[HWRT] Improvements towards supporting offline BVH for skeletal meshes.
- extended FRayTracingGeometryInitializer with new flag bTemplate (indicating the geometry will only be used as "SourceGeometry" when building other BLAS.
- don't build FRayTracingGeometry when bTemplate flag is set.
- added missing logic to serialize RawDataHeader in FSkeletalMeshLODRenderData.
- use SourceRayTracingGeometry to build StaticRayTracingGeometry for FSkeletalMeshRenderData.

- these code paths are still not compatible with mesh streaming / referenced based BLAS residency / toggling ray tracing yet, will be implemented in follow up CLs.

diff --git a/Engine/Source/Runtime/Core/Private/UObject/FortniteMainBranchObjectVersion.cpp b/Engine/Source/Runtime/Core/Private/UObject/FortniteMainBranchObjectVersion.cpp
--- a/Engine/Source/Runtime/Core/Private/UObject/FortniteMainBranchObjectVersion.cpp
+++ b/Engine/Source/Runtime/Core/Private/UObject/FortniteMainBranchObjectVersion.cpp
@@ -12,7 +12,7 @@ TMap<FGuid, FGuid> FFortniteMainBranchObjectVersion::GetSystemGuids()
 	SystemGuids.Add(DevGuids.NANITE_DERIVEDDATA_VER, FGuid("42C46B24E1586986880AFE62B0F575AB"));
 	SystemGuids.Add(DevGuids.NIAGARASHADERMAP_DERIVEDDATA_VER, FGuid("45D5978F33D940EEA6602B78559900DC"));
 	SystemGuids.Add(DevGuids.Niagara_LatestScriptCompileVersion, FGuid("670A2EBDEE1A4794BB77BBFF331EA7EF"));
-	SystemGuids.Add(DevGuids.SkeletalMeshDerivedDataVersion, FGuid("5E93192FCF3647EEB7F9D4FFF1750050"));
+	SystemGuids.Add(DevGuids.SkeletalMeshDerivedDataVersion, FGuid("716424BCEC314E969ACFE36D2C2ECE49"));
 	SystemGuids.Add(DevGuids.STATICMESH_DERIVEDDATA_VER, FGuid("3ABF32149F9D4D838750368AB95573BE"));
 	SystemGuids.Add(DevGuids.MaterialTranslationDDCVersion, FGuid("0D223CD149C54D5B8B09436D1B8AC0BF"));
 	return SystemGuids;
diff --git a/Engine/Source/Runtime/Engine/Private/SkeletalMeshLODRenderData.cpp b/Engine/Source/Runtime/Engine/Private/SkeletalMeshLODRenderData.cpp
--- a/Engine/Source/Runtime/Engine/Private/SkeletalMeshLODRenderData.cpp
+++ b/Engine/Source/Runtime/Engine/Private/SkeletalMeshLODRenderData.cpp
@@ -302,8 +302,35 @@ void FSkeletalMeshLODRenderData::InitResources(bool bNeedsVertexColors, int32 LO
 #if RHI_RAYTRACING
 	if (IsRayTracingAllowed())
 	{
-		if (SourceRayTracingGeometry.RawData.Num() > 0)
+		if (SourceRayTracingGeometry.RawDataHeader.IsValid())
 		{
+			uint32 TrianglesCount = 0;
+
+			TArray<FRayTracingGeometrySegment> GeometrySections;
+			GeometrySections.Reserve(RenderSections.Num());
+
+			for (const FSkelMeshRenderSection& Section : RenderSections)
+			{
+				FRayTracingGeometrySegment Segment;
+				Segment.FirstPrimitive = Section.BaseIndex / 3;
+				Segment.NumPrimitives = Section.NumTriangles;
+				Segment.bEnabled = !Section.bDisabled && Section.bVisibleInRayTracing;
+				GeometrySections.Add(Segment);
+
+				TrianglesCount += Section.NumTriangles;
+			}
+
+			FRayTracingGeometryInitializer Initializer;
+			Initializer.DebugName = (Owner ? Owner->GetFName() : NAME_None);
+			Initializer.OwnerName = OwnerName;
+			Initializer.bTemplate = true;
+			Initializer.TotalPrimitiveCount = TrianglesCount;
+			Initializer.Segments = GeometrySections;
+
+			SourceRayTracingGeometry.SetInitializer(Initializer);
+
+			SourceRayTracingGeometry.LODIndex = LODIndex;
+
 			BeginInitResource(OwnerName, &SourceRayTracingGeometry, &UE::RenderCommandPipe::SkeletalMesh);
 		}
 	}
@@ -787,11 +814,14 @@ void FSkeletalMeshLODRenderData::SerializeStreamedData(FArchive& Ar, USkinnedAss
 	{
 		Owner->SetSkinWeightProfilesData(LODIdx, SkinWeightProfilesData);
 	}
-	Ar << SourceRayTracingGeometry.RawData;
+
+	Ar << SourceRayTracingGeometry.RawDataHeader;
+	SourceRayTracingGeometry.RawData.BulkSerialize(Ar);
 	SourceRayTracingGeometry.SetResourceName(SkeletalMeshResourceNames::SourceRayTracingGeometryName);
 
 	if (Ar.IsLoading() && RayTracing::ShouldForceRuntimeBLAS())
 	{
+		SourceRayTracingGeometry.RawDataHeader = FRayTracingGeometryOfflineDataHeader();
 		SourceRayTracingGeometry.RawData.Discard();
 	}
 
@@ -920,6 +950,15 @@ void FSkeletalMeshLODRenderData::SerializeAvailabilityInfo(FArchive& Ar, int32 L
 	}
 	SkinWeightProfilesData.SerializeMetaData(Ar);
 	SkinWeightProfilesData.Init(&SkinWeightVertexBuffer);
+
+	{
+		Ar << SourceRayTracingGeometry.RawDataHeader;
+
+		if (Ar.IsLoading() && RayTracing::ShouldForceRuntimeBLAS())
+		{
+			SourceRayTracingGeometry.RawDataHeader = FRayTracingGeometryOfflineDataHeader();
+		}
+	}
 }
 
 void FSkeletalMeshLODRenderData::Serialize(FArchive& Ar, UObject* Owner, int32 Idx)
diff --git a/Engine/Source/Runtime/Engine/Private/SkeletalMeshRenderData.cpp b/Engine/Source/Runtime/Engine/Private/SkeletalMeshRenderData.cpp
--- a/Engine/Source/Runtime/Engine/Private/SkeletalMeshRenderData.cpp
+++ b/Engine/Source/Runtime/Engine/Private/SkeletalMeshRenderData.cpp
@@ -706,6 +706,7 @@ void FSkeletalMeshRenderData::InitStaticRayTracingGeometry(int32 LODIndex)
 			LODIndex = LODIndex,
 			SkelMeshRenderData = this,
 			&RayTracingGeometry = LODData.StaticRayTracingGeometry,
+			&SourceRayTracingGeometry = LODData.SourceRayTracingGeometry,
 			&bReferencedByStaticSkeletalMeshObjects_RenderThread = LODData.bReferencedByStaticSkeletalMeshObjects_RenderThread](FRHICommandList& RHICmdList)
 			{
 				FRayTracingGeometryInitializer Initializer;
@@ -740,6 +741,8 @@ void FSkeletalMeshRenderData::InitStaticRayTracingGeometry(int32 LODIndex)
 					Initializer.Type = ERayTracingGeometryInitializerType::StreamingDestination;
 				}
 
+				Initializer.SourceGeometry = SourceRayTracingGeometry.GetRHI();
+
 				RayTracingGeometry.GroupHandle = SkelMeshRenderData->RayTracingGeometryGroupHandle;
 				RayTracingGeometry.LODIndex = LODIndex;
 
diff --git a/Engine/Source/Runtime/Engine/Private/SkeletalMeshSceneProxy.cpp b/Engine/Source/Runtime/Engine/Private/SkeletalMeshSceneProxy.cpp
--- a/Engine/Source/Runtime/Engine/Private/SkeletalMeshSceneProxy.cpp
+++ b/Engine/Source/Runtime/Engine/Private/SkeletalMeshSceneProxy.cpp
@@ -899,8 +899,14 @@ TArray<FRayTracingGeometry*> FSkeletalMeshSceneProxy::GetStaticRayTracingGeometr
 		RayTracingGeometries.AddDefaulted(SkeletalMeshRenderData->LODRenderData.Num());
 		for (int32 LODIndex = 0; LODIndex < SkeletalMeshRenderData->LODRenderData.Num(); LODIndex++)
 		{
-			ensure(SkeletalMeshRenderData->LODRenderData[LODIndex].NumReferencingStaticSkeletalMeshObjects > 0);
-			RayTracingGeometries[LODIndex] = &SkeletalMeshRenderData->LODRenderData[LODIndex].StaticRayTracingGeometry;
+			FSkeletalMeshLODRenderData& LODRenderData = SkeletalMeshRenderData->LODRenderData[LODIndex];
+
+			// Skip LODs that have their render data stripped
+			if (LODRenderData.GetNumVertices() > 0)
+			{
+				ensure(LODRenderData.NumReferencingStaticSkeletalMeshObjects > 0);
+				RayTracingGeometries[LODIndex] = &LODRenderData.StaticRayTracingGeometry;
+			}
 		}
 
 		return MoveTemp(RayTracingGeometries);
diff --git a/Engine/Source/Runtime/RHI/Public/RHIResources.h b/Engine/Source/Runtime/RHI/Public/RHIResources.h
--- a/Engine/Source/Runtime/RHI/Public/RHIResources.h
+++ b/Engine/Source/Runtime/RHI/Public/RHIResources.h
@@ -3494,6 +3494,7 @@ public:
 	bool bFastBuild = false;
 	bool bAllowUpdate = false;
 	bool bAllowCompaction = true;
+	bool bTemplate = false;
 	ERayTracingGeometryInitializerType Type = ERayTracingGeometryInitializerType::Rendering;
 
 	// Partitions of geometry to allow different shader and resource bindings.
diff --git a/Engine/Source/Runtime/RenderCore/Private/RayTracingGeometry.cpp b/Engine/Source/Runtime/RenderCore/Private/RayTracingGeometry.cpp
--- a/Engine/Source/Runtime/RenderCore/Private/RayTracingGeometry.cpp
+++ b/Engine/Source/Runtime/RenderCore/Private/RayTracingGeometry.cpp
@@ -32,7 +32,12 @@ PRAGMA_DISABLE_DEPRECATION_WARNINGS // Remove this in UE 5.6 once RayTracingGeom
 
 bool FRayTracingGeometry::HasValidInitializer() const
 {
-	bool bAllSegmentsAreValid = Initializer.Segments.Num() > 0 || Initializer.OfflineData;
+	if (Initializer.OfflineData != nullptr)
+	{
+		return true;
+	}
+
+	bool bAllSegmentsAreValid = Initializer.Segments.Num() > 0;
 	for (const FRayTracingGeometrySegment& Segment : Initializer.Segments)
 	{
 		if (!Segment.VertexBuffer)
@@ -156,16 +161,22 @@ void FRayTracingGeometry::CreateRayTracingGeometry(FRHICommandListBase& RHICmdLi
 			EnumAddFlags(GeometryState, EGeometryStateFlags::Valid);
 		}
 
+		const bool bWithOfflineData = Initializer.OfflineData != nullptr;
+
 		if (IsRayTracingEnabled())
 		{
 			RayTracingGeometryRHI = RHICmdList.CreateRayTracingGeometry(Initializer);
+
+			// Offline data ownership is transferred to the RHI, which discards it after use.
+			// It is no longer valid to use it after this point.
+			Initializer.OfflineData = nullptr;
 		}
 		else
 		{
 			EnumAddFlags(GeometryState, EGeometryStateFlags::Evicted);
 		}
 
-		if (Initializer.OfflineData == nullptr)
+		if (!bWithOfflineData)
 		{
 			// Request build if not skip
 			if (InBuildPriority != ERTAccelerationStructureBuildPriority::Skip)
@@ -183,16 +194,12 @@ void FRayTracingGeometry::CreateRayTracingGeometry(FRHICommandListBase& RHICmdLi
 		}
 		else
 		{
-			if (RayTracingGeometryRHI && RayTracingGeometryRHI->IsCompressed())
+			if (RayTracingGeometryRHI && RayTracingGeometryRHI->IsCompressed() && !Initializer.bTemplate)
 			{
 				GRayTracingGeometryManager->RequestBuildAccelerationStructure(this, InBuildPriority);
 			}
 
 			SetRequiresBuild(false);
-
-			// Offline data ownership is transferred to the RHI, which discards it after use.
-			// It is no longer valid to use it after this point.
-			Initializer.OfflineData = nullptr;
 		}
 	}
 
@@ -235,7 +242,7 @@ void FRayTracingGeometry::InitRHI(FRHICommandListBase& RHICmdList)
 	if (!IsRayTracingAllowed())
 		return;
 
-	ERTAccelerationStructureBuildPriority BuildPriority = Initializer.Type != ERayTracingGeometryInitializerType::Rendering
+	ERTAccelerationStructureBuildPriority BuildPriority = (Initializer.Type != ERayTracingGeometryInitializerType::Rendering || Initializer.bTemplate)
 		? ERTAccelerationStructureBuildPriority::Skip
 		: ERTAccelerationStructureBuildPriority::Normal;
 	CreateRayTracingGeometry(RHICmdList, BuildPriority);
