Added GetInstanceDataUpdateTaskInfo to FInstancedSkinnedMeshSceneProxy to fix race condition that occurs otherwise as the RT doesn't sync the task updating the instance data.

diff --git a/Engine/Source/Runtime/Engine/Private/Components/InstancedSkinnedMeshComponent.cpp b/Engine/Source/Runtime/Engine/Private/Components/InstancedSkinnedMeshComponent.cpp
--- a/Engine/Source/Runtime/Engine/Private/Components/InstancedSkinnedMeshComponent.cpp
+++ b/Engine/Source/Runtime/Engine/Private/Components/InstancedSkinnedMeshComponent.cpp
@@ -189,6 +189,11 @@ public:
 		InstanceEndCullDistance = EndCullDistance;
 	}
 
+	virtual FInstanceDataUpdateTaskInfo *GetInstanceDataUpdateTaskInfo() const override
+	{
+		return InstanceDataSceneProxy ? InstanceDataSceneProxy->GetUpdateTaskInfo() : nullptr;
+	}
+
 public:
 	// FPrimitiveSceneProxy interface.
 	virtual SIZE_T GetTypeHash() const override
