Fixed bug preventing material expression parameters and comments from being created.

diff --git a/Engine/Source/Runtime/Engine/Public/Materials/MaterialExpressionComment.h b/Engine/Source/Runtime/Engine/Public/Materials/MaterialExpressionComment.h
--- a/Engine/Source/Runtime/Engine/Public/Materials/MaterialExpressionComment.h
+++ b/Engine/Source/Runtime/Engine/Public/Materials/MaterialExpressionComment.h
@@ -58,7 +58,7 @@ public:
 	virtual void GetCaption(TArray<FString>& OutCaptions) const override;
 	virtual bool MatchesSearchQuery( const TCHAR* SearchQuery ) override;
 #endif // WITH_EDITOR
-	virtual bool IsAllowedIn(const UObject* MaterialOrFunction) const override { return false; }
+	virtual bool IsAllowedIn(const UObject* MaterialOrFunction) const override { return true; }
 	//~ End UMaterialExpression Interface
 };
 
diff --git a/Engine/Source/Runtime/Engine/Public/Materials/MaterialExpressionParameter.h b/Engine/Source/Runtime/Engine/Public/Materials/MaterialExpressionParameter.h
--- a/Engine/Source/Runtime/Engine/Public/Materials/MaterialExpressionParameter.h
+++ b/Engine/Source/Runtime/Engine/Public/Materials/MaterialExpressionParameter.h
@@ -61,7 +61,9 @@ class UMaterialExpressionParameter : public UMaterialExpression
 
 	virtual void Build(MIR::FEmitter& Emitter) override;
 #endif
-	virtual bool IsAllowedIn(const UObject* MaterialOrFunction) const override { return false; }
+	
+	virtual bool IsAllowedIn(const UObject* MaterialOrFunction) const override { return true; }
+
 	//~ End UMaterialExpression Interface
 
 	virtual FGuid& GetParameterExpressionId() override
