Improve Light unit descriptions to make it clear which ones are normalized vs. un-normalized

diff --git a/Engine/Source/Runtime/Engine/Classes/Components/LocalLightComponent.h b/Engine/Source/Runtime/Engine/Classes/Components/LocalLightComponent.h
--- a/Engine/Source/Runtime/Engine/Classes/Components/LocalLightComponent.h
+++ b/Engine/Source/Runtime/Engine/Classes/Components/LocalLightComponent.h
@@ -20,8 +20,8 @@ class ULocalLightComponent : public ULightComponent
 
 	/** 
 	 * Units used for the intensity. 
-	 * The peak luminous intensity is measured in candelas,
-	 * while the luminous power is measured in lumens.
+	 * The peak luminous intensity is measured in candelas, while the luminous flux is measured in lumens.
+	 * When the units are set in Nits, the light's power is also determined by the size of the light source (larger sources will emit more light).
 	 */
 	UPROPERTY(EditAnywhere, BlueprintReadOnly, Category=Light, meta=(DisplayName="Intensity Units", EditCondition="!bUseIESBrightness"))
 	ELightUnits IntensityUnits;
diff --git a/Engine/Source/Runtime/Engine/Classes/Engine/Scene.h b/Engine/Source/Runtime/Engine/Classes/Engine/Scene.h
--- a/Engine/Source/Runtime/Engine/Classes/Engine/Scene.h
+++ b/Engine/Source/Runtime/Engine/Classes/Engine/Scene.h
@@ -73,11 +73,11 @@ enum ETemperatureMethod : int
 UENUM() 
 enum class ELightUnits : uint8
 {
-	Unitless,
-	Candelas,
-	Lumens,
-	EV,
-	Nits
+	Unitless UMETA(DisplayName = "Unitless"), 
+	Candelas UMETA(DisplayName = "Candelas - Luminous Intensity, Normalized"),
+	Lumens   UMETA(DisplayName = "Lumens - Luminous Flux, Normalized"),
+	EV       UMETA(DisplayName = "EV - Exposure Value, Normalized"),
+	Nits     UMETA(DisplayName = "Nits - Luminance, Non-Normalized")
 };
 
 UENUM()
