[Primitive] rollback of change preventing component from invalid BodyInstance from being considered relevant to navigation. In some scenario the order of operations may prevent some component to properly register to the NavigationSystem. A different fix need to be done to fix the issue this change was trying to solve.

diff --git a/Engine/Source/Runtime/Engine/Private/Components/PrimitiveComponent.cpp b/Engine/Source/Runtime/Engine/Private/Components/PrimitiveComponent.cpp
--- a/Engine/Source/Runtime/Engine/Private/Components/PrimitiveComponent.cpp
+++ b/Engine/Source/Runtime/Engine/Private/Components/PrimitiveComponent.cpp
@@ -134,13 +134,6 @@ namespace PrimitiveComponentCVars
 		TEXT("p.EnableFastOverlapCheck"),
 		bEnableFastOverlapCheck,
 		TEXT("Enable fast overlap check against sweep hits, avoiding UpdateOverlaps (for the swept component)."));
-
-	bool bInvalidBodyInstanceNeverRelevantToNavigation = true;
-	static FAutoConsoleVariableRef InvalidBodyInstanceNeverRelevantToNavigation(
-		TEXT("p.InvalidBodyInstanceIsNeverRelevantToNavigation"),
-		bInvalidBodyInstanceNeverRelevantToNavigation,
-		TEXT("0: IsNavigationRelevant ignores the validity of the body instance (legacy), 1: IsNavigationRelevant returns false for invalid body instance (default)"),
-		ECVF_Default);
 }
 
 namespace PhysicsReplicationCVars
@@ -3424,12 +3417,6 @@ bool UPrimitiveComponent::IsNavigationRelevant() const
 		return false;
 	}
 
-	if (PrimitiveComponentCVars::bInvalidBodyInstanceNeverRelevantToNavigation
-		&& !BodyInstance.IsValidBodyInstance())
-	{
-		return false;
-	}
-
 	if (HasCustomNavigableGeometry() >= EHasCustomNavigableGeometry::EvenIfNotCollidable)
 	{
 		return true;
