Build fix for incorrect auto-merge(?) in CL 8584458

diff --git a/Engine/Source/Runtime/Engine/Private/PhysicsEngine/ImmediatePhysics/ImmediatePhysicsAdapters.cpp b/Engine/Source/Runtime/Engine/Private/PhysicsEngine/ImmediatePhysics/ImmediatePhysicsAdapters.cpp
--- a/Engine/Source/Runtime/Engine/Private/PhysicsEngine/ImmediatePhysics/ImmediatePhysicsAdapters.cpp
+++ b/Engine/Source/Runtime/Engine/Private/PhysicsEngine/ImmediatePhysics/ImmediatePhysicsAdapters.cpp
@@ -360,79 +360,6 @@
 		JointSettings.TwistContactDistance = Constraint.ContactDistance;
 	}
 
-	void UpdateJointSettingsFromLinearDriveConstraint(const FLinearDriveConstraint& LinearDriveConstraint, Chaos::FPBDJointSettings& JointSettings)
-	{
-		using namespace Chaos;
-
-		JointSettings.LinearDrivePositionTarget = LinearDriveConstraint.PositionTarget;
-		JointSettings.LinearDriveVelocityTarget = LinearDriveConstraint.VelocityTarget;
-		JointSettings.bLinearPositionDriveEnabled[0] = LinearDriveConstraint.XDrive.bEnablePositionDrive;
-		JointSettings.bLinearPositionDriveEnabled[1] = LinearDriveConstraint.YDrive.bEnablePositionDrive;
-		JointSettings.bLinearPositionDriveEnabled[2] = LinearDriveConstraint.ZDrive.bEnablePositionDrive;
-		JointSettings.bLinearVelocityDriveEnabled[0] = LinearDriveConstraint.XDrive.bEnableVelocityDrive;
-		JointSettings.bLinearVelocityDriveEnabled[1] = LinearDriveConstraint.YDrive.bEnableVelocityDrive;
-		JointSettings.bLinearVelocityDriveEnabled[2] = LinearDriveConstraint.ZDrive.bEnableVelocityDrive;
-
-		JointSettings.LinearDriveForceMode = EJointForceMode::Acceleration; // hardcoded!
-		JointSettings.LinearDriveStiffness = ConstraintSettings::LinearDriveStiffnessScale() * FVec3(
-			LinearDriveConstraint.XDrive.Stiffness,
-			LinearDriveConstraint.YDrive.Stiffness,
-			LinearDriveConstraint.ZDrive.Stiffness);
-		JointSettings.LinearDriveDamping = ConstraintSettings::LinearDriveDampingScale() * FVec3(
-			LinearDriveConstraint.XDrive.Damping,
-			LinearDriveConstraint.YDrive.Damping,
-			LinearDriveConstraint.ZDrive.Damping);
-		JointSettings.LinearDriveMaxForce[0] = LinearDriveConstraint.XDrive.MaxForce;
-		JointSettings.LinearDriveMaxForce[1] = LinearDriveConstraint.YDrive.MaxForce;
-		JointSettings.LinearDriveMaxForce[2] = LinearDriveConstraint.ZDrive.MaxForce;
-	}
-
-	void UpdateJointSettingsFromAngularDriveConstraint(
-		const FAngularDriveConstraint& AngularDriveConstraint, Chaos::FPBDJointSettings& JointSettings)
-	{
-		using namespace Chaos;
-
-		JointSettings.AngularDrivePositionTarget = FQuat(AngularDriveConstraint.OrientationTarget);
-		JointSettings.AngularDriveVelocityTarget = AngularDriveConstraint.AngularVelocityTarget * UE_TWO_PI; // rev/s to rad/s
-
-		JointSettings.AngularDriveForceMode = EJointForceMode::Acceleration; // hardcoded!
-		if (AngularDriveConstraint.AngularDriveMode == EAngularDriveMode::SLERP)
-		{
-			JointSettings.AngularDriveStiffness = FVec3(
-				ConstraintSettings::AngularDriveStiffnessScale() * AngularDriveConstraint.SlerpDrive.Stiffness);
-			JointSettings.AngularDriveDamping = FVec3(
-				ConstraintSettings::AngularDriveDampingScale() * AngularDriveConstraint.SlerpDrive.Damping);
-			JointSettings.AngularDriveMaxTorque = FVec3(
-				AngularDriveConstraint.SlerpDrive.MaxForce);
-			JointSettings.bAngularSLerpPositionDriveEnabled = AngularDriveConstraint.SlerpDrive.bEnablePositionDrive;
-			JointSettings.bAngularSLerpVelocityDriveEnabled = AngularDriveConstraint.SlerpDrive.bEnableVelocityDrive;
-			JointSettings.bAngularTwistPositionDriveEnabled = false;
-			JointSettings.bAngularTwistVelocityDriveEnabled = false;
-			JointSettings.bAngularSwingPositionDriveEnabled = false;
-			JointSettings.bAngularSwingVelocityDriveEnabled = false;
-		}
-		else
-		{
-			JointSettings.AngularDriveStiffness = ConstraintSettings::AngularDriveStiffnessScale() * FVec3(
-				AngularDriveConstraint.TwistDrive.Stiffness,
-				AngularDriveConstraint.SwingDrive.Stiffness,
-				AngularDriveConstraint.SwingDrive.Stiffness);
-			JointSettings.AngularDriveDamping = ConstraintSettings::AngularDriveDampingScale() * FVec3(
-				AngularDriveConstraint.TwistDrive.Damping,
-				AngularDriveConstraint.SwingDrive.Damping,
-				AngularDriveConstraint.SwingDrive.Damping);
-			JointSettings.AngularDriveMaxTorque[0] = AngularDriveConstraint.TwistDrive.MaxForce;
-			JointSettings.AngularDriveMaxTorque[1] = AngularDriveConstraint.SwingDrive.MaxForce;
-			JointSettings.AngularDriveMaxTorque[2] = AngularDriveConstraint.SwingDrive.MaxForce;
-			JointSettings.bAngularSLerpPositionDriveEnabled = false;
-			JointSettings.bAngularSLerpVelocityDriveEnabled = false;
-			JointSettings.bAngularTwistPositionDriveEnabled = AngularDriveConstraint.TwistDrive.bEnablePositionDrive;
-			JointSettings.bAngularTwistVelocityDriveEnabled = AngularDriveConstraint.TwistDrive.bEnableVelocityDrive;
-			JointSettings.bAngularSwingPositionDriveEnabled = AngularDriveConstraint.SwingDrive.bEnablePositionDrive;
-			JointSettings.bAngularSwingVelocityDriveEnabled = AngularDriveConstraint.SwingDrive.bEnableVelocityDrive;
-		}
-	}
-
 	// Copies data from a profile into settings
 	void UpdateJointSettingsFromConstraintProfile(const FConstraintProfileProperties& Profile, Chaos::FPBDJointSettings& JointSettings)
 	{
@@ -522,5 +449,78 @@
 
 		return JointSetup;
 	}
+	
+	void UpdateJointSettingsFromLinearDriveConstraint(const FLinearDriveConstraint& LinearDriveConstraint, Chaos::FPBDJointSettings& JointSettings)
+	{
+		using namespace Chaos;
+
+		JointSettings.LinearDrivePositionTarget = LinearDriveConstraint.PositionTarget;
+		JointSettings.LinearDriveVelocityTarget = LinearDriveConstraint.VelocityTarget;
+		JointSettings.bLinearPositionDriveEnabled[0] = LinearDriveConstraint.XDrive.bEnablePositionDrive;
+		JointSettings.bLinearPositionDriveEnabled[1] = LinearDriveConstraint.YDrive.bEnablePositionDrive;
+		JointSettings.bLinearPositionDriveEnabled[2] = LinearDriveConstraint.ZDrive.bEnablePositionDrive;
+		JointSettings.bLinearVelocityDriveEnabled[0] = LinearDriveConstraint.XDrive.bEnableVelocityDrive;
+		JointSettings.bLinearVelocityDriveEnabled[1] = LinearDriveConstraint.YDrive.bEnableVelocityDrive;
+		JointSettings.bLinearVelocityDriveEnabled[2] = LinearDriveConstraint.ZDrive.bEnableVelocityDrive;
+
+		JointSettings.LinearDriveForceMode = EJointForceMode::Acceleration; // hardcoded!
+		JointSettings.LinearDriveStiffness = Chaos::ConstraintSettings::LinearDriveStiffnessScale() * FVec3(
+			LinearDriveConstraint.XDrive.Stiffness,
+			LinearDriveConstraint.YDrive.Stiffness,
+			LinearDriveConstraint.ZDrive.Stiffness);
+		JointSettings.LinearDriveDamping = Chaos::ConstraintSettings::LinearDriveDampingScale() * FVec3(
+			LinearDriveConstraint.XDrive.Damping,
+			LinearDriveConstraint.YDrive.Damping,
+			LinearDriveConstraint.ZDrive.Damping);
+		JointSettings.LinearDriveMaxForce[0] = LinearDriveConstraint.XDrive.MaxForce;
+		JointSettings.LinearDriveMaxForce[1] = LinearDriveConstraint.YDrive.MaxForce;
+		JointSettings.LinearDriveMaxForce[2] = LinearDriveConstraint.ZDrive.MaxForce;
+	}
+
+	void UpdateJointSettingsFromAngularDriveConstraint(
+		const FAngularDriveConstraint& AngularDriveConstraint, Chaos::FPBDJointSettings& JointSettings)
+	{
+		using namespace Chaos;
+
+		JointSettings.AngularDrivePositionTarget = FQuat(AngularDriveConstraint.OrientationTarget);
+		JointSettings.AngularDriveVelocityTarget = AngularDriveConstraint.AngularVelocityTarget * UE_TWO_PI; // rev/s to rad/s
+
+		JointSettings.AngularDriveForceMode = EJointForceMode::Acceleration; // hardcoded!
+		if (AngularDriveConstraint.AngularDriveMode == EAngularDriveMode::SLERP)
+		{
+			JointSettings.AngularDriveStiffness = FVec3(
+				ConstraintSettings::AngularDriveStiffnessScale() * AngularDriveConstraint.SlerpDrive.Stiffness);
+			JointSettings.AngularDriveDamping = FVec3(
+				ConstraintSettings::AngularDriveDampingScale() * AngularDriveConstraint.SlerpDrive.Damping);
+			JointSettings.AngularDriveMaxTorque = FVec3(
+				AngularDriveConstraint.SlerpDrive.MaxForce);
+			JointSettings.bAngularSLerpPositionDriveEnabled = AngularDriveConstraint.SlerpDrive.bEnablePositionDrive;
+			JointSettings.bAngularSLerpVelocityDriveEnabled = AngularDriveConstraint.SlerpDrive.bEnableVelocityDrive;
+			JointSettings.bAngularTwistPositionDriveEnabled = false;
+			JointSettings.bAngularTwistVelocityDriveEnabled = false;
+			JointSettings.bAngularSwingPositionDriveEnabled = false;
+			JointSettings.bAngularSwingVelocityDriveEnabled = false;
+		}
+		else
+		{
+			JointSettings.AngularDriveStiffness = ConstraintSettings::AngularDriveStiffnessScale() * FVec3(
+				AngularDriveConstraint.TwistDrive.Stiffness,
+				AngularDriveConstraint.SwingDrive.Stiffness,
+				AngularDriveConstraint.SwingDrive.Stiffness);
+			JointSettings.AngularDriveDamping = ConstraintSettings::AngularDriveDampingScale() * FVec3(
+				AngularDriveConstraint.TwistDrive.Damping,
+				AngularDriveConstraint.SwingDrive.Damping,
+				AngularDriveConstraint.SwingDrive.Damping);
+			JointSettings.AngularDriveMaxTorque[0] = AngularDriveConstraint.TwistDrive.MaxForce;
+			JointSettings.AngularDriveMaxTorque[1] = AngularDriveConstraint.SwingDrive.MaxForce;
+			JointSettings.AngularDriveMaxTorque[2] = AngularDriveConstraint.SwingDrive.MaxForce;
+			JointSettings.bAngularSLerpPositionDriveEnabled = false;
+			JointSettings.bAngularSLerpVelocityDriveEnabled = false;
+			JointSettings.bAngularTwistPositionDriveEnabled = AngularDriveConstraint.TwistDrive.bEnablePositionDrive;
+			JointSettings.bAngularTwistVelocityDriveEnabled = AngularDriveConstraint.TwistDrive.bEnableVelocityDrive;
+			JointSettings.bAngularSwingPositionDriveEnabled = AngularDriveConstraint.SwingDrive.bEnablePositionDrive;
+			JointSettings.bAngularSwingVelocityDriveEnabled = AngularDriveConstraint.SwingDrive.bEnableVelocityDrive;
+		}
+	}
 }
 
