Fix validation errors on textures expecting ERHIPipeline::All instead of ERHIPipeline::Graphics only

diff --git a/Engine/Source/Runtime/Engine/Private/Texture2D.cpp b/Engine/Source/Runtime/Engine/Private/Texture2D.cpp
--- a/Engine/Source/Runtime/Engine/Private/Texture2D.cpp
+++ b/Engine/Source/Runtime/Engine/Private/Texture2D.cpp
@@ -1772,7 +1772,8 @@ void FVirtualTexture2DResource::InitializeEditorResources(IVirtualTexture* InVir
 				CopyInfo.Size = FIntVector(MipWidth, MipHeight, 1);
 
 				// Put the source texture in CopySrc mode. The destination texture is already in CopyDest mode because we created it that way.
-				RHICommandList.Transition(FRHITransitionInfo(Texture2DRHI, ERHIAccess::SRVMask, ERHIAccess::CopySrc));
+				RHICommandList.Transition({FRHITransitionInfo(Texture2DRHI, ERHIAccess::SRVMask, ERHIAccess::CopySrc)}, ERHIPipeline::All, ERHIPipeline::Graphics);
+
 				RHICommandList.CopyTexture(Texture2DRHI, ResizedTexture2DRHI, CopyInfo);
 				// Make the destination texture SRVMask again. We don't care about the source texture after this, so we won't bother transitioning it.
 				RHICommandList.Transition(FRHITransitionInfo(ResizedTexture2DRHI, ERHIAccess::CopyDest, ERHIAccess::SRVMask));
