New material translator changes:
- Added support for mip value modes Level and Bias in `MIR::FTextureSample` and for explicit Derivative mode.
- Added `FEmitter::DefaultInputToIntScalar()` API.
- Added a `NumUses` constexpr entry to each MIR value type so that the `GetUse()` function implementation can refer to this constant, instead of numeric literals.
- Changed inline documentation comments in MaterialIR from // to /** */ for consistency. In coming changes I will apply this format to the whole MIR codebase.

#rnx

diff --git a/Engine/Source/Runtime/Engine/Private/Materials/MaterialExpressionsIR.cpp b/Engine/Source/Runtime/Engine/Private/Materials/MaterialExpressionsIR.cpp
--- a/Engine/Source/Runtime/Engine/Private/Materials/MaterialExpressionsIR.cpp
+++ b/Engine/Source/Runtime/Engine/Private/Materials/MaterialExpressionsIR.cpp
@@ -883,7 +883,32 @@ static void BuildTextureSampleExpression(MIR::FValue* Texture, UMaterialExpressi
 		TexCoords = Emitter.EmitExternalInput(MIR::TexCoordIndexToExternalInput(Expr->ConstCoordinate));
 	}
 
-	MIR::FValue* OutputValue = Emitter.EmitTextureSample(Texture, TexCoords, Expr->SamplerSource, Expr->MipValueMode, Expr->SamplerType);
+	MIR::FValue* TexCoordsDdx{};
+	MIR::FValue* TexCoordsDdy{};
+	MIR::FValue* MipValue{};
+	if (Expr->MipValueMode == TMVM_MipLevel || Expr->MipValueMode == TMVM_MipBias)
+	{
+		// Get the mip value level (either through the expression input or using the given constant if disconnected).
+		MipValue = Emitter.DefaultInputToIntScalar(&Expr->MipValue, Expr->ConstMipValue).InputScalar(&Expr->MipValue);
+
+		// Determine if automatic view mip bias should be used
+		bool bAutomaticViewMipBias = Emitter.DefaultInputToBoolScalar(&Expr->AutomaticViewMipBiasValue, Expr->AutomaticViewMipBias).InputConstantBoolScalar(&Expr->AutomaticViewMipBiasValue);
+		UE_MIR_CHECKPOINT();
+
+		// If so, add "View.MaterialTextureViewBias" to the specified mip value
+		if (bAutomaticViewMipBias)
+		{
+			MipValue = Emitter.EmitOperator(MIR::BO_Add, MipValue, Emitter.EmitExternalInput(MIR::EExternalInput::ViewMaterialTextureMipBias));
+		}
+	}
+	else
+	{
+		check(Expr->MipValueMode = TMVM_Derivative);
+		TexCoordsDdx = Emitter.EmitConstruct(TexCoords->Type, Emitter.Input(&Expr->CoordinatesDX));
+		TexCoordsDdy = Emitter.EmitConstruct(TexCoords->Type, Emitter.Input(&Expr->CoordinatesDY));
+	}
+
+	MIR::FValue* OutputValue = Emitter.EmitTextureSample(Texture, TexCoords, TexCoordsDdx, TexCoordsDdy, MipValue, Expr->SamplerSource, Expr->MipValueMode, Expr->SamplerType);
 	Emitter.Output(0, Emitter.EmitSwizzle(OutputValue, MIR::FSwizzleMask::XYZ()));
 	Emitter.Output(1, Emitter.EmitSubscript(OutputValue, 0)); 
 	Emitter.Output(2, Emitter.EmitSubscript(OutputValue, 1)); 
diff --git a/Engine/Source/Runtime/Engine/Private/Materials/MaterialIR.cpp b/Engine/Source/Runtime/Engine/Private/Materials/MaterialIR.cpp
--- a/Engine/Source/Runtime/Engine/Private/Materials/MaterialIR.cpp
+++ b/Engine/Source/Runtime/Engine/Private/Materials/MaterialIR.cpp
@@ -131,31 +131,31 @@ TConstArrayView<FValue*> FValue::GetUses() const
 		case VK_SetMaterialOutput:
 		{
 			auto This = static_cast<const FSetMaterialOutput*>(this);
-			return { &This->Arg, 1 };
+			return { &This->Arg, FSetMaterialOutput::NumUses };
 		}
 
 		case VK_Operator:
 		{
 			auto This = static_cast<const FOperator*>(this);
-			return { &This->AArg, 3 };
+			return { &This->AArg, FOperator::NumUses };
 		}
 
 		case VK_Branch:
 		{
 			auto This = static_cast<const FBranch*>(this);
-			return { &This->ConditionArg, 3 };
+			return { &This->ConditionArg, FBranch::NumUses };
 		}
 
 		case VK_Subscript:
 		{
 			auto This = static_cast<const FSubscript*>(this);
-			return { &This->Arg, 1 };
+			return { &This->Arg, FSubscript::NumUses };
 		}
 
 		case VK_TextureSample:
 		{
 			auto This = static_cast<const FTextureSample*>(this);
-			return { &This->TextureObject, 4 };
+			return { &This->TextureObject, FTextureSample::NumUses };
 		}
 
 		case VK_InlineHLSL:
@@ -333,6 +333,7 @@ const TCHAR* LexToString(EExternalInput Input)
 		case EExternalInput::TexCoord5_Ddy: return TEXT("TexCoord5_Ddy");
 		case EExternalInput::TexCoord6_Ddy: return TEXT("TexCoord6_Ddy");
 		case EExternalInput::TexCoord7_Ddy: return TEXT("TexCoord7_Ddy");
+		case EExternalInput::ViewMaterialTextureMipBias: return TEXT("ViewMaterialTextureMipBias");
 		default: UE_MIR_UNREACHABLE();
 	}
 }
@@ -343,8 +344,12 @@ const FType* GetExternalInputType(EExternalInput Id)
 	{
 		return FPrimitiveType::GetFloat2();
 	}
-
-	UE_MIR_UNREACHABLE();
+	
+	switch (Id)
+	{
+		case EExternalInput::ViewMaterialTextureMipBias: return FPrimitiveType::GetFloat1();
+		default: UE_MIR_UNREACHABLE();
+	}
 }
 
 MIR::EExternalInput TexCoordIndexToExternalInput(int TexCoordIndex)
diff --git a/Engine/Source/Runtime/Engine/Private/Materials/MaterialIREmitter.cpp b/Engine/Source/Runtime/Engine/Private/Materials/MaterialIREmitter.cpp
--- a/Engine/Source/Runtime/Engine/Private/Materials/MaterialIREmitter.cpp
+++ b/Engine/Source/Runtime/Engine/Private/Materials/MaterialIREmitter.cpp
@@ -190,6 +190,16 @@ FEmitter& FEmitter::DefaultInputToBoolScalar(const FExpressionInput* InInput, bo
 	return *this;
 }
 
+FEmitter& FEmitter::DefaultInputToIntScalar(const FExpressionInput* InInput, TInteger Value)
+{
+	if (!InInput->IsConnected())
+	{
+		Internal::BindValueToExpressionInput(BuilderImpl, InInput, EmitConstantIntScalar(Value));
+	}
+
+	return *this;
+}
+
 FEmitter& FEmitter::DefaultInputToFloatScalar(const FExpressionInput* InInput, TFloat Value)
 {
 	if (!InInput->IsConnected())
@@ -1820,6 +1830,11 @@ static FValue* ConstructPrimitiveValue(FEmitter* Emitter, const FPrimitiveType*
 
 FValue* FEmitter::EmitConstruct(const FType* Type, FValue* Initializer)
 {
+	if (Initializer->IsPoison())
+	{
+		return FPoison::Get();
+	}
+
 	// If target type matches initializer's, simply return the same value.
 	const FType* InitializerType = Initializer->Type;
 	if (InitializerType == Type)
@@ -1843,13 +1858,19 @@ FValue* FEmitter::EmitConstruct(const FType* Type, FValue* Initializer)
 	return Result;
 }
 
-FValue* FEmitter::EmitTextureSample(FValue* Texture, FValue* TexCoord, ESamplerSourceMode SamplerSourceMode, ETextureMipValueMode MipValueMode, EMaterialSamplerType SamplerType)
+FValue* FEmitter::EmitTextureSample(FValue* Texture, FValue* TexCoord, FValue* TexCoordDdx, FValue* TexCoordDdy, FValue* MipValue, ESamplerSourceMode SamplerSourceMode, ETextureMipValueMode MipValueMode, EMaterialSamplerType SamplerType)
 {
+	if (IsAnyNonNullAPoison(Texture, TexCoord, TexCoordDdx, TexCoordDdy, MipValue))
+	{
+		return FPoison::Get();
+	}
+
 	MIR::FTextureSample Prototype = MakePrototype<MIR::FTextureSample>(FPrimitiveType::GetFloat4());
 	Prototype.TextureObject = Texture;
-	Prototype.TexCoordArg = TexCoord;
-	Prototype.MipValueArg = nullptr;
-	Prototype.AutomaticMipBiasArg = nullptr;
+	Prototype.TexCoord = TexCoord;
+	Prototype.TexCoordDdx = TexCoordDdx;
+	Prototype.TexCoordDdy = TexCoordDdy;
+	Prototype.MipValue = MipValue;
 	Prototype.SamplerSourceMode = SamplerSourceMode;
 	Prototype.MipValueMode = MipValueMode;
 	Prototype.SamplerType = SamplerType;
diff --git a/Engine/Source/Runtime/Engine/Private/Materials/MaterialIRToHLSLTranslator.cpp b/Engine/Source/Runtime/Engine/Private/Materials/MaterialIRToHLSLTranslator.cpp
--- a/Engine/Source/Runtime/Engine/Private/Materials/MaterialIRToHLSLTranslator.cpp
+++ b/Engine/Source/Runtime/Engine/Private/Materials/MaterialIRToHLSLTranslator.cpp
@@ -412,7 +412,13 @@ struct FTranslator : FMaterialIRToHLSLTranslation
 		}
 		else
 		{
-			UE_MIR_UNREACHABLE();
+			const TCHAR* Code{};
+			switch (ExternalInput->Id)
+			{
+				case MIR::EExternalInput::ViewMaterialTextureMipBias: Code = TEXT("View.MaterialTextureMipBias"); break;
+				default: UE_MIR_UNREACHABLE();
+			}
+			Printer << Code;
 		}
 	}
 
@@ -675,29 +681,44 @@ struct FTranslator : FMaterialIRToHLSLTranslation
 
 	void LowerTextureSample(MIR::FTextureSample* TextureSample)
 	{
-		bool bUsesSpecialSampler = LowerSamplerType(TextureSample->SamplerType);
-		if (bUsesSpecialSampler)
+		bool bSamplerNeedsBrackets = LowerSamplerType(TextureSample->SamplerType);
+		if (bSamplerNeedsBrackets)
 		{
 			Printer << TEXT("(");
 		}
 
 		switch (TextureSample->TextureObject->Type->AsObject()->ObjectKind)
 		{
-			case MIR::ObjectKind_Texture2D:
-				Printer << TEXT("Texture2DSample");
-				break;
+			case MIR::ObjectKind_Texture2D: Printer << TEXT("Texture2DSample"); break;
+			default: UE_MIR_UNREACHABLE();
+		}
 
-			default:
-				UE_MIR_UNREACHABLE();
+		switch (TextureSample->MipValueMode)
+		{
+			case TMVM_None: break;
+			case TMVM_MipLevel: Printer << TEXT("Level"); break;
+			case TMVM_MipBias: Printer << TEXT("Bias"); break;
+			case TMVM_Derivative: Printer << TEXT("Grad"); break;
+			default: UE_MIR_UNREACHABLE();
 		}
 
 		Printer << BeginArgs
-			<< ListSeparator << LowerValue(TextureSample->TextureObject)
-			<< ListSeparator << LowerTextureSamplerReference(TextureSample->TextureObject, TextureSample->SamplerSourceMode)
-			<< ListSeparator << LowerValue(TextureSample->TexCoordArg)
-			<< EndArgs;
+				<< ListSeparator << LowerValue(TextureSample->TextureObject)
+				<< ListSeparator << LowerTextureSamplerReference(TextureSample->TextureObject, TextureSample->SamplerSourceMode)
+				<< ListSeparator << LowerValue(TextureSample->TexCoord);
+
+		switch (TextureSample->MipValueMode)
+		{
+			case TMVM_None: break;
+			case TMVM_MipLevel: Printer << ListSeparator << LowerValue(TextureSample->MipValue); break;
+			case TMVM_MipBias: Printer << ListSeparator << LowerValue(TextureSample->MipValue); break;
+			case TMVM_Derivative: Printer << ListSeparator << LowerValue(TextureSample->TexCoordDdx) << ListSeparator << LowerValue(TextureSample->TexCoordDdy); break;
+			default: UE_MIR_UNREACHABLE();
+		}
+
+		Printer << EndArgs;
 
-		if (bUsesSpecialSampler)
+		if (bSamplerNeedsBrackets)
 		{
 			Printer << TEXT(")");
 		}
diff --git a/Engine/Source/Runtime/Engine/Public/Materials/MaterialIR.h b/Engine/Source/Runtime/Engine/Public/Materials/MaterialIR.h
--- a/Engine/Source/Runtime/Engine/Public/Materials/MaterialIR.h
+++ b/Engine/Source/Runtime/Engine/Public/Materials/MaterialIR.h
@@ -11,8 +11,10 @@ struct FMaterialExternalCodeDeclaration;
 namespace MIR
 {
 
-// Identifies the block of execution in which an instruction runs.
-// Note: if you introduce a new stage, make sure you update EValueFlags accordingly.
+/**
+ * @brief Identifies the block of execution in which an instruction runs.
+ * @note If you introduce a new stage, make sure you update EValueFlags accordingly.
+ */
 enum EStage
 {
 	Stage_Vertex,
@@ -20,27 +22,31 @@ enum EStage
 	NumStages
 };
 
-// Returns the stage in which specified material property shall be evaluated.
+/**
+ * @brief Returns the stage in which specified material property shall be evaluated.
+ */
 EStage GetMaterialOutputStage(EMaterialProperty Property);
 
-// A collection of bit flags for a specific Value instance.
+/**
+ * @brief A collection of bit flags for a specific Value instance.
+ */
 enum class EValueFlags : uint8
 {
 	None = 0,
 
-	// Value has been analyzed for the vertex stage.
+	/** Value has been analyzed for the vertex stage. */
 	AnalyzedForStageVertex = 1 << 0,
 
-	// Value has been analyzed for the pixel stage.
+	/** Value has been analyzed for the pixel stage. */
 	AnalyzedForStagePixel = 1 << 1,
 
-	// Mask that includes all "analyzed in stage" flags.
+	/** Mask that includes all "analyzed in stage" flags. */
 	AnalyzedInAnyStageMask = AnalyzedForStageVertex | AnalyzedForStagePixel,
 
-	// Whether to skip substution of abstract tokens like "<PREV>". Only true for custom nodes.
+	/** Whether to skip substution of abstract tokens like "<PREV>". Only true for custom nodes. */
 	SubstituteTagsInInlineHLSL = 1 << 2,
 
-	// Whether 'FInlineHLSL::Code' field is used. Otherwise, 'ExternalCodeDeclaration' is used.
+	/** Whether 'FInlineHLSL::Code' field is used. Otherwise, 'ExternalCodeDeclaration' is used. */
 	HasDynamicHLSLCode = 1 << 3, 
 };
 
@@ -48,24 +54,29 @@ ENUM_CLASS_FLAGS(EValueFlags);
 
 static_assert(uint32(1 << NumStages) - 1 == uint32(EValueFlags::AnalyzedInAnyStageMask), "Discrepancy between the number of stages in NumStages and the flags in EValueFlags.");
 
-// A collection of graph properties used by a value. If a graph property is true, it means
-// that either the value itself makes has of that property, or one of its dependencies
-// (direct or indirect) has it. This entails these flags are automatically propagated to
-// all dependant values (values that depend on a given one).
-// As an example, if "ReadsPixelNormal" is true on a specific value it means that either
-// that value itself or some other upstream value that value is dependent on reads
-// the pixel normal.
+/** 
+ * @brief A collection of graph properties used by a value.
+ * 
+ * If a graph property is true, it means that either the value itself makes has of that
+ * property, or one of its dependencies (direct or indirect) has it. This entails these
+ * flags are automatically propagated to all dependant values (values that depend on a given one).
+ * As an example, if "ReadsPixelNormal" is true on a specific value it means that either
+ * that value itself or some other upstream value that value is dependent on reads
+ * the pixel normal.
+ */
 enum class EGraphProperties : uint8
 {
 	None = 0,
 
-	// Some value reads the pixel normal.
+	/** Some value reads the pixel normal. */
 	ReadsPixelNormal = 1 << 0,
 };
 
 ENUM_CLASS_FLAGS(EGraphProperties);
 
-// Enumeration of all different structs deriving from FValue. Used for dynamic casting.
+/**
+ * @brief Enumeration of all different structs deriving from FValue. Used for dynamic casting.
+ */
 enum EValueKind
 {
 	/* Values */
@@ -92,136 +103,160 @@ enum EValueKind
 	VK_InstructionEnd,
 };
 
-// Returns the string representation of given value kind.
+/** @return The string representation of given value kind. */
 const TCHAR* LexToString(EValueKind Kind);
 
 /* Values */
 
-// Base entity of all IR graph nodes.
-// An IR module is a graph of values, connected by their "uses" relations. The graph of IR
-// values is built by the `MaterialIRModuleBuilder` as the result of crawling through and 
-// analyizing the MaterialExpression graph contained in the translated Material. During 
-// this processing, IR values are emitted, and linked together. After the graph is 
-// constructed, it is itself analyzed: the builder will call `FMaterialIRValueAnalyzer::Analyze()`
-// in each active* (i.e. truly used) value in the graph, making sure a value is analyzed only 
-// after its dependencies have been analyzed.
-// A few notes:
-//  - FValues are automatically zero-initialized.
-//  - FValues are intended to be simple and inert data records. They cannot have non-trivia
-//    ctor, dtor or copy operators.
-//  - The ModuleBuilder relies on this property to efficiently hashing values so that it
-//    will reuse the same value instance instead of creating multiple instances of the
-//    same computation (for a more efficient output shader).
-//  - All values have a MIR type.
-//  - Pure FValue instances are values that do not have other dependencies (called "uses").
-//  - If a value has some other value as dependency, it means that it is the result of a
-//    calculation on those values. Values that have dependencies are Instructions (they
-//    derive from FInstruction).
+/**
+ * @brief Base entity of all IR graph nodes.
+ * 
+ * An IR module is a graph of values, connected by their "uses" relations. The graph of IR
+ * values is built by the `MaterialIRModuleBuilder` as the result of crawling through and 
+ * analyizing the MaterialExpression graph contained in the translated Material. During 
+ * this processing, IR values are emitted, and linked together. After the graph is 
+ * constructed, it is itself analyzed: the builder will call `FMaterialIRValueAnalyzer::Analyze()`
+ * in each active* (i.e. truly used) value in the graph, making sure a value is analyzed only 
+ * after its dependencies have been analyzed.
+ * A few notes:
+ *  - FValues are automatically zero-initialized.
+ *  - FValues are intended to be simple and inert data records. They cannot have non-trivia
+ *    ctor, dtor or copy operators.
+ *  - The ModuleBuilder relies on this property to efficiently hashing values so that it
+ *    will reuse the same value instance instead of creating multiple instances of the
+ *    same computation (for a more efficient output shader).
+ *  - All values have a MIR type.
+ *  - Pure FValue instances are values that do not have other dependencies (called "uses").
+ *  - If a value has some other value as dependency, it means that it is the result of a
+ *    calculation on those values. Values that have dependencies are Instructions (they
+ *    derive from FInstruction).
+ */
 struct FValue
 {
-	// Used to discern the concrete C++ type of this value (e.g. Subscript)
+	/** Used to discern the concrete C++ type of this value (e.g. Subscript) */
 	EValueKind Kind : 8;
 
-	// Set of fundamental flags true for this value.
+	/** Set of fundamental flags true for this value. */
 	EValueFlags Flags;
 
-	// The set of properties that are true for this value. Some flags might have been
-	// set but some upstream dependency leading to this value and not this value directly.
-	// See `EGraphProperties` for more information.
+	/** The set of properties that are true for this value. Some flags might have been
+	 * set but some upstream dependency leading to this value and not this value directly.
+	 * See `EGraphProperties` for more information. */
 	EGraphProperties GraphProperties;
 
-	// The runtime type this value has.
+	/** The runtime type this value has. */
 	const FType* Type;
 
-	// Returns whether this value has been analyzed for specified stage.
+	/** @return Whether this value has been analyzed for specified stage. */
 	bool IsAnalyzed(EStage State) const;
 
-	// Returns whether specified flags are true for this value.
+	/** @return Whether specified flags are true for this value. */
 	bool HasFlags(EValueFlags InFlags) const;
 
-	// Enables the specified flags without affecting others.
+	/** Enables the specified flags without affecting others. */
 	void SetFlags(EValueFlags InFlags);
 
-	// Disables the specified flags without affecting others.
+	/** Disables the specified flags without affecting others. */
 	void ClearFlags(EValueFlags InFlags);
 
-	// Returns whether specified properties are true for this value.
+	/** @return Whether specified properties are true for this value. */
 	bool HasSubgraphProperties(EGraphProperties Properties) const;
 
-	// Enables specified properties for this value.
+	/** Enables specified properties for this value. */
 	void UseSubgraphProperties(EGraphProperties Properties);
 
-	// Returns the size in bytes of this value instance.
+	/** @return The size in bytes of this value instance. */
 	uint32 GetSizeInBytes() const;
 	
-	// Returns whether this value is of specified kind.
+	/** @return Whether this value is of specified kind. */
 	bool IsA(EValueKind InKind) const;
 
-	// Returns the array of this value's uses. An use is another value referenced by this 
-	// one (e.g. the operands of a binary expression).
+	/**
+	 * Gets the immutable array of this value uses.
+	 * 
+	 * An use is another value referenced by this one (e.g. the operands of a binary expression).
+	 * 
+	 * @return The immutable array of uses.
+	 */
 	TConstArrayView<FValue*> GetUses() const;
 
-	// Returns whether this value is poison.
+	/** @return Whether this value is poison. */
 	bool IsPoison() const;
 
-	// Returns whether this value exactly equals Other.
+	/** @return Whether this value exactly equals Other. */
 	bool Equals(const FValue* Other) const;
 
-	// Returns whether this value is a scalar (its type is Primitive with exactly 1 component).
+	/** @return Whether this value is a scalar (its type is Primitive with exactly 1 component). */
 	bool IsScalar() const;
 
-	// Returns whether this value is a vector (tis type is Primitive with 1-4 rows and exactly 1 column).
+	/** @return Whether this value is a vector (tis type is Primitive with 1-4 rows and exactly 1 column). */
 	bool IsVector() const;
 
-	// Returns whether this value is a constant boolean with value true.
+	/** @return Whether this value is a constant boolean with value true. */
 	bool IsTrue() const;
 
-	// Returns whether this value is a constant boolean with value false.
+	/** @return Whether this value is a constant boolean with value false. */
 	bool IsFalse() const;
 
-	// Returns whether this value is arithmetic and exactly zero.
+	/** @return Whether this value is arithmetic and exactly zero. */
 	bool AreAllExactlyZero() const;
 
-	// Returns whether this value is arithmetic and approximately zero.
+	/** @return Whether this value is arithmetic and approximately zero. */
 	bool AreAllNearlyZero() const;
 
-	// Returns whether this value is arithmetic and exactly one.
+	/** @return Whether this value is arithmetic and exactly one. */
 	bool AreAllExactlyOne() const;
 
-	// Returns whether this value is arithmetic and approximately one.
+	/** @return Whether this value is arithmetic and approximately one. */
 	bool AreAllNearlyOne() const;
 	
-	// Tries to cast this value to specified type T and returns the casted pointer, if possible (nullptr otherwise).
+	/** Tries to cast this value to specified type T and returns the casted pointer, if possible (nullptr otherwise). */
 	template <typename T>
 	T* As() { return IsA(T::TypeKind) ? static_cast<T*>(this) : nullptr; }
 
-	// Tries to cast this value to specified type T and returns the casted pointer, if possible (nullptr otherwise).
+	/** Tries to cast this value to specified type T and returns the casted pointer, if possible (nullptr otherwise). */
 	template <typename T>
 	const T* As() const { return IsA(T::TypeKind) ? static_cast<const T*>(this) : nullptr; }
 };
 
-// If specified value is not null, it tries to cast this value T and returns it.
-// Otherwise, it returns null.
+/**
+ * Tries to cast a value to a derived type.
+ * 
+ * @return  If specified value is not null, it tries to cast this value T and returns it.
+ *          Otherwise, it returns null.
+ */
 template <typename T>
 T* As(FValue* Value)
 {
 	return Value && Value->IsA(T::TypeKind) ? static_cast<T*>(Value) : nullptr;
 }
 
-// If specified value is not null, it tries to cast this value T and returns it.
-// Otherwise, it returns null.
+/**
+ * Tries to cast a value to a derived type.
+ * 
+ * @return  If specified value is not null, it tries to cast this value T and returns it.
+ *          Otherwise, it returns null.
+ */
 template <typename T>
 const T* As(const FValue* Value)
 {
 	return Value && Value->IsA(T::TypeKind) ? static_cast<const T*>(Value) : nullptr;
 }
 
-// If specified value is not null, it tries to cast this value to an instruction and returns it.
-// Otherwise, it returns null.
+/**
+ * Casts a value to an instruction.
+ * 
+ * @return If specified value is not null, it tries to cast this value to an instruction
+ *         and returns it. Otherwise, it returns null.
+ */
 FInstruction* AsInstruction(FValue* Value);
 
-// If specified value is not null, it tries to cast this value to an instruction and returns it.
-// Otherwise, it returns null.
+/**
+ * Casts a value to an instruction.
+ * 
+ * @return If specified value is not null, it tries to cast this value to an instruction
+ *         and returns it. Otherwise, it returns null.
+ */
 const FInstruction* AsInstruction(const FValue* Value);
 
 template <EValueKind TTypeKind>
@@ -230,21 +265,51 @@ struct TValue : FValue
 	static constexpr EValueKind TypeKind = TTypeKind;
 };
 
-// A poison value represents an invalid value. It is produced by the emitter when an
-// invalid operation is performed. Poison values can be passed as arguments to other operations,
-// but they are "contagious": any instruction emitted with a poison value as an argument
-// will itself produce a poison value.
+/**
+ * @brief A placeholder for an invalid value.
+ * 
+ * A poison value represents an invalid value. It is produced by the emitter when an
+ * invalid operation is performed. Poison values can be passed as arguments to other operations,
+ * but they are "contagious": any instruction emitted with a poison value as an argument
+ * will itself produce a poison value.
+ */
 struct FPoison : TValue<VK_Poison>
 {
 	static FPoison* Get();
 };
 
+static inline bool IsAnyNonNullAPoison()
+{
+	return false;
+}
+
+/** @return Whether there is a poison in the values specified, which could be null. */
+template <typename... TTail>
+bool IsAnyNonNullAPoison(const FValue* Head, TTail... Tail)
+{
+	if (Head && Head->IsPoison())
+	{
+		return true;
+	}
+	else
+	{
+		return IsAnyNonNullAPoison(Tail...);
+	}
+}
+
+/** The integer type used inside MIR. */
 using TInteger = int64_t;
+
+/** The floating point type used inside MIR. */
 using TFloat = float;
 
-// A constant represents a translation-time known scalar primitive value. Operations on
-// constant values can be folded by the builder, that is they can be evaluated statically
-// while the builder constructs the IR graph of an input material.
+/**
+ * @brief A constant value.
+ * 
+ * A constant represents a translation-time known scalar primitive value. Operations on
+ * constant values can be folded by the builder, that is they can be evaluated statically
+ * while the builder constructs the IR graph of an input material.
+ */
 struct FConstant : TValue<VK_Constant>
 {
 	union
@@ -254,16 +319,16 @@ struct FConstant : TValue<VK_Constant>
 		TFloat 		Float;
 	};
 
-	// Returns whether this constant is a boolean.
+	/** @return Whether this constant is a boolean. */
 	bool IsBool() const;
 
-	// Returns whether this constant is an integer.
+	/** @return Whether this constant is an integer. */
 	bool IsInteger() const;
 	
-	// Returns whether this constant is a float.
+	/** @return Whether this constant is a float. */
 	bool IsFloat() const;
 
-	// Returns the constant value of given type T. The type must be bool, integral or floating point.
+	/** @return The constant value of given type T. The type must be bool, integral or floating point. */
 	template <typename T>
 	T Get() const
 	{
@@ -317,132 +382,156 @@ enum class EExternalInput
 	TexCoord6_Ddy,
 	TexCoord7_Ddy,
 
+	ViewMaterialTextureMipBias,
+
 	Count,
 };
 
-// Maximum number of supported texture coordinates.
+/** Maximum number of supported texture coordinates. */
 static constexpr int TexCoordMaxNum = 8;
 
-// Converts the given external input to its string representation.
+/** @return The string representation of given given external input. */
 const TCHAR* LexToString(EExternalInput Input);
 
-// Returns the given external input type.
+/** @return The given external input type. */
 const FType* GetExternalInputType(EExternalInput Id);
 
-// Converts the given texture coordinate index to its corresponding external input mapping.
+/** Converts the given texture coordinate index to its corresponding external input mapping. */
 MIR::EExternalInput TexCoordIndexToExternalInput(int TexCoordIndex);
 
-// Returns whether the given external is a texture coordinate.
+/** @return Whether the given external is a texture coordinate. */
 bool IsExternalInputTexCoord(EExternalInput Id);
 
-// Returns whether the given external is a texture coordinate ddx.
+/** @return Whether the given external is a texture coordinate ddx. */
 bool IsExternalInputTexCoordDdx(EExternalInput Id);
 
-// Returns whether the given external is a texture coordinate ddy.
+/** @return Whether the given external is a texture coordinate ddy. */
 bool IsExternalInputTexCoordDdy(EExternalInput Id);
 
-// It represents an external input (like texture coordinates) value.
+/**
+ * It represents an external input (like texture coordinates) value.
+ */
 struct FExternalInput : TValue<VK_ExternalInput>
 {
 	EExternalInput Id;
 };
 
-// It represents a texture object value.
+/**
+ * It represents a texture object value. 
+ */
 struct FTextureObject : TValue<VK_TextureObject>
 {
-	// The texture object.
+	/** The texture object. */
 	UTexture* Texture;
 
-	// The sampler type associated to this texture object.
+	/** The sampler type associated to this texture object. */
 	EMaterialSamplerType SamplerType;
 
-	// Index of this parameter in the uniform expression set.
-	// Note: This field  is automatically set by the builder.
+	/** Index of this parameter in the uniform expression set.
+	 * @note This field is automatically set by the builder. */
 	uint32 Analysis_UniformParameterIndex;
 };
 
-// It represents a material uniform parameter.
+/**
+ * It represents a material uniform parameter.
+ */
 struct FUniformParameter: TValue<VK_UniformParameter>
 {
-	// Index of the parameter as registered in the module.
+	/** Index of the parameter as registered in the module. */
 	uint32 ParameterIdInModule;
 
-	// Eventual sampler type to use to sample this parameter's texture (if it is one)
+	/** Eventual sampler type to use to sample this parameter's texture (if it is one) */
 	EMaterialSamplerType SamplerType;
 
-	// Index of this parameter in the uniform expression set
-	// Note: This field is automatically set by the builder.
+	/** Index of this parameter in the uniform expression set
+	 * @note This field is automatically set by the builder. */
 	uint32 Analysis_UniformParameterIndex;
 };
 
 
 /* Instructions */
 
-// A block is a sequence of instructions in order of execution. Blocks are organized in a
-// tree-like structure, used to model blocks nested inside other blocks.
+/**
+ * A block of instructions.
+ * 
+ * A block is a sequence of instructions in order of execution. Blocks are organized in a
+ * tree-like structure, used to model blocks nested inside other blocks.
+ */
 struct FBlock
 {
-	// This block's parent block, if any. If this is null, this is a *root block*.
+	/** This block's parent block, if any. If this is null, this is a *root block. */
 	FBlock* Parent;
 
-	// The linked-list head of the instructions contained in this block. Links are contained
-	// in the FInstruction::Next field.
+	/** The linked-list head of the instructions contained in this block.
+	 * Links are contained in the FInstruction::Next field. */
 	FInstruction* Instructions;
 
-	// Depth of this block in the tree structure (root blocks have level zero).
+	/** Depth of this block in the tree structure (root blocks have level zero). */
 	int32 Level;
 
-	// Finds and returns the common block between this and Other, if any. It returns null
-	// if no common block was found.
-	// Note: this has O(n) complexity, where n is the maximum depth of the tree structure.
+	/** Finds and returns the common block between this and Other, if any. It returns null
+	 * if no common block was found.
+	 * @note This has O(n) complexity, where n is the maximum depth of the tree structure. */
 	FBlock* FindCommonParentWith(MIR::FBlock* Other);
 };
 
-// An instruction is a value in a well defined order of evaluation.
-// Instructions have a parent block and are organized in a linked list: the Next field
-// indicates the instruction that will execute immediately after this one.
-// Since the material shader has multiple stages, the same instruction can belong to two
-// different graphs of execution, which explains why all fields in this structure have
-// a different possible value per stage.
-// Note: all fields in this struct are not expected to be set by the user when emitting
-// 		an instruction, and are instead automatically populated by the builder.
+/**
+ * @brief Base struct of an instruction value.
+ * 
+ * An instruction is a value in a well defined order of evaluation.
+ * Instructions have a parent block and are organized in a linked list: the Next field
+ * indicates the instruction that will execute immediately after this one.
+ * Since the material shader has multiple stages, the same instruction can belong to two
+ * different graphs of execution, which explains why all fields in this structure have
+ * a different possible value per stage.
+ * 
+ * @note All fields in this struct are not expected to be set by the user when emitting 
+ *       an instruction, and are instead automatically populated by the builder.
+ */
 struct FInstruction : FValue
 {
-	// The next instruction executing after this one in all stages of execution.
+	/** The next instruction executing after this one in all stages of execution. */
 	FInstruction* Next[NumStages];
 
-	// This instruction parent block in all stages of execution.
+	/** This instruction parent block in all stages of execution. */
 	FBlock* Block[NumStages];
 
-	// How many users (i.e., dependencies) this instruction has in each stage of execution.
+	/** How many users (i.e., dependencies) this instruction has in each stage of execution. */
 	uint32 NumUsers[NumStages];
 
-	// The number of users processed during instruction graph linking in each stage of execution.
-	// Note: This information combined with NumUsers is used by the builder to push instructions
-	//       in the appropriate block automatically.
+	/** The number of users processed during instruction graph linking in each stage of execution.
+	 * @note This information combined with NumUsers is used by the builder to push instructions
+	 *       in the appropriate block automatically. */
 	uint32 NumProcessedUsers[NumStages];
 
-	// Returns the block in which the dependency with specified index should execute.
+	/** @return The block in which the dependency with specified index should execute. */
 	FBlock* GetDesiredBlockForUse(EStage Stage, int32 UseIndex);
 };
 
-template <EValueKind TTypeKind>
+template <EValueKind TTypeKind, uint32 TNumUses>
 struct TInstruction : FInstruction
 {
+	/** The kind of this instruction. */
 	static constexpr EValueKind TypeKind = TTypeKind;
+
+	/** The number of uses this instruction has. */
+	static constexpr uint32 NumUses = TNumUses;
 };
 
-// A dimensional is a fixed array of other values. This value is used to model vectors
-// and matrices.
-struct FDimensional : TInstruction<VK_Dimensional>
+/**
+ * @brief An aggregate of other values.
+ * 
+ * A dimensional is a fixed array of other values. This value is used to model vectors and matrices.
+ */
+struct FDimensional : TInstruction<VK_Dimensional, 0>
 {
-	// Returns the constant array of component values. 
+	/** @return The constant array of component values.  */
 	TConstArrayView<FValue*> GetComponents() const;
 
-	// Returns the mutable array of component values. 
+	/** @return The mutable array of component values.  */
 	TArrayView<FValue*> GetMutableComponents();
 
-	// Returns whether all components are constant (i.e., they're instances of FConstant).
+	/** @return Whether all components are constant (i.e., they're instances of FConstant). */
 	bool AreComponentsConstant() const;
 };
 
@@ -452,28 +541,42 @@ struct TDimensional : FDimensional
 	FValue* Components[TDimension];
 };
 
-// This instruction sets material attribute (i.e., "BaseColor") to its value.
-struct FSetMaterialOutput : TInstruction<VK_SetMaterialOutput>
+/**
+ * This instruction sets material attribute (i.e., "BaseColor") to its value.
+ */
+struct FSetMaterialOutput : TInstruction<VK_SetMaterialOutput, 1>
 {
-	// The material attribute to set.
-	EMaterialProperty Property;
-
-	// The value this material attribute should be set to.
+	/** The value this material attribute should be set to. */
 	FValue* Arg;
+
+	/** The material attribute to set. */
+	EMaterialProperty Property;
 };
 
-// MIR instruction for external code declarations; See BaseMaterialExpressions.ini file.
-struct FInlineHLSL : TInstruction<VK_InlineHLSL>
+/**
+ * @brief Inline HLSL instruction.
+ * 
+ * A value that is the result of an arbitrary snippet of HLSL code.
+ * 
+ * @note See BaseMaterialExpressions.ini file. 
+ */
+struct FInlineHLSL : TInstruction<VK_InlineHLSL, 0>
 {
 	union
 	{
+		/** The declaration this instruction refers to. */
 		const FMaterialExternalCodeDeclaration* ExternalCodeDeclaration;
-		const TCHAR* Code; // The inlined HLSL code string.
+
+		/** The arbitrary inlined HLSL code snippet */
+		const TCHAR* Code;
 	};
 };
 
-// Note: if you modify this enum, update the implementations of the
-// helper functions below.
+/**
+ * @brief Operator enumeration.
+ * 
+ * @note If you modify this enum, update the implementations of the helper functions below.
+ */
 enum EOperator
 {
 	O_Invalid,
@@ -544,128 +647,126 @@ enum EOperator
 	OperatorCount,
 };
 
-// Whether the given operator identifies a comparison operation (e.g., ">=", "==").
+/** Whether the given operator identifies a comparison operation (e.g., ">=", "=="). */
 bool IsComparisonOperator(EOperator Op);
 
-// Whether the given operator identifies a bit operation (e.g., "&", "|").
+/** Whether the given operator identifies a bit operation (e.g., "&", "|"). */
 bool IsBitOperator(EOperator Op);
 
-// Whether the given operator identifies a unary operator.
+/** Whether the given operator identifies a unary operator. */
 bool IsUnaryOperator(EOperator Op);
 
-// Whether the given operator identifies a binary operator.
+/** Whether the given operator identifies a binary operator. */
 bool IsBinaryOperator(EOperator Op);
 
-// Whether the given operator identifies a ternary operator.
+/** Whether the given operator identifies a ternary operator. */
 bool IsTernaryOperator(EOperator Op);
 
-// Returns the string representation of the given operator.
+/** @return The string representation of the given operator. */
 const TCHAR* LexToString(EOperator Op);
 
-// This instruction identifies a built-in operation on one, two or three argument values.
-struct FOperator : TInstruction<VK_Operator>
+/**
+ * @brief A mathematical operator instruction.
+ * 
+ * This instruction identifies a built-in operation on one, two or three argument values.
+ */
+struct FOperator : TInstruction<VK_Operator, 3>
 {
-	// It identifies which supported operation to carry.
-	EOperator Op;
-
-	// The first argument of the operation. This value is never null.
+	/** The first argument of the operation. This value is never null. */
 	FValue* AArg;
 
-	// The second argument of the operation. This value is null for unary operators.
+	/** The second argument of the operation. This value is null for unary operators. */
 	FValue* BArg;
 
-	// The third argument of the operation. This value is null for unary and binary operators.
+	/** The third argument of the operation. This value is null for unary and binary operators. */
 	FValue* CArg;
+
+	/** It identifies which supported operation to carry. */
+	EOperator Op;
 };
 
-// This instruction evaluates to one or another argument based on whether a third boolean
-// condition argument is true or false. The builder will automatically place as many instructions
-// as possible in the true/false inner blocks whilst respecting dependency requirements.
-// This is done in an effort to avoid the unnecessary computation of the input value that was
-// not selected by the condition.
-struct FBranch : TInstruction<VK_Branch>
+/**
+ * @brief A branch instruction.
+ * 
+ * This instruction evaluates to one or another argument based on whether a third boolean
+ * condition argument is true or false. The builder will automatically place as many instructions
+ * as possible in the true/false inner blocks whilst respecting dependency requirements.
+ * This is done in an effort to avoid the unnecessary computation of the input value that was
+ * not selected by the condition.
+ */
+struct FBranch : TInstruction<VK_Branch, 3>
 {
-	// The boolean condition argument used to forward the "true" or "false" argument.
+	/** The boolean condition argument used to forward the "true" or "false" argument. */
 	FValue* ConditionArg;
 
-	// Value this branch evaluates to when the condition is true.
+	/** Value this branch evaluates to when the condition is true. */
 	FValue* TrueArg;
 
-	// Value this branch evaluates to when the condition is false.
+	/** Value this branch evaluates to when the condition is false. */
 	FValue* FalseArg;
 
-	// The inner block (in each execution stage) the subgraph evaluating `TrueArg` should be placed in.
+	/** The inner block (in each execution stage) the subgraph evaluating `TrueArg` should be placed in. */
 	FBlock TrueBlock[NumStages];
 
-	// The inner block (in each execution stage) the subgraph evaluating `FalseArg` should be placed in.
+	/** The inner block (in each execution stage) the subgraph evaluating `FalseArg` should be placed in. */
 	FBlock FalseBlock[NumStages];
 };
 
-// This instruction is used to pull the inner value making a compound one. For example,
-// this instruction is used to extract an individual component of a vector value.
-struct FSubscript : TInstruction<VK_Subscript>
+/**
+ * @brief A subscript instruction.
+ * 
+ * This instruction is used to pull the inner value making a compound one. For example,
+ * it is used to extract an individual component of a vector value.
+ */
+struct FSubscript : TInstruction<VK_Subscript, 1>
 {
-	// The argument to subscript.
+	/** The argument to subscript. */
 	FValue* Arg;
 
-	// The subscript index, i.e. the index of the component to extract.
+	/** The subscript index, i.e. the index of the component to extract. */
 	int Index;
 };
 
-// This instruction casts an argument value to a different type.
-struct FCast : TInstruction<VK_Cast>
+/**
+ * @brief A type cast instruction.
+ * 
+ * This instruction casts an argument value to a different type.
+ */
+struct FCast : TInstruction<VK_Cast, 1>
 {
-	// The argument value to cast.
+	/** The argument value to cast. */
 	FValue* Arg;
 };
 
-// This instruction performs texture sampling.
-struct FTextureSample : TInstruction<VK_TextureSample>
+/**
+ * @brief This instruction performs texture sampling.
+ */
+struct FTextureSample : TInstruction<VK_TextureSample, 5>
 {
-	// The texture object to sample.
+	/** The texture object to sample. */
 	FValue* TextureObject;
 
-	// The texture coordinate at which to sample.
-	FValue* TexCoordArg;
+	/** The texture coordinate at which to sample. */
+	FValue* TexCoord;
 
-	// The mip index to sample, if any provided.
-	FValue* MipValueArg;
+	/** Optional. The mip index to sample, if any provided. */
+	FValue* MipValue;
 
-	// The automatic mip bias to use, if any provided.
-	FValue* AutomaticMipBiasArg;
+	/** Optional. The analytical partial derivative of the coordinates along the X axis. */
+	FValue* TexCoordDdx;
+		
+	/** Optional. The analytical partial derivative of the coordinates along the Y axis. */
+	FValue* TexCoordDdy;
 
-	// The sampler source mode to use for sampling.
+	/** The sampler source mode to use for sampling. */
 	ESamplerSourceMode SamplerSourceMode;
 
-	// The mip value mode to use for sampling.
+	/** The mip value mode to use for sampling. */
 	ETextureMipValueMode MipValueMode;
 
-	// The sampler type to use for sampling.
+	/** The sampler type to use for sampling. */
 	EMaterialSamplerType SamplerType;
-
-	// TODO: TexCoordDerivatives;
 };
 
-// --- Helper functions ---
-
-static inline bool IsAnyNonNullAPoison()
-{
-	return false;
-}
-
-// Returns whether there is a poison in the values specified, which could be null.
-template <typename... TTail>
-bool IsAnyNonNullAPoison(const FValue* Head, TTail... Tail)
-{
-	if (Head && Head->IsPoison())
-	{
-		return true;
-	}
-	else
-	{
-		return IsAnyNonNullAPoison(Tail...);
-	}
-}
-
 } // namespace MIR
 #endif
diff --git a/Engine/Source/Runtime/Engine/Public/Materials/MaterialIREmitter.h b/Engine/Source/Runtime/Engine/Public/Materials/MaterialIREmitter.h
--- a/Engine/Source/Runtime/Engine/Public/Materials/MaterialIREmitter.h
+++ b/Engine/Source/Runtime/Engine/Public/Materials/MaterialIREmitter.h
@@ -50,6 +50,9 @@ public:
 	//
 	FEmitter& DefaultInputToBoolScalar(const FExpressionInput* Input, bool Value);
 
+	//
+	FEmitter& DefaultInputToIntScalar(const FExpressionInput* Input, TInteger Value);
+	
 	//
 	FEmitter& DefaultInputToFloatScalar(const FExpressionInput* Input, TFloat Value);
 	
@@ -163,7 +166,23 @@ public:
 	FValue* EmitOperator(EOperator Operator, FValue* A, FValue* B = nullptr, FValue* C = nullptr);
 	FValue* EmitBranch(FValue* Condition, FValue* True, FValue* False);
 	FValue* EmitConstruct(const FType* Type, FValue* Initializer);
-	FValue* EmitTextureSample(FValue* Texture, FValue* TexCoord, ESamplerSourceMode SamplerSourceMode, ETextureMipValueMode MipValueMode, EMaterialSamplerType SamplerType);
+
+	/**
+	 * @brief Emits a texture sampling instruction.
+	 * 
+	 * @param Texture The texture to sample. Must be of type MIR::FObjectType and represent a valid texture.
+	 * @param TexCoord The UV coordinates at which to sample the texture.
+	 * @param TexCoordDdx Optional. The manually provided partial derivative of TexCoord along the X axis. Required if MipValueMode is Derivative.
+	 * @param TexCoordDdy Optional. The manually provided partial derivative of TexCoord along the Y axis. Required if MipValueMode is Derivative.
+	 * @param MipValue Optional. The mip level or mip bias to apply. Required if MipValueMode is Level or Bias.
+	 * @param SamplerSourceMode The sampler source mode to use.
+	 * @param MipValueMode The mip value mode to use.
+	 * @param SamplerType The sampler type to use.
+	 * 
+	 * @return A pointer to the emitted texture sample instruction.
+	 */
+	FValue* EmitTextureSample(FValue* Texture, FValue* TexCoord, FValue* TexCoordDdx, FValue* TexCoordDdy, FValue* MipValue, ESamplerSourceMode SamplerSourceMode, ETextureMipValueMode MipValueMode, EMaterialSamplerType SamplerType);
+	
 	FValue* EmitInlineHLSL(const FType* Type, FString Code, EValueFlags ValueFlags = EValueFlags::None, EGraphProperties UsedGraphProperties = EGraphProperties::None);
 	FValue* EmitInlineHLSL(const FMaterialExternalCodeDeclaration* InExternalCodeDeclaration, EGraphProperties UsedGraphProperties = EGraphProperties::None);
 
