Add default member initialization values for FCellDesc and FCellDescInstance following 40787720.

diff --git a/Engine/Source/Runtime/Engine/Public/WorldPartition/RuntimeHashSet/RuntimePartition.h b/Engine/Source/Runtime/Engine/Public/WorldPartition/RuntimeHashSet/RuntimePartition.h
--- a/Engine/Source/Runtime/Engine/Public/WorldPartition/RuntimeHashSet/RuntimePartition.h
+++ b/Engine/Source/Runtime/Engine/Public/WorldPartition/RuntimeHashSet/RuntimePartition.h
@@ -33,14 +33,14 @@ public:
 	struct FCellDesc
 	{
 		FName Name;
-		bool bIsSpatiallyLoaded;
-		bool bBlockOnSlowStreaming;
-		bool bClientOnlyVisible;
-		bool bIs2D;
-		int32 Priority;
+		bool bIsSpatiallyLoaded = false;
+		bool bBlockOnSlowStreaming = false;
+		bool bClientOnlyVisible = false;
+		bool bIs2D = false;
+		int32 Priority = 0;
 
 		/** Optional level value that can be used to filter debug display */
-		int32 Level;
+		int32 Level = 0;
 
 		/** Optional cell bounds for partitions that work on uniform grids */
 		TOptional<FBox> CellBounds;
@@ -55,7 +55,7 @@ public:
 	{
 		FCellDescInstance(const FCellDesc& InCellDesc, URuntimePartition* InSourcePartition, const TArray<const UDataLayerInstance*>& InDataLayerInstances, const FGuid& InContentBundleID);
 
-		URuntimePartition* SourcePartition;
+		URuntimePartition* SourcePartition = nullptr;
 		TArray<const UDataLayerInstance*> DataLayerInstances;
 		FGuid ContentBundleID;
 	};
