Fix crash in the GPUScene upload

diff --git a/Engine/Source/Runtime/Engine/Public/PrimitiveUniformShaderParametersBuilder.h b/Engine/Source/Runtime/Engine/Public/PrimitiveUniformShaderParametersBuilder.h
--- a/Engine/Source/Runtime/Engine/Public/PrimitiveUniformShaderParametersBuilder.h
+++ b/Engine/Source/Runtime/Engine/Public/PrimitiveUniformShaderParametersBuilder.h
@@ -254,7 +254,7 @@ public:
 	inline FPrimitiveUniformShaderParametersBuilder& CustomPrimitiveData(const FCustomPrimitiveData* InCustomPrimitiveData)
 	{
 		// If this primitive has custom primitive data, set it
-		if (InCustomPrimitiveData)
+		if (InCustomPrimitiveData && !InCustomPrimitiveData->Data.IsEmpty())
 		{
 			// Copy at most up to the max supported number of floats for safety
 			FMemory::Memcpy(
