[ZenLoader]
- Add test for ExternalReadDependency support that was recently fixed

diff --git a/Engine/Source/Runtime/Engine/Private/Tests/Loading/AsyncLoadingTests_ExternalReadDependencies.cpp b/Engine/Source/Runtime/Engine/Private/Tests/Loading/AsyncLoadingTests_ExternalReadDependencies.cpp
new file mode 100644
--- /dev/null
+++ b/Engine/Source/Runtime/Engine/Private/Tests/Loading/AsyncLoadingTests_ExternalReadDependencies.cpp
@@ -0,0 +1,44 @@
+// Copyright Epic Games, Inc. All Rights Reserved.
+
+#include "AsyncLoadingTests_Shared.h"
+#include "Misc/AutomationTest.h"
+#include "Misc/PackageName.h"
+
+#if WITH_DEV_AUTOMATION_TESTS
+
+IMPLEMENT_SIMPLE_AUTOMATION_TEST(
+	FLoadingTests_ExternalReadDependencies,
+	TEXT("System.Engine.Loading.ExternalReadDependencies"),
+	EAutomationTestFlags::EditorContext | EAutomationTestFlags::EngineFilter
+)
+bool FLoadingTests_ExternalReadDependencies::RunTest(const FString& Parameters)
+{
+	FLoadingTestsScope LoadingTestScope(this);
+
+	std::atomic<int32> NumExternalReads = 0;
+	UAsyncLoadingTests_Shared::OnSerialize.BindLambda(
+		[this, &NumExternalReads](FArchive& Ar, UAsyncLoadingTests_Shared* Object)
+		{
+			FExternalReadCallback ExternalReadCallback =
+				[&NumExternalReads](double RemainingTime)
+				{
+					NumExternalReads++;
+					return true;
+				};
+
+			Ar.AttachExternalReadDependency(ExternalReadCallback);
+		}
+	);
+
+	// Trigger the async loading
+	int32 RequestId = LoadPackageAsync(FLoadingTestsScope::PackagePath1);
+
+	// When a flush occurs, we have no choice but to ignore bIsReadyForAsyncPostLoad
+	FlushAsyncLoading(RequestId);
+
+	TestTrue("ExternalReadDependency callbacks should have been called", NumExternalReads != 0);
+
+	return true;
+}
+
+#endif // WITH_DEV_AUTOMATION_TESTS
