Added a speculative fix for a crash with an invalid body instance.

diff --git a/Engine/Source/Runtime/Engine/Private/Components/PrimitiveComponent.cpp b/Engine/Source/Runtime/Engine/Private/Components/PrimitiveComponent.cpp
--- a/Engine/Source/Runtime/Engine/Private/Components/PrimitiveComponent.cpp
+++ b/Engine/Source/Runtime/Engine/Private/Components/PrimitiveComponent.cpp
@@ -3612,19 +3612,22 @@ bool UPrimitiveComponent::ComponentOverlapComponentImpl(class UPrimitiveComponen
 	{
 		if (FBodyInstance* ThisBodyInstance = GetBodyInstance())
 		{
-			return ThisBodyInstance->OverlapTestForBodies(Pos, Quat, InstancedStaticMesh->InstanceBodies);
-		}
-		else
-		{
-			return false;
+			if (ensure(ThisBodyInstance->IsValidBodyInstance()))
+			{
+				return ThisBodyInstance->OverlapTestForBodies(Pos, Quat, InstancedStaticMesh->InstanceBodies);
+			}
 		}
+		return false;
 	}
 
 	FBodyInstance* BI = PrimComp->GetBodyInstance();
 	FBodyInstance* ThisBodyInstance = GetBodyInstance();
 	if(BI && ThisBodyInstance)
 	{
-		return BI->OverlapTestForBody(Pos, Quat, ThisBodyInstance);
+		if (ensure(BI->IsValidBodyInstance() && ThisBodyInstance->IsValidBodyInstance()))
+		{
+			return BI->OverlapTestForBody(Pos, Quat, ThisBodyInstance);
+		}
 	}
 
 	TArray<Chaos::FPhysicsObjectHandle> InObjects = PrimComp->GetAllPhysicsObjects();
