Component unregistration for borrowing a borrowed object: components of borrowed objects are not unregistered when they are loaned. Adding an extra check to make sure all components are unregistered.

diff --git a/Engine/Source/Runtime/Engine/Private/Actor.cpp b/Engine/Source/Runtime/Engine/Private/Actor.cpp
--- a/Engine/Source/Runtime/Engine/Private/Actor.cpp
+++ b/Engine/Source/Runtime/Engine/Private/Actor.cpp
@@ -1252,7 +1252,7 @@ void AActor::PostMigrate(const FUObjectMigrationContext& MigrationContext)
 		ensureMsgf(MigrationContext.RemoteServerId != LocalServerId, TEXT("%hs: Sending Object %s from %s to %s but destination is the same as the Local Server Id."), __func__, *GetName(), *LocalServerId.ToString(), *MigrationContext.RemoteServerId.ToString());
 
 		const bool bKeepOwnership = (MigrationSendType == EObjectMigrationSendType::Loan);
-		if (!bKeepOwnership)
+		if (!bKeepOwnership || !bHasOwnership)
 		{
 			// Remove the actor from the world on this server as it is about to be migrated to another server.
 			// Migrated actors are unique to the server instance that currently owns their memory and are not replicated to other servers.
