World Partition - HLOD: Disable distance culling on primitive components that are included in HLOD as this can potentially lead to a gap of visibility between the HLOD representation of the component & the component itself.

diff --git a/Engine/Source/Runtime/Engine/Classes/Components/PrimitiveComponent.h b/Engine/Source/Runtime/Engine/Classes/Components/PrimitiveComponent.h
--- a/Engine/Source/Runtime/Engine/Classes/Components/PrimitiveComponent.h
+++ b/Engine/Source/Runtime/Engine/Classes/Components/PrimitiveComponent.h
@@ -2947,6 +2947,11 @@ public:
 	 */
 	ENGINE_API bool ShouldComponentAddToScene() const;
 	
+	/**
+	 * Returns whether the component can ever be distanced culled.
+	 */
+	ENGINE_API bool CanBeDistanceCulled() const;
+
 	/**
 	 * Changes the value of CullDistance.
 	 * @param NewCullDistance - The value to assign to CullDistance.
diff --git a/Engine/Source/Runtime/Engine/Private/Components/PrimitiveComponent.cpp b/Engine/Source/Runtime/Engine/Private/Components/PrimitiveComponent.cpp
--- a/Engine/Source/Runtime/Engine/Private/Components/PrimitiveComponent.cpp
+++ b/Engine/Source/Runtime/Engine/Private/Components/PrimitiveComponent.cpp
@@ -638,8 +638,7 @@ void UPrimitiveComponent::CreateRenderState_Concurrent(FRegisterComponentContext
 	// Make sure cached cull distance is up-to-date if its zero and we have an LD cull distance
 	if( CachedMaxDrawDistance == 0.f && LDMaxDrawDistance > 0.f )
 	{
-		bool bNeverCull = bNeverDistanceCull || GetLODParentPrimitive();
-		CachedMaxDrawDistance = bNeverCull ? 0.f : LDMaxDrawDistance;
+		CachedMaxDrawDistance = CanBeDistanceCulled() ? LDMaxDrawDistance : 0;
 	}
 
 	Super::CreateRenderState_Concurrent(Context);
@@ -1724,8 +1723,7 @@ void UPrimitiveComponent::PostLoad()
 			CachedMaxDrawDistance = FMath::Min( LDMaxDrawDistance, CachedMaxDrawDistance );
 		}
 
-		bool bNeverCull = bNeverDistanceCull || GetLODParentPrimitive();
-		CachedMaxDrawDistance = bNeverCull ? 0.f : CachedMaxDrawDistance;
+		CachedMaxDrawDistance = CanBeDistanceCulled() ? CachedMaxDrawDistance : 0.0f;
 	} 
 
 	if (GetLightmapType() == ELightmapType::ForceSurface && GetStaticLightingType() == LMIT_None)
@@ -2342,10 +2340,32 @@ void UPrimitiveComponent::SetCullDistance(const float NewCullDistance)
 	}
 }
 
+bool UPrimitiveComponent::CanBeDistanceCulled() const
+{
+	if (bNeverDistanceCull)
+	{
+		return false;
+	}
+
+	if (GetLODParentPrimitive())
+	{
+		return false;
+	}
+
+	if (ULevel* Level = GetTypedOuter<ULevel>())
+	{
+		if (Level->IsWorldPartitionRuntimeCell())
+		{
+			return !IsHLODRelevant();
+		}
+	}
+	
+	return true;
+}
+
 void UPrimitiveComponent::SetCachedMaxDrawDistance(const float NewCachedMaxDrawDistance)
 {
-	bool bNeverCull = bNeverDistanceCull || GetLODParentPrimitive();
-	float NewMaxDrawDistance = bNeverCull ? 0.f : NewCachedMaxDrawDistance;
+	float NewMaxDrawDistance = CanBeDistanceCulled() ? NewCachedMaxDrawDistance : 0;
 
 	if( !FMath::IsNearlyEqual(CachedMaxDrawDistance, NewMaxDrawDistance) )
 	{
