Add function to RemoveAxisKeyBinding from AxisKeyBindings array.

diff --git a/Engine/Source/Runtime/Engine/Classes/Components/InputComponent.h b/Engine/Source/Runtime/Engine/Classes/Components/InputComponent.h
--- a/Engine/Source/Runtime/Engine/Classes/Components/InputComponent.h
+++ b/Engine/Source/Runtime/Engine/Classes/Components/InputComponent.h
@@ -963,6 +963,13 @@ public:
 		return AxisKeyBindings.Last();
 	}
 
+	/**
+	 * Removes the axis key binding with the specified key.
+	 * 
+	 * @param AxisKey the key of the axis to remove.
+	 */
+	ENGINE_API void RemoveAxisKeyBinding( const FKey AxisKey);
+
 	/**
 	 * Binds a delegate function to a vector axis key (e.g. Tilt)
 	 * Returned reference is only guaranteed to be valid until another vector axis key is bound.
diff --git a/Engine/Source/Runtime/Engine/Private/Components/InputComponent.cpp b/Engine/Source/Runtime/Engine/Private/Components/InputComponent.cpp
--- a/Engine/Source/Runtime/Engine/Private/Components/InputComponent.cpp
+++ b/Engine/Source/Runtime/Engine/Private/Components/InputComponent.cpp
@@ -405,6 +405,18 @@ void UInputComponent::ClearAxisBindings()
 	AxisBindings.Reset();
 }
 
+void UInputComponent::RemoveAxisKeyBinding(FKey AxisKey)
+{
+	for (int32 AxisIdx = AxisKeyBindings.Num(); --AxisIdx;)
+	{
+		const FInputAxisKeyBinding& Binding = AxisKeyBindings[AxisIdx];
+		if (Binding.AxisKey == AxisKey)
+		{
+			AxisKeyBindings.RemoveAt(AxisIdx, EAllowShrinking::No);
+		}
+	}
+}
+
 /* Deprecated functions (needed for Blueprints)
  *****************************************************************************/
 
