Substrate - clearly prevent unlit node usage for decals.

#tests pc editor dx12
#rnx

diff --git a/Engine/Source/Runtime/Core/Private/UObject/FortniteMainBranchObjectVersion.cpp b/Engine/Source/Runtime/Core/Private/UObject/FortniteMainBranchObjectVersion.cpp
--- a/Engine/Source/Runtime/Core/Private/UObject/FortniteMainBranchObjectVersion.cpp
+++ b/Engine/Source/Runtime/Core/Private/UObject/FortniteMainBranchObjectVersion.cpp
@@ -14,6 +14,6 @@ TMap<FGuid, FGuid> FFortniteMainBranchObjectVersion::GetSystemGuids()
 	SystemGuids.Add(DevGuids.Niagara_LatestScriptCompileVersion, FGuid("670A2EBDEE1A4794BB77BBFF331EA7EF"));
 	SystemGuids.Add(DevGuids.SkeletalMeshDerivedDataVersion, FGuid("716424BCEC314E969ACFE36D2C2ECE49"));
 	SystemGuids.Add(DevGuids.STATICMESH_DERIVEDDATA_VER, FGuid("3ABF32149F9D4D838750368AB95573BE"));
-	SystemGuids.Add(DevGuids.MaterialTranslationDDCVersion, FGuid("125CF21C881744418BD4A614B3B225FF"));
+	SystemGuids.Add(DevGuids.MaterialTranslationDDCVersion, FGuid("56EE7432297C42BF90244A36A8D8D841"));
 	return SystemGuids;
 }
diff --git a/Engine/Source/Runtime/Engine/Private/Materials/HLSLMaterialTranslator.cpp b/Engine/Source/Runtime/Engine/Private/Materials/HLSLMaterialTranslator.cpp
--- a/Engine/Source/Runtime/Engine/Private/Materials/HLSLMaterialTranslator.cpp
+++ b/Engine/Source/Runtime/Engine/Private/Materials/HLSLMaterialTranslator.cpp
@@ -13458,7 +13458,14 @@ bool FHLSLMaterialTranslator::FSubstrateCompilationContext::SubstrateGenerateDer
 			if ((bHasUnlit || bHasVFogCloud || bHasHair || bHasEye || bHasSLW) && SubstrateMaterialEffectiveClosureCount > 1)
 			{
 				Compiler->Errorf(TEXT("Unlit, Fog/Cloud, Hair or SingleLayerWater must result in a single one Closure (Found %d). See %s (asset: %s).\r\n"), SubstrateMaterialEffectiveClosureCount, *CompilerMaterial->GetDebugName(), *CompilerMaterial->GetAssetPath().ToString());
-				// Even though technically we could support Unlit combined with Slab.
+				// Even though technically we could support Unlit parameter blending with Slab.
+				return false;
+			}
+
+			if (bHasDecal && bHasUnlit)
+			{
+				// Because Unlit does not support parameter blending.
+				Compiler->Errorf(TEXT("Decals only support Substrate Slabs and Shading Model nodes, not Unlit node. See %s (asset: %s).\r\n"), *CompilerMaterial->GetDebugName(), *CompilerMaterial->GetAssetPath().ToString());
 				return false;
 			}
 
