[HWRT] Enable Reference Based Residency by default.

diff --git a/Engine/Source/Runtime/Engine/Private/Rendering/RayTracingGeometryManager.cpp b/Engine/Source/Runtime/Engine/Private/Rendering/RayTracingGeometryManager.cpp
--- a/Engine/Source/Runtime/Engine/Private/Rendering/RayTracingGeometryManager.cpp
+++ b/Engine/Source/Runtime/Engine/Private/Rendering/RayTracingGeometryManager.cpp
@@ -42,8 +42,8 @@ static TAutoConsoleVariable<int32> CVarRayTracingEnable(
 
 static TAutoConsoleVariable<bool> CVarRayTracingUseReferenceBasedResidency(
 	TEXT("r.RayTracing.UseReferenceBasedResidency"),
-	false,
-	TEXT("(EXPERIMENTAL) Whether raytracing geometries should be resident or evicted based on whether they're referenced in TLAS"),
+	true,
+	TEXT("Whether raytracing geometries should be resident or evicted based on whether they're referenced in TLAS."),
 	FConsoleVariableDelegate::CreateLambda([](IConsoleVariable* InVariable)
 		{
 			FGlobalComponentRecreateRenderStateContext Context;
