[HWRT] Fix SplineVertexFactory not being initialized for ray tracing proxies.

#test Cooked project with Nanite fallback meshes stripped + level with Nanite spline mesh

diff --git a/Engine/Source/Runtime/Engine/Private/SplineMeshSceneProxy.cpp b/Engine/Source/Runtime/Engine/Private/SplineMeshSceneProxy.cpp
--- a/Engine/Source/Runtime/Engine/Private/SplineMeshSceneProxy.cpp
+++ b/Engine/Source/Runtime/Engine/Private/SplineMeshSceneProxy.cpp
@@ -100,6 +100,17 @@ FSplineMeshSceneProxy::FSplineMeshSceneProxy(const FStaticMeshSceneProxyDesc& In
 			FSplineMeshSceneProxyDesc::InitVertexFactory(InMeshDesc.GetStaticMesh(), FeatureLevel, LODIndex, InMeshDesc.LODData[LODIndex].OverrideVertexColors);
 		}
 	}
+
+#if RHI_RAYTRACING
+	if (bSupportRayTracing)
+	{
+		const int32 NumRayTracingLods = InMeshDesc.GetStaticMesh() && !InMeshDesc.GetStaticMesh()->GetRenderData()->RayTracingProxy->bUsingRenderingLODs ? InMeshDesc.GetStaticMesh()->GetRenderData()->RayTracingProxy->LODs.Num() : 0;
+		for (int32 LODIndex = 0; LODIndex < NumRayTracingLods; LODIndex++)
+		{
+			FSplineMeshSceneProxyDesc::InitRayTracingProxyVertexFactory(InMeshDesc.GetStaticMesh(), FeatureLevel, LODIndex, nullptr);
+		}
+	}
+#endif
 }
 
 SIZE_T FSplineMeshSceneProxy::GetTypeHash() const
@@ -355,6 +366,12 @@ FNaniteSplineMeshSceneProxy::FNaniteSplineMeshSceneProxy(const Nanite::FMaterial
 				FSplineMeshSceneProxyDesc::InitVertexFactory(InMeshDesc.GetStaticMesh(), FeatureLevel, LODIndex, InMeshDesc.LODData[LODIndex].OverrideVertexColors);
 			}
 		}
+
+		const int32 NumRayTracingLods = InMeshDesc.GetStaticMesh() && !InMeshDesc.GetStaticMesh()->GetRenderData()->RayTracingProxy->bUsingRenderingLODs ? InMeshDesc.GetStaticMesh()->GetRenderData()->RayTracingProxy->LODs.Num() : 0;
+		for (int32 LODIndex = 0; LODIndex < NumRayTracingLods; LODIndex++)
+		{
+			FSplineMeshSceneProxyDesc::InitRayTracingProxyVertexFactory(InMeshDesc.GetStaticMesh(), FeatureLevel, LODIndex, nullptr);
+		}
 	}
 #endif
 }
@@ -405,7 +422,7 @@ void FNaniteSplineMeshSceneProxy::SetupFallbackRayTracingMaterials(int32 LODInde
 	Nanite::FSceneProxy::SetupFallbackRayTracingMaterials(LODIndex, OutMaterials);
 
 	// set up the vertex factories
-	const FStaticMeshVertexFactories& VFs = RenderData->LODVertexFactories[LODIndex];
+	const FStaticMeshVertexFactories& VFs = (*RenderData->RayTracingProxy->LODVertexFactories)[LODIndex];
 	for (auto& MeshBatch : OutMaterials)
 	{
 		MeshBatch.VertexFactory = MeshBatch.Elements[0].bUserDataIsColorVertexBuffer ? VFs.SplineVertexFactoryOverrideColorVertexBuffer : VFs.SplineVertexFactory;
diff --git a/Engine/Source/Runtime/Engine/Private/SplineMeshSceneProxyDesc.cpp b/Engine/Source/Runtime/Engine/Private/SplineMeshSceneProxyDesc.cpp
--- a/Engine/Source/Runtime/Engine/Private/SplineMeshSceneProxyDesc.cpp
+++ b/Engine/Source/Runtime/Engine/Private/SplineMeshSceneProxyDesc.cpp
@@ -171,24 +171,16 @@ extern void InitSplineMeshVertexFactoryComponents(
 	bool bOverrideColorVertexBuffer,
 	FLocalVertexFactory::FDataType& OutData);
 
-void FSplineMeshSceneProxyDesc::InitVertexFactory(UStaticMesh* Mesh, const ERHIFeatureLevel::Type FeatureLevel, int32 InLODIndex, FColorVertexBuffer* InOverrideColorVertexBuffer)
-{	
-	if (Mesh == nullptr)
-	{
-		return;
-	}
-
-	FStaticMeshLODResources& LODRenderData = Mesh->GetRenderData()->LODResources[InLODIndex];
-	FStaticMeshVertexFactories& VertexFactories = Mesh->GetRenderData()->LODVertexFactories[InLODIndex];
-
+static void InitSplineVertexFactory_Internal(FStaticMeshVertexFactories& VertexFactories, const FStaticMeshVertexBuffers& VertexBuffers, int32 LightMapCoordinateIndex, const ERHIFeatureLevel::Type FeatureLevel, FColorVertexBuffer* InOverrideColorVertexBuffer)
+{
 	// Skip LODs that have their render data stripped (eg. platform MinLod settings)
-	if (LODRenderData.VertexBuffers.StaticMeshVertexBuffer.GetNumVertices() == 0)
+	if (VertexBuffers.StaticMeshVertexBuffer.GetNumVertices() == 0)
 	{
 		return;
 	}
 
 	bool bOverrideColorVertexBuffer = !!InOverrideColorVertexBuffer;
-	
+
 	if ((VertexFactories.SplineVertexFactory && !bOverrideColorVertexBuffer) || (VertexFactories.SplineVertexFactoryOverrideColorVertexBuffer && bOverrideColorVertexBuffer))
 	{
 		// we already have it
@@ -205,16 +197,39 @@ void FSplineMeshSceneProxyDesc::InitVertexFactory(UStaticMesh* Mesh, const ERHIF
 		VertexFactories.SplineVertexFactory = VertexFactory;
 	}
 
-	int32 LightMapCoordinateIndex = Mesh->GetLightMapCoordinateIndex();
 	// Initialize the static mesh's vertex factory.
 	ENQUEUE_RENDER_COMMAND(InitSplineMeshVertexFactory)(
-		[VertexFactory, &LODRenderData, bOverrideColorVertexBuffer, LightMapCoordinateIndex](FRHICommandListBase& RHICmdList)
+		[VertexFactory, &VertexBuffers, bOverrideColorVertexBuffer, LightMapCoordinateIndex](FRHICommandListBase& RHICmdList)
+		{
+			FLocalVertexFactory::FDataType Data;
+			InitSplineMeshVertexFactoryComponents(VertexBuffers, VertexFactory, LightMapCoordinateIndex, bOverrideColorVertexBuffer, Data);
+			VertexFactory->SetData(RHICmdList, Data);
+			VertexFactory->InitResource(RHICmdList);
+		});
+}
+
+void FSplineMeshSceneProxyDesc::InitVertexFactory(UStaticMesh* Mesh, const ERHIFeatureLevel::Type FeatureLevel, int32 InLODIndex, FColorVertexBuffer* InOverrideColorVertexBuffer)
+{	
+	if (Mesh == nullptr)
+	{
+		return;
+	}
+
+	FStaticMeshRenderData* RenderData = Mesh->GetRenderData();
+
+	InitSplineVertexFactory_Internal(RenderData->LODVertexFactories[InLODIndex], RenderData->LODResources[InLODIndex].VertexBuffers, Mesh->GetLightMapCoordinateIndex(), FeatureLevel, InOverrideColorVertexBuffer);
+}
+
+void FSplineMeshSceneProxyDesc::InitRayTracingProxyVertexFactory(UStaticMesh* Mesh, const ERHIFeatureLevel::Type FeatureLevel, int32 InLODIndex, FColorVertexBuffer* InOverrideColorVertexBuffer)
+{
+	if (Mesh == nullptr || Mesh->GetRenderData()->RayTracingProxy->bUsingRenderingLODs)
 	{
-		FLocalVertexFactory::FDataType Data;
-		InitSplineMeshVertexFactoryComponents(LODRenderData.VertexBuffers, VertexFactory, LightMapCoordinateIndex, bOverrideColorVertexBuffer, Data);
-		VertexFactory->SetData(RHICmdList, Data);
-		VertexFactory->InitResource(RHICmdList);
-	});
+		return;
+	}
+
+	FStaticMeshRenderData* RenderData = Mesh->GetRenderData();
+
+	InitSplineVertexFactory_Internal((*RenderData->RayTracingProxy->LODVertexFactories)[InLODIndex], *RenderData->RayTracingProxy->LODs[InLODIndex].VertexBuffers, Mesh->GetLightMapCoordinateIndex(), FeatureLevel, InOverrideColorVertexBuffer);
 }
 
 FBox FSplineMeshSceneProxyDesc::ComputeDistortedBounds(UStaticMesh* StaticMesh, const FTransform& InLocalToWorld, const FBoxSphereBounds& InMeshBounds, const FBoxSphereBounds* InBoundsToDistort) const
diff --git a/Engine/Source/Runtime/Engine/Public/SplineMeshSceneProxyDesc.h b/Engine/Source/Runtime/Engine/Public/SplineMeshSceneProxyDesc.h
--- a/Engine/Source/Runtime/Engine/Public/SplineMeshSceneProxyDesc.h
+++ b/Engine/Source/Runtime/Engine/Public/SplineMeshSceneProxyDesc.h
@@ -30,4 +30,5 @@ struct FSplineMeshSceneProxyDesc
 	FTransform CalcSliceTransformAtSplineOffset(const float Alpha, const float MinT, const float MaxT) const;
 	
 	static void InitVertexFactory(UStaticMesh* Mesh, const ERHIFeatureLevel::Type FeatureLevel, int32 InLODIndex, FColorVertexBuffer* InOverrideColorVertexBuffer);
+	static void InitRayTracingProxyVertexFactory(UStaticMesh* Mesh, const ERHIFeatureLevel::Type FeatureLevel, int32 InLODIndex, FColorVertexBuffer* InOverrideColorVertexBuffer);
 };
\ No newline at end of file
