Prevent crash when passing an invalid class to UAnimInstance::PerformLinkedLayerOverlayOperation by returning early

diff --git a/Engine/Source/Runtime/Engine/Private/Animation/AnimInstance.cpp b/Engine/Source/Runtime/Engine/Private/Animation/AnimInstance.cpp
--- a/Engine/Source/Runtime/Engine/Private/Animation/AnimInstance.cpp
+++ b/Engine/Source/Runtime/Engine/Private/Animation/AnimInstance.cpp
@@ -3412,6 +3412,13 @@ void UAnimInstance::PerformLinkedLayerOverlayOperation(TSubclassOf<UAnimInstance
 	if (IAnimClassInterface* AnimBlueprintClass = IAnimClassInterface::GetFromClass(GetClass()))
 	{
 		UClass* NewClass = InClass.Get();
+		if (NewClass && IAnimClassInterface::GetFromClass(NewClass) == nullptr)
+		{
+			// Invalid layer class (i.e passing in a native C++ class instead of an animation blueprint)
+			UE_LOG(LogAnimation, Warning, TEXT("Performing linked layer operations with an invalid layer class: %s"), *NewClass->GetName());
+			return;
+		}
+
 		USkeletalMeshComponent* MeshComp = GetSkelMeshComponent();
 
 		MeshComp->ForEachAnimInstance([](UAnimInstance* InInstance)
