Use correct test for UMaterialExpressionParameter::AllowedIn to ensure subclass parameter types can be added to a materail graph.
Revert back to returning false in UMaterialExpressionComment::AllowedIn as the UI special cases the comments and does not allow direct creation of comment nodes.

diff --git a/Engine/Source/Runtime/Engine/Public/Materials/MaterialExpressionComment.h b/Engine/Source/Runtime/Engine/Public/Materials/MaterialExpressionComment.h
--- a/Engine/Source/Runtime/Engine/Public/Materials/MaterialExpressionComment.h
+++ b/Engine/Source/Runtime/Engine/Public/Materials/MaterialExpressionComment.h
@@ -58,7 +58,7 @@ public:
 	virtual void GetCaption(TArray<FString>& OutCaptions) const override;
 	virtual bool MatchesSearchQuery( const TCHAR* SearchQuery ) override;
 #endif // WITH_EDITOR
-	virtual bool IsAllowedIn(const UObject* MaterialOrFunction) const override { return true; }
+	virtual bool IsAllowedIn(const UObject* MaterialOrFunction) const override { return false; }
 	//~ End UMaterialExpression Interface
 };
 
diff --git a/Engine/Source/Runtime/Engine/Public/Materials/MaterialExpressionParameter.h b/Engine/Source/Runtime/Engine/Public/Materials/MaterialExpressionParameter.h
--- a/Engine/Source/Runtime/Engine/Public/Materials/MaterialExpressionParameter.h
+++ b/Engine/Source/Runtime/Engine/Public/Materials/MaterialExpressionParameter.h
@@ -61,9 +61,11 @@ class UMaterialExpressionParameter : public UMaterialExpression
 
 	virtual void Build(MIR::FEmitter& Emitter) override;
 #endif
-	
-	virtual bool IsAllowedIn(const UObject* MaterialOrFunction) const override { return true; }
-
+	virtual bool IsAllowedIn(const UObject* MaterialOrFunction) const
+	{
+		// UMaterialExpressionParameter is not allowed, but subclasses of UMaterialExpressionParameter are allowed.
+		return GetClass() != UMaterialExpressionParameter::StaticClass();
+	}
 	//~ End UMaterialExpression Interface
 
 	virtual FGuid& GetParameterExpressionId() override
