Expand incoming RPC transactional scope to cover all of ProcessBunch, as some object dereferences weren't covered before.

diff --git a/Engine/Source/Runtime/Engine/Classes/Engine/ActorChannel.h b/Engine/Source/Runtime/Engine/Classes/Engine/ActorChannel.h
--- a/Engine/Source/Runtime/Engine/Classes/Engine/ActorChannel.h
+++ b/Engine/Source/Runtime/Engine/Classes/Engine/ActorChannel.h
@@ -487,6 +487,9 @@ private:
 	/** Used internally by ProcessQueuedBunches to check time limits */
 	bool ShouldProcessAllQueuedBunches(float CurrentTimeSeconds);
 
+	/** Handles the internal logic for processing bunches */
+	void ProcessBunchInternal(FInBunch& Bunch);
+
 private:
 
 	// TODO: It would be nice to merge the tracking of these with PendingGuidResolves, to not duplicate memory,
diff --git a/Engine/Source/Runtime/Engine/Classes/Engine/NetDriver.h b/Engine/Source/Runtime/Engine/Classes/Engine/NetDriver.h
--- a/Engine/Source/Runtime/Engine/Classes/Engine/NetDriver.h
+++ b/Engine/Source/Runtime/Engine/Classes/Engine/NetDriver.h
@@ -2322,11 +2322,18 @@ public:
 	/** Called internally by channels to track how many hit net.QueuedBunchTimeFailsafeSeconds */
 	void AddQueuedBunchFailsafeChannel() { ++QueuedBunchFailsafeNumChannels; }
 
+	/** Sets whether this driver should run certain operations via UE::RemoteExecutor::ExecuteTransactional */
 	void SetReplicateTransactionally(bool bInReplicateTransactionally)
 	{
 		bReplicateTransactionally = bInReplicateTransactionally;
 	}
 
+	/** Returns true if this driver should run certain operations via UE::RemoteExecutor::ExecuteTransactional */
+	bool ShouldReplicateTransactionally() const
+	{
+		return bReplicateTransactionally;
+	}
+
 	/** Returns true if references to UObjects should be serialized as FRemoteObjectReference */
 	bool IsUsingRemoteObjectReferences() const { return bUseRemoteObjectRefs; }
 
diff --git a/Engine/Source/Runtime/Engine/Private/DataChannel.cpp b/Engine/Source/Runtime/Engine/Private/DataChannel.cpp
--- a/Engine/Source/Runtime/Engine/Private/DataChannel.cpp
+++ b/Engine/Source/Runtime/Engine/Private/DataChannel.cpp
@@ -3224,7 +3224,7 @@ void UActorChannel::ReceivedBunch( FInBunch & Bunch )
 	ProcessBunch(Bunch);
 }
 
-void UActorChannel::ProcessBunch( FInBunch & Bunch )
+void UActorChannel::ProcessBunchInternal(FInBunch & Bunch)
 {
 	if ( Broken )
 	{
@@ -3492,6 +3492,24 @@ void UActorChannel::ProcessBunch( FInBunch & Bunch )
 	}
 }
 
+void UActorChannel::ProcessBunch(FInBunch& Bunch)
+{
+#if UE_WITH_REMOTE_OBJECT_HANDLE
+	if (Connection && Connection->Driver && Connection->Driver->ShouldReplicateTransactionally())
+	{
+		static FName TransactionalWorkName(TEXT("UActorChannel::ProcessBunch"));
+		UE::RemoteExecutor::ExecuteTransactional(TransactionalWorkName, [this, &Bunch]()
+			{
+				ProcessBunchInternal(Bunch);
+			});
+	}
+	else
+#endif // UE_WITH_REMOTE_OBJECT_HANDLE
+	{
+		ProcessBunchInternal(Bunch);
+	}
+}
+
 // Helper class to downgrade a non owner of an actor to simulated while replicating
 class FScopedRoleDowngrade
 {
diff --git a/Engine/Source/Runtime/Engine/Private/DataReplication.cpp b/Engine/Source/Runtime/Engine/Private/DataReplication.cpp
--- a/Engine/Source/Runtime/Engine/Private/DataReplication.cpp
+++ b/Engine/Source/Runtime/Engine/Private/DataReplication.cpp
@@ -1175,23 +1175,7 @@ bool FObjectReplicator::ReceivedBunch(FNetBitReader& Bunch, const FReplicationFl
 		{
 			bool bDelayFunction = false;
 			TSet<FNetworkGUID> UnmappedGuids;
-
-			bool bSuccess = false;
-
-#if UE_WITH_REMOTE_OBJECT_HANDLE
-			if (Connection->Driver->bReplicateTransactionally)
-			{
-				static FName TransactionalWorkName(TEXT("FObjectReplicator ReceivedRPC"));
-				UE::RemoteExecutor::ExecuteTransactional(TransactionalWorkName, [this, &bSuccess, &Reader, RepFlags, FieldCache, bCanDelayRPCs, &bDelayFunction, &UnmappedGuids]()
-					{
-						bSuccess = ReceivedRPC(Reader, RepFlags, FieldCache, bCanDelayRPCs, bDelayFunction, UnmappedGuids);
-					});
-			}
-			else
-#endif
-			{
-				bSuccess = ReceivedRPC(Reader, RepFlags, FieldCache, bCanDelayRPCs, bDelayFunction, UnmappedGuids);
-			}
+			bool bSuccess = ReceivedRPC(Reader, RepFlags, FieldCache, bCanDelayRPCs, bDelayFunction, UnmappedGuids);
 
 			if (!bSuccess)
 			{
