Add virtual texture specific lod bias to texture group lod settings.
The new setting is named LODBias_VT and is used instead of LODBias when compiling a virtual texture.

diff --git a/Engine/Source/Editor/DetailCustomizations/Private/TextureLODSettingsDetails.cpp b/Engine/Source/Editor/DetailCustomizations/Private/TextureLODSettingsDetails.cpp
--- a/Engine/Source/Editor/DetailCustomizations/Private/TextureLODSettingsDetails.cpp
+++ b/Engine/Source/Editor/DetailCustomizations/Private/TextureLODSettingsDetails.cpp
@@ -201,6 +201,26 @@ void FTextureLODGroupLayout::GenerateChildContent(IDetailChildrenBuilder& Childr
 			];
 	}
 
+	{
+		ChildrenBuilder.AddCustomRow(LOCTEXT("LODBiasVT_Row", "LOD Bias VT"))
+			.NameContent()
+			[
+				SNew(STextBlock)
+					.Font(IDetailLayoutBuilder::GetDetailFont())
+					.Text(LOCTEXT("LODBiasVT_Name", "LOD Bias VT"))
+			]
+			.ValueContent()
+			[
+				SNew(SSpinBox<int32>)
+					.Font(IDetailLayoutBuilder::GetDetailFont())
+					.MinValue(-MAX_TEXTURE_MIP_COUNT)
+					.MaxValue(MAX_TEXTURE_MIP_COUNT)
+					.Value(this, &FTextureLODGroupLayout::GetLODBiasVT)
+					.OnValueChanged(this, &FTextureLODGroupLayout::OnLODBiasVTChanged)
+					.OnValueCommitted(this, &FTextureLODGroupLayout::OnLODBiasVTCommitted)
+			];
+	}
+
 
 	// Filter properties
 
@@ -369,6 +389,21 @@ void FTextureLODGroupLayout::OnLODBiasCommitted(int32 NewValue, ETextCommit::Typ
 	OnLODBiasChanged(NewValue);
 }
 
+int32 FTextureLODGroupLayout::GetLODBiasVT() const
+{
+	return LodGroup->LODBias_VT;
+}
+
+void FTextureLODGroupLayout::OnLODBiasVTChanged(int32 NewValue)
+{
+	LodGroup->LODBias_VT = NewValue;
+}
+
+void FTextureLODGroupLayout::OnLODBiasVTCommitted(int32 NewValue, ETextCommit::Type TextCommitType)
+{
+	OnLODBiasVTChanged(NewValue);
+}
+
 
 TSharedRef<SWidget> FTextureLODGroupLayout::MakeMinMagFilterComboWidget(TSharedPtr<FName> InItem)
 {
diff --git a/Engine/Source/Editor/DetailCustomizations/Private/TextureLODSettingsDetails.h b/Engine/Source/Editor/DetailCustomizations/Private/TextureLODSettingsDetails.h
--- a/Engine/Source/Editor/DetailCustomizations/Private/TextureLODSettingsDetails.h
+++ b/Engine/Source/Editor/DetailCustomizations/Private/TextureLODSettingsDetails.h
@@ -63,6 +63,11 @@ private:
 	void OnLODBiasChanged(int32 NewValue);
 	void OnLODBiasCommitted(int32 NewValue, ETextCommit::Type TextCommitType);
 
+	// Controls for the LOD Bias VT property editing
+	int32 GetLODBiasVT() const;
+	void OnLODBiasVTChanged(int32 NewValue);
+	void OnLODBiasVTCommitted(int32 NewValue, ETextCommit::Type TextCommitType);
+
 	// Controls for the MinMag Filter property editing
 	TSharedRef<SWidget> MakeMinMagFilterComboWidget(TSharedPtr<FName> InItem);
 	void OnMinMagFilterChanged(TSharedPtr<FName> NewSelection, ESelectInfo::Type SelectInfo);
diff --git a/Engine/Source/Runtime/Engine/Classes/Engine/TextureLODSettings.h b/Engine/Source/Runtime/Engine/Classes/Engine/TextureLODSettings.h
--- a/Engine/Source/Runtime/Engine/Classes/Engine/TextureLODSettings.h
+++ b/Engine/Source/Runtime/Engine/Classes/Engine/TextureLODSettings.h
@@ -48,6 +48,9 @@ struct FTextureLODGroup
 	UPROPERTY()
 	int32 LODBias_Smallest = -1;
 	
+	UPROPERTY()
+	int32 LODBias_VT = 0;
+
 	/** Number of mip-levels that can be streamed. -1 means all mips can stream.	*/
 	UPROPERTY()
 	int32 NumStreamedMips = -1;
diff --git a/Engine/Source/Runtime/Engine/Private/Streaming/StreamingManagerTexture.cpp b/Engine/Source/Runtime/Engine/Private/Streaming/StreamingManagerTexture.cpp
--- a/Engine/Source/Runtime/Engine/Private/Streaming/StreamingManagerTexture.cpp
+++ b/Engine/Source/Runtime/Engine/Private/Streaming/StreamingManagerTexture.cpp
@@ -2858,23 +2858,25 @@ bool FRenderAssetStreamingManager::HandleInvestigateRenderAssetCommand(const TCH
 
 						// LOD group Bias : see UTextureLODSettings::CalculateLODBias(), included in CachedCombinedLODBias
 						const FTextureLODGroup& LODGroupInfo = UDeviceProfileManager::Get().GetActiveProfile()->GetTextureLODSettings()->TextureLODGroups[StreamingRenderAsset.LODGroup];
-						if (LODGroupInfo.LODBias)
+						const bool bIsVirtualTexture = Texture && Texture->IsCurrentlyVirtualTextured();
+						const int32 GroupLODBias = bIsVirtualTexture ? LODGroupInfo.LODBias_VT : LODGroupInfo.LODBias;
+						if (GroupLODBias)
 						{
 							if (FPlatformProperties::RequiresCookedData())
 							{
-								BiasDesc += FString::Printf(TEXT(" [LODGroup.Bias:0(%d)]"), LODGroupInfo.LODBias);
+								BiasDesc += FString::Printf(TEXT(" [LODGroup.Bias:0(%d)]"), GroupLODBias);
 							}
 							else
 							{
-								BiasDesc += FString::Printf(TEXT(" [LODGroup.Bias:%d]"), LODGroupInfo.LODBias);
-								CumuBias += LODGroupInfo.LODBias;
+								BiasDesc += FString::Printf(TEXT(" [LODGroup.Bias:%d]"), GroupLODBias);
+								CumuBias += GroupLODBias;
 							}
 						}
 
 						// LOD group MaxResolution clamp : see UTextureLODSettings::CalculateLODBias(), included in CachedCombinedLODBias
 						const int32 MipCountBeforeMaxRes = ResourceState.MaxNumLODs - RenderAsset->NumCinematicMipLevels -
 							(StreamingRenderAsset.LODGroup == TEXTUREGROUP_UI ? GUITextureLODBias : 0) - 
-							(FPlatformProperties::RequiresCookedData() ? 0 : (LODGroupInfo.LODBias + (Texture ? Texture->LODBias : 0)));
+							(FPlatformProperties::RequiresCookedData() ? 0 : (GroupLODBias + (Texture ? Texture->LODBias : 0)));
 						const int32 MaxResBias = MipCountBeforeMaxRes - (LODGroupInfo.MaxLODMipCount + 1);
 						if (MaxResBias > 0)
 						{
diff --git a/Engine/Source/Runtime/Engine/Private/TextureLODSettings.cpp b/Engine/Source/Runtime/Engine/Private/TextureLODSettings.cpp
--- a/Engine/Source/Runtime/Engine/Private/TextureLODSettings.cpp
+++ b/Engine/Source/Runtime/Engine/Private/TextureLODSettings.cpp
@@ -198,11 +198,7 @@ int32 UTextureLODSettings::CalculateLODBias(int32 Width, int32 Height, int32 Max
 		// When cooking, LODBias and LODGroupInfo.LODBias are taken into account to strip the top mips.
 		// Considering them again here would apply them twice.
 		UsedLODBias	+= LODBias;
-		if (!bVirtualTexture)
-		{
-			// Don't include LOD group's bias for VT, could include a separate bias that applies only to VT if needed
-			UsedLODBias += LODGroupInfo.LODBias;
-		}
+		UsedLODBias += bVirtualTexture ? LODGroupInfo.LODBias_VT : LODGroupInfo.LODBias;
 	}
 	
 	if (LODGroup == TEXTUREGROUP_UI)
