Fix ISMs not rendering when Nanite fallback mesh is stripped from cooked build.
- Remove problematic validation from UInstancedStaticMeshComponent::CreateSceneProxy() since it's already done in UStaticMeshComponent::CreateSceneProxy() where it correctly handles stripped fallback and also logs relevant warnings instead of failing silently.

diff --git a/Engine/Source/Runtime/Engine/Private/InstancedStaticMesh.cpp b/Engine/Source/Runtime/Engine/Private/InstancedStaticMesh.cpp
--- a/Engine/Source/Runtime/Engine/Private/InstancedStaticMesh.cpp
+++ b/Engine/Source/Runtime/Engine/Private/InstancedStaticMesh.cpp
@@ -2382,19 +2382,15 @@ FPrimitiveSceneProxy* UInstancedStaticMeshComponent::CreateSceneProxy()
 
 	PrimitiveInstanceDataManager.ResetComponentDirtyTracking();
 
-	// Verify that both mesh and instance data is valid before using it.
-	const bool bIsMeshAndInstanceDataValid =
+	// Verify that instance data is valid before using it.
+	const bool bIsInstanceDataValid =
 #if WITH_EDITOR
 		bIsInstanceDataApplyCompleted && 
 #endif
-		// make sure we have instances - or have explicitly permitted creating emtpty ISMs
-		(PerInstanceSMData.Num() > 0 || CVarAllowCreateEmptyISMs.GetValueOnGameThread()) &&
-		// make sure we have an actual static mesh
-		GetStaticMesh() &&
-		GetStaticMesh()->IsCompiling() == false &&
-		GetStaticMesh()->HasValidRenderData();
-
-	if (!bIsMeshAndInstanceDataValid)
+		// make sure we have instances - or have explicitly permitted creating empty ISMs
+		(PerInstanceSMData.Num() > 0 || CVarAllowCreateEmptyISMs.GetValueOnGameThread());
+
+	if (!bIsInstanceDataValid)
 	{
 		return nullptr;
 	}
