Extend WorldPartition StreamingPerformance to be more generic (previously was only reported for blocking sources).
Added another level of granularity (Immediate) to EWorldPartitionStreamingPerformance which triggers when we are within unloaded cell content bounds

#tests RR, local tests to make sure streaming perf is reported correctly with current hash set. (Spatial hash needs further testing)

[RN] WorldPartition - GetStreamingPerformance now reports streaming performance also for non blocking cells/sources. Poor streaming performance caused by blocking cells/sources take preference over poor streaming performance from non blocking cells/sources. The thresholds for reporting streaming performance from non blocking sources can be configured via "wp.Runtime.SlowStreamingRatio" and "wp.Runtime.SlowStreamingWarningFactor". GetStreamingPerformance also reports an additional enum for streaming performance ("Immediate") that triggers when sources are within unloaded cell content bounds.

diff --git a/Engine/Source/Runtime/Engine/Private/WorldPartition/WorldPartitionRuntimeCellData.cpp b/Engine/Source/Runtime/Engine/Private/WorldPartition/WorldPartitionRuntimeCellData.cpp
--- a/Engine/Source/Runtime/Engine/Private/WorldPartition/WorldPartitionRuntimeCellData.cpp
+++ b/Engine/Source/Runtime/Engine/Private/WorldPartition/WorldPartitionRuntimeCellData.cpp
@@ -11,6 +11,8 @@ UWorldPartitionRuntimeCellData::UWorldPartitionRuntimeCellData(const FObjectInit
 	, bCachedWasRequestedByBlockingSource(false)
 	, CachedMinSquareDistanceToBlockingSource(MAX_dbl)
 	, CachedMinBlockOnSlowStreamingRatio(MAX_flt)
+	, CachedMinSquareDistanceToSource(MAX_dbl)
+	, CachedMinSlowStreamingRatio(MAX_flt)
 	, CachedMinSpatialSortingPriority(MAX_dbl)
 	, CachedSourceInfoEpoch(MIN_int32)
 	, ContentBounds(ForceInit)
@@ -42,6 +44,8 @@ void UWorldPartitionRuntimeCellData::ResetStreamingSourceInfo(const FWorldPartit
 	bCachedWasRequestedByBlockingSource = false;
 	CachedMinSquareDistanceToBlockingSource = MAX_dbl;
 	CachedMinBlockOnSlowStreamingRatio = MAX_flt;
+	CachedMinSquareDistanceToSource = MAX_dbl;
+	CachedMinSlowStreamingRatio = MAX_flt;
 	CachedMinSpatialSortingPriority = MAX_dbl;
 	CachedSourceInfoEpoch = Context.GetUpdateStreamingStateEpoch();
 }
@@ -77,6 +81,11 @@ void UWorldPartitionRuntimeCellData::AppendStreamingSourceInfo(const FWorldParti
 		CachedMinBlockOnSlowStreamingRatio = FMath::Min(CachedMinBlockOnSlowStreamingRatio, BlockOnSlowStreamingRatio);
 	}
 
+	CachedMinSquareDistanceToSource = FMath::Min(CellToSourceSquareDistance, CachedMinSquareDistanceToSource);
+
+	const double SlowStreamingRatio = FMath::Sqrt(CachedMinSquareDistanceToSource) / SourceShape.GetRadius();
+	CachedMinSlowStreamingRatio = FMath::Min(CachedMinSlowStreamingRatio, SlowStreamingRatio);
+
 	// Compute square distance from cell to source ratio
 	const double SoureDistanceRatio = FMath::Clamp(ContentBounds.ComputeSquaredDistanceToPoint(SourceShape.GetCenter()) / FMath::Square(SourceShape.GetRadius()), 0.0f, 1.0f);
 
diff --git a/Engine/Source/Runtime/Engine/Private/WorldPartition/WorldPartitionRuntimeCellDataSpatialHash.cpp b/Engine/Source/Runtime/Engine/Private/WorldPartition/WorldPartitionRuntimeCellDataSpatialHash.cpp
--- a/Engine/Source/Runtime/Engine/Private/WorldPartition/WorldPartitionRuntimeCellDataSpatialHash.cpp
+++ b/Engine/Source/Runtime/Engine/Private/WorldPartition/WorldPartitionRuntimeCellDataSpatialHash.cpp
@@ -27,21 +27,23 @@ static FAutoConsoleVariableRef CVarRuntimeSpatialHashSortUsingCellPriority(
 UWorldPartitionRuntimeCellDataSpatialHash::UWorldPartitionRuntimeCellDataSpatialHash(const FObjectInitializer& ObjectInitializer)
 	: Super(ObjectInitializer)
 	, Extent(0)
-	, CachedMinSquareDistanceToSource(MAX_dbl)
 	, CachedMinBlockOnSlowStreamingRatio2D(MAX_dbl)
 	, CachedMinSquareDistanceToBlockingSource2D(MAX_dbl)
+	, CachedMinSlowStreamingRatio2D(MAX_dbl)
+	, CachedMinSquareDistanceToSource2D(MAX_dbl)
 {}
 
 void UWorldPartitionRuntimeCellDataSpatialHash::ResetStreamingSourceInfo(const FWorldPartitionStreamingContext& Context) const
 {
 	Super::ResetStreamingSourceInfo(Context);
 
-	CachedMinSquareDistanceToSource = MAX_dbl;
 	CachedSourcePriorityWeights.Reset();
 	CachedSourceSquaredDistances.Reset();
 	CachedInstersectingShapes.Reset();
 	CachedMinBlockOnSlowStreamingRatio2D = MAX_dbl;
 	CachedMinSquareDistanceToBlockingSource2D = MAX_dbl;
+	CachedMinSlowStreamingRatio2D = MAX_dbl;
+	CachedMinSquareDistanceToSource2D = MAX_dbl;
 }
 
 float UWorldPartitionRuntimeCellDataSpatialHash::ComputeSourceToCellAngleFactor(const FSphericalSector& SourceShape) const
@@ -95,6 +97,9 @@ void UWorldPartitionRuntimeCellDataSpatialHash::AppendStreamingSourceInfo(const
 		CachedMinSquareDistanceToBlockingSource2D = FMath::Min(SquareDistance2D, CachedMinSquareDistanceToBlockingSource2D);
 		CachedMinBlockOnSlowStreamingRatio2D = FMath::Min(CachedMinBlockOnSlowStreamingRatio2D, FMath::Sqrt(CachedMinSquareDistanceToBlockingSource2D) / SourceShape.GetRadius());
 	}
+	CachedMinSquareDistanceToSource2D = FMath::Min(SquareDistance2D, CachedMinSquareDistanceToSource2D);
+	CachedMinSlowStreamingRatio2D = FMath::Min(CachedMinSlowStreamingRatio2D, FMath::Sqrt(CachedMinSquareDistanceToSource2D) / SourceShape.GetRadius());
+
 	CachedSourceSquaredDistances.Add(SquareDistance2D);
 	CachedSourcePriorityWeights.Add(1.0f - ((float)Source.Priority / (float)EStreamingSourcePriority::Lowest));
 	CachedInstersectingShapes.Add(SourceShape);
@@ -108,9 +113,10 @@ void UWorldPartitionRuntimeCellDataSpatialHash::MergeStreamingSourceInfo() const
 	check(Count == CachedSourcePriorityWeights.Num());
 	check(Count == CachedInstersectingShapes.Num());
 
-	CachedMinSquareDistanceToSource = Count ? FMath::Min(CachedSourceSquaredDistances) : MAX_dbl;
 	CachedMinBlockOnSlowStreamingRatio = CachedMinBlockOnSlowStreamingRatio2D;
 	CachedMinSquareDistanceToBlockingSource = CachedMinSquareDistanceToBlockingSource2D;
+	CachedMinSquareDistanceToSource = CachedMinSquareDistanceToSource2D;
+	CachedMinSlowStreamingRatio = CachedMinSlowStreamingRatio2D;
 
 	if (Count)
 	{
@@ -169,7 +175,7 @@ int32 UWorldPartitionRuntimeCellDataSpatialHash::SortCompare(const UWorldPartiti
 					const double Diff = CachedSourceSortingDistance - Other->CachedSourceSortingDistance;
 					if (FMath::IsNearlyZero(Diff))
 					{
-						const double RawDistanceDiff = CachedMinSquareDistanceToSource - Other->CachedMinSquareDistanceToSource;
+						const double RawDistanceDiff = CachedMinSquareDistanceToSource2D - Other->CachedMinSquareDistanceToSource2D;
 						Result = RawDistanceDiff < 0 ? -1 : (RawDistanceDiff > 0.f ? 1 : 0);
 					}
 					else
diff --git a/Engine/Source/Runtime/Engine/Private/WorldPartition/WorldPartitionRuntimeHash.cpp b/Engine/Source/Runtime/Engine/Private/WorldPartition/WorldPartitionRuntimeHash.cpp
--- a/Engine/Source/Runtime/Engine/Private/WorldPartition/WorldPartitionRuntimeHash.cpp
+++ b/Engine/Source/Runtime/Engine/Private/WorldPartition/WorldPartitionRuntimeHash.cpp
@@ -24,6 +24,18 @@
 
 #define LOCTEXT_NAMESPACE "WorldPartition"
 
+ENGINE_API float GSlowStreamingRatio = 0.25f;
+static FAutoConsoleVariableRef CVarSlowStreamingRatio(
+	TEXT("wp.Runtime.SlowStreamingRatio"),
+	GSlowStreamingRatio,
+	TEXT("Ratio of DistanceToCell / LoadingRange to use to determine if World Partition streaming is considered to be slow"));
+
+ENGINE_API float GSlowStreamingWarningFactor = 2.f;
+static FAutoConsoleVariableRef CVarSlowStreamingWarningFactor(
+	TEXT("wp.Runtime.SlowStreamingWarningFactor"),
+	GSlowStreamingWarningFactor,
+	TEXT("Factor of wp.Runtime.SlowStreamingRatio we want to start notifying the user"));
+
 ENGINE_API float GBlockOnSlowStreamingRatio = 0.25f;
 static FAutoConsoleVariableRef CVarBlockOnSlowStreamingRatio(
 	TEXT("wp.Runtime.BlockOnSlowStreamingRatio"),
@@ -36,6 +48,23 @@ static FAutoConsoleVariableRef CVarBlockOnSlowStreamingWarningFactor(
 	GBlockOnSlowStreamingWarningFactor,
 	TEXT("Factor of wp.Runtime.BlockOnSlowStreamingRatio we want to start notifying the user"));
 
+const TCHAR* EnumToString(EWorldPartitionStreamingPerformance InStreamingPerformance)
+{
+	switch (InStreamingPerformance)
+	{
+	case EWorldPartitionStreamingPerformance::Immediate:
+		return TEXT("Immediate");
+	case EWorldPartitionStreamingPerformance::Critical:
+		return TEXT("Critical");
+	case EWorldPartitionStreamingPerformance::Slow:
+		return TEXT("Slow");
+	case EWorldPartitionStreamingPerformance::Good:
+		return TEXT("Good");
+	}
+	ensure(false);
+	return TEXT("Unknown");
+}
+
 #if DO_CHECK
 void URuntimeHashExternalStreamingObjectBase::BeginDestroy()
 {
@@ -202,21 +231,52 @@ UWorldPartitionRuntimeCell* UWorldPartitionRuntimeHash::CreateRuntimeCell(UClass
 
 EWorldPartitionStreamingPerformance UWorldPartitionRuntimeHash::GetStreamingPerformanceForCell(const UWorldPartitionRuntimeCell* Cell) const
 {
-	check(Cell->GetBlockOnSlowLoading());
+	return EWorldPartitionStreamingPerformance::Good;
+}
 
-	if (Cell->RuntimeCellData->bCachedWasRequestedByBlockingSource)
+EWorldPartitionStreamingPerformance UWorldPartitionRuntimeHash::GetStreamingPerformanceForCell(const UWorldPartitionRuntimeCell* Cell, bool& bOutShouldBlock) const
+{
+	EWorldPartitionStreamingPerformance BlockingPerformance = EWorldPartitionStreamingPerformance::Good;
+	if (Cell->GetBlockOnSlowLoading() && Cell->RuntimeCellData->bCachedWasRequestedByBlockingSource)
 	{
-		if (Cell->RuntimeCellData->CachedMinBlockOnSlowStreamingRatio < GBlockOnSlowStreamingRatio)
+		if (Cell->RuntimeCellData->CachedMinBlockOnSlowStreamingRatio <= 0.0)
+		{
+			BlockingPerformance = EWorldPartitionStreamingPerformance::Immediate;
+		}
+		else if (Cell->RuntimeCellData->CachedMinBlockOnSlowStreamingRatio < GBlockOnSlowStreamingRatio)
 		{
-			return EWorldPartitionStreamingPerformance::Critical;
+			BlockingPerformance = EWorldPartitionStreamingPerformance::Critical;
 		}
 		else if (Cell->RuntimeCellData->CachedMinBlockOnSlowStreamingRatio < (GBlockOnSlowStreamingRatio * GBlockOnSlowStreamingWarningFactor))
 		{
-			return EWorldPartitionStreamingPerformance::Slow;
+			BlockingPerformance = EWorldPartitionStreamingPerformance::Slow;
 		}
 	}
 
-	return EWorldPartitionStreamingPerformance::Good;
+	EWorldPartitionStreamingPerformance NonBlockingPerformance = EWorldPartitionStreamingPerformance::Good;
+	if (Cell->RuntimeCellData->CachedMinSlowStreamingRatio <= 0.0)
+	{
+		NonBlockingPerformance = EWorldPartitionStreamingPerformance::Immediate;
+	}
+	else if (Cell->RuntimeCellData->CachedMinSlowStreamingRatio < GSlowStreamingRatio)
+	{
+		NonBlockingPerformance = EWorldPartitionStreamingPerformance::Critical;
+	}
+	else if (Cell->RuntimeCellData->CachedMinSlowStreamingRatio < (GSlowStreamingRatio * GSlowStreamingWarningFactor))
+	{
+		NonBlockingPerformance = EWorldPartitionStreamingPerformance::Slow;
+	}
+
+	if (BlockingPerformance >= NonBlockingPerformance)
+	{
+		bOutShouldBlock = BlockingPerformance >= EWorldPartitionStreamingPerformance::Critical;
+		return BlockingPerformance;
+	}
+	else
+	{
+		bOutShouldBlock = false;
+		return NonBlockingPerformance;
+	}
 }
 
 URuntimeHashExternalStreamingObjectBase* UWorldPartitionRuntimeHash::CreateExternalStreamingObject(TSubclassOf<URuntimeHashExternalStreamingObjectBase> InClass, UObject* InOuter, UWorld* InOuterWorld)
@@ -516,8 +576,14 @@ bool UWorldPartitionRuntimeHash::IsCellRelevantFor(bool bClientOnlyVisible) cons
 }
 
 EWorldPartitionStreamingPerformance UWorldPartitionRuntimeHash::GetStreamingPerformance(const TSet<const UWorldPartitionRuntimeCell*>& CellsToActivate) const
+{
+	return EWorldPartitionStreamingPerformance::Good;
+}
+
+EWorldPartitionStreamingPerformance UWorldPartitionRuntimeHash::GetStreamingPerformance(const TSet<const UWorldPartitionRuntimeCell*>& CellsToActivate, bool& bOutShouldBlock) const
 {
 	EWorldPartitionStreamingPerformance StreamingPerformance = EWorldPartitionStreamingPerformance::Good;
+	bOutShouldBlock = false;
 
 	if (!CellsToActivate.IsEmpty() && GetWorld()->bMatchStarted)
 	{
@@ -525,15 +591,17 @@ EWorldPartitionStreamingPerformance UWorldPartitionRuntimeHash::GetStreamingPerf
 
 		for (const UWorldPartitionRuntimeCell* Cell : CellsToActivate)
 		{
-			if (Cell->GetBlockOnSlowLoading() && !Cell->IsAlwaysLoaded() && Cell->GetStreamingStatus() != LEVEL_Visible)
+			if (!Cell->IsAlwaysLoaded() && Cell->GetStreamingStatus() != LEVEL_Visible)
 			{
-				EWorldPartitionStreamingPerformance CellPerformance = GetStreamingPerformanceForCell(Cell);
+				bool bOutShouldBlockCell;
+				EWorldPartitionStreamingPerformance CellPerformance = GetStreamingPerformanceForCell(Cell, bOutShouldBlockCell);
+				bOutShouldBlock |= bOutShouldBlockCell;
 				// Cell Performance is worst than previous cell performance
 				if (CellPerformance > StreamingPerformance)
 				{
 					StreamingPerformance = CellPerformance;
 					// Early out performance is critical
-					if (StreamingPerformance == EWorldPartitionStreamingPerformance::Critical)
+					if (bOutShouldBlockCell && StreamingPerformance >= EWorldPartitionStreamingPerformance::Critical)
 					{
 						return StreamingPerformance;
 					}
diff --git a/Engine/Source/Runtime/Engine/Private/WorldPartition/WorldPartitionRuntimeSpatialHash.cpp b/Engine/Source/Runtime/Engine/Private/WorldPartition/WorldPartitionRuntimeSpatialHash.cpp
--- a/Engine/Source/Runtime/Engine/Private/WorldPartition/WorldPartitionRuntimeSpatialHash.cpp
+++ b/Engine/Source/Runtime/Engine/Private/WorldPartition/WorldPartitionRuntimeSpatialHash.cpp
@@ -1936,33 +1936,69 @@ const FSpatialHashStreamingGrid* UWorldPartitionRuntimeSpatialHash::GetStreaming
 
 EWorldPartitionStreamingPerformance UWorldPartitionRuntimeSpatialHash::GetStreamingPerformanceForCell(const UWorldPartitionRuntimeCell* Cell) const
 {
-	check(Cell->GetBlockOnSlowLoading());
-	const double BlockOnSlowStreamingWarningRatio = GBlockOnSlowStreamingRatio * GBlockOnSlowStreamingWarningFactor;
-	
+	return EWorldPartitionStreamingPerformance::Good;
+}
+
+ENGINE_API EWorldPartitionStreamingPerformance UWorldPartitionRuntimeSpatialHash::GetStreamingPerformanceForCell(const UWorldPartitionRuntimeCell* Cell, bool& bOutShouldBlock) const
+{
 	const UWorldPartitionRuntimeCellDataSpatialHash* CellDataSpatialHash = CastChecked<UWorldPartitionRuntimeCellDataSpatialHash>(Cell->RuntimeCellData);
 	const FSpatialHashStreamingGrid* StreamingGrid = GetStreamingGridByName(CellDataSpatialHash->GridName);
 
 	if (ensure(StreamingGrid))
 	{
 		const float LoadingRange = StreamingGrid->LoadingRange;
-
-		if (CellDataSpatialHash->bCachedWasRequestedByBlockingSource)
+		const double BlockOnSlowStreamingWarningRatio = GBlockOnSlowStreamingRatio * GBlockOnSlowStreamingWarningFactor;
+		const double OnSlowStreamingWarningRatio = GSlowStreamingRatio * GSlowStreamingWarningFactor;
+		EWorldPartitionStreamingPerformance BlockingPerformance = EWorldPartitionStreamingPerformance::Good;
+		EWorldPartitionStreamingPerformance NonBlockingPerformance = EWorldPartitionStreamingPerformance::Good;
+		if (Cell->GetBlockOnSlowLoading() && CellDataSpatialHash->bCachedWasRequestedByBlockingSource)
 		{
 			const double Distance = FMath::Sqrt(CellDataSpatialHash->CachedMinSquareDistanceToBlockingSource) - ((double)StreamingGrid->GetCellSize(CellDataSpatialHash->HierarchicalLevel) / 2);
 
 			const double Ratio = Distance / LoadingRange;
-
-			if (Ratio < GBlockOnSlowStreamingRatio)
+			if (Ratio <= 0.0)
+			{
+				BlockingPerformance = EWorldPartitionStreamingPerformance::Immediate;
+			}
+			else if (Ratio < GBlockOnSlowStreamingRatio)
 			{
-				return EWorldPartitionStreamingPerformance::Critical;
+				BlockingPerformance = EWorldPartitionStreamingPerformance::Critical;
 			}
 			else if (Ratio < BlockOnSlowStreamingWarningRatio)
 			{
-				return EWorldPartitionStreamingPerformance::Slow;
+				BlockingPerformance = EWorldPartitionStreamingPerformance::Slow;
 			}
 		}
+
+		const double Distance = FMath::Sqrt(CellDataSpatialHash->CachedMinSquareDistanceToSource) - ((double)StreamingGrid->GetCellSize(CellDataSpatialHash->HierarchicalLevel) / 2);
+
+		const double Ratio = Distance / LoadingRange;
+		if (Ratio <= 0.0)
+		{
+			NonBlockingPerformance = EWorldPartitionStreamingPerformance::Immediate;
+		}
+		else if (Ratio < GSlowStreamingRatio)
+		{
+			NonBlockingPerformance = EWorldPartitionStreamingPerformance::Critical;
+		}
+		else if (Ratio < OnSlowStreamingWarningRatio)
+		{
+			NonBlockingPerformance = EWorldPartitionStreamingPerformance::Slow;
+		}
+
+		if (BlockingPerformance >= NonBlockingPerformance)
+		{
+			bOutShouldBlock = BlockingPerformance >= EWorldPartitionStreamingPerformance::Critical;
+			return BlockingPerformance;
+		}
+		else
+		{
+			bOutShouldBlock = false;
+			return NonBlockingPerformance;
+		}
 	}
 
+	bOutShouldBlock = false;
 	return EWorldPartitionStreamingPerformance::Good;
 }
 
diff --git a/Engine/Source/Runtime/Engine/Private/WorldPartition/WorldPartitionStreamingPolicy.cpp b/Engine/Source/Runtime/Engine/Private/WorldPartition/WorldPartitionStreamingPolicy.cpp
--- a/Engine/Source/Runtime/Engine/Private/WorldPartition/WorldPartitionStreamingPolicy.cpp
+++ b/Engine/Source/Runtime/Engine/Private/WorldPartition/WorldPartitionStreamingPolicy.cpp
@@ -67,6 +67,7 @@ UWorldPartitionStreamingPolicy::UWorldPartitionStreamingPolicy(const FObjectInit
 	: Super(ObjectInitializer) 
 	, WorldPartition(nullptr)
 	, CurrentStreamingPerformance(EWorldPartitionStreamingPerformance::Good)
+	, bCurrentBlockOnSlowStreaming(false)
 	, bShouldMergeStreamingSourceInfo(false)
 	, bCriticalPerformanceRequestedBlockTillOnWorld(false)
 	, CriticalPerformanceBlockTillLevelStreamingCompletedEpoch(0)
@@ -82,6 +83,7 @@ UWorldPartitionStreamingPolicy::UWorldPartitionStreamingPolicy(const FObjectInit
 #if !UE_BUILD_SHIPPING
 	, OnScreenMessageStartTime(0.0)
 	, OnScreenMessageStreamingPerformance(EWorldPartitionStreamingPerformance::Good)
+	, bOnScreenMessageShouldBlock(false)
 #endif
 {
 	if (!IsTemplate())
@@ -648,7 +650,7 @@ void UWorldPartitionStreamingPolicy::UpdateStreamingStateInternal(const FUpdateS
 
 	if (InParams.RuntimeHash)
 	{
-		OutTargetState.StreamingPerformance = InParams.RuntimeHash->GetStreamingPerformance(InParams.FrameActivateCells);
+		OutTargetState.StreamingPerformance = InParams.RuntimeHash->GetStreamingPerformance(InParams.FrameActivateCells, OutTargetState.bBlockOnSlowStreaming);
 	}
 }
 
@@ -701,7 +703,7 @@ void UWorldPartitionStreamingPolicy::PostUpdateStreamingStateInternal_GameThread
 	}
 
 	// Evaluate streaming performance based on cells that should be activated
-	UpdateStreamingPerformance(InOutTargetState.StreamingPerformance);
+	UpdateStreamingPerformance(InOutTargetState.StreamingPerformance, InOutTargetState.bBlockOnSlowStreaming);
 }
 
 void UWorldPartitionStreamingPolicy::OnPreChangeStreamingContent()
@@ -805,19 +807,18 @@ void UWorldPartitionStreamingPolicy::UpdateDebugCellsStreamingPriority(const TSe
 }
 #endif
 
-void UWorldPartitionStreamingPolicy::UpdateStreamingPerformance(EWorldPartitionStreamingPerformance NewStreamingPerformance)
+void UWorldPartitionStreamingPolicy::UpdateStreamingPerformance(EWorldPartitionStreamingPerformance NewStreamingPerformance, bool bBlockOnSlowStreaming)
 {
 	check(IsInGameThread());
 
 	TRACE_CPUPROFILER_EVENT_SCOPE(UWorldPartitionStreamingPolicy::UpdateStreamingPerformance);
 	UWorld* World = GetWorld();
 	// If we are currently in a blocked loading just reset the on screen message time and return
-	if (CurrentStreamingPerformance == EWorldPartitionStreamingPerformance::Critical && IsInBlockTillLevelStreamingCompleted())
+	if (CurrentStreamingPerformance >= EWorldPartitionStreamingPerformance::Critical && IsInBlockTillLevelStreamingCompleted())
 	{
 #if !UE_BUILD_SHIPPING
 		OnScreenMessageStartTime = FPlatformTime::Seconds();
 #endif
-		return;
 	}
 
 	if (WorldPartition->RuntimeHash)
@@ -830,6 +831,10 @@ void UWorldPartitionStreamingPolicy::UpdateStreamingPerformance(EWorldPartitionS
 
 			CurrentStreamingPerformance = NewStreamingPerformance;
 		}
+		if (bCurrentBlockOnSlowStreaming != bBlockOnSlowStreaming)
+		{
+			bCurrentBlockOnSlowStreaming = bBlockOnSlowStreaming;
+		}
 	}
 
 #if !UE_BUILD_SHIPPING
@@ -838,10 +843,11 @@ void UWorldPartitionStreamingPolicy::UpdateStreamingPerformance(EWorldPartitionS
 		// performance still bad keep message alive
 		OnScreenMessageStartTime = FPlatformTime::Seconds();
 		OnScreenMessageStreamingPerformance = CurrentStreamingPerformance;
+		bOnScreenMessageShouldBlock = bBlockOnSlowStreaming;
 	}
 #endif
 	
-	if (CurrentStreamingPerformance == EWorldPartitionStreamingPerformance::Critical)
+	if (CurrentStreamingPerformance >= EWorldPartitionStreamingPerformance::Critical && bCurrentBlockOnSlowStreaming)
 	{
 		const bool bIsServer = WorldPartition->IsServer();
 		const bool bIsServerStreamingEnabled = WorldPartition->IsServerStreamingEnabled();
@@ -865,11 +871,14 @@ void UWorldPartitionStreamingPolicy::GetOnScreenMessages(FCoreDelegates::FSeveri
 	{
 		switch (OnScreenMessageStreamingPerformance)
 		{
+		case EWorldPartitionStreamingPerformance::Immediate:
+			OutMessages.Add(FCoreDelegates::EOnScreenMessageSeverity::Error, FText::Format(LOCTEXT("WPStreamingImmediate", "[Immediate] WorldPartition Streaming Performance [Blocking:{0}]"), bOnScreenMessageShouldBlock ? LOCTEXT("True", "True") : LOCTEXT("False", "False")));
+			break;
 		case EWorldPartitionStreamingPerformance::Critical:
-			OutMessages.Add(FCoreDelegates::EOnScreenMessageSeverity::Error, LOCTEXT("WPStreamingCritical", "[Critical] WorldPartition Streaming Performance"));
+			OutMessages.Add(FCoreDelegates::EOnScreenMessageSeverity::Error, FText::Format(LOCTEXT("WPStreamingCritical", "[Critical] WorldPartition Streaming Performance [Blocking:{0}]"), bOnScreenMessageShouldBlock ? LOCTEXT("True", "True") : LOCTEXT("False", "False")));
 			break;
 		case EWorldPartitionStreamingPerformance::Slow:
-			OutMessages.Add(FCoreDelegates::EOnScreenMessageSeverity::Warning, LOCTEXT("WPStreamingWarning", "[Slow] WorldPartition Streaming Performance"));
+			OutMessages.Add(FCoreDelegates::EOnScreenMessageSeverity::Warning, FText::Format(LOCTEXT("WPStreamingWarning", "[Slow] WorldPartition Streaming Performance [Blocking:{0}]"), bOnScreenMessageShouldBlock ? LOCTEXT("True", "True") : LOCTEXT("False", "False")));
 			break;
 		default:
 			break;
diff --git a/Engine/Source/Runtime/Engine/Public/WorldPartition/WorldPartitionRuntimeCellData.h b/Engine/Source/Runtime/Engine/Public/WorldPartition/WorldPartitionRuntimeCellData.h
--- a/Engine/Source/Runtime/Engine/Public/WorldPartition/WorldPartitionRuntimeCellData.h
+++ b/Engine/Source/Runtime/Engine/Public/WorldPartition/WorldPartitionRuntimeCellData.h
@@ -62,6 +62,12 @@ class UWorldPartitionRuntimeCellData : public UObject
 	// Ratio used to determine the cell streaming performance status
 	mutable float CachedMinBlockOnSlowStreamingRatio;
 
+	// Square distance from the cell to the closest streaming source
+	mutable double CachedMinSquareDistanceToSource;
+
+	// Ratio used to determine the cell streaming performance status
+	mutable float CachedMinSlowStreamingRatio;
+
 	// Spatial priority based on distance and angle from source
 	mutable double CachedMinSpatialSortingPriority;
 
diff --git a/Engine/Source/Runtime/Engine/Public/WorldPartition/WorldPartitionRuntimeCellDataSpatialHash.h b/Engine/Source/Runtime/Engine/Public/WorldPartition/WorldPartitionRuntimeCellDataSpatialHash.h
--- a/Engine/Source/Runtime/Engine/Public/WorldPartition/WorldPartitionRuntimeCellDataSpatialHash.h
+++ b/Engine/Source/Runtime/Engine/Public/WorldPartition/WorldPartitionRuntimeCellDataSpatialHash.h
@@ -32,9 +32,6 @@ class UWorldPartitionRuntimeCellDataSpatialHash : public UWorldPartitionRuntimeC
 private:
 	ENGINE_API float ComputeSourceToCellAngleFactor(const FSphericalSector& SourceShape) const;
 
-	// Square distance from the cell to the closest streaming source
-	mutable double CachedMinSquareDistanceToSource;
-
 	// Modulated distance to the different streaming sources used to sort relative priority amongst streaming cells
 	// The value is affected by :
 	// - All sources intersecting the cell
@@ -57,4 +54,10 @@ private:
 
 	// 2D version of CachedMinSquareDistanceToBlockingSource
 	mutable double CachedMinSquareDistanceToBlockingSource2D;
+
+	// 2D version of CachedMinSlowStreamingRatio
+	mutable double CachedMinSlowStreamingRatio2D;
+
+	// 2D version of CachedMinSquareDistanceToSource
+	mutable double CachedMinSquareDistanceToSource2D;
 };
diff --git a/Engine/Source/Runtime/Engine/Public/WorldPartition/WorldPartitionRuntimeHash.h b/Engine/Source/Runtime/Engine/Public/WorldPartition/WorldPartitionRuntimeHash.h
--- a/Engine/Source/Runtime/Engine/Public/WorldPartition/WorldPartitionRuntimeHash.h
+++ b/Engine/Source/Runtime/Engine/Public/WorldPartition/WorldPartitionRuntimeHash.h
@@ -28,6 +28,9 @@ class FWorldPartitionDraw2DContext;
 class UExternalDataLayerAsset;
 class UExternalDataLayerInstance;
 
+extern ENGINE_API float GSlowStreamingRatio;
+extern ENGINE_API float GSlowStreamingWarningFactor;
+
 extern ENGINE_API float GBlockOnSlowStreamingRatio;
 extern ENGINE_API float GBlockOnSlowStreamingWarningFactor;
 
@@ -36,8 +39,10 @@ enum class EWorldPartitionStreamingPerformance : uint8
 {
 	Good,
 	Slow,
-	Critical
+	Critical,
+	Immediate,
 };
+ENGINE_API const TCHAR* EnumToString(EWorldPartitionStreamingPerformance InState);
 
 USTRUCT()
 struct FWorldPartitionRuntimeCellStreamingData
@@ -259,7 +264,10 @@ public:
 	virtual uint32 ComputeUpdateStreamingHash() const { return 0; }
 
 	ENGINE_API bool IsCellRelevantFor(bool bClientOnlyVisible) const;
-	ENGINE_API EWorldPartitionStreamingPerformance GetStreamingPerformance(const TSet<const UWorldPartitionRuntimeCell*>& CellToActivate) const;
+
+	UE_DEPRECATED(5.6, "Use version that has also flags if streaming should block (bOutShouldBlock) instead.")
+	ENGINE_API EWorldPartitionStreamingPerformance GetStreamingPerformance(const TSet<const UWorldPartitionRuntimeCell*>& CellsToActivate) const;	
+	ENGINE_API EWorldPartitionStreamingPerformance GetStreamingPerformance(const TSet<const UWorldPartitionRuntimeCell*>& CellsToActivate, bool& bOutShouldBlock) const;	
 
 	ENGINE_API virtual bool IsExternalStreamingObjectInjected(URuntimeHashExternalStreamingObjectBase* InExternalStreamingObject) const;
 	ENGINE_API virtual bool InjectExternalStreamingObject(URuntimeHashExternalStreamingObjectBase* ExternalStreamingObject);
@@ -275,7 +283,9 @@ public:
 	ENGINE_API UWorldPartitionRuntimeCell* CreateRuntimeCell(UClass* CellClass, UClass* CellDataClass, const FString& CellName, const FString& CellInstanceSuffix, UObject* InOuter = nullptr);
 
 protected:
+	UE_DEPRECATED(5.6, "Use version that has also flags if streaming should block (bOutShouldBlock) instead.")
 	virtual EWorldPartitionStreamingPerformance GetStreamingPerformanceForCell(const UWorldPartitionRuntimeCell* Cell) const;
+	virtual EWorldPartitionStreamingPerformance GetStreamingPerformanceForCell(const UWorldPartitionRuntimeCell* Cell, bool& bOutShouldBlock) const;
 
 #if WITH_EDITORONLY_DATA
 	struct FEditorAlwaysLoadedActor
diff --git a/Engine/Source/Runtime/Engine/Public/WorldPartition/WorldPartitionRuntimeSpatialHash.h b/Engine/Source/Runtime/Engine/Public/WorldPartition/WorldPartitionRuntimeSpatialHash.h
--- a/Engine/Source/Runtime/Engine/Public/WorldPartition/WorldPartitionRuntimeSpatialHash.h
+++ b/Engine/Source/Runtime/Engine/Public/WorldPartition/WorldPartitionRuntimeSpatialHash.h
@@ -363,8 +363,10 @@ public:
 protected:
 	ENGINE_API void ForEachStreamingGrid(TFunctionRef<void(FSpatialHashStreamingGrid&)> Func);
 	ENGINE_API void ForEachStreamingGridBreakable(TFunctionRef<bool(const FSpatialHashStreamingGrid&)> Func) const;
-
+	
+	UE_DEPRECATED(5.6, "Use version that has also flags if streaming should block (bOutShouldBlock) instead.")
 	ENGINE_API virtual EWorldPartitionStreamingPerformance GetStreamingPerformanceForCell(const UWorldPartitionRuntimeCell* Cell) const override;
+	ENGINE_API virtual EWorldPartitionStreamingPerformance GetStreamingPerformanceForCell(const UWorldPartitionRuntimeCell* Cell, bool& bOutShouldBlock) const override;
 
 #if WITH_EDITOR
 	ENGINE_API virtual bool HasStreamingContent() const override;
diff --git a/Engine/Source/Runtime/Engine/Public/WorldPartition/WorldPartitionStreamingPolicy.h b/Engine/Source/Runtime/Engine/Public/WorldPartition/WorldPartitionStreamingPolicy.h
--- a/Engine/Source/Runtime/Engine/Public/WorldPartition/WorldPartitionStreamingPolicy.h
+++ b/Engine/Source/Runtime/Engine/Public/WorldPartition/WorldPartitionStreamingPolicy.h
@@ -63,6 +63,7 @@ struct FWorldPartitionUpdateStreamingTargetState
 	TArray<TObjectPtr<const UWorldPartitionRuntimeCell>> ToUnloadCells;
 
 	EWorldPartitionStreamingPerformance StreamingPerformance = EWorldPartitionStreamingPerformance::Good;
+	bool bBlockOnSlowStreaming = false;
 
 	bool bUpdateServerEpoch = false;
 
@@ -150,7 +151,10 @@ protected:
 	virtual void GetCellsToUpdate(TArray<const UWorldPartitionRuntimeCell*>& OutToLoadCells, TArray<const UWorldPartitionRuntimeCell*>& OutToActivateCells);
 	virtual void GetCellsToReprioritize(TArray<const UWorldPartitionRuntimeCell*>& OutToLoadCells, TArray<const UWorldPartitionRuntimeCell*>& OutToActivateCells);
 	virtual void UpdateStreamingSources(bool bCanOptimizeUpdate);
+
+	UE_DEPRECATED(5.6, "Use version that also sets if streaming should block on slow loading")
 	void UpdateStreamingPerformance(EWorldPartitionStreamingPerformance NewStreamingPerformance);
+	void UpdateStreamingPerformance(EWorldPartitionStreamingPerformance NewStreamingPerformance, bool bBlockOnSlowStreaming);
 	bool IsInBlockTillLevelStreamingCompleted(bool bIsCausedByBadStreamingPerformance = false) const;
 
 	// Outer World Partition
@@ -211,6 +215,10 @@ private:
 	UPROPERTY(Transient)
 	EWorldPartitionStreamingPerformance CurrentStreamingPerformance;
 
+	// Current block on slow streaming
+	UPROPERTY(Transient)
+	bool bCurrentBlockOnSlowStreaming;
+
 	// Target state
 	UPROPERTY(Transient)
 	FWorldPartitionUpdateStreamingTargetState TargetState;
@@ -260,6 +268,7 @@ private:
 
 	double OnScreenMessageStartTime;
 	EWorldPartitionStreamingPerformance  OnScreenMessageStreamingPerformance;
+	bool bOnScreenMessageShouldBlock;
 #endif
 
 	friend class UWorldPartition;
