FLevelStreamingGCHelper::TrashPackage() - Make sure the trashed package is not seen as relevant by various systems
- Explicitly mark trashed packages as transient

diff --git a/Engine/Source/Runtime/Engine/Private/LevelStreamingGCHelper.cpp b/Engine/Source/Runtime/Engine/Private/LevelStreamingGCHelper.cpp
--- a/Engine/Source/Runtime/Engine/Private/LevelStreamingGCHelper.cpp
+++ b/Engine/Source/Runtime/Engine/Private/LevelStreamingGCHelper.cpp
@@ -262,4 +262,7 @@ void FLevelStreamingGCHelper::TrashPackage(UPackage* InPackage)
 {
 	const FName NewName = MakeUniqueObjectName(nullptr, UPackage::StaticClass(), NAME_TrashedPackage);
 	InPackage->Rename(*NewName.ToString(), nullptr, REN_DontCreateRedirectors | REN_NonTransactional | REN_DoNotDirty);
+
+	// Mark it transient to be sure it is not seen as relevant by various systems
+	InPackage->SetFlags(RF_Transient);
 }
