Fix UStaticMeshComponent::GetStreamingRenderAssetInfo(...) not returning references when Nanite fallback meshes are stripped from cooked builds.

diff --git a/Engine/Source/Runtime/Engine/Private/Components/StaticMeshComponent.cpp b/Engine/Source/Runtime/Engine/Private/Components/StaticMeshComponent.cpp
--- a/Engine/Source/Runtime/Engine/Private/Components/StaticMeshComponent.cpp
+++ b/Engine/Source/Runtime/Engine/Private/Components/StaticMeshComponent.cpp
@@ -1235,7 +1235,7 @@ float UStaticMeshComponent::GetTextureStreamingTransformScale() const
 
 void UStaticMeshComponent::GetStreamingRenderAssetInfo(FStreamingTextureLevelContext& LevelContext, TArray<FStreamingRenderAssetPrimitiveInfo>& OutStreamingRenderAssets) const
 {
-	if (bIgnoreInstanceForTextureStreaming || !GetStaticMesh() || GetStaticMesh()->IsCompiling() || !GetStaticMesh()->HasValidRenderData())
+	if (bIgnoreInstanceForTextureStreaming || !GetStaticMesh() || GetStaticMesh()->IsCompiling() || !(GetStaticMesh()->HasValidRenderData() || GetStaticMesh()->HasValidNaniteData()))
 	{
 		return;
 	}
