Chaos - simplify lock logic

diff --git a/Engine/Source/Runtime/Engine/Private/PhysicsEngine/Experimental/ChaosScopedSceneLock.cpp b/Engine/Source/Runtime/Engine/Private/PhysicsEngine/Experimental/ChaosScopedSceneLock.cpp
--- a/Engine/Source/Runtime/Engine/Private/PhysicsEngine/Experimental/ChaosScopedSceneLock.cpp
+++ b/Engine/Source/Runtime/Engine/Private/PhysicsEngine/Experimental/ChaosScopedSceneLock.cpp
@@ -3,17 +3,25 @@
 #include "Physics/Experimental/ChaosScopedSceneLock.h"
 
 #include "Chaos/PBDJointConstraintData.h"
+#include "Components/SkeletalMeshComponent.h"
 #include "Physics/Experimental/PhysScene_Chaos.h"
 #include "PhysicsEngine/ConstraintInstance.h"
 #include "PhysicsEngine/PhysicsObjectExternalInterface.h"
-#include "Components/SkeletalMeshComponent.h"
+#include "PhysicsProxy/SingleParticlePhysicsProxy.h"
 #include "PBDRigidsSolver.h"
 
+namespace
+{
+	Chaos::FPBDRigidsSolver* GetSolverForActor(const FPhysicsActorHandle& InActorHandle)
+	{
+		return InActorHandle->GetSolver<Chaos::FPBDRigidsSolver>();
+	}
+}
+
 FScopedSceneLock_Chaos::FScopedSceneLock_Chaos(const FPhysicsActorHandle& InActorHandle, EPhysicsInterfaceScopedLockType InLockType)
 	: LockType(InLockType)
 {
-	FChaosScene* Scene = GetSceneForActor(InActorHandle);
-	Solver = Scene ? Scene->GetSolver() : nullptr;
+	Solver = GetSolverForActor(InActorHandle);
 	LockScene();
 }
 
