NullCheck for a UniformBuffer update crash

diff --git a/Engine/Source/Runtime/Engine/Private/PrimitiveSceneProxy.cpp b/Engine/Source/Runtime/Engine/Private/PrimitiveSceneProxy.cpp
--- a/Engine/Source/Runtime/Engine/Private/PrimitiveSceneProxy.cpp
+++ b/Engine/Source/Runtime/Engine/Private/PrimitiveSceneProxy.cpp
@@ -838,7 +838,10 @@ void FPrimitiveSceneProxy::UpdateUniformBuffer(FRHICommandList& RHICmdList)
 		FPrimitiveUniformShaderParameters PrimitiveParams = Builder.Build();
 
 		check(UniformBuffer);
-		UniformBuffer.UpdateUniformBufferImmediate(RHICmdList, PrimitiveParams);
+		if (UniformBuffer.IsValid())
+		{
+			UniformBuffer.UpdateUniformBufferImmediate(RHICmdList, PrimitiveParams);
+		}
 	}
 }
 
