Updated bit-width of material value types to 64-bits

diff --git a/Engine/Plugins/Enterprise/DatasmithImporter/Source/DatasmithImporter/Private/DatasmithMaterialExpressions.cpp b/Engine/Plugins/Enterprise/DatasmithImporter/Source/DatasmithImporter/Private/DatasmithMaterialExpressions.cpp
--- a/Engine/Plugins/Enterprise/DatasmithImporter/Source/DatasmithImporter/Private/DatasmithMaterialExpressions.cpp
+++ b/Engine/Plugins/Enterprise/DatasmithImporter/Source/DatasmithImporter/Private/DatasmithMaterialExpressions.cpp
@@ -543,7 +543,7 @@ UMaterialExpression* FDatasmithMaterialExpressions::Lerp(UObject* UnrealMaterial
 	bool bInvertOrder = false;
 	if (ToBeConnectedA && ToBeConnectedB)
 	{
-		if (ToBeConnectedA->GetOutputType(0) == MCT_Float3)
+		if (ToBeConnectedA->GetOutputValueType(0) == MCT_Float3)
 		{
 			bInvertOrder = true;
 		}
diff --git a/Engine/Plugins/Experimental/DynamicMaterial/Source/DynamicMaterialEditor/Private/Components/DMMaterialEffectFunction.cpp b/Engine/Plugins/Experimental/DynamicMaterial/Source/DynamicMaterialEditor/Private/Components/DMMaterialEffectFunction.cpp
--- a/Engine/Plugins/Experimental/DynamicMaterial/Source/DynamicMaterialEditor/Private/Components/DMMaterialEffectFunction.cpp
+++ b/Engine/Plugins/Experimental/DynamicMaterial/Source/DynamicMaterialEditor/Private/Components/DMMaterialEffectFunction.cpp
@@ -417,10 +417,10 @@ bool UDMMaterialEffectFunction::NeedsFunctionInit() const
 		return false;
 	}
 
-	const uint32 InputType = Inputs[0].ExpressionInput ? Inputs[0].ExpressionInput->GetInputType(0) : MCT_Unknown;
+	const EMaterialValueType InputType = Inputs[0].ExpressionInput ? Inputs[0].ExpressionInput->GetInputValueType(0) : MCT_Unknown;
 	const bool bInputTypeIsFloat = (InputType & MCT_Float) != 0;
 
-	const uint32 OutputType = Outputs[0].ExpressionOutput ? Outputs[0].ExpressionOutput->GetOutputType(0) : MCT_Unknown;
+	const EMaterialValueType OutputType = Outputs[0].ExpressionOutput ? Outputs[0].ExpressionOutput->GetOutputValueType(0) : MCT_Unknown;
 	const bool bOutputTypeIsFloat = (InputType & MCT_Float) != 0;
 
 	/**
diff --git a/Engine/Plugins/Experimental/DynamicMaterial/Source/DynamicMaterialEditor/Private/Components/DMMaterialStageExpression.cpp b/Engine/Plugins/Experimental/DynamicMaterial/Source/DynamicMaterialEditor/Private/Components/DMMaterialStageExpression.cpp
--- a/Engine/Plugins/Experimental/DynamicMaterial/Source/DynamicMaterialEditor/Private/Components/DMMaterialStageExpression.cpp
+++ b/Engine/Plugins/Experimental/DynamicMaterial/Source/DynamicMaterialEditor/Private/Components/DMMaterialStageExpression.cpp
@@ -138,7 +138,7 @@ void UDMMaterialStageExpression::GenerateExpressionList()
 
 		for (int32 OutputIdx = 0; OutputIdx < ExpressionOutputs.Num(); ++OutputIdx)
 		{
-			if (ExpressionCDO->GetOutputType(OutputIdx) == MCT_Unknown)
+			if (ExpressionCDO->GetOutputValueType(OutputIdx) == MCT_Unknown)
 			{
 				bValidOutputs = false;
 				break;
diff --git a/Engine/Plugins/Experimental/DynamicMaterial/Source/DynamicMaterialEditor/Private/Components/DMMaterialStageFunction.cpp b/Engine/Plugins/Experimental/DynamicMaterial/Source/DynamicMaterialEditor/Private/Components/DMMaterialStageFunction.cpp
--- a/Engine/Plugins/Experimental/DynamicMaterial/Source/DynamicMaterialEditor/Private/Components/DMMaterialStageFunction.cpp
+++ b/Engine/Plugins/Experimental/DynamicMaterial/Source/DynamicMaterialEditor/Private/Components/DMMaterialStageFunction.cpp
@@ -316,10 +316,10 @@ bool UDMMaterialStageFunction::NeedsFunctionInit() const
 		return false;
 	}
 
-	const uint32 InputType = Inputs[0].ExpressionInput ? Inputs[0].ExpressionInput->GetInputType(0) : MCT_Unknown;
+	const EMaterialValueType InputType = Inputs[0].ExpressionInput ? Inputs[0].ExpressionInput->GetInputValueType(0) : MCT_Unknown;
 	const bool bInputTypeIsFloat = (InputType & MCT_Float) != 0;
 
-	const uint32 OutputType = Outputs[0].ExpressionOutput ? Outputs[0].ExpressionOutput->GetOutputType(0) : MCT_Unknown;
+	const EMaterialValueType OutputType = Outputs[0].ExpressionOutput ? Outputs[0].ExpressionOutput->GetOutputValueType(0) : MCT_Unknown;
 	const bool bOutputTypeIsFloat = (InputType & MCT_Float) != 0;
 
 	/**
diff --git a/Engine/Plugins/Experimental/DynamicMaterial/Source/DynamicMaterialEditor/Private/Components/DMMaterialStageSource.cpp b/Engine/Plugins/Experimental/DynamicMaterial/Source/DynamicMaterialEditor/Private/Components/DMMaterialStageSource.cpp
--- a/Engine/Plugins/Experimental/DynamicMaterial/Source/DynamicMaterialEditor/Private/Components/DMMaterialStageSource.cpp
+++ b/Engine/Plugins/Experimental/DynamicMaterial/Source/DynamicMaterialEditor/Private/Components/DMMaterialStageSource.cpp
@@ -184,7 +184,7 @@ bool UDMMaterialStageSource::GenerateStagePreviewMaterial(UDMMaterialStage* InSt
 
 	for (int32 OutputIdx = 0; OutputIdx < LastExpression->GetOutputs().Num(); ++OutputIdx)
 	{
-		EMaterialValueType CurrentOutputType = static_cast<EMaterialValueType>(LastExpression->GetOutputType(OutputIdx));
+		EMaterialValueType CurrentOutputType = LastExpression->GetOutputValueType(OutputIdx);
 		int32 CurrentOutputCount = 0;
 
 		switch (CurrentOutputType)
diff --git a/Engine/Plugins/Experimental/DynamicMaterial/Source/DynamicMaterialEditor/Private/Model/DMMaterialBuildUtils.cpp b/Engine/Plugins/Experimental/DynamicMaterial/Source/DynamicMaterialEditor/Private/Model/DMMaterialBuildUtils.cpp
--- a/Engine/Plugins/Experimental/DynamicMaterial/Source/DynamicMaterialEditor/Private/Model/DMMaterialBuildUtils.cpp
+++ b/Engine/Plugins/Experimental/DynamicMaterial/Source/DynamicMaterialEditor/Private/Model/DMMaterialBuildUtils.cpp
@@ -309,7 +309,7 @@ UMaterialExpressionComponentMask* FDMMaterialBuildUtils::CreateExpressionBitMask
 
 	if (InOutputChannels == FDMMaterialStageConnectorChannel::WHOLE_CHANNEL)
 	{
-		const uint32 OutputType = InExpression->GetOutputType(InOutputIndex);
+		const EMaterialValueType OutputType = InExpression->GetOutputValueType(InOutputIndex);
 
 		Mask->R = (OutputType == MCT_Float1);
 		Mask->G = (OutputType == MCT_Float2);
diff --git a/Engine/Plugins/Experimental/VirtualHeightfieldMesh/Source/VirtualHeightfieldMesh/Private/HeightfieldMinMaxTextureMaterialExpression.cpp b/Engine/Plugins/Experimental/VirtualHeightfieldMesh/Source/VirtualHeightfieldMesh/Private/HeightfieldMinMaxTextureMaterialExpression.cpp
--- a/Engine/Plugins/Experimental/VirtualHeightfieldMesh/Source/VirtualHeightfieldMesh/Private/HeightfieldMinMaxTextureMaterialExpression.cpp
+++ b/Engine/Plugins/Experimental/VirtualHeightfieldMesh/Source/VirtualHeightfieldMesh/Private/HeightfieldMinMaxTextureMaterialExpression.cpp
@@ -76,7 +76,7 @@ void UMaterialExpressionHeightfieldMinMaxTexture::GetCaption(TArray<FString>& Ou
 	OutCaptions.Add(TEXT("HeightfieldMinMaxTextureObject"));
 }
 
-uint32 UMaterialExpressionHeightfieldMinMaxTexture::GetOutputType(int32 OutputIndex)
+EMaterialValueType UMaterialExpressionHeightfieldMinMaxTexture::GetOutputValueType(int32 OutputIndex)
 {
 	return MCT_Texture2D;
 }
diff --git a/Engine/Plugins/Experimental/VirtualHeightfieldMesh/Source/VirtualHeightfieldMesh/Private/HeightfieldMinMaxTextureMaterialExpression.h b/Engine/Plugins/Experimental/VirtualHeightfieldMesh/Source/VirtualHeightfieldMesh/Private/HeightfieldMinMaxTextureMaterialExpression.h
--- a/Engine/Plugins/Experimental/VirtualHeightfieldMesh/Source/VirtualHeightfieldMesh/Private/HeightfieldMinMaxTextureMaterialExpression.h
+++ b/Engine/Plugins/Experimental/VirtualHeightfieldMesh/Source/VirtualHeightfieldMesh/Private/HeightfieldMinMaxTextureMaterialExpression.h
@@ -33,7 +33,7 @@ protected:
 	virtual int32 Compile(class FMaterialCompiler* Compiler, int32 OutputIndex) override;
 	virtual int32 CompilePreview(class FMaterialCompiler* Compiler, int32 OutputIndex) override;
 	virtual void GetCaption(TArray<FString>& OutCaptions) const override;
-	virtual uint32 GetOutputType(int32 OutputIndex) override;
+	virtual EMaterialValueType GetOutputValueType(int32 OutputIndex) override;
 	virtual bool CanReferenceTexture() const override { return true; }
 #endif
 	virtual UObject* GetReferencedTexture() const override;
diff --git a/Engine/Plugins/Interchange/Runtime/Source/Import/Private/MaterialX/MaterialExpressions/MaterialExpressionSwizzle.cpp b/Engine/Plugins/Interchange/Runtime/Source/Import/Private/MaterialX/MaterialExpressions/MaterialExpressionSwizzle.cpp
--- a/Engine/Plugins/Interchange/Runtime/Source/Import/Private/MaterialX/MaterialExpressions/MaterialExpressionSwizzle.cpp
+++ b/Engine/Plugins/Interchange/Runtime/Source/Import/Private/MaterialX/MaterialExpressions/MaterialExpressionSwizzle.cpp
@@ -107,7 +107,7 @@ int32 UMaterialExpressionMaterialXSwizzle::Compile(FMaterialCompiler* Compiler,
 		Channels.Len() >= 4 ? Compiler->ComponentMask(Index, SwizzleX[3], SwizzleY[3], SwizzleZ[3], SwizzleW[3]) : -1,
 	};
 
-	Input.Expression->GetOutputType(Input.OutputIndex);
+	Input.Expression->GetOutputValueType(Input.OutputIndex);
 	int32 Res = INDEX_NONE;
 
 	switch(Channels.Len())
diff --git a/Engine/Plugins/Runtime/RenderTrace/Source/Private/Materials/MaterialExpressionPhysicalMaterialOutput.h b/Engine/Plugins/Runtime/RenderTrace/Source/Private/Materials/MaterialExpressionPhysicalMaterialOutput.h
--- a/Engine/Plugins/Runtime/RenderTrace/Source/Private/Materials/MaterialExpressionPhysicalMaterialOutput.h
+++ b/Engine/Plugins/Runtime/RenderTrace/Source/Private/Materials/MaterialExpressionPhysicalMaterialOutput.h
@@ -53,7 +53,7 @@ class RENDERTRACE_API UMaterialExpressionPhysicalMaterialOutput : public UMateri
 	virtual TArrayView<FExpressionInput*> GetInputsView() override;
 	virtual FExpressionInput* GetInput(int32 InputIndex) override;
 	virtual FName GetInputName(int32 InputIndex) const override;
-	virtual uint32 GetInputType(int32 InputIndex) override { return MCT_Float; }
+	virtual EMaterialValueType GetInputValueType(int32 InputIndex) override { return MCT_Float; }
 	virtual int32 Compile(class FMaterialCompiler* Compiler, int32 OutputIndex) override;
 	//~ End UMaterialExpression Interface
 #endif
diff --git a/Engine/Plugins/TextureGraph/Source/TextureGraph/Private/Expressions/Input/TG_Expression_MaterialFunction.cpp b/Engine/Plugins/TextureGraph/Source/TextureGraph/Private/Expressions/Input/TG_Expression_MaterialFunction.cpp
--- a/Engine/Plugins/TextureGraph/Source/TextureGraph/Private/Expressions/Input/TG_Expression_MaterialFunction.cpp
+++ b/Engine/Plugins/TextureGraph/Source/TextureGraph/Private/Expressions/Input/TG_Expression_MaterialFunction.cpp
@@ -279,7 +279,7 @@ UMaterialInterface* UTG_Expression_MaterialFunction::CreateMaterialReference()
 		{
 			auto Name = ExpInput->InputName;
 			auto InputName = FName(Name.ToString() + TEXT("_param"));
-			auto InputType = ExpInput->GetInputType(0);
+			auto InputType = ExpInput->GetInputValueType(0);
 			FString TypeName;
 			switch (InputType)
 			{
@@ -380,7 +380,7 @@ UMaterialInterface* UTG_Expression_MaterialFunction::CreateMaterialReference()
 		if (ExpOutput)
 		{
 			auto Name = ExpOutput->OutputName;
-			uint32 Type = ExpOutput->GetOutputType(0);
+			uint32 Type = ExpOutput->GetOutputValueType(0);
 
 			FString TypeName;
 			switch (Type)
diff --git a/Engine/Source/Editor/GraphEditor/Private/MaterialNodes/SGraphSubstrateMaterial.cpp b/Engine/Source/Editor/GraphEditor/Private/MaterialNodes/SGraphSubstrateMaterial.cpp
--- a/Engine/Source/Editor/GraphEditor/Private/MaterialNodes/SGraphSubstrateMaterial.cpp
+++ b/Engine/Source/Editor/GraphEditor/Private/MaterialNodes/SGraphSubstrateMaterial.cpp
@@ -36,7 +36,7 @@ bool FSubstrateWidget::HasInputSubstrateType(const UEdGraphPin* InPin)
 
 	if (UMaterialGraphNode_Root* RootPinNode = Cast<UMaterialGraphNode_Root>(InPin->GetOwningNode()))
 	{
-		if(RootPinNode->GetPinMaterialType(InPin) == MCT_Substrate)
+		if(RootPinNode->GetPinMaterialValueType(InPin) == MCT_Substrate)
 		{
 			return true;
 		}
@@ -52,7 +52,7 @@ bool FSubstrateWidget::HasInputSubstrateType(const UEdGraphPin* InPin)
 
 			if (InputName == TargetPinName)
 			{
-				switch (PinNode->MaterialExpression->GetInputType(It.Index))
+				switch (PinNode->MaterialExpression->GetInputValueType(It.Index))
 				{
 					case MCT_Substrate:
 						return true;
@@ -71,7 +71,7 @@ bool FSubstrateWidget::HasOutputSubstrateType(const UEdGraphPin* InPin)
 	if (UMaterialGraphNode* PinNode = Cast<UMaterialGraphNode>(InPin->GetOwningNode()))
 	{
 		const TArray<FExpressionOutput>& ExpressionOutputs = PinNode->MaterialExpression->GetOutputs();
-		if (InPin->SourceIndex < ExpressionOutputs.Num() && PinNode->MaterialExpression->GetOutputType(InPin->SourceIndex) == MCT_Substrate)
+		if (InPin->SourceIndex < ExpressionOutputs.Num() && PinNode->MaterialExpression->GetOutputValueType(InPin->SourceIndex) == MCT_Substrate)
 		{
 			return true;
 		}
diff --git a/Engine/Source/Editor/MaterialEditor/Private/MaterialEditingLibrary.cpp b/Engine/Source/Editor/MaterialEditor/Private/MaterialEditingLibrary.cpp
--- a/Engine/Source/Editor/MaterialEditor/Private/MaterialEditingLibrary.cpp
+++ b/Engine/Source/Editor/MaterialEditor/Private/MaterialEditingLibrary.cpp
@@ -924,11 +924,11 @@ TArray<int32> UMaterialEditingLibrary::GetMaterialExpressionInputTypes(UMaterial
 		UMaterialExpression* Expression = It->Expression;
 		if (Expression != nullptr)
 		{
-			InputTypes.Add(Expression->GetOutputType(It->OutputIndex));
+			InputTypes.Add(Expression->GetOutputValueType(It->OutputIndex));
 		}
 		else
 		{
-			InputTypes.Add(MaterialExpression->GetInputType(It.Index));
+			InputTypes.Add(MaterialExpression->GetInputValueType(It.Index));
 		}
 	}
 	return InputTypes;
diff --git a/Engine/Source/Editor/MaterialEditor/Private/MaterialEditor.cpp b/Engine/Source/Editor/MaterialEditor/Private/MaterialEditor.cpp
--- a/Engine/Source/Editor/MaterialEditor/Private/MaterialEditor.cpp
+++ b/Engine/Source/Editor/MaterialEditor/Private/MaterialEditor.cpp
@@ -5317,7 +5317,7 @@ UClass* FMaterialEditor::GetOnPromoteToParameterClass(const UEdGraphPin* TargetP
 
 			if (InputName == TargetPinName)
 			{
-				switch (OtherPinNode->MaterialExpression->GetInputType(It.Index))
+				switch (OtherPinNode->MaterialExpression->GetInputValueType(It.Index))
 				{
 					case MCT_Float1:
 					case MCT_Float: return UMaterialExpressionScalarParameter::StaticClass();
diff --git a/Engine/Source/Editor/MaterialEditor/Private/MaterialEditorHelpers.cpp b/Engine/Source/Editor/MaterialEditor/Private/MaterialEditorHelpers.cpp
--- a/Engine/Source/Editor/MaterialEditor/Private/MaterialEditorHelpers.cpp
+++ b/Engine/Source/Editor/MaterialEditor/Private/MaterialEditorHelpers.cpp
@@ -221,7 +221,7 @@ void FMaterialEditorHelpers::CollapseToFunction(FMaterialEditor& MaterialEditor)
 				UEdGraphPin* PinLinkedToInput = It.Key;
 				UMaterialGraphNode* GraphNodeLinkedToInput = CastChecked<UMaterialGraphNode>(PinLinkedToInput->GetOwningNode());
 
-				EMaterialValueType PinMaterialType = EMaterialValueType(GraphNodeLinkedToInput->GetOutputType(PinLinkedToInput));
+				EMaterialValueType PinMaterialType = GraphNodeLinkedToInput->GetOutputValueType(PinLinkedToInput);
 				EFunctionInputType PinFunctionType = {};
 
 				// This could maybe be shared with UMaterialExpressionFunctionInput, but there are tricky details involved with MCT_Float1 vs MCT_Float
diff --git a/Engine/Source/Editor/MaterialEditor/Private/MaterialEditorUtilities.cpp b/Engine/Source/Editor/MaterialEditor/Private/MaterialEditorUtilities.cpp
--- a/Engine/Source/Editor/MaterialEditor/Private/MaterialEditorUtilities.cpp
+++ b/Engine/Source/Editor/MaterialEditor/Private/MaterialEditorUtilities.cpp
@@ -742,7 +742,7 @@ bool FMaterialEditorUtilities::HasCompatibleConnection(UClass* ExpressionClass,
 		{
 			for (FExpressionInputIterator It{ DefaultExpression }; It; ++It)
 			{
-				uint32 InputType = DefaultExpression->GetInputType(It.Index);
+				EMaterialValueType InputType = DefaultExpression->GetInputValueType(It.Index);
 				if (CanConnectMaterialValueTypes(InputType, TestType))
 				{
 					return true;
@@ -754,7 +754,7 @@ bool FMaterialEditorUtilities::HasCompatibleConnection(UClass* ExpressionClass,
 			int32 NumOutputs = DefaultExpression->GetOutputs().Num();
 			for (int32 Index = 0; Index < NumOutputs; ++Index)
 			{
-				uint32 OutputType = DefaultExpression->GetOutputType(Index);
+				EMaterialValueType OutputType = DefaultExpression->GetOutputValueType(Index);
 				if (CanConnectMaterialValueTypes(TestType, OutputType))
 				{
 					return true;
diff --git a/Engine/Source/Editor/UnrealEd/Classes/MaterialGraph/MaterialGraphNode.h b/Engine/Source/Editor/UnrealEd/Classes/MaterialGraph/MaterialGraphNode.h
--- a/Engine/Source/Editor/UnrealEd/Classes/MaterialGraph/MaterialGraphNode.h
+++ b/Engine/Source/Editor/UnrealEd/Classes/MaterialGraph/MaterialGraphNode.h
@@ -72,7 +72,7 @@ public:
 	virtual void GetPinHoverText(const UEdGraphPin& Pin, FString& HoverTextOut) const override;
 	virtual FString GetDocumentationExcerptName() const override;
 	virtual bool CanUserDeleteNode() const override;
-	virtual uint32 GetPinMaterialType(const UEdGraphPin* Pin) const override;
+	virtual EMaterialValueType GetPinMaterialValueType(const UEdGraphPin* Pin) const override;
 	virtual void PinDefaultValueChanged(UEdGraphPin* Pin) override;
 	virtual TSharedPtr<SGraphNode> CreateVisualWidget() override;
 	virtual bool CanJumpToDefinition() const override;
diff --git a/Engine/Source/Editor/UnrealEd/Classes/MaterialGraph/MaterialGraphNode_Base.h b/Engine/Source/Editor/UnrealEd/Classes/MaterialGraph/MaterialGraphNode_Base.h
--- a/Engine/Source/Editor/UnrealEd/Classes/MaterialGraph/MaterialGraphNode_Base.h
+++ b/Engine/Source/Editor/UnrealEd/Classes/MaterialGraph/MaterialGraphNode_Base.h
@@ -5,6 +5,7 @@
 #include "CoreMinimal.h"
 #include "UObject/ObjectMacros.h"
 #include "EdGraph/EdGraphNode.h"
+#include "MaterialValueType.h"
 #include "MaterialGraphNode_Base.generated.h"
 
 class UEdGraphPin;
@@ -37,10 +38,17 @@ class UMaterialGraphNode_Base : public UEdGraphNode
 	UNREALED_API void ReplaceNode(UMaterialGraphNode_Base* OldNode);
 
 	/** Get the Material value type of an input pin */
-	uint32 GetInputType(const UEdGraphPin* InputPin) const;
+	EMaterialValueType GetInputValueType(const UEdGraphPin* InputPin) const;
 
 	/** Get the Material value type of an output pin */
-	UNREALED_API uint32 GetOutputType(const UEdGraphPin* OutputPin) const;
+	UNREALED_API EMaterialValueType GetOutputValueType(const UEdGraphPin* OutputPin) const;
+	
+	UNREALED_API uint32 GetOutputType(const UEdGraphPin* OutputPin)
+	{
+		PRAGMA_DISABLE_DEPRECATION_WARNINGS
+		return static_cast<uint32>(GetOutputValueType(OutputPin));
+		PRAGMA_ENABLE_DEPRECATION_WARNINGS
+	}
 
 	/**
 	 * Handles inserting the node between the FromPin and what the FromPin was original connected to
@@ -64,7 +72,7 @@ class UMaterialGraphNode_Base : public UEdGraphNode
 protected:
 	void ModifyAndCopyPersistentPinData(UEdGraphPin& TargetPin, const UEdGraphPin& SourcePin) const;
 
-	virtual uint32 GetPinMaterialType(const UEdGraphPin* Pin) const;
+	virtual EMaterialValueType GetPinMaterialValueType(const UEdGraphPin* Pin) const;
 
 	void EmptyPins();
 
diff --git a/Engine/Source/Editor/UnrealEd/Classes/MaterialGraph/MaterialGraphNode_Root.h b/Engine/Source/Editor/UnrealEd/Classes/MaterialGraph/MaterialGraphNode_Root.h
--- a/Engine/Source/Editor/UnrealEd/Classes/MaterialGraph/MaterialGraphNode_Root.h
+++ b/Engine/Source/Editor/UnrealEd/Classes/MaterialGraph/MaterialGraphNode_Root.h
@@ -34,7 +34,7 @@ class UMaterialGraphNode_Root : public UMaterialGraphNode_Base
 	//~ Begin UMaterialGraphNode_Base Interface
 	virtual UObject* GetMaterialNodeOwner() const override;
 	virtual int32 GetSourceIndexForInputIndex(int32 InputIndex) const override;
-	virtual uint32 GetPinMaterialType(const UEdGraphPin* Pin) const override;
+	virtual EMaterialValueType GetPinMaterialValueType(const UEdGraphPin* Pin) const override;
 	virtual void CreateInputPins() override;
 	virtual bool IsRootNode() const override {return true;}
 	//~ End UMaterialGraphNode_Base Interface
diff --git a/Engine/Source/Editor/UnrealEd/Private/MaterialGraphNode.cpp b/Engine/Source/Editor/UnrealEd/Private/MaterialGraphNode.cpp
--- a/Engine/Source/Editor/UnrealEd/Private/MaterialGraphNode.cpp
+++ b/Engine/Source/Editor/UnrealEd/Private/MaterialGraphNode.cpp
@@ -631,7 +631,7 @@ FName UMaterialGraphNode::GetShortenPinName(const FName PinName)
 	return InputName;
 }
 
-uint32 UMaterialGraphNode::GetPinMaterialType(const UEdGraphPin* Pin) const
+EMaterialValueType UMaterialGraphNode::GetPinMaterialValueType(const UEdGraphPin* Pin) const
 {
 	if (Pin->PinType.PinCategory == UMaterialGraphSchema::PC_Exec)
 	{
@@ -640,11 +640,11 @@ uint32 UMaterialGraphNode::GetPinMaterialType(const UEdGraphPin* Pin) const
 	
 	if (Pin->Direction == EGPD_Input)
 	{
-		return MaterialExpression->GetInputType(Pin->SourceIndex);
+		return MaterialExpression->GetInputValueType(Pin->SourceIndex);
 	}
 	else
 	{
-		return MaterialExpression->GetOutputType(Pin->SourceIndex);
+		return MaterialExpression->GetOutputValueType(Pin->SourceIndex);
 	}
 }
 
diff --git a/Engine/Source/Editor/UnrealEd/Private/MaterialGraphNode_Base.cpp b/Engine/Source/Editor/UnrealEd/Private/MaterialGraphNode_Base.cpp
--- a/Engine/Source/Editor/UnrealEd/Private/MaterialGraphNode_Base.cpp
+++ b/Engine/Source/Editor/UnrealEd/Private/MaterialGraphNode_Base.cpp
@@ -95,17 +95,17 @@ UEdGraphPin* UMaterialGraphNode_Base::GetExecOutputPin(int32 OutputIndex) const
 	return nullptr;
 }
 
-uint32 UMaterialGraphNode_Base::GetOutputType(const UEdGraphPin* OutputPin) const
+EMaterialValueType UMaterialGraphNode_Base::GetOutputValueType(const UEdGraphPin* OutputPin) const
 {
-	return GetPinMaterialType(OutputPin);
+	return GetPinMaterialValueType(OutputPin);
 }
 
-uint32 UMaterialGraphNode_Base::GetInputType(const UEdGraphPin* InputPin) const
+EMaterialValueType UMaterialGraphNode_Base::GetInputValueType(const UEdGraphPin* InputPin) const
 {
-	return GetPinMaterialType(InputPin);
+	return GetPinMaterialValueType(InputPin);
 }
 
-uint32 UMaterialGraphNode_Base::GetPinMaterialType(const UEdGraphPin* Pin) const
+EMaterialValueType UMaterialGraphNode_Base::GetPinMaterialValueType(const UEdGraphPin* Pin) const
 {
 	if (Pin->PinType.PinCategory == UMaterialGraphSchema::PC_Exec)
 	{
diff --git a/Engine/Source/Editor/UnrealEd/Private/MaterialGraphNode_Root.cpp b/Engine/Source/Editor/UnrealEd/Private/MaterialGraphNode_Root.cpp
--- a/Engine/Source/Editor/UnrealEd/Private/MaterialGraphNode_Root.cpp
+++ b/Engine/Source/Editor/UnrealEd/Private/MaterialGraphNode_Root.cpp
@@ -157,7 +157,7 @@ int32 UMaterialGraphNode_Root::GetSourceIndexForInputIndex(int32 InputIndex) con
 	return INDEX_NONE;
 }
 
-uint32 UMaterialGraphNode_Root::GetPinMaterialType(const UEdGraphPin* Pin) const
+EMaterialValueType UMaterialGraphNode_Root::GetPinMaterialValueType(const UEdGraphPin* Pin) const
 {
 	if (Pin->PinType.PinCategory == UMaterialGraphSchema::PC_Exec)
 	{
@@ -168,7 +168,7 @@ uint32 UMaterialGraphNode_Root::GetPinMaterialType(const UEdGraphPin* Pin) const
 	const FMaterialInputInfo& MaterialInput = MaterialGraph->MaterialInputs[Pin->SourceIndex];
 	EMaterialProperty Property = MaterialInput.GetProperty();
 
-	uint32 MaterialType = 0u;
+	EMaterialValueType MaterialType{};
 	if (Property == MP_MaterialAttributes)
 	{
 		MaterialType = MCT_MaterialAttributes;
diff --git a/Engine/Source/Editor/UnrealEd/Private/MaterialGraphSchema.cpp b/Engine/Source/Editor/UnrealEd/Private/MaterialGraphSchema.cpp
--- a/Engine/Source/Editor/UnrealEd/Private/MaterialGraphSchema.cpp
+++ b/Engine/Source/Editor/UnrealEd/Private/MaterialGraphSchema.cpp
@@ -473,11 +473,11 @@ uint32 UMaterialGraphSchema::GetMaterialValueType(const UEdGraphPin* MaterialPin
 	const UMaterialGraphNode_Base* OwningNode = CastChecked<UMaterialGraphNode_Base>(MaterialPin->GetOwningNode());
 	if (MaterialPin->Direction == EGPD_Output)
 	{
-		return OwningNode->GetOutputType(MaterialPin);
+		return OwningNode->GetOutputValueType(MaterialPin);
 	}
 	else
 	{
-		return OwningNode->GetInputType(MaterialPin);
+		return OwningNode->GetInputValueType(MaterialPin);
 	}
 }
 
diff --git a/Engine/Source/Runtime/Engine/Internal/Materials/MaterialExpressionMaterialSample.cpp b/Engine/Source/Runtime/Engine/Internal/Materials/MaterialExpressionMaterialSample.cpp
--- a/Engine/Source/Runtime/Engine/Internal/Materials/MaterialExpressionMaterialSample.cpp
+++ b/Engine/Source/Runtime/Engine/Internal/Materials/MaterialExpressionMaterialSample.cpp
@@ -67,7 +67,7 @@ int32 UMaterialExpressionMaterialSample::CompilePreview(FMaterialCompiler* Compi
 	return DynamicCompile(Compiler, OutputIndex, true);
 }
 
-uint32 UMaterialExpressionMaterialSample::GetOutputType(int32 OutputIndex)
+EMaterialValueType UMaterialExpressionMaterialSample::GetOutputValueType(int32 OutputIndex)
 {
 	if (Substrate::IsSubstrateEnabled())
 	{
diff --git a/Engine/Source/Runtime/Engine/Internal/Materials/MaterialExpressionMaterialSample.h b/Engine/Source/Runtime/Engine/Internal/Materials/MaterialExpressionMaterialSample.h
--- a/Engine/Source/Runtime/Engine/Internal/Materials/MaterialExpressionMaterialSample.h
+++ b/Engine/Source/Runtime/Engine/Internal/Materials/MaterialExpressionMaterialSample.h
@@ -28,7 +28,7 @@ class UMaterialExpressionMaterialSample: public UMaterialExpression
 #if ENABLE_MATERIAL_SAMPLE_PROTOTYPE
 	ENGINE_API virtual int32 Compile(class FMaterialCompiler* Compiler, int32 OutputIndex) override;
 	ENGINE_API virtual int32 CompilePreview(FMaterialCompiler* Compiler, int32 OutputIndex) override;
-	ENGINE_API virtual uint32 GetOutputType(int32 OutputIndex) override;
+	ENGINE_API virtual EMaterialValueType GetOutputValueType(int32 OutputIndex) override;
 	virtual bool IsResultSubstrateMaterial(int32 OutputIndex) override;
 	virtual FSubstrateOperator * SubstrateGenerateMaterialTopologyTree(class FMaterialCompiler* Compiler, class UMaterialExpression* Parent, int32 OutputIndex) override; 
 
diff --git a/Engine/Source/Runtime/Engine/Private/Materials/Material.cpp b/Engine/Source/Runtime/Engine/Private/Materials/Material.cpp
--- a/Engine/Source/Runtime/Engine/Private/Materials/Material.cpp
+++ b/Engine/Source/Runtime/Engine/Private/Materials/Material.cpp
@@ -6824,7 +6824,7 @@ bool UMaterial::RecursiveGetExpressionChain(
 	}
 	else if ((SetMaterialAttributesExp = Cast<UMaterialExpressionSetMaterialAttributes>(InExpression)) != nullptr)
 	{
-		checkf(SetMaterialAttributesExp->GetInput(0) && (SetMaterialAttributesExp->GetInputType(0) == MCT_MaterialAttributes), TEXT("There must always be one input at least : the material attribute pin"));
+		checkf(SetMaterialAttributesExp->GetInput(0) && (SetMaterialAttributesExp->GetInputValueType(0) == MCT_MaterialAttributes), TEXT("There must always be one input at least : the material attribute pin"));
 		// Always add the material attribute input, so that we keep on traversing up the property chain : 
 		Inputs.Add(SetMaterialAttributesExp->GetInput(0));
 		InputsFrequency.Add(InShaderFrequency);
diff --git a/Engine/Source/Runtime/Engine/Private/Materials/MaterialExpressionNoise.cpp b/Engine/Source/Runtime/Engine/Private/Materials/MaterialExpressionNoise.cpp
--- a/Engine/Source/Runtime/Engine/Private/Materials/MaterialExpressionNoise.cpp
+++ b/Engine/Source/Runtime/Engine/Private/Materials/MaterialExpressionNoise.cpp
@@ -272,7 +272,7 @@ UMaterialExpressionScalarBlueNoise::UMaterialExpressionScalarBlueNoise(const FOb
 
 #if WITH_EDITOR
 
-uint32 UMaterialExpressionScalarBlueNoise::GetOutputType(int32 OutputIndex) 
+EMaterialValueType UMaterialExpressionScalarBlueNoise::GetOutputValueType(int32 OutputIndex) 
 {
 	return MCT_Float1;
 }
diff --git a/Engine/Source/Runtime/Engine/Private/Materials/MaterialExpressionSubstrate.cpp b/Engine/Source/Runtime/Engine/Private/Materials/MaterialExpressionSubstrate.cpp
--- a/Engine/Source/Runtime/Engine/Private/Materials/MaterialExpressionSubstrate.cpp
+++ b/Engine/Source/Runtime/Engine/Private/Materials/MaterialExpressionSubstrate.cpp
@@ -398,12 +398,12 @@ void UMaterialExpressionSubstrateShadingModels::GetCaption(TArray<FString>& OutC
 	OutCaptions.Add(TEXT("Substrate Shading Models"));
 }
 
-uint32 UMaterialExpressionSubstrateShadingModels::GetOutputType(int32 OutputIndex)
+EMaterialValueType UMaterialExpressionSubstrateShadingModels::GetOutputValueType(int32 OutputIndex)
 {
 	return MCT_Substrate;
 }
 
-uint32 UMaterialExpressionSubstrateShadingModels::GetInputType(int32 InputIndex)
+EMaterialValueType UMaterialExpressionSubstrateShadingModels::GetInputValueType(int32 InputIndex)
 {
 	if (InputIndex == 0)	   return MCT_Float3; // BaseColor
 	else if (InputIndex == 1)  return MCT_Float1; // Metallic
@@ -1036,12 +1036,12 @@ void UMaterialExpressionSubstrateSlabBSDF::GetExpressionToolTip(TArray<FString>&
 	OutToolTip.Add(TEXT(" - Complex Special means the Slab is using more advanced features such as glints or specular LUT. This is the most expenssive shading path."));
 }
 
-uint32 UMaterialExpressionSubstrateSlabBSDF::GetOutputType(int32 OutputIndex)
+EMaterialValueType UMaterialExpressionSubstrateSlabBSDF::GetOutputValueType(int32 OutputIndex)
 {
 	return MCT_Substrate;
 }
 
-uint32 UMaterialExpressionSubstrateSlabBSDF::GetInputType(int32 InputIndex)
+EMaterialValueType UMaterialExpressionSubstrateSlabBSDF::GetInputValueType(int32 InputIndex)
 {
 	if (InputIndex == 0)
 	{
@@ -1416,12 +1416,12 @@ void UMaterialExpressionSubstrateSimpleClearCoatBSDF::GetCaption(TArray<FString>
 	OutCaptions.Add(TEXT("Substrate Simple Clear Coat"));
 }
 
-uint32 UMaterialExpressionSubstrateSimpleClearCoatBSDF::GetOutputType(int32 OutputIndex)
+EMaterialValueType UMaterialExpressionSubstrateSimpleClearCoatBSDF::GetOutputValueType(int32 OutputIndex)
 {
 	return MCT_Substrate;
 }
 
-uint32 UMaterialExpressionSubstrateSimpleClearCoatBSDF::GetInputType(int32 InputIndex)
+EMaterialValueType UMaterialExpressionSubstrateSimpleClearCoatBSDF::GetInputValueType(int32 InputIndex)
 {
 	if (InputIndex == 0)
 	{
@@ -1578,12 +1578,12 @@ void UMaterialExpressionSubstrateVolumetricFogCloudBSDF::GetCaption(TArray<FStri
 	OutCaptions.Add(Caption);
 }
 
-uint32 UMaterialExpressionSubstrateVolumetricFogCloudBSDF::GetOutputType(int32 OutputIndex)
+EMaterialValueType UMaterialExpressionSubstrateVolumetricFogCloudBSDF::GetOutputValueType(int32 OutputIndex)
 {
 	return MCT_Substrate;
 }
 
-uint32 UMaterialExpressionSubstrateVolumetricFogCloudBSDF::GetInputType(int32 InputIndex)
+EMaterialValueType UMaterialExpressionSubstrateVolumetricFogCloudBSDF::GetInputValueType(int32 InputIndex)
 {
 	switch (InputIndex)
 	{
@@ -1674,12 +1674,12 @@ void UMaterialExpressionSubstrateLightFunction::GetCaption(TArray<FString>& OutC
 	OutCaptions.Add(TEXT("Substrate Light Function"));
 }
 
-uint32 UMaterialExpressionSubstrateLightFunction::GetOutputType(int32 OutputIndex)
+EMaterialValueType UMaterialExpressionSubstrateLightFunction::GetOutputValueType(int32 OutputIndex)
 {
 	return MCT_Substrate;
 }
 
-uint32 UMaterialExpressionSubstrateLightFunction::GetInputType(int32 InputIndex)
+EMaterialValueType UMaterialExpressionSubstrateLightFunction::GetInputValueType(int32 InputIndex)
 {
 	switch (InputIndex)
 	{
@@ -1761,12 +1761,12 @@ void UMaterialExpressionSubstratePostProcess::GetCaption(TArray<FString>& OutCap
 	OutCaptions.Add(TEXT("Substrate Post Process"));
 }
 
-uint32 UMaterialExpressionSubstratePostProcess::GetOutputType(int32 OutputIndex)
+EMaterialValueType UMaterialExpressionSubstratePostProcess::GetOutputValueType(int32 OutputIndex)
 {
 	return MCT_Substrate;
 }
 
-uint32 UMaterialExpressionSubstratePostProcess::GetInputType(int32 InputIndex)
+EMaterialValueType UMaterialExpressionSubstratePostProcess::GetInputValueType(int32 InputIndex)
 {
 	switch (InputIndex)
 	{
@@ -1848,12 +1848,12 @@ void UMaterialExpressionSubstrateUI::GetCaption(TArray<FString>& OutCaptions) co
 	OutCaptions.Add(TEXT("Substrate UI"));
 }
 
-uint32 UMaterialExpressionSubstrateUI::GetOutputType(int32 OutputIndex)
+EMaterialValueType UMaterialExpressionSubstrateUI::GetOutputValueType(int32 OutputIndex)
 {
 	return MCT_Substrate;
 }
 
-uint32 UMaterialExpressionSubstrateUI::GetInputType(int32 InputIndex)
+EMaterialValueType UMaterialExpressionSubstrateUI::GetInputValueType(int32 InputIndex)
 {
 	switch (InputIndex)
 	{
@@ -1950,12 +1950,12 @@ void UMaterialExpressionSubstrateConvertToDecal::GetCaption(TArray<FString>& Out
 	OutCaptions.Add(TEXT("Substrate Convert To Decal"));
 }
 
-uint32 UMaterialExpressionSubstrateConvertToDecal::GetOutputType(int32 OutputIndex)
+EMaterialValueType UMaterialExpressionSubstrateConvertToDecal::GetOutputValueType(int32 OutputIndex)
 {
 	return MCT_Substrate;
 }
 
-uint32 UMaterialExpressionSubstrateConvertToDecal::GetInputType(int32 InputIndex)
+EMaterialValueType UMaterialExpressionSubstrateConvertToDecal::GetInputValueType(int32 InputIndex)
 {
 	switch (InputIndex)
 	{
@@ -2051,12 +2051,12 @@ void UMaterialExpressionSubstrateUnlitBSDF::GetCaption(TArray<FString>& OutCapti
 	OutCaptions.Add(TEXT("Substrate Unlit BSDF"));
 }
 
-uint32 UMaterialExpressionSubstrateUnlitBSDF::GetOutputType(int32 OutputIndex)
+EMaterialValueType UMaterialExpressionSubstrateUnlitBSDF::GetOutputValueType(int32 OutputIndex)
 {
 	return MCT_Substrate;
 }
 
-uint32 UMaterialExpressionSubstrateUnlitBSDF::GetInputType(int32 InputIndex)
+EMaterialValueType UMaterialExpressionSubstrateUnlitBSDF::GetInputValueType(int32 InputIndex)
 {
 	switch (InputIndex)
 	{
@@ -2148,12 +2148,12 @@ void UMaterialExpressionSubstrateHairBSDF::GetCaption(TArray<FString>& OutCaptio
 	OutCaptions.Add(TEXT("Substrate Hair BSDF"));
 }
 
-uint32 UMaterialExpressionSubstrateHairBSDF::GetOutputType(int32 OutputIndex)
+EMaterialValueType UMaterialExpressionSubstrateHairBSDF::GetOutputValueType(int32 OutputIndex)
 {
 	return MCT_Substrate;
 }
 
-uint32 UMaterialExpressionSubstrateHairBSDF::GetInputType(int32 InputIndex)
+EMaterialValueType UMaterialExpressionSubstrateHairBSDF::GetInputValueType(int32 InputIndex)
 {
 	switch (InputIndex)
 	{
@@ -2269,12 +2269,12 @@ void UMaterialExpressionSubstrateEyeBSDF::GetCaption(TArray<FString>& OutCaption
 	OutCaptions.Add(TEXT("Substrate Eye BSDF"));
 }
 
-uint32 UMaterialExpressionSubstrateEyeBSDF::GetOutputType(int32 OutputIndex)
+EMaterialValueType UMaterialExpressionSubstrateEyeBSDF::GetOutputValueType(int32 OutputIndex)
 {
 	return MCT_Substrate;
 }
 
-uint32 UMaterialExpressionSubstrateEyeBSDF::GetInputType(int32 InputIndex)
+EMaterialValueType UMaterialExpressionSubstrateEyeBSDF::GetInputValueType(int32 InputIndex)
 {
 	switch (InputIndex)
 	{
@@ -2392,12 +2392,12 @@ void UMaterialExpressionSubstrateSingleLayerWaterBSDF::GetCaption(TArray<FString
 	OutCaptions.Add(TEXT("Substrate Single Layer Water BSDF"));
 }
 
-uint32 UMaterialExpressionSubstrateSingleLayerWaterBSDF::GetOutputType(int32 OutputIndex)
+EMaterialValueType UMaterialExpressionSubstrateSingleLayerWaterBSDF::GetOutputValueType(int32 OutputIndex)
 {
 	return MCT_Substrate;
 }
 
-uint32 UMaterialExpressionSubstrateSingleLayerWaterBSDF::GetInputType(int32 InputIndex)
+EMaterialValueType UMaterialExpressionSubstrateSingleLayerWaterBSDF::GetInputValueType(int32 InputIndex)
 {
 	switch (InputIndex)
 	{
@@ -2568,12 +2568,12 @@ void UMaterialExpressionSubstrateHorizontalMixing::GetCaption(TArray<FString>& O
 	}
 }
 
-uint32 UMaterialExpressionSubstrateHorizontalMixing::GetOutputType(int32 OutputIndex)
+EMaterialValueType UMaterialExpressionSubstrateHorizontalMixing::GetOutputValueType(int32 OutputIndex)
 {
 	return MCT_Substrate;
 }
 
-uint32 UMaterialExpressionSubstrateHorizontalMixing::GetInputType(int32 InputIndex)
+EMaterialValueType UMaterialExpressionSubstrateHorizontalMixing::GetInputValueType(int32 InputIndex)
 {
 	return InputIndex == 2 ? MCT_Float1 : MCT_Substrate;
 }
@@ -2764,12 +2764,12 @@ FName UMaterialExpressionSubstrateVerticalLayering::GetInputName(int32 InputInde
 	return TEXT("Unknown");
 }
 
-uint32 UMaterialExpressionSubstrateVerticalLayering::GetOutputType(int32 OutputIndex)
+EMaterialValueType UMaterialExpressionSubstrateVerticalLayering::GetOutputValueType(int32 OutputIndex)
 {
 	return MCT_Substrate;
 }
 
-uint32 UMaterialExpressionSubstrateVerticalLayering::GetInputType(int32 InputIndex)
+EMaterialValueType UMaterialExpressionSubstrateVerticalLayering::GetInputValueType(int32 InputIndex)
 {
 	if (InputIndex == 2)
 	{
@@ -2947,12 +2947,12 @@ void UMaterialExpressionSubstrateAdd::GetCaption(TArray<FString>& OutCaptions) c
 	}
 }
 
-uint32 UMaterialExpressionSubstrateAdd::GetOutputType(int32 OutputIndex)
+EMaterialValueType UMaterialExpressionSubstrateAdd::GetOutputValueType(int32 OutputIndex)
 {
 	return MCT_Substrate;
 }
 
-uint32 UMaterialExpressionSubstrateAdd::GetInputType(int32 InputIndex)
+EMaterialValueType UMaterialExpressionSubstrateAdd::GetInputValueType(int32 InputIndex)
 {
 	return MCT_Substrate;
 }
@@ -3086,12 +3086,12 @@ void UMaterialExpressionSubstrateWeight::GetCaption(TArray<FString>& OutCaptions
 	OutCaptions.Add(TEXT("Substrate Coverage Weight"));
 }
 
-uint32 UMaterialExpressionSubstrateWeight::GetOutputType(int32 OutputIndex)
+EMaterialValueType UMaterialExpressionSubstrateWeight::GetOutputValueType(int32 OutputIndex)
 {
 	return MCT_Substrate;
 }
 
-uint32 UMaterialExpressionSubstrateWeight::GetInputType(int32 InputIndex)
+EMaterialValueType UMaterialExpressionSubstrateWeight::GetInputValueType(int32 InputIndex)
 {
 	if (InputIndex == 0)
 	{
@@ -3248,12 +3248,12 @@ void UMaterialExpressionSubstrateSelect::GetCaption(TArray<FString>& OutCaptions
 	}
 }
 
-uint32 UMaterialExpressionSubstrateSelect::GetOutputType(int32 OutputIndex)
+EMaterialValueType UMaterialExpressionSubstrateSelect::GetOutputValueType(int32 OutputIndex)
 {
 	return MCT_Substrate;
 }
 
-uint32 UMaterialExpressionSubstrateSelect::GetInputType(int32 InputIndex)
+EMaterialValueType UMaterialExpressionSubstrateSelect::GetInputValueType(int32 InputIndex)
 {
 	if (InputIndex == 2)
 	{
@@ -3376,7 +3376,7 @@ void UMaterialExpressionSubstrateTransmittanceToMFP::GetCaption(TArray<FString>&
 	OutCaptions.Add(TEXT("Substrate Transmittance-To-MeanFreePath"));
 }
 
-uint32 UMaterialExpressionSubstrateTransmittanceToMFP::GetOutputType(int32 OutputIndex)
+EMaterialValueType UMaterialExpressionSubstrateTransmittanceToMFP::GetOutputValueType(int32 OutputIndex)
 {
 	switch (OutputIndex)
 	{
@@ -3390,7 +3390,7 @@ uint32 UMaterialExpressionSubstrateTransmittanceToMFP::GetOutputType(int32 Outpu
 	return MCT_Float1;
 }
 
-uint32 UMaterialExpressionSubstrateTransmittanceToMFP::GetInputType(int32 InputIndex)
+EMaterialValueType UMaterialExpressionSubstrateTransmittanceToMFP::GetInputValueType(int32 InputIndex)
 {
 	switch (InputIndex)
 	{
@@ -3489,7 +3489,7 @@ void UMaterialExpressionSubstrateMetalnessToDiffuseAlbedoF0::GetCaption(TArray<F
 	OutCaptions.Add(TEXT("Substrate Metalness-To-DiffuseAlbedo-F0"));
 }
 
-uint32 UMaterialExpressionSubstrateMetalnessToDiffuseAlbedoF0::GetOutputType(int32 OutputIndex)
+EMaterialValueType UMaterialExpressionSubstrateMetalnessToDiffuseAlbedoF0::GetOutputValueType(int32 OutputIndex)
 {
 	switch (OutputIndex)
 	{
@@ -3503,7 +3503,7 @@ uint32 UMaterialExpressionSubstrateMetalnessToDiffuseAlbedoF0::GetOutputType(int
 	return MCT_Float1;
 }
 
-uint32 UMaterialExpressionSubstrateMetalnessToDiffuseAlbedoF0::GetInputType(int32 InputIndex)
+EMaterialValueType UMaterialExpressionSubstrateMetalnessToDiffuseAlbedoF0::GetInputValueType(int32 InputIndex)
 {
 	if (InputIndex == 0) { return MCT_Float3; }
 	if (InputIndex == 1) { return MCT_Float1; }
@@ -3574,7 +3574,7 @@ void UMaterialExpressionSubstrateHazinessToSecondaryRoughness::GetCaption(TArray
 	OutCaptions.Add(TEXT("Substrate Haziness-To-Secondary-Roughness"));
 }
 
-uint32 UMaterialExpressionSubstrateHazinessToSecondaryRoughness::GetOutputType(int32 OutputIndex)
+EMaterialValueType UMaterialExpressionSubstrateHazinessToSecondaryRoughness::GetOutputValueType(int32 OutputIndex)
 {
 	switch (OutputIndex)
 	{
@@ -3588,7 +3588,7 @@ uint32 UMaterialExpressionSubstrateHazinessToSecondaryRoughness::GetOutputType(i
 	return MCT_Float1;
 }
 
-uint32 UMaterialExpressionSubstrateHazinessToSecondaryRoughness::GetInputType(int32 InputIndex)
+EMaterialValueType UMaterialExpressionSubstrateHazinessToSecondaryRoughness::GetInputValueType(int32 InputIndex)
 {
 	switch (InputIndex)
 	{
@@ -3698,7 +3698,7 @@ void UMaterialExpressionSubstrateThinFilm::GetCaption(TArray<FString>& OutCaptio
 	OutCaptions.Add(TEXT("Substrate Thin-Film"));
 }
 
-uint32 UMaterialExpressionSubstrateThinFilm::GetOutputType(int32 OutputIndex)
+EMaterialValueType UMaterialExpressionSubstrateThinFilm::GetOutputValueType(int32 OutputIndex)
 {
 	switch (OutputIndex)
 	{
@@ -3712,7 +3712,7 @@ uint32 UMaterialExpressionSubstrateThinFilm::GetOutputType(int32 OutputIndex)
 	return MCT_Float1;
 }
 
-uint32 UMaterialExpressionSubstrateThinFilm::GetInputType(int32 InputIndex)
+EMaterialValueType UMaterialExpressionSubstrateThinFilm::GetInputValueType(int32 InputIndex)
 {
 	if (InputIndex == 0) { return MCT_Float3; } // Normal
 	if (InputIndex == 1) { return MCT_Float3; } // F0
@@ -3993,7 +3993,7 @@ void UMaterialExpressionSubstrateConvertMaterialAttributes::GetCaption(TArray<FS
 	OutCaptions.Add(TEXT("Substrate Convert Material Attributes"));
 }
 
-uint32 UMaterialExpressionSubstrateConvertMaterialAttributes::GetOutputType(int32 OutputIndex)
+EMaterialValueType UMaterialExpressionSubstrateConvertMaterialAttributes::GetOutputValueType(int32 OutputIndex)
 {
 	switch (OutputIndex)
 	{
@@ -4003,7 +4003,7 @@ uint32 UMaterialExpressionSubstrateConvertMaterialAttributes::GetOutputType(int3
 	return MCT_Float1;
 }
 
-uint32 UMaterialExpressionSubstrateConvertMaterialAttributes::GetInputType(int32 InputIndex)
+EMaterialValueType UMaterialExpressionSubstrateConvertMaterialAttributes::GetInputValueType(int32 InputIndex)
 {
 	if (InputIndex == 0)	  return MCT_MaterialAttributes; // MaterialAttributes
 	else if (InputIndex == 1) return MCT_Float3; // WaterScatteringCoefficients
diff --git a/Engine/Source/Runtime/Engine/Private/Materials/MaterialExpressions.cpp b/Engine/Source/Runtime/Engine/Private/Materials/MaterialExpressions.cpp
--- a/Engine/Source/Runtime/Engine/Private/Materials/MaterialExpressions.cpp
+++ b/Engine/Source/Runtime/Engine/Private/Materials/MaterialExpressions.cpp
@@ -2555,7 +2555,7 @@ void UMaterialExpressionTextureSample::ApplyChannelNames()
 	static const FName Green("G");
 	static const FName Blue("B");
 	static const FName Alpha("A");
-	if (GetOutputType(0) != MCT_Texture)
+	if (GetOutputValueType(0) != MCT_Texture)
 	{
 		Outputs[1].OutputName = !ChannelNames.R.IsEmpty() ? FName(*ChannelNames.R.ToString()) : Red;
 		Outputs[2].OutputName = !ChannelNames.G.IsEmpty() ? FName(*ChannelNames.G.ToString()) : Green;
@@ -2910,7 +2910,8 @@ void UMaterialExpressionTextureSample::GetCaption(TArray<FString>& OutCaptions)
 
 // this define is only used for the following function
 #define IF_INPUT_RETURN(Type) if(!InputIndex) return (Type); --InputIndex
-uint32 UMaterialExpressionTextureSample::GetInputType(int32 InputIndex)
+
+EMaterialValueType UMaterialExpressionTextureSample::GetInputValueType(int32 InputIndex)
 {
 	// Coordinates
 	IF_INPUT_RETURN(MCT_Float);
@@ -3830,7 +3831,7 @@ void UMaterialExpressionMeshPaintTextureObject::GetExpressionToolTip(TArray<FStr
 	ConvertToMultilineToolTip(TEXT("Get the Mesh Paint Texture object for feeding to a Texture Sample node."), 40, OutToolTip);
 }
 
-uint32 UMaterialExpressionMeshPaintTextureObject::GetOutputType(int32 OutputIndex)
+EMaterialValueType UMaterialExpressionMeshPaintTextureObject::GetOutputValueType(int32 OutputIndex)
 {
 	return MCT_Texture2D;
 }
@@ -4110,7 +4111,7 @@ int32 UMaterialExpressionTextureObject::CompilePreview(class FMaterialCompiler*
 	return CompileTextureSample(Compiler, Texture, Compiler->TextureCoordinate(0, false, false), UMaterialExpressionTextureBase::GetSamplerTypeForTexture( Texture ));
 }
 
-uint32 UMaterialExpressionTextureObject::GetOutputType(int32 OutputIndex)
+EMaterialValueType UMaterialExpressionTextureObject::GetOutputValueType(int32 OutputIndex)
 {
 	if (Cast<UTextureCube>(Texture) != nullptr)
 	{
@@ -4196,10 +4197,11 @@ void UMaterialExpressionTextureProperty::GetCaption(TArray<FString>& OutCaptions
 
 // this define is only used for the following function
 #define IF_INPUT_RETURN(Type) if(!InputIndex) return Type; --InputIndex
-uint32 UMaterialExpressionTextureProperty::GetInputType(int32 InputIndex)
+
+EMaterialValueType UMaterialExpressionTextureProperty::GetInputValueType(int32 InputIndex)
 {
 	// TextureObject
-	IF_INPUT_RETURN(MCT_Texture | MCT_SparseVolumeTexture);
+	IF_INPUT_RETURN(static_cast<EMaterialValueType>(MCT_Texture | MCT_SparseVolumeTexture));
 	return MCT_Unknown;
 }
 #undef IF_INPUT_RETURN
@@ -5845,7 +5847,7 @@ void UMaterialExpressionMakeMaterialAttributes::GetCaption(TArray<FString>& OutC
 	OutCaptions.Add(TEXT("MakeMaterialAttributes"));
 }
 
-uint32 UMaterialExpressionMakeMaterialAttributes::GetInputType(int32 InputIndex)
+EMaterialValueType UMaterialExpressionMakeMaterialAttributes::GetInputValueType(int32 InputIndex)
 {
 	if (GetInputName(InputIndex).IsEqual("ShadingModel"))
 	{
@@ -5853,7 +5855,7 @@ uint32 UMaterialExpressionMakeMaterialAttributes::GetInputType(int32 InputIndex)
 	}
 	else
 	{
-		return UMaterialExpression::GetInputType(InputIndex);
+		return UMaterialExpression::GetInputValueType(InputIndex);
 	}
 }
 
@@ -6019,7 +6021,7 @@ bool UMaterialExpressionBreakMaterialAttributes::IsInputConnectionRequired(int32
 	return true;
 }
 
-uint32 UMaterialExpressionBreakMaterialAttributes::GetOutputType(int32 OutputIndex)
+EMaterialValueType UMaterialExpressionBreakMaterialAttributes::GetOutputValueType(int32 OutputIndex)
 {
 	BuildPropertyToIOIndexMap();
 
@@ -6031,7 +6033,7 @@ uint32 UMaterialExpressionBreakMaterialAttributes::GetOutputType(int32 OutputInd
 	}
 	else
 	{
-		return UMaterialExpression::GetOutputType(OutputIndex);
+		return UMaterialExpression::GetOutputValueType(OutputIndex);
 	}
 }
 #endif // WITH_EDITOR
@@ -6136,10 +6138,10 @@ FName UMaterialExpressionGetMaterialAttributes::GetInputName(int32 InputIndex) c
 	return NAME_None;
 }
 
-uint32 UMaterialExpressionGetMaterialAttributes::GetOutputType(int32 OutputIndex)
+EMaterialValueType UMaterialExpressionGetMaterialAttributes::GetOutputValueType(int32 OutputIndex)
 {
 	// Call base class impl to get the type
-	uint32 OutputType = Super::GetOutputType(OutputIndex);
+	EMaterialValueType OutputType = Super::GetOutputValueType(OutputIndex);
 
 	// Override the type if it's a ShadingModel type
 	if (OutputIndex > GET_SET_MA_MATERIALATTRIBUTESINDEX) // "0th" place is the mandatory MaterialAttribute itself, skip it
@@ -6440,9 +6442,9 @@ FName UMaterialExpressionSetMaterialAttributes::GetInputName(int32 InputIndex) c
 	return Name;
 }
 
-uint32 UMaterialExpressionSetMaterialAttributes::GetInputType(int32 InputIndex)
+EMaterialValueType UMaterialExpressionSetMaterialAttributes::GetInputValueType(int32 InputIndex)
 {
-	uint32 InputType = MCT_Unknown;
+	EMaterialValueType InputType = MCT_Unknown;
 
 	if (InputIndex == GET_SET_MA_MATERIALATTRIBUTESINDEX)
 	{
@@ -6519,7 +6521,7 @@ bool UMaterialExpressionSetMaterialAttributes::GetSubstrateMaterialInputIndex(in
 {
 	for (InputIndex = Inputs.Num() - 1; InputIndex > GET_SET_MA_MATERIALATTRIBUTESINDEX; InputIndex--)
 	{
-		if (GetInputType(InputIndex) == MCT_Substrate && Inputs[InputIndex].IsConnected())
+		if (GetInputValueType(InputIndex) == MCT_Substrate && Inputs[InputIndex].IsConnected())
 		{
 			return true;
 		}
@@ -6527,7 +6529,7 @@ bool UMaterialExpressionSetMaterialAttributes::GetSubstrateMaterialInputIndex(in
 
 	//MA input is always in position 0
 	InputIndex = GET_SET_MA_MATERIALATTRIBUTESINDEX;
-	if(GetInputType(InputIndex) == MCT_MaterialAttributes && Inputs[InputIndex].IsConnected())
+	if(GetInputValueType(InputIndex) == MCT_MaterialAttributes && Inputs[InputIndex].IsConnected())
 	{
 		return Inputs[InputIndex].Expression->IsResultSubstrateMaterial(Inputs[InputIndex].OutputIndex);
 	}	
@@ -6885,7 +6887,8 @@ bool UMaterialExpressionLegacyBlendMaterialAttributes::IsInputConnectionRequired
 		default: return Super::IsInputConnectionRequired(InputIndex);
 	}
 }
-uint32 UMaterialExpressionLegacyBlendMaterialAttributes::GetInputType(int32 InputIndex)
+
+EMaterialValueType UMaterialExpressionLegacyBlendMaterialAttributes::GetInputValueType(int32 InputIndex)
 {
 	switch (InputIndex)
 	{
@@ -6893,7 +6896,7 @@ uint32 UMaterialExpressionLegacyBlendMaterialAttributes::GetInputType(int32 Inpu
 		case 4:
 		case 5:
 		case 6: return MCT_StaticBool;
-		default: return Super::GetInputType(InputIndex);
+		default: return Super::GetInputValueType(InputIndex);
 	}
 }
 
@@ -7681,7 +7684,7 @@ FName UMaterialExpressionMaterialAttributeLayers::GetInputName(int32 InputIndex)
 	return NAME_None;
 }
 
-uint32 UMaterialExpressionMaterialAttributeLayers::GetInputType(int32 InputIndex)
+EMaterialValueType UMaterialExpressionMaterialAttributeLayers::GetInputValueType(int32 InputIndex)
 {
 	return MCT_MaterialAttributes;
 }
@@ -7724,12 +7727,12 @@ void UMaterialExpressionFloatToUInt::GetCaption(TArray<FString>& OutCaptions) co
 	OutCaptions.Add(TEXT("FloatToUInt"));
 }
 
-uint32 UMaterialExpressionFloatToUInt::GetInputType(int32 InputIndex)
+EMaterialValueType UMaterialExpressionFloatToUInt::GetInputValueType(int32 InputIndex)
 {
 	return MCT_Float;
 }
 
-uint32 UMaterialExpressionFloatToUInt::GetOutputType(int32 OutputIndex)
+EMaterialValueType UMaterialExpressionFloatToUInt::GetOutputValueType(int32 OutputIndex)
 {
 	return MCT_UInt;
 }
@@ -7764,12 +7767,12 @@ void UMaterialExpressionUIntToFloat::GetCaption(TArray<FString>& OutCaptions) co
 	OutCaptions.Add(TEXT("UIntToFloat"));
 }
 
-uint32 UMaterialExpressionUIntToFloat::GetInputType(int32 InputIndex)
+EMaterialValueType UMaterialExpressionUIntToFloat::GetInputValueType(int32 InputIndex)
 {
 	return MCT_UInt;
 }
 
-uint32 UMaterialExpressionUIntToFloat::GetOutputType(int32 OutputIndex)
+EMaterialValueType UMaterialExpressionUIntToFloat::GetOutputValueType(int32 OutputIndex)
 {
 	return MCT_Float;
 }
@@ -7789,9 +7792,9 @@ int32 UMaterialExpressionFloor::Compile(class FMaterialCompiler* Compiler, int32
 	return Compiler->Floor(Input.Compile(Compiler));
 }
 
-uint32 UMaterialExpressionFloor::GetOutputType(int32 OutputIndex)
+EMaterialValueType UMaterialExpressionFloor::GetOutputValueType(int32 OutputIndex)
 {
-	return MCT_Float | MCT_UInt;
+	return static_cast<EMaterialValueType>(MCT_Float | MCT_UInt);
 }
 
 void UMaterialExpressionFloor::GetCaption(TArray<FString>& OutCaptions) const
@@ -7810,9 +7813,9 @@ int32 UMaterialExpressionCeil::Compile(class FMaterialCompiler* Compiler, int32
 	return Compiler->Ceil(Input.Compile(Compiler));
 }
 
-uint32 UMaterialExpressionCeil::GetOutputType(int32 OutputIndex)
+EMaterialValueType UMaterialExpressionCeil::GetOutputValueType(int32 OutputIndex)
 {
-	return MCT_Float | MCT_UInt;
+	return static_cast<EMaterialValueType>(MCT_Float | MCT_UInt);
 }
 
 void UMaterialExpressionCeil::GetCaption(TArray<FString>& OutCaptions) const
@@ -7834,9 +7837,9 @@ int32 UMaterialExpressionRound::Compile(class FMaterialCompiler* Compiler, int32
 	return Compiler->Round(Input.Compile(Compiler));
 }
 
-uint32 UMaterialExpressionRound::GetOutputType(int32 OutputIndex)
+EMaterialValueType UMaterialExpressionRound::GetOutputValueType(int32 OutputIndex)
 {
-	return MCT_Float | MCT_UInt;
+	return static_cast<EMaterialValueType>(MCT_Float | MCT_UInt);
 }
 
 void UMaterialExpressionRound::GetCaption(TArray<FString>& OutCaptions) const
@@ -7863,9 +7866,9 @@ int32 UMaterialExpressionTruncate::Compile(class FMaterialCompiler* Compiler, in
 	return Compiler->Truncate(Input.Compile(Compiler));
 }
 
-uint32 UMaterialExpressionTruncate::GetOutputType(int32 OutputIndex)
+EMaterialValueType UMaterialExpressionTruncate::GetOutputValueType(int32 OutputIndex)
 {
-	return MCT_Float | MCT_UInt;
+	return static_cast<EMaterialValueType>(MCT_Float | MCT_UInt);
 }
 
 void UMaterialExpressionTruncate::GetCaption(TArray<FString>& OutCaptions) const
@@ -7943,14 +7946,14 @@ int32 UMaterialExpressionModulo::Compile(class FMaterialCompiler* Compiler, int3
 	return Compiler->Modulo(A.Compile(Compiler), B.Compile(Compiler));
 }
 
-uint32 UMaterialExpressionModulo::GetInputType(int32 InputIndex)
+EMaterialValueType UMaterialExpressionModulo::GetInputValueType(int32 InputIndex)
 {
-	return MCT_Float | MCT_LWCType | MCT_UInt;
+	return static_cast<EMaterialValueType>(MCT_Float | MCT_LWCType | MCT_UInt);
 }
 
-uint32 UMaterialExpressionModulo::GetOutputType(int32 OutputIndex)
+EMaterialValueType UMaterialExpressionModulo::GetOutputValueType(int32 OutputIndex)
 {
-	return MCT_Float | MCT_LWCType | MCT_UInt;
+	return static_cast<EMaterialValueType>(MCT_Float | MCT_LWCType | MCT_UInt);
 }
 
 void UMaterialExpressionModulo::GetCaption(TArray<FString>& OutCaptions) const
@@ -8854,7 +8857,7 @@ FName UMaterialExpressionStaticSwitch::GetInputName(int32 InputIndex) const
 	}
 }
 
-uint32 UMaterialExpressionStaticSwitch::GetInputType(int32 InputIndex)
+EMaterialValueType UMaterialExpressionStaticSwitch::GetInputValueType(int32 InputIndex)
 {
 	if (InputIndex == 0 || InputIndex == 1)
 	{
@@ -8949,7 +8952,7 @@ FName UMaterialExpressionPreviousFrameSwitch::GetInputName(int32 InputIndex) con
 	}
 }
 
-uint32 UMaterialExpressionPreviousFrameSwitch::GetInputType(int32 InputIndex)
+EMaterialValueType UMaterialExpressionPreviousFrameSwitch::GetInputValueType(int32 InputIndex)
 {
 	return MCT_Unknown;
 }
@@ -9435,7 +9438,7 @@ void UMaterialExpressionTextureCollection::GetCaption(TArray<FString>& OutCaptio
 	OutCaptions.Emplace(TEXT("Texture Collection"));
 }
 
-uint32 UMaterialExpressionTextureCollection::GetOutputType(int32 OutputIndex)
+EMaterialValueType UMaterialExpressionTextureCollection::GetOutputValueType(int32 OutputIndex)
 {
 	if (OutputIndex == 1)
 	{
@@ -9592,7 +9595,7 @@ int32 UMaterialExpressionTextureObjectFromCollection::Compile(FMaterialCompiler*
 	return TextureFromCollectionCodeIndex;
 }
 
-uint32 UMaterialExpressionTextureObjectFromCollection::GetInputType(int32 InputIndex)
+EMaterialValueType UMaterialExpressionTextureObjectFromCollection::GetInputValueType(int32 InputIndex)
 {
 	if (InputIndex == 0)
 	{
@@ -9602,7 +9605,7 @@ uint32 UMaterialExpressionTextureObjectFromCollection::GetInputType(int32 InputI
 	return MCT_UInt1;
 }
 
-uint32 UMaterialExpressionTextureObjectFromCollection::GetOutputType(int32 OutputIndex)
+EMaterialValueType UMaterialExpressionTextureObjectFromCollection::GetOutputValueType(int32 OutputIndex)
 {
 	return MaterialValueTypeFromTextureCollectionMemberType(TextureType);
 }
@@ -10931,12 +10934,12 @@ void UMaterialExpressionIf::GetCaption(TArray<FString>& OutCaptions) const
 	OutCaptions.Add(TEXT("If"));
 }
 
-uint32 UMaterialExpressionIf::GetInputType(int32 InputIndex)
+EMaterialValueType UMaterialExpressionIf::GetInputValueType(int32 InputIndex)
 {
 	// First two inputs are always float
 	if (InputIndex == 0 || InputIndex == 1)
 	{
-		return MCT_MaterialAttributes | MCT_Numeric | MCT_ShadingModel | MCT_StaticBool | MCT_Bool;
+		return static_cast<EMaterialValueType>(MCT_MaterialAttributes | MCT_Numeric | MCT_ShadingModel | MCT_StaticBool | MCT_Bool);
 	}
 
 	return MCT_Unknown;
@@ -11441,22 +11444,22 @@ FName UMaterialExpressionComposite::GetInputName(int32 InputIndex) const
 	return FName();
 }
 
-uint32 UMaterialExpressionComposite::GetInputType(int32 InputIndex)
+EMaterialValueType UMaterialExpressionComposite::GetInputValueType(int32 InputIndex)
 {
 	if (InputExpressions && InputExpressions->ReroutePins.IsValidIndex(InputIndex))
 	{
-		return InputExpressions->ReroutePins[InputIndex].Expression->GetInputType(0);
+		return InputExpressions->ReroutePins[InputIndex].Expression->GetInputValueType(0);
 	}
 
 	check(false);
 	return MCT_Float;
 }
 
-uint32 UMaterialExpressionComposite::GetOutputType(int32 OutputIndex)
+EMaterialValueType UMaterialExpressionComposite::GetOutputValueType(int32 OutputIndex)
 {
 	if (OutputExpressions && OutputExpressions->ReroutePins.IsValidIndex(OutputIndex))
 	{
-		return OutputExpressions->ReroutePins[OutputIndex].Expression->GetOutputType(0);
+		return OutputExpressions->ReroutePins[OutputIndex].Expression->GetOutputValueType(0);
 	}
 
 	check(false);
@@ -11684,22 +11687,22 @@ FName UMaterialExpressionPinBase::GetInputName(int32 InputIndex) const
 	return FName();
 }
 
-uint32 UMaterialExpressionPinBase::GetInputType(int32 InputIndex)
+EMaterialValueType UMaterialExpressionPinBase::GetInputValueType(int32 InputIndex)
 {
 	if (InputIndex >= 0 && InputIndex < ReroutePins.Num())
 	{
-		return PinDirection == EGPD_Input ? ReroutePins[InputIndex].Expression->GetInputType(0) : MCT_Float;
+		return PinDirection == EGPD_Input ? ReroutePins[InputIndex].Expression->GetInputValueType(0) : MCT_Float;
 	}
 
 	check(false);
 	return MCT_Float;
 }
 
-uint32 UMaterialExpressionPinBase::GetOutputType(int32 OutputIndex)
+EMaterialValueType UMaterialExpressionPinBase::GetOutputValueType(int32 OutputIndex)
 {
 	if (OutputIndex >= 0 && OutputIndex < ReroutePins.Num())
 	{
-		return PinDirection == EGPD_Output ? ReroutePins[OutputIndex].Expression->GetOutputType(0) : MCT_Float;
+		return PinDirection == EGPD_Output ? ReroutePins[OutputIndex].Expression->GetOutputValueType(0) : MCT_Float;
 	}
 
 	check(false);
@@ -12722,7 +12725,7 @@ void UMaterialExpressionCustom::RebuildOutputs()
 	}
 }
 
-uint32 UMaterialExpressionCustom::GetOutputType(int32 OutputIndex)
+EMaterialValueType UMaterialExpressionCustom::GetOutputValueType(int32 OutputIndex)
 {
 	ECustomMaterialOutputType Type = CMOT_MAX;
 	if (OutputIndex == 0)
@@ -12753,7 +12756,7 @@ uint32 UMaterialExpressionCustom::GetOutputType(int32 OutputIndex)
 
 bool UMaterialExpressionCustom::IsResultMaterialAttributes(int32 OutputIndex)
 {
-	return GetOutputType(OutputIndex) == MCT_MaterialAttributes;
+	return GetOutputValueType(OutputIndex) == MCT_MaterialAttributes;
 }
 
 void UMaterialExpressionCustom::GetIncludeFilePaths(TSet<FString>& OutIncludeFilePaths) const
@@ -13591,11 +13594,11 @@ void UMaterialFunction::UpdateInputOutputTypes()
 
 		if (InputExpression)
 		{
-			CombinedInputTypes |= InputExpression->GetInputType(0);
+			CombinedInputTypes |= InputExpression->GetInputValueType(0);
 		}
 		else if (OutputExpression)
 		{
-			CombinedOutputTypes |= OutputExpression->GetOutputType(0);
+			CombinedOutputTypes |= OutputExpression->GetOutputValueType(0);
 		}
 	}
 }
@@ -13762,7 +13765,7 @@ void UMaterialFunction::ConvertExpressionsBetweenLegacyAndSubstrate()
 				//0 index is always MaterialAttributes so no need to check that entry
 				for(int32 InputIndex = InputsArray.Num() - 1; InputIndex > 0; InputIndex--)
 				{
-					if(SetAttributes->GetInputType(InputIndex) == MCT_Substrate)
+					if(SetAttributes->GetInputValueType(InputIndex) == MCT_Substrate)
 					{
 						return;
 					}
@@ -13774,7 +13777,7 @@ void UMaterialFunction::ConvertExpressionsBetweenLegacyAndSubstrate()
 				//0 index is always MaterialAttributes so no need to check that entry
 				for(int32 OutputIndex = GetAttributes->GetOutputs().Num() - 1; OutputIndex > 0; OutputIndex--)
 				{
-					if (GetAttributes->GetOutputType(OutputIndex) == MCT_Substrate)
+					if (GetAttributes->GetOutputValueType(OutputIndex) == MCT_Substrate)
 					{
 						return;
 					}
@@ -17642,13 +17645,13 @@ FSubstrateOperator* UMaterialExpressionMaterialFunctionCall::SubstrateGenerateMa
 	return nullptr;
 }
 
-uint32 UMaterialExpressionMaterialFunctionCall::GetInputType(int32 InputIndex)
+EMaterialValueType UMaterialExpressionMaterialFunctionCall::GetInputValueType(int32 InputIndex)
 {
 	if (InputIndex < FunctionInputs.Num())
 	{
 		if (FunctionInputs[InputIndex].ExpressionInput)
 		{
-			return FunctionInputs[InputIndex].ExpressionInput->GetInputType(0);
+			return FunctionInputs[InputIndex].ExpressionInput->GetInputValueType(0);
 		}
 	}
 	return MCT_Unknown;
@@ -18089,7 +18092,7 @@ FSubstrateOperator* UMaterialExpressionFunctionInput::SubstrateGenerateMaterialT
 	return SlabOperator;
 }
 
-uint32 UMaterialExpressionFunctionInput::GetInputType(int32 InputIndex)
+EMaterialValueType UMaterialExpressionFunctionInput::GetInputValueType(int32 InputIndex)
 {
 	switch (InputType)
 	{
@@ -18124,9 +18127,9 @@ uint32 UMaterialExpressionFunctionInput::GetInputType(int32 InputIndex)
 	}
 }
 
-uint32 UMaterialExpressionFunctionInput::GetOutputType(int32 OutputIndex)
+EMaterialValueType UMaterialExpressionFunctionInput::GetOutputValueType(int32 OutputIndex)
 {
-	return GetInputType(0);
+	return GetInputValueType(0);
 }
 #endif // WITH_EDITOR
 
@@ -18216,10 +18219,10 @@ void UMaterialExpressionFunctionOutput::GetExpressionToolTip(TArray<FString>& Ou
 	ConvertToMultilineToolTip(Description, 40, OutToolTip);
 }
 
-uint32 UMaterialExpressionFunctionOutput::GetInputType(int32 InputIndex)
+EMaterialValueType UMaterialExpressionFunctionOutput::GetInputValueType(int32 InputIndex)
 {
 	// Acceptable types for material function outputs
-	return MCT_Float | MCT_MaterialAttributes | MCT_Substrate;
+	return static_cast<EMaterialValueType>(MCT_Float | MCT_MaterialAttributes | MCT_Substrate);
 }
 
 int32 UMaterialExpressionFunctionOutput::Compile(class FMaterialCompiler* Compiler, int32 OutputIndex)
@@ -18978,7 +18981,7 @@ void UMaterialExpressionRayTracingQualitySwitch::GetCaption(TArray<FString>& Out
 	OutCaptions.Add(TEXT("RayTracingQualitySwitchReplace"));
 }
 
-uint32 UMaterialExpressionRayTracingQualitySwitch::GetInputType(int32 InputIndex)
+EMaterialValueType UMaterialExpressionRayTracingQualitySwitch::GetInputValueType(int32 InputIndex)
 {
 	return MCT_Unknown;
 }
@@ -19035,7 +19038,7 @@ void UMaterialExpressionPathTracingQualitySwitch::GetCaption(TArray<FString>& Ou
 	OutCaptions.Add(TEXT("PathTracingQualitySwitchReplace"));
 }
 
-uint32 UMaterialExpressionPathTracingQualitySwitch::GetInputType(int32 InputIndex)
+EMaterialValueType UMaterialExpressionPathTracingQualitySwitch::GetInputValueType(int32 InputIndex)
 {
 	return MCT_Unknown;
 }
@@ -19196,7 +19199,7 @@ FExpressionInput UMaterialExpressionRerouteBase::TraceInputsToRealExpressionInte
 }
 
 #if WITH_EDITOR
-uint32 UMaterialExpressionRerouteBase::GetInputType(int32 InputIndex)
+EMaterialValueType UMaterialExpressionRerouteBase::GetInputValueType(int32 InputIndex)
 {
 	FExpressionInput Input;
 	if (GetRerouteInput(Input))
@@ -19210,17 +19213,17 @@ uint32 UMaterialExpressionRerouteBase::GetInputType(int32 InputIndex)
 			// If we found a valid connection to a real output, then our type becomes that type.
 			if (RealExpression != nullptr && RealExpressionOutputIndex != -1 && RealExpression->Outputs.Num() > RealExpressionOutputIndex && RealExpressionOutputIndex >= 0)
 			{
-				return RealExpression->GetOutputType(RealExpressionOutputIndex);
+				return RealExpression->GetOutputValueType(RealExpressionOutputIndex);
 			}
 		}
 	}
 	return MCT_Unknown;
 }
 
-uint32 UMaterialExpressionRerouteBase::GetOutputType(int32 OutputIndex)
+EMaterialValueType UMaterialExpressionRerouteBase::GetOutputValueType(int32 OutputIndex)
 {
 	// Our node is a passthrough so input and output types must match.
-	return GetInputType(0);
+	return GetInputValueType(0);
 }
 
 bool UMaterialExpressionRerouteBase::IsResultMaterialAttributes(int32 OutputIndex)
@@ -19604,15 +19607,15 @@ void UMaterialExpressionNamedRerouteUsage::GetCaption(TArray<FString>& OutCaptio
 	OutCaptions.Add(IsDeclarationValid() ? Declaration->Name.ToString() : TEXT("Invalid named reroute"));
 }
 
-uint32 UMaterialExpressionNamedRerouteUsage::GetOutputType(int32 OutputIndex)
+EMaterialValueType UMaterialExpressionNamedRerouteUsage::GetOutputValueType(int32 OutputIndex)
 {
 	if (IsDeclarationValid())
 	{
-		return Declaration->GetInputType(OutputIndex);
+		return Declaration->GetInputValueType(OutputIndex);
 	}
 	else
 	{
-		return Super::GetOutputType(OutputIndex);
+		return Super::GetOutputValueType(OutputIndex);
 	}
 }
 
@@ -21332,15 +21335,15 @@ FExpressionInput* UMaterialExpressionVertexInterpolator::GetInput(int32 InputInd
 	return InputIndex == 0 ? &Input : nullptr;
 }
 
-uint32 UMaterialExpressionVertexInterpolator::GetInputType(int32 InputIndex)
+EMaterialValueType UMaterialExpressionVertexInterpolator::GetInputValueType(int32 InputIndex)
 {
 	// New HLSL generator allows struct interpolators
 	return IsUsingNewHLSLGenerator() ? MCT_Unknown : MCT_Float4;
 }
 
-uint32 UMaterialExpressionVertexInterpolator::GetOutputType(int32 OutputIndex)
+EMaterialValueType UMaterialExpressionVertexInterpolator::GetOutputValueType(int32 OutputIndex)
 {
-	return IsUsingNewHLSLGenerator() ? MCT_Unknown : UMaterialExpression::GetOutputType(OutputIndex);
+	return IsUsingNewHLSLGenerator() ? MCT_Unknown : UMaterialExpression::GetOutputValueType(OutputIndex);
 }
 
 #endif // WITH_EDITOR
@@ -21999,7 +22002,7 @@ void UMaterialExpressionShadingModel::GetCaption(TArray<FString>& OutCaptions) c
 	OutCaptions.Add(TEXT("Shading Model"));
 }
 
-uint32 UMaterialExpressionShadingModel::GetOutputType(int32 OutputIndex)
+EMaterialValueType UMaterialExpressionShadingModel::GetOutputValueType(int32 OutputIndex)
 {
 	return MCT_ShadingModel;
 }
@@ -22227,7 +22230,7 @@ void UMaterialExpressionVolumetricCloudEmptySpaceSkippingOutput::GetCaption(TArr
 	OutCaptions.Add(FString(TEXT("Volumetric Cloud Empty Space Skipping Output <Experimental>")));
 }
 
-uint32 UMaterialExpressionVolumetricCloudEmptySpaceSkippingOutput::GetOutputType(int32 OutputIndex)
+EMaterialValueType UMaterialExpressionVolumetricCloudEmptySpaceSkippingOutput::GetOutputValueType(int32 OutputIndex)
 {
 	return MCT_Float1;
 }
@@ -22287,7 +22290,7 @@ void UMaterialExpressionVolumetricCloudEmptySpaceSkippingInput::GetConnectorTool
 	Super::GetConnectorToolTip(InputIndex, OutputIndex, OutToolTip);
 }
 
-uint32 UMaterialExpressionVolumetricCloudEmptySpaceSkippingInput::GetInputType(int32 InputIndex)
+EMaterialValueType UMaterialExpressionVolumetricCloudEmptySpaceSkippingInput::GetInputValueType(int32 InputIndex)
 {
 	if (InputIndex == 0)
 	{
@@ -22405,7 +22408,7 @@ void UMaterialExpressionThinTranslucentMaterialOutput::GetCaption(TArray<FString
 	OutCaptions.Add(FString(TEXT("Thin Translucent Material")));
 }
 
-uint32 UMaterialExpressionThinTranslucentMaterialOutput::GetInputType(int32 InputIndex)
+EMaterialValueType UMaterialExpressionThinTranslucentMaterialOutput::GetInputValueType(int32 InputIndex)
 {
 	switch (InputIndex)
 	{
@@ -22554,14 +22557,14 @@ void UMaterialExpressionIfThenElse::GetCaption(TArray<FString>& OutCaptions) con
 	OutCaptions.Add(TEXT("If"));
 }
 
-uint32 UMaterialExpressionIfThenElse::GetInputType(int32 InputIndex)
+EMaterialValueType UMaterialExpressionIfThenElse::GetInputValueType(int32 InputIndex)
 {
 	switch (InputIndex)
 	{
 		case 0: return MCT_Bool;
 		case 1: return MCT_Unknown;
 		case 2: return MCT_Unknown;
-		default: checkNoEntry(); return 0u;
+		default: checkNoEntry(); return {};
 	}
 }
 #endif
@@ -22649,7 +22652,7 @@ void UMaterialExpressionSparseVolumeTextureObject::GetCaption(TArray<FString>& O
 	OutCaptions.Add(TEXT("SparseVolumeTexture Object"));
 }
 
-uint32 UMaterialExpressionSparseVolumeTextureObject::GetOutputType(int32 OutputIndex)
+EMaterialValueType UMaterialExpressionSparseVolumeTextureObject::GetOutputValueType(int32 OutputIndex)
 {
 	return MCT_SparseVolumeTexture;
 }
@@ -22693,7 +22696,7 @@ int32 UMaterialExpressionSparseVolumeTextureObjectParameter::Compile(class FMate
 	return Compiler->SparseVolumeTextureParameter(ParameterName, SparseVolumeTexture, TextureReferenceIndex, SAMPLERTYPE_LinearColor);
 }
 
-uint32 UMaterialExpressionSparseVolumeTextureObjectParameter::GetOutputType(int32 OutputIndex)
+EMaterialValueType UMaterialExpressionSparseVolumeTextureObjectParameter::GetOutputValueType(int32 OutputIndex)
 {
 	return MCT_SparseVolumeTexture;
 }
@@ -22778,7 +22781,7 @@ void UMaterialExpressionSparseVolumeTextureSample::PostEditChangeProperty(FPrope
 	Super::PostEditChangeProperty(PropertyChangedEvent);
 }
 
-uint32 UMaterialExpressionSparseVolumeTextureSample::GetOutputType(int32 OutputIndex)
+EMaterialValueType UMaterialExpressionSparseVolumeTextureSample::GetOutputValueType(int32 OutputIndex)
 {
 	switch (OutputIndex)
 	{
@@ -22847,7 +22850,8 @@ FName UMaterialExpressionSparseVolumeTextureSample::GetInputName(int32 InputInde
 
 // this define is only used for the following function
 #define IF_INPUT_RETURN(Type) if(!InputIndex) return (Type); --InputIndex
-uint32 UMaterialExpressionSparseVolumeTextureSample::GetInputType(int32 InputIndex)
+
+EMaterialValueType UMaterialExpressionSparseVolumeTextureSample::GetInputValueType(int32 InputIndex)
 {
 	// Coordinates
 	IF_INPUT_RETURN(MCT_Float3);
@@ -23078,7 +23082,7 @@ void UMaterialExpressionSubsurfaceMediumMaterialOutput::GetCaption(TArray<FStrin
 	OutCaptions.Add(FString(TEXT("Subsurface Medium (Path Tracer Only)")));
 }
 
-uint32 UMaterialExpressionSubsurfaceMediumMaterialOutput::GetInputType(int32 InputIndex)
+EMaterialValueType UMaterialExpressionSubsurfaceMediumMaterialOutput::GetInputValueType(int32 InputIndex)
 {
 	switch (InputIndex)
 	{
@@ -23172,7 +23176,7 @@ void UMaterialExpressionNeuralNetworkInput::GetCaption(TArray<FString>& OutCapti
 	}
 }
 
-uint32 UMaterialExpressionNeuralNetworkInput::GetInputType(int32 InputIndex)
+EMaterialValueType UMaterialExpressionNeuralNetworkInput::GetInputValueType(int32 InputIndex)
 {
 
 	switch (InputIndex)
@@ -23238,7 +23242,8 @@ void UMaterialExpressionNeuralNetworkOutput::GetCaption(TArray<FString>& OutCapt
 		OutCaptions.Add(TEXT("Neural Output (Buffer)"));
 	}
 }
-uint32 UMaterialExpressionNeuralNetworkOutput::GetInputType(int32 InputIndex)
+
+EMaterialValueType UMaterialExpressionNeuralNetworkOutput::GetInputValueType(int32 InputIndex)
 {
 	if (NeuralIndexType == ENeuralIndexType::NIT_TextureIndex)
 	{
@@ -23488,7 +23493,7 @@ int32 UMaterialExpressionConvert::CountInputs() const
 	return ConvertInputs.Num();
 }
 
-uint32 UMaterialExpressionConvert::GetInputType(int32 InputIndex)
+EMaterialValueType UMaterialExpressionConvert::GetInputValueType(int32 InputIndex)
 {
 	if (ConvertInputs.IsValidIndex(InputIndex))
 	{
@@ -23509,7 +23514,7 @@ FExpressionInput* UMaterialExpressionConvert::GetInput(int32 InputIndex)
 	return nullptr;
 }
 
-uint32 UMaterialExpressionConvert::GetOutputType(int32 OutputIndex)
+EMaterialValueType UMaterialExpressionConvert::GetOutputValueType(int32 OutputIndex)
 {
 	if (ConvertOutputs.IsValidIndex(OutputIndex))
 	{
diff --git a/Engine/Source/Runtime/Engine/Private/Materials/MaterialExternalCodeRegistry.h b/Engine/Source/Runtime/Engine/Private/Materials/MaterialExternalCodeRegistry.h
--- a/Engine/Source/Runtime/Engine/Private/Materials/MaterialExternalCodeRegistry.h
+++ b/Engine/Source/Runtime/Engine/Private/Materials/MaterialExternalCodeRegistry.h
@@ -29,7 +29,7 @@ ENUM_CLASS_FLAGS(EMaterialShaderFrequency);
  * NOTE: Remove UMETA(Hidden) markers once the respective entry is used in BaseMaterialExpressions.ini file.
  */
 UENUM(Flags)
-enum class EMaterialValueTypeBridge : uint32
+enum class EMaterialValueTypeBridge : uint64
 {
 	Float1              = MCT_Float1,
 	Float2              = MCT_Float2,
diff --git a/Engine/Source/Runtime/Engine/Private/Materials/MaterialIRModuleBuilder.cpp b/Engine/Source/Runtime/Engine/Private/Materials/MaterialIRModuleBuilder.cpp
--- a/Engine/Source/Runtime/Engine/Private/Materials/MaterialIRModuleBuilder.cpp
+++ b/Engine/Source/Runtime/Engine/Private/Materials/MaterialIRModuleBuilder.cpp
@@ -226,7 +226,7 @@ struct FMaterialIRModuleBuilderImpl
 			MIR::FValue* Value = Emitter->Input(Call->GetInput(i));
 			if (Value)
 			{
-				const MIR::FType* Type = MIR::FType::FromMaterialValueType((EMaterialValueType)FunctionInput.ExpressionInput->GetInputType(0));
+				const MIR::FType* Type = MIR::FType::FromMaterialValueType(FunctionInput.ExpressionInput->GetInputValueType(0));
 				CallInputValues[i] = Emitter->EmitConstruct(Type, Value);
 			}
 		}
@@ -283,7 +283,7 @@ struct FMaterialIRModuleBuilderImpl
 			MIR::FValue* Value = Emitter->Input(FunctionOutput.ExpressionOutput->GetInput(0));
 
 			// Get the function output type.
-			const MIR::FType* OutputType = MIR::FType::FromMaterialValueType((EMaterialValueType)FunctionOutput.ExpressionOutput->GetOutputType(0));
+			const MIR::FType* OutputType = MIR::FType::FromMaterialValueType(FunctionOutput.ExpressionOutput->GetOutputValueType(0));
 
 			// Cast the value to the expected output type. This may fail (value will be poison). 
 			Value = Emitter->EmitConstruct(OutputType, Value);
diff --git a/Engine/Source/Runtime/Engine/Public/MaterialValueType.h b/Engine/Source/Runtime/Engine/Public/MaterialValueType.h
--- a/Engine/Source/Runtime/Engine/Public/MaterialValueType.h
+++ b/Engine/Source/Runtime/Engine/Public/MaterialValueType.h
@@ -11,7 +11,7 @@
 /**
  * The types which can be used by materials.
  */
-enum EMaterialValueType : uint32
+enum EMaterialValueType : uint64
 {
 	/** 
 	 * A scalar float type.  
@@ -66,7 +66,11 @@ enum EMaterialValueType : uint32
 
 	MCT_TextureCollection     = 1u << 29,
 	MCT_TextureMeshPaint      = 1u << 30,
-	MCT_TextureMaterialCache    = 1u << 31,
+	MCT_TextureMaterialCache  = 1u << 31,
+
+	/** Reserved range for licensee modifications, any bits in the (inclusive) range is free for modifications */
+	MCT_LicenseeReservedBegin     = 1ull << 48,
+	MCT_LicenseeReservedEnd       = 1ull << 63,
 
 	/** MCT_SparseVolumeTexture is intentionally not (yet) included here because it differs a lot from the other texture types and may not be supported/appropriate for all MCT_Texture use cases. */
 	MCT_Texture = MCT_Texture2D | MCT_TextureCube | MCT_Texture2DArray | MCT_TextureCubeArray | MCT_VolumeTexture | MCT_TextureExternal | MCT_TextureVirtual | MCT_TextureMeshPaint | MCT_TextureMaterialCache,
diff --git a/Engine/Source/Runtime/Engine/Public/Materials/MaterialExpression.h b/Engine/Source/Runtime/Engine/Public/Materials/MaterialExpression.h
--- a/Engine/Source/Runtime/Engine/Public/Materials/MaterialExpression.h
+++ b/Engine/Source/Runtime/Engine/Public/Materials/MaterialExpression.h
@@ -8,6 +8,7 @@
 #include "UObject/UnrealType.h"
 #include "Misc/Guid.h"
 #include "MaterialExpressionIO.h"
+#include "MaterialValueType.h"
 
 #include "MaterialExpression.generated.h"
 
@@ -341,10 +342,27 @@ class UMaterialExpression : public UObject
 	
 	ENGINE_API virtual bool IsInputConnectionRequired(int32 InputIndex) const;
 		
+	UE_DEPRECATED(5.6, "Use EMaterialValueType GetInputValueType(...) instead")
 	ENGINE_API virtual uint32 GetInputType(int32 InputIndex);
 
 	ENGINE_API virtual TArray<FExpressionOutput>& GetOutputs();
+
+	UE_DEPRECATED(5.6, "Use EMaterialValueType GetOutputValueType(...) instead")
 	ENGINE_API virtual uint32 GetOutputType(int32 OutputIndex);
+	
+	ENGINE_API virtual EMaterialValueType GetInputValueType(int32 InputIndex)
+	{
+		PRAGMA_DISABLE_DEPRECATION_WARNINGS
+		return static_cast<EMaterialValueType>(GetInputType(InputIndex));
+		PRAGMA_ENABLE_DEPRECATION_WARNINGS
+	}
+	
+	ENGINE_API virtual EMaterialValueType GetOutputValueType(int32 OutputIndex)
+	{
+		PRAGMA_DISABLE_DEPRECATION_WARNINGS
+		return static_cast<EMaterialValueType>(GetOutputType(OutputIndex));
+		PRAGMA_ENABLE_DEPRECATION_WARNINGS
+	}
 
 	virtual bool CanUserDeleteExpression() const { return true; }
 
diff --git a/Engine/Source/Runtime/Engine/Public/Materials/MaterialExpressionBentNormalCustomOutput.h b/Engine/Source/Runtime/Engine/Public/Materials/MaterialExpressionBentNormalCustomOutput.h
--- a/Engine/Source/Runtime/Engine/Public/Materials/MaterialExpressionBentNormalCustomOutput.h
+++ b/Engine/Source/Runtime/Engine/Public/Materials/MaterialExpressionBentNormalCustomOutput.h
@@ -21,7 +21,7 @@ class UMaterialExpressionBentNormalCustomOutput : public UMaterialExpressionCust
 #if WITH_EDITOR
 	virtual int32 Compile(class FMaterialCompiler* Compiler, int32 OutputIndex) override;
 	virtual void GetCaption(TArray<FString>& OutCaptions) const override;
-	virtual uint32 GetInputType(int32 InputIndex) override { return MCT_Float3; }
+	virtual EMaterialValueType GetInputValueType(int32 InputIndex) override { return MCT_Float3; }
 	virtual FExpressionInput* GetInput(int32 InputIndex) override;
 #endif
 	virtual int32 GetNumOutputs() const override { return 1; }
diff --git a/Engine/Source/Runtime/Engine/Public/Materials/MaterialExpressionBlendMaterialAttributes.h b/Engine/Source/Runtime/Engine/Public/Materials/MaterialExpressionBlendMaterialAttributes.h
--- a/Engine/Source/Runtime/Engine/Public/Materials/MaterialExpressionBlendMaterialAttributes.h
+++ b/Engine/Source/Runtime/Engine/Public/Materials/MaterialExpressionBlendMaterialAttributes.h
@@ -63,7 +63,7 @@ class UMaterialExpressionBlendMaterialAttributes : public UMaterialExpression
 	virtual bool IsResultSubstrateMaterial(int32 OutputIndex) override;
 	virtual void GatherSubstrateMaterialInfo(FSubstrateMaterialInfo& SubstrateMaterialInfo, int32 OutputIndex) override;
 	virtual FSubstrateOperator* SubstrateGenerateMaterialTopologyTree(class FMaterialCompiler* Compiler, class UMaterialExpression* Parent, int32 OutputIndex) override;
-	virtual uint32 GetInputType(int32 InputIndex) override {return InputIndex == 2 ? MCT_Float1 : MCT_MaterialAttributes;}
+	virtual EMaterialValueType GetInputValueType(int32 InputIndex) override {return InputIndex == 2 ? MCT_Float1 : MCT_MaterialAttributes;}
 #endif
 	//~ End UMaterialExpression Interface
 };
@@ -97,7 +97,7 @@ public:
 	virtual FExpressionInput* GetInput(int32 InputIndex) override;
 	virtual FName GetInputName(int32 InputIndex) const override;
 	virtual bool IsInputConnectionRequired(int32 InputIndex) const override;
-	virtual uint32 GetInputType(int32 InputIndex) override;
+	virtual EMaterialValueType GetInputValueType(int32 InputIndex) override;
 	virtual int32 Compile(class FMaterialCompiler* Compiler, int32 OutputIndex) override;
 #endif // WITH_EDITOR
 };
\ No newline at end of file
diff --git a/Engine/Source/Runtime/Engine/Public/Materials/MaterialExpressionBreakMaterialAttributes.h b/Engine/Source/Runtime/Engine/Public/Materials/MaterialExpressionBreakMaterialAttributes.h
--- a/Engine/Source/Runtime/Engine/Public/Materials/MaterialExpressionBreakMaterialAttributes.h
+++ b/Engine/Source/Runtime/Engine/Public/Materials/MaterialExpressionBreakMaterialAttributes.h
@@ -29,8 +29,8 @@ class UMaterialExpressionBreakMaterialAttributes : public UMaterialExpression
 	virtual void GetCaption(TArray<FString>& OutCaptions) const override;
 	virtual FName GetInputName(int32 InputIndex) const override;
 	virtual bool IsInputConnectionRequired(int32 InputIndex) const override;
-	virtual uint32 GetInputType(int32 InputIndex) override {return MCT_MaterialAttributes;}
-	virtual uint32 GetOutputType(int32 OutputIndex) override;
+	virtual EMaterialValueType GetInputValueType(int32 InputIndex) override {return MCT_MaterialAttributes;}
+	virtual EMaterialValueType GetOutputValueType(int32 OutputIndex) override;
 
 	static TMap<EMaterialProperty, int32> PropertyToIOIndexMap;
 	static void BuildPropertyToIOIndexMap();
diff --git a/Engine/Source/Runtime/Engine/Public/Materials/MaterialExpressionCeil.h b/Engine/Source/Runtime/Engine/Public/Materials/MaterialExpressionCeil.h
--- a/Engine/Source/Runtime/Engine/Public/Materials/MaterialExpressionCeil.h
+++ b/Engine/Source/Runtime/Engine/Public/Materials/MaterialExpressionCeil.h
@@ -24,7 +24,7 @@ public:
 #if WITH_EDITOR
 	virtual void Build(MIR::FEmitter& Emitter) override;
 	virtual int32 Compile(class FMaterialCompiler* Compiler, int32 OutputIndex) override;
-	virtual uint32 GetOutputType(int32 OutputIndex) override;
+	virtual EMaterialValueType GetOutputValueType(int32 OutputIndex) override;
 	virtual void GetCaption(TArray<FString>& OutCaptions) const override;
 #endif
 	//~ End UMaterialExpression Interface
diff --git a/Engine/Source/Runtime/Engine/Public/Materials/MaterialExpressionChannelMaskParameter.h b/Engine/Source/Runtime/Engine/Public/Materials/MaterialExpressionChannelMaskParameter.h
--- a/Engine/Source/Runtime/Engine/Public/Materials/MaterialExpressionChannelMaskParameter.h
+++ b/Engine/Source/Runtime/Engine/Public/Materials/MaterialExpressionChannelMaskParameter.h
@@ -39,7 +39,7 @@ class UMaterialExpressionChannelMaskParameter : public UMaterialExpressionVector
 	virtual void GetCaption(TArray<FString>& OutCaptions) const override;
 	
 	virtual bool IsInputConnectionRequired(int32 InputIndex) const override {return true;}
-	virtual uint32 GetInputType(int32 InputIndex) override {return MCT_Float4;}
+	virtual EMaterialValueType GetInputValueType(int32 InputIndex) override {return MCT_Float4;}
 
 #endif
 
diff --git a/Engine/Source/Runtime/Engine/Public/Materials/MaterialExpressionClearCoatNormalCustomOutput.h b/Engine/Source/Runtime/Engine/Public/Materials/MaterialExpressionClearCoatNormalCustomOutput.h
--- a/Engine/Source/Runtime/Engine/Public/Materials/MaterialExpressionClearCoatNormalCustomOutput.h
+++ b/Engine/Source/Runtime/Engine/Public/Materials/MaterialExpressionClearCoatNormalCustomOutput.h
@@ -23,7 +23,7 @@ class UMaterialExpressionClearCoatNormalCustomOutput : public UMaterialExpressio
 	virtual int32 Compile(class FMaterialCompiler* Compiler, int32 OutputIndex) override;
 	virtual void GetCaption(TArray<FString>& OutCaptions) const override;
 	virtual FExpressionInput* GetInput(int32 InputIndex) override;
-	virtual uint32 GetInputType(int32 InputIndex) override { return MCT_Float3; }
+	virtual EMaterialValueType GetInputValueType(int32 InputIndex) override { return MCT_Float3; }
 #endif
 
 
diff --git a/Engine/Source/Runtime/Engine/Public/Materials/MaterialExpressionComposite.h b/Engine/Source/Runtime/Engine/Public/Materials/MaterialExpressionComposite.h
--- a/Engine/Source/Runtime/Engine/Public/Materials/MaterialExpressionComposite.h
+++ b/Engine/Source/Runtime/Engine/Public/Materials/MaterialExpressionComposite.h
@@ -76,8 +76,8 @@ class UMaterialExpressionComposite : public UMaterialExpression
 	virtual TArrayView<FExpressionInput*> GetInputsView() override;
 	virtual FExpressionInput* GetInput(int32 InputIndex) override;
 	virtual FName GetInputName(int32 InputIndex) const override;
-	virtual uint32 GetInputType(int32 InputIndex) override;
-	virtual uint32 GetOutputType(int32 OutputIndex) override;
+	virtual EMaterialValueType GetInputValueType(int32 InputIndex) override;
+	virtual EMaterialValueType GetOutputValueType(int32 OutputIndex) override;
 	virtual bool IsExpressionConnected(FExpressionInput* Input, int32 OutputIndex) override;
 	virtual void ConnectExpression(FExpressionInput* Input, int32 OutputIndex) override;
 
diff --git a/Engine/Source/Runtime/Engine/Public/Materials/MaterialExpressionConstant2Vector.h b/Engine/Source/Runtime/Engine/Public/Materials/MaterialExpressionConstant2Vector.h
--- a/Engine/Source/Runtime/Engine/Public/Materials/MaterialExpressionConstant2Vector.h
+++ b/Engine/Source/Runtime/Engine/Public/Materials/MaterialExpressionConstant2Vector.h
@@ -27,7 +27,7 @@ class UMaterialExpressionConstant2Vector : public UMaterialExpression
 	virtual int32 Compile(class FMaterialCompiler* Compiler, int32 OutputIndex) override;
 	virtual void GetCaption(TArray<FString>& OutCaptions) const override;
 	virtual FString GetDescription() const override;
-	virtual uint32 GetOutputType(int32 OutputIndex) override {return MCT_Float2;}
+	virtual EMaterialValueType GetOutputValueType(int32 OutputIndex) override {return MCT_Float2;}
 #endif // WITH_EDITOR
 	//~ End UMaterialExpression Interface
 };
diff --git a/Engine/Source/Runtime/Engine/Public/Materials/MaterialExpressionConstant3Vector.h b/Engine/Source/Runtime/Engine/Public/Materials/MaterialExpressionConstant3Vector.h
--- a/Engine/Source/Runtime/Engine/Public/Materials/MaterialExpressionConstant3Vector.h
+++ b/Engine/Source/Runtime/Engine/Public/Materials/MaterialExpressionConstant3Vector.h
@@ -25,6 +25,6 @@ public:
 	virtual int32 Compile(class FMaterialCompiler* Compiler, int32 OutputIndex) override;
 	virtual void GetCaption(TArray<FString>& OutCaptions) const override;
 	virtual FString GetDescription() const override;
-	virtual uint32 GetOutputType(int32 OutputIndex) override {return MCT_Float3;}
+	virtual EMaterialValueType GetOutputValueType(int32 OutputIndex) override {return MCT_Float3;}
 #endif // WITH_EDITOR
 };
diff --git a/Engine/Source/Runtime/Engine/Public/Materials/MaterialExpressionConstant4Vector.h b/Engine/Source/Runtime/Engine/Public/Materials/MaterialExpressionConstant4Vector.h
--- a/Engine/Source/Runtime/Engine/Public/Materials/MaterialExpressionConstant4Vector.h
+++ b/Engine/Source/Runtime/Engine/Public/Materials/MaterialExpressionConstant4Vector.h
@@ -23,7 +23,7 @@ class UMaterialExpressionConstant4Vector : public UMaterialExpression
 	virtual int32 Compile(class FMaterialCompiler* Compiler, int32 OutputIndex) override;
 	virtual void GetCaption(TArray<FString>& OutCaptions) const override;
 	virtual FString GetDescription() const override;
-	virtual uint32 GetOutputType(int32 OutputIndex) override {return MCT_Float4;}
+	virtual EMaterialValueType GetOutputValueType(int32 OutputIndex) override {return MCT_Float4;}
 #endif // WITH_EDITOR
 	//~ End UMaterialExpression Interface
 };
diff --git a/Engine/Source/Runtime/Engine/Public/Materials/MaterialExpressionConvert.h b/Engine/Source/Runtime/Engine/Public/Materials/MaterialExpressionConvert.h
--- a/Engine/Source/Runtime/Engine/Public/Materials/MaterialExpressionConvert.h
+++ b/Engine/Source/Runtime/Engine/Public/Materials/MaterialExpressionConvert.h
@@ -156,10 +156,10 @@ public:
 	virtual void GetConnectorToolTip(int32 InputIndex, int32 OutputIndex, TArray<FString>& OutToolTip) override;
 
 	virtual int32 CountInputs() const override;
-	virtual uint32 GetInputType(int32 InputIndex) override;
+	virtual EMaterialValueType GetInputValueType(int32 InputIndex) override;
 	virtual FExpressionInput* GetInput(int32 InputIndex) override;
 
-	virtual uint32 GetOutputType(int32 OutputIndex) override;
+	virtual EMaterialValueType GetOutputValueType(int32 OutputIndex) override;
 	virtual TArray<FExpressionOutput>& GetOutputs() override;
 
 	virtual bool CanDeletePin(EEdGraphPinDirection PinDirection, int32 PinIndex) const override;
diff --git a/Engine/Source/Runtime/Engine/Public/Materials/MaterialExpressionCurveAtlasRowParameter.h b/Engine/Source/Runtime/Engine/Public/Materials/MaterialExpressionCurveAtlasRowParameter.h
--- a/Engine/Source/Runtime/Engine/Public/Materials/MaterialExpressionCurveAtlasRowParameter.h
+++ b/Engine/Source/Runtime/Engine/Public/Materials/MaterialExpressionCurveAtlasRowParameter.h
@@ -30,7 +30,7 @@ class UMaterialExpressionCurveAtlasRowParameter : public UMaterialExpressionScal
 	virtual bool GetParameterValue(FMaterialParameterMetadata& OutMeta) const override;
 	virtual bool SetParameterValue(const FName& Name, const FMaterialParameterMetadata& Meta, EMaterialExpressionSetParameterValueFlags Flags) override;
 
-	virtual uint32 GetInputType(int32 InputIndex) override 
+	virtual EMaterialValueType GetInputValueType(int32 InputIndex) override 
 	{
 		return MCT_Float;
 	}
diff --git a/Engine/Source/Runtime/Engine/Public/Materials/MaterialExpressionCustom.h b/Engine/Source/Runtime/Engine/Public/Materials/MaterialExpressionCustom.h
--- a/Engine/Source/Runtime/Engine/Public/Materials/MaterialExpressionCustom.h
+++ b/Engine/Source/Runtime/Engine/Public/Materials/MaterialExpressionCustom.h
@@ -104,8 +104,8 @@ class UMaterialExpressionCustom : public UMaterialExpression
 	virtual TArrayView<FExpressionInput*> GetInputsView() override;
 	virtual FExpressionInput* GetInput(int32 InputIndex) override;
 	virtual FName GetInputName(int32 InputIndex) const override;
-	virtual uint32 GetInputType(int32 InputIndex) override {return MCT_Unknown;}
-	virtual uint32 GetOutputType(int32 OutputIndex) override;
+	virtual EMaterialValueType GetInputValueType(int32 InputIndex) override {return MCT_Unknown;}
+	virtual EMaterialValueType GetOutputValueType(int32 OutputIndex) override;
 	virtual bool IsResultMaterialAttributes(int32 OutputIndex) override;
 	virtual void GetIncludeFilePaths(TSet<FString>& OutIncludeFilePaths) const override;
 #endif // WITH_EDITOR
diff --git a/Engine/Source/Runtime/Engine/Public/Materials/MaterialExpressionDataDrivenShaderPlatformInfoSwitch.h b/Engine/Source/Runtime/Engine/Public/Materials/MaterialExpressionDataDrivenShaderPlatformInfoSwitch.h
--- a/Engine/Source/Runtime/Engine/Public/Materials/MaterialExpressionDataDrivenShaderPlatformInfoSwitch.h
+++ b/Engine/Source/Runtime/Engine/Public/Materials/MaterialExpressionDataDrivenShaderPlatformInfoSwitch.h
@@ -63,8 +63,8 @@ class UMaterialExpressionDataDrivenShaderPlatformInfoSwitch : public UMaterialEx
 	virtual bool IsInputConnectionRequired(int32 InputIndex) const override;
 	virtual FName GetInputName(int32 InputIndex) const override;
 	virtual bool IsResultMaterialAttributes(int32 OutputIndex) override;
-	virtual uint32 GetInputType(int32 InputIndex) override { return MCT_Unknown; }
-	virtual uint32 GetOutputType(int32 OutputIndex) override { return MCT_Unknown; }
+	virtual EMaterialValueType GetInputValueType(int32 InputIndex) override { return MCT_Unknown; }
+	virtual EMaterialValueType GetOutputValueType(int32 OutputIndex) override { return MCT_Unknown; }
 #endif // WITH_EDITOR
 	//~ End UMaterialExpression Interface
 };
diff --git a/Engine/Source/Runtime/Engine/Public/Materials/MaterialExpressionDistanceFieldsRenderingSwitch.h b/Engine/Source/Runtime/Engine/Public/Materials/MaterialExpressionDistanceFieldsRenderingSwitch.h
--- a/Engine/Source/Runtime/Engine/Public/Materials/MaterialExpressionDistanceFieldsRenderingSwitch.h
+++ b/Engine/Source/Runtime/Engine/Public/Materials/MaterialExpressionDistanceFieldsRenderingSwitch.h
@@ -28,8 +28,8 @@ public:
 #if WITH_EDITOR
 	virtual int32 Compile(class FMaterialCompiler* Compiler, int32 OutputIndex) override;
 	virtual bool IsResultMaterialAttributes(int32 OutputIndex) override;
-	virtual uint32 GetInputType(int32 InputIndex) override { return MCT_Unknown; }
-	virtual uint32 GetOutputType(int32 OutputIndex) override { return MCT_Unknown; }
+	virtual EMaterialValueType GetInputValueType(int32 InputIndex) override { return MCT_Unknown; }
+	virtual EMaterialValueType GetOutputValueType(int32 OutputIndex) override { return MCT_Unknown; }
 	virtual void GetCaption(TArray<FString>& OutCaptions) const override;
 
 #endif
diff --git a/Engine/Source/Runtime/Engine/Public/Materials/MaterialExpressionFeatureLevelSwitch.h b/Engine/Source/Runtime/Engine/Public/Materials/MaterialExpressionFeatureLevelSwitch.h
--- a/Engine/Source/Runtime/Engine/Public/Materials/MaterialExpressionFeatureLevelSwitch.h
+++ b/Engine/Source/Runtime/Engine/Public/Materials/MaterialExpressionFeatureLevelSwitch.h
@@ -37,8 +37,8 @@ public:
 	virtual FName GetInputName(int32 InputIndex) const override;
 	virtual bool IsInputConnectionRequired(int32 InputIndex) const override;
 	virtual bool IsResultMaterialAttributes(int32 OutputIndex) override;
-	virtual uint32 GetInputType(int32 InputIndex) override {return MCT_Unknown;}
-	virtual uint32 GetOutputType(int32 OutputIndex) override {return MCT_Unknown;}
+	virtual EMaterialValueType GetInputValueType(int32 InputIndex) override {return MCT_Unknown;}
+	virtual EMaterialValueType GetOutputValueType(int32 OutputIndex) override {return MCT_Unknown;}
 #endif // WITH_EDITOR
 	//~ End UMaterialExpression Interface
 };
diff --git a/Engine/Source/Runtime/Engine/Public/Materials/MaterialExpressionFloatToUInt.h b/Engine/Source/Runtime/Engine/Public/Materials/MaterialExpressionFloatToUInt.h
--- a/Engine/Source/Runtime/Engine/Public/Materials/MaterialExpressionFloatToUInt.h
+++ b/Engine/Source/Runtime/Engine/Public/Materials/MaterialExpressionFloatToUInt.h
@@ -34,8 +34,8 @@ public:
 #if WITH_EDITOR
 	virtual int32 Compile(class FMaterialCompiler* Compiler, int32 OutputIndex) override;
 	virtual void GetCaption(TArray<FString>& OutCaptions) const override;
-	virtual uint32 GetInputType(int32 InputIndex) override;
-	virtual uint32 GetOutputType(int32 OutputIndex) override;
+	virtual EMaterialValueType GetInputValueType(int32 InputIndex) override;
+	virtual EMaterialValueType GetOutputValueType(int32 OutputIndex) override;
 #endif
 	//~ End UMaterialExpression Interface
 };
@@ -54,8 +54,8 @@ public:
 #if WITH_EDITOR
 	virtual int32 Compile(class FMaterialCompiler* Compiler, int32 OutputIndex) override;
 	virtual void GetCaption(TArray<FString>& OutCaptions) const override;
-	virtual uint32 GetInputType(int32 InputIndex) override;
-	virtual uint32 GetOutputType(int32 OutputIndex) override;
+	virtual EMaterialValueType GetInputValueType(int32 InputIndex) override;
+	virtual EMaterialValueType GetOutputValueType(int32 OutputIndex) override;
 #endif
 	//~ End UMaterialExpression Interface
 };
diff --git a/Engine/Source/Runtime/Engine/Public/Materials/MaterialExpressionFloor.h b/Engine/Source/Runtime/Engine/Public/Materials/MaterialExpressionFloor.h
--- a/Engine/Source/Runtime/Engine/Public/Materials/MaterialExpressionFloor.h
+++ b/Engine/Source/Runtime/Engine/Public/Materials/MaterialExpressionFloor.h
@@ -24,7 +24,7 @@ public:
 #if WITH_EDITOR
 	virtual void Build(MIR::FEmitter& Emitter) override;
 	virtual int32 Compile(class FMaterialCompiler* Compiler, int32 OutputIndex) override;
-	virtual uint32 GetOutputType(int32 OutputIndex) override;
+	virtual EMaterialValueType GetOutputValueType(int32 OutputIndex) override;
 	virtual void GetCaption(TArray<FString>& OutCaptions) const override;
 #endif
 	//~ End UMaterialExpression Interface
diff --git a/Engine/Source/Runtime/Engine/Public/Materials/MaterialExpressionFunctionInput.h b/Engine/Source/Runtime/Engine/Public/Materials/MaterialExpressionFunctionInput.h
--- a/Engine/Source/Runtime/Engine/Public/Materials/MaterialExpressionFunctionInput.h
+++ b/Engine/Source/Runtime/Engine/Public/Materials/MaterialExpressionFunctionInput.h
@@ -109,8 +109,8 @@ class UMaterialExpressionFunctionInput : public UMaterialExpression
 	virtual void GatherSubstrateMaterialInfo(FSubstrateMaterialInfo& SubstrateMaterialInfo, int32 OutputIndex) override;
 	virtual FSubstrateOperator* SubstrateGenerateMaterialTopologyTree(class FMaterialCompiler* Compiler, class UMaterialExpression* Parent, int32 OutputIndex) override;
 	virtual bool IsResultMaterialAttributes(int32 OutputIndex) override;
-	virtual uint32 GetInputType(int32 InputIndex) override;
-	virtual uint32 GetOutputType(int32 OutputIndex) override;
+	virtual EMaterialValueType GetInputValueType(int32 InputIndex) override;
+	virtual EMaterialValueType GetOutputValueType(int32 OutputIndex) override;
 #endif // WITH_EDITOR
 	virtual bool IsAllowedIn(const UObject* MaterialOrFunction) const override;
 	//~ End UMaterialExpression Interface
diff --git a/Engine/Source/Runtime/Engine/Public/Materials/MaterialExpressionFunctionOutput.h b/Engine/Source/Runtime/Engine/Public/Materials/MaterialExpressionFunctionOutput.h
--- a/Engine/Source/Runtime/Engine/Public/Materials/MaterialExpressionFunctionOutput.h
+++ b/Engine/Source/Runtime/Engine/Public/Materials/MaterialExpressionFunctionOutput.h
@@ -63,7 +63,7 @@ class UMaterialExpressionFunctionOutput : public UMaterialExpression
 		return NAME_None;
 	}
 	virtual void GetExpressionToolTip(TArray<FString>& OutToolTip) override;
-	virtual uint32 GetInputType(int32 InputIndex) override;
+	virtual EMaterialValueType GetInputValueType(int32 InputIndex) override;
 	virtual bool IsResultMaterialAttributes(int32 OutputIndex) override;
 	virtual bool IsResultSubstrateMaterial(int32 OutputIndex) override;
 	virtual void GatherSubstrateMaterialInfo(FSubstrateMaterialInfo& SubstrateMaterialInfo, int32 OutputIndex) override;
diff --git a/Engine/Source/Runtime/Engine/Public/Materials/MaterialExpressionGetMaterialAttributes.h b/Engine/Source/Runtime/Engine/Public/Materials/MaterialExpressionGetMaterialAttributes.h
--- a/Engine/Source/Runtime/Engine/Public/Materials/MaterialExpressionGetMaterialAttributes.h
+++ b/Engine/Source/Runtime/Engine/Public/Materials/MaterialExpressionGetMaterialAttributes.h
@@ -40,9 +40,9 @@ class UMaterialExpressionGetMaterialAttributes : public UMaterialExpression
 	virtual int32 Compile(class FMaterialCompiler* Compiler, int32 OutputIndex) override;
 	virtual void GetCaption(TArray<FString>& OutCaptions) const override;
 	virtual FName GetInputName(int32 InputIndex) const override;
-	virtual uint32 GetOutputType(int32 OutputIndex) override;
+	virtual EMaterialValueType GetOutputValueType(int32 OutputIndex) override;
 	virtual bool IsInputConnectionRequired(int32 InputIndex) const override {return true;}
-	virtual uint32 GetInputType(int32 InputIndex) override { return MCT_MaterialAttributes; }
+	virtual EMaterialValueType GetInputValueType(int32 InputIndex) override { return MCT_MaterialAttributes; }
 	virtual bool IsResultMaterialAttributes(int32 OutputIndex) override { return OutputIndex == 0; }
 	virtual bool IsResultSubstrateMaterial(int32 OutputIndex) override;
 	virtual void GatherSubstrateMaterialInfo(FSubstrateMaterialInfo& SubstrateMaterialInfo, int32 OutputIndex) override;
diff --git a/Engine/Source/Runtime/Engine/Public/Materials/MaterialExpressionIf.h b/Engine/Source/Runtime/Engine/Public/Materials/MaterialExpressionIf.h
--- a/Engine/Source/Runtime/Engine/Public/Materials/MaterialExpressionIf.h
+++ b/Engine/Source/Runtime/Engine/Public/Materials/MaterialExpressionIf.h
@@ -47,8 +47,8 @@ public:
 	virtual void Build(MIR::FEmitter& Emitter) override;
 	virtual int32 Compile(class FMaterialCompiler* Compiler, int32 OutputIndex) override;
 	virtual void GetCaption(TArray<FString>& OutCaptions) const override;
-	virtual uint32 GetInputType(int32 InputIndex) override;
-	virtual uint32 GetOutputType(int32 InputIndex) override {return MCT_Unknown;}
+	virtual EMaterialValueType GetInputValueType(int32 InputIndex) override;
+	virtual EMaterialValueType GetOutputValueType(int32 InputIndex) override {return MCT_Unknown;}
 	virtual bool IsResultMaterialAttributes(int32 OutputIndex) override;
 #endif
 	//~ End UMaterialExpression Interface
diff --git a/Engine/Source/Runtime/Engine/Public/Materials/MaterialExpressionIfThenElse.h b/Engine/Source/Runtime/Engine/Public/Materials/MaterialExpressionIfThenElse.h
--- a/Engine/Source/Runtime/Engine/Public/Materials/MaterialExpressionIfThenElse.h
+++ b/Engine/Source/Runtime/Engine/Public/Materials/MaterialExpressionIfThenElse.h
@@ -30,7 +30,7 @@ public:
 #if WITH_EDITOR
 	virtual void Build(MIR::FEmitter& Emitter) override;
 	virtual void GetCaption(TArray<FString>& OutCaptions) const override;
-	virtual uint32 GetInputType(int32 InputIndex) override;
+	virtual EMaterialValueType GetInputValueType(int32 InputIndex) override;
 #endif
 	//~ End UMaterialExpression Interface
 };
diff --git a/Engine/Source/Runtime/Engine/Public/Materials/MaterialExpressionMakeMaterialAttributes.h b/Engine/Source/Runtime/Engine/Public/Materials/MaterialExpressionMakeMaterialAttributes.h
--- a/Engine/Source/Runtime/Engine/Public/Materials/MaterialExpressionMakeMaterialAttributes.h
+++ b/Engine/Source/Runtime/Engine/Public/Materials/MaterialExpressionMakeMaterialAttributes.h
@@ -88,7 +88,7 @@ public:
 	virtual int32 Compile(class FMaterialCompiler* Compiler, int32 OutputIndex) override;
 	virtual void GetCaption(TArray<FString>& OutCaptions) const override;
 	virtual bool IsResultMaterialAttributes(int32 OutputIndex) override {return true;}
-	virtual uint32 GetInputType(int32 InputIndex) override;
+	virtual EMaterialValueType GetInputValueType(int32 InputIndex) override;
 
 	uint64 GetConnectedInputs() const;
 #endif // WITH_EDITOR
diff --git a/Engine/Source/Runtime/Engine/Public/Materials/MaterialExpressionMaterialAttributeLayers.h b/Engine/Source/Runtime/Engine/Public/Materials/MaterialExpressionMaterialAttributeLayers.h
--- a/Engine/Source/Runtime/Engine/Public/Materials/MaterialExpressionMaterialAttributeLayers.h
+++ b/Engine/Source/Runtime/Engine/Public/Materials/MaterialExpressionMaterialAttributeLayers.h
@@ -110,7 +110,7 @@ class UMaterialExpressionMaterialAttributeLayers : public UMaterialExpression
 	virtual void GetExpressionToolTip(TArray<FString>& OutToolTip) override;
 	virtual FName GetInputName(int32 InputIndex) const override;
 	virtual bool IsInputConnectionRequired(int32 InputIndex) const override {return false;}
-	virtual uint32 GetInputType(int32 InputIndex) override;
+	virtual EMaterialValueType GetInputValueType(int32 InputIndex) override;
 	virtual bool IsResultMaterialAttributes(int32 OutputIndex) override {return true;}
 	virtual bool IsResultSubstrateMaterial(int32 OutputIndex) override;
 	virtual FSubstrateOperator* SubstrateGenerateMaterialTopologyTree(class FMaterialCompiler* Compiler, class UMaterialExpression* Parent, int32 OutputIndex) override;
diff --git a/Engine/Source/Runtime/Engine/Public/Materials/MaterialExpressionMaterialFunctionCall.h b/Engine/Source/Runtime/Engine/Public/Materials/MaterialExpressionMaterialFunctionCall.h
--- a/Engine/Source/Runtime/Engine/Public/Materials/MaterialExpressionMaterialFunctionCall.h
+++ b/Engine/Source/Runtime/Engine/Public/Materials/MaterialExpressionMaterialFunctionCall.h
@@ -134,7 +134,7 @@ class UMaterialExpressionMaterialFunctionCall : public UMaterialExpression
 	ENGINE_API virtual bool IsResultSubstrateMaterial(int32 OutputIndex) override;
 	ENGINE_API virtual void GatherSubstrateMaterialInfo(FSubstrateMaterialInfo& SubstrateMaterialInfo, int32 OutputIndex) override;
 	ENGINE_API virtual FSubstrateOperator* SubstrateGenerateMaterialTopologyTree(class FMaterialCompiler* Compiler, class UMaterialExpression* Parent, int32 OutputIndex) override;
-	ENGINE_API virtual uint32 GetInputType(int32 InputIndex) override;
+	ENGINE_API virtual EMaterialValueType GetInputValueType(int32 InputIndex) override;
 	//~ End UMaterialExpression Interface
 
 	/** Util to get name of a particular type, optionally with type appended in parenthesis */
diff --git a/Engine/Source/Runtime/Engine/Public/Materials/MaterialExpressionMaterialLayerOutput.h b/Engine/Source/Runtime/Engine/Public/Materials/MaterialExpressionMaterialLayerOutput.h
--- a/Engine/Source/Runtime/Engine/Public/Materials/MaterialExpressionMaterialLayerOutput.h
+++ b/Engine/Source/Runtime/Engine/Public/Materials/MaterialExpressionMaterialLayerOutput.h
@@ -16,7 +16,7 @@ class UMaterialExpressionMaterialLayerOutput : public UMaterialExpressionFunctio
 {
 	GENERATED_UCLASS_BODY()
 #if WITH_EDITOR
-	virtual uint32 GetInputType(int32 InputIndex) override
+	virtual EMaterialValueType GetInputValueType(int32 InputIndex) override
 	{
 		return MCT_MaterialAttributes;
 	}
diff --git a/Engine/Source/Runtime/Engine/Public/Materials/MaterialExpressionMaterialProxyReplace.h b/Engine/Source/Runtime/Engine/Public/Materials/MaterialExpressionMaterialProxyReplace.h
--- a/Engine/Source/Runtime/Engine/Public/Materials/MaterialExpressionMaterialProxyReplace.h
+++ b/Engine/Source/Runtime/Engine/Public/Materials/MaterialExpressionMaterialProxyReplace.h
@@ -28,8 +28,8 @@ public:
 #if WITH_EDITOR
 	virtual int32 Compile(class FMaterialCompiler* Compiler, int32 OutputIndex) override;
 	virtual bool IsResultMaterialAttributes(int32 OutputIndex) override;
-	virtual uint32 GetInputType(int32 InputIndex) override { return MCT_Unknown; }
-	virtual uint32 GetOutputType(int32 OutputIndex) override { return MCT_Unknown; }
+	virtual EMaterialValueType GetInputValueType(int32 InputIndex) override { return MCT_Unknown; }
+	virtual EMaterialValueType GetOutputValueType(int32 OutputIndex) override { return MCT_Unknown; }
 	virtual void GetCaption(TArray<FString>& OutCaptions) const override;
 #endif
 	// End UMaterialExpression Interface
diff --git a/Engine/Source/Runtime/Engine/Public/Materials/MaterialExpressionMeshPaintTextureObject.h b/Engine/Source/Runtime/Engine/Public/Materials/MaterialExpressionMeshPaintTextureObject.h
--- a/Engine/Source/Runtime/Engine/Public/Materials/MaterialExpressionMeshPaintTextureObject.h
+++ b/Engine/Source/Runtime/Engine/Public/Materials/MaterialExpressionMeshPaintTextureObject.h
@@ -15,7 +15,7 @@ class UMaterialExpressionMeshPaintTextureObject : public UMaterialExpressionExte
 #if WITH_EDITOR
 	virtual void GetCaption(TArray<FString>& OutCaptions) const override;
 	virtual void GetExpressionToolTip(TArray<FString>& OutToolTip) override;
-	virtual uint32 GetOutputType(int32 OutputIndex) override;
+	virtual EMaterialValueType GetOutputValueType(int32 OutputIndex) override;
 #endif
 	//~ End UMaterialExpression Interface
 };
diff --git a/Engine/Source/Runtime/Engine/Public/Materials/MaterialExpressionModulo.h b/Engine/Source/Runtime/Engine/Public/Materials/MaterialExpressionModulo.h
--- a/Engine/Source/Runtime/Engine/Public/Materials/MaterialExpressionModulo.h
+++ b/Engine/Source/Runtime/Engine/Public/Materials/MaterialExpressionModulo.h
@@ -21,8 +21,8 @@ public:
 	//~ Begin UMaterialExpression Interface
 #if WITH_EDITOR
 	virtual int32 Compile(class FMaterialCompiler* Compiler, int32 OutputIndex) override;
-	virtual uint32 GetInputType(int32 InputIndex) override;
-	virtual uint32 GetOutputType(int32 OutputIndex) override;
+	virtual EMaterialValueType GetInputValueType(int32 InputIndex) override;
+	virtual EMaterialValueType GetOutputValueType(int32 OutputIndex) override;
 	virtual void GetCaption(TArray<FString>& OutCaptions) const override;
 	virtual FText GetKeywords() const override;
 #endif // WITH_EDITOR
diff --git a/Engine/Source/Runtime/Engine/Public/Materials/MaterialExpressionNamedReroute.h b/Engine/Source/Runtime/Engine/Public/Materials/MaterialExpressionNamedReroute.h
--- a/Engine/Source/Runtime/Engine/Public/Materials/MaterialExpressionNamedReroute.h
+++ b/Engine/Source/Runtime/Engine/Public/Materials/MaterialExpressionNamedReroute.h
@@ -112,7 +112,7 @@ public:
 	ENGINE_API virtual int32 CompilePreview(FMaterialCompiler* Compiler, int32 OutputIndex) override;
 
 	ENGINE_API virtual void GetCaption(TArray<FString>& OutCaptions) const override;
-	ENGINE_API virtual uint32 GetOutputType(int32 OutputIndex) override;
+	ENGINE_API virtual EMaterialValueType GetOutputValueType(int32 OutputIndex) override;
 	
 	ENGINE_API virtual bool MatchesSearchQuery(const TCHAR* SearchQuery) override;
 	
diff --git a/Engine/Source/Runtime/Engine/Public/Materials/MaterialExpressionNeuralPostProcessNode.h b/Engine/Source/Runtime/Engine/Public/Materials/MaterialExpressionNeuralPostProcessNode.h
--- a/Engine/Source/Runtime/Engine/Public/Materials/MaterialExpressionNeuralPostProcessNode.h
+++ b/Engine/Source/Runtime/Engine/Public/Materials/MaterialExpressionNeuralPostProcessNode.h
@@ -57,7 +57,7 @@ class UMaterialExpressionNeuralNetworkInput : public UMaterialExpressionCustomOu
 
 	virtual int32 Compile(class FMaterialCompiler* Compiler, int32 OutputIndex) override;
 	virtual void GetCaption(TArray<FString>& OutCaptions) const override;
-	virtual uint32 GetInputType(int32 InputIndex) override;
+	virtual EMaterialValueType GetInputValueType(int32 InputIndex) override;
 
 #endif // WITH_EDITOR
 
@@ -91,7 +91,7 @@ class UMaterialExpressionNeuralNetworkOutput : public UMaterialExpression
 #if WITH_EDITOR
 	virtual int32 Compile(class FMaterialCompiler* Compiler, int32 OutputIndex) override;
 	virtual void GetCaption(TArray<FString>& OutCaptions) const override;
-	virtual uint32 GetInputType(int32 InputIndex) override;
+	virtual EMaterialValueType GetInputValueType(int32 InputIndex) override;
 
 #endif // WITH_EDITOR
 
diff --git a/Engine/Source/Runtime/Engine/Public/Materials/MaterialExpressionNoise.h b/Engine/Source/Runtime/Engine/Public/Materials/MaterialExpressionNoise.h
--- a/Engine/Source/Runtime/Engine/Public/Materials/MaterialExpressionNoise.h
+++ b/Engine/Source/Runtime/Engine/Public/Materials/MaterialExpressionNoise.h
@@ -130,7 +130,7 @@ class UMaterialExpressionScalarBlueNoise : public UMaterialExpression
 
 	//~ Begin UMaterialExpression Interface
 #if WITH_EDITOR
-	virtual uint32 GetOutputType(int32 OutputIndex) override;
+	virtual EMaterialValueType GetOutputValueType(int32 OutputIndex) override;
 	virtual int32 Compile(class FMaterialCompiler* Compiler, int32 OutputIndex) override;
 	virtual void GetCaption(TArray<FString>& OutCaptions) const override;
 	virtual void GetExpressionToolTip(TArray<FString>& OutToolTip) override;
diff --git a/Engine/Source/Runtime/Engine/Public/Materials/MaterialExpressionPathTracingQualitySwitch.h b/Engine/Source/Runtime/Engine/Public/Materials/MaterialExpressionPathTracingQualitySwitch.h
--- a/Engine/Source/Runtime/Engine/Public/Materials/MaterialExpressionPathTracingQualitySwitch.h
+++ b/Engine/Source/Runtime/Engine/Public/Materials/MaterialExpressionPathTracingQualitySwitch.h
@@ -29,8 +29,8 @@ public:
 	virtual int32 Compile(class FMaterialCompiler* Compiler, int32 OutputIndex) override;
 	virtual void GetCaption(TArray<FString>& OutCaptions) const override;
 	virtual bool IsResultMaterialAttributes(int32 OutputIndex) override;
-	virtual uint32 GetInputType(int32 InputIndex) override;
-	virtual uint32 GetOutputType(int32 OutputIndex) override { return MCT_Unknown; }
+	virtual EMaterialValueType GetInputValueType(int32 InputIndex) override;
+	virtual EMaterialValueType GetOutputValueType(int32 OutputIndex) override { return MCT_Unknown; }
 
 #endif
 	//~ End UMaterialExpression Interface
diff --git a/Engine/Source/Runtime/Engine/Public/Materials/MaterialExpressionPathTracingRayTypeSwitch.h b/Engine/Source/Runtime/Engine/Public/Materials/MaterialExpressionPathTracingRayTypeSwitch.h
--- a/Engine/Source/Runtime/Engine/Public/Materials/MaterialExpressionPathTracingRayTypeSwitch.h
+++ b/Engine/Source/Runtime/Engine/Public/Materials/MaterialExpressionPathTracingRayTypeSwitch.h
@@ -41,8 +41,8 @@ public:
 	virtual int32 Compile(class FMaterialCompiler* Compiler, int32 OutputIndex) override;
 	virtual void GetCaption(TArray<FString>& OutCaptions) const override;
 	virtual bool IsResultMaterialAttributes(int32 OutputIndex) override;
-	virtual uint32 GetInputType(int32 InputIndex) override { return MCT_Unknown; }
-	virtual uint32 GetOutputType(int32 OutputIndex) override { return MCT_Unknown; }
+	virtual EMaterialValueType GetInputValueType(int32 InputIndex) override { return MCT_Unknown; }
+	virtual EMaterialValueType GetOutputValueType(int32 OutputIndex) override { return MCT_Unknown; }
 #endif
 	//~ End UMaterialExpression Interface
 };
diff --git a/Engine/Source/Runtime/Engine/Public/Materials/MaterialExpressionPinBase.h b/Engine/Source/Runtime/Engine/Public/Materials/MaterialExpressionPinBase.h
--- a/Engine/Source/Runtime/Engine/Public/Materials/MaterialExpressionPinBase.h
+++ b/Engine/Source/Runtime/Engine/Public/Materials/MaterialExpressionPinBase.h
@@ -80,8 +80,8 @@ class UMaterialExpressionPinBase : public UMaterialExpression
 	virtual TArrayView<FExpressionInput*> GetInputsView() override;
 	virtual FExpressionInput* GetInput(int32 InputIndex) override;
 	virtual FName GetInputName(int32 InputIndex) const override;
-	virtual uint32 GetInputType(int32 InputIndex) override;
-	virtual uint32 GetOutputType(int32 OutputIndex) override;
+	virtual EMaterialValueType GetInputValueType(int32 InputIndex) override;
+	virtual EMaterialValueType GetOutputValueType(int32 OutputIndex) override;
 	virtual bool IsExpressionConnected(FExpressionInput* Input, int32 OutputIndex) override;
 	virtual void ConnectExpression(FExpressionInput* Input, int32 OutputIndex) override;
 
diff --git a/Engine/Source/Runtime/Engine/Public/Materials/MaterialExpressionPreviousFrameSwitch.h b/Engine/Source/Runtime/Engine/Public/Materials/MaterialExpressionPreviousFrameSwitch.h
--- a/Engine/Source/Runtime/Engine/Public/Materials/MaterialExpressionPreviousFrameSwitch.h
+++ b/Engine/Source/Runtime/Engine/Public/Materials/MaterialExpressionPreviousFrameSwitch.h
@@ -27,8 +27,8 @@ public:
 	virtual void GetExpressionToolTip(TArray<FString>& OutToolTip) override;
 	virtual FName GetInputName(int32 InputIndex) const override;
 	virtual bool IsResultMaterialAttributes(int32 OutputIndex) override;
-	virtual uint32 GetInputType(int32 InputIndex) override;
-	virtual uint32 GetOutputType(int32 OutputIndex) override {return MCT_Unknown;}
+	virtual EMaterialValueType GetInputValueType(int32 InputIndex) override;
+	virtual EMaterialValueType GetOutputValueType(int32 OutputIndex) override {return MCT_Unknown;}
 
 #endif // WITH_EDITOR
 	//~ End UMaterialExpression Interface
diff --git a/Engine/Source/Runtime/Engine/Public/Materials/MaterialExpressionQualitySwitch.h b/Engine/Source/Runtime/Engine/Public/Materials/MaterialExpressionQualitySwitch.h
--- a/Engine/Source/Runtime/Engine/Public/Materials/MaterialExpressionQualitySwitch.h
+++ b/Engine/Source/Runtime/Engine/Public/Materials/MaterialExpressionQualitySwitch.h
@@ -32,8 +32,8 @@ public:
 	virtual FName GetInputName(int32 InputIndex) const override;
 	virtual bool IsInputConnectionRequired(int32 InputIndex) const override;
 	virtual bool IsResultMaterialAttributes(int32 OutputIndex) override;
-	virtual uint32 GetInputType(int32 InputIndex) override {return MCT_Unknown;}
-	virtual uint32 GetOutputType(int32 InputIndex) override { return MCT_Unknown; }
+	virtual EMaterialValueType GetInputValueType(int32 InputIndex) override {return MCT_Unknown;}
+	virtual EMaterialValueType GetOutputValueType(int32 InputIndex) override { return MCT_Unknown; }
 
 	virtual bool IsResultSubstrateMaterial(int32 OutputIndex) override;
 	virtual void GatherSubstrateMaterialInfo(FSubstrateMaterialInfo& SubstrateMaterialInfo, int32 OutputIndex) override;
diff --git a/Engine/Source/Runtime/Engine/Public/Materials/MaterialExpressionRayTracingQualitySwitch.h b/Engine/Source/Runtime/Engine/Public/Materials/MaterialExpressionRayTracingQualitySwitch.h
--- a/Engine/Source/Runtime/Engine/Public/Materials/MaterialExpressionRayTracingQualitySwitch.h
+++ b/Engine/Source/Runtime/Engine/Public/Materials/MaterialExpressionRayTracingQualitySwitch.h
@@ -29,8 +29,8 @@ public:
 	virtual int32 Compile(class FMaterialCompiler* Compiler, int32 OutputIndex) override;
 	virtual void GetCaption(TArray<FString>& OutCaptions) const override;
 	virtual bool IsResultMaterialAttributes(int32 OutputIndex) override;
-	virtual uint32 GetInputType(int32 InputIndex) override;
-	virtual uint32 GetOutputType(int32 OutputIndex) override { return MCT_Unknown; }
+	virtual EMaterialValueType GetInputValueType(int32 InputIndex) override;
+	virtual EMaterialValueType GetOutputValueType(int32 OutputIndex) override { return MCT_Unknown; }
 
 #endif
 	//~ End UMaterialExpression Interface
diff --git a/Engine/Source/Runtime/Engine/Public/Materials/MaterialExpressionRequiredSamplersSwitch.h b/Engine/Source/Runtime/Engine/Public/Materials/MaterialExpressionRequiredSamplersSwitch.h
--- a/Engine/Source/Runtime/Engine/Public/Materials/MaterialExpressionRequiredSamplersSwitch.h
+++ b/Engine/Source/Runtime/Engine/Public/Materials/MaterialExpressionRequiredSamplersSwitch.h
@@ -34,8 +34,8 @@ public:
 	virtual bool IsInputConnectionRequired(int32 InputIndex) const override;
 	virtual FName GetInputName(int32 InputIndex) const override;
 	virtual bool IsResultMaterialAttributes(int32 OutputIndex) override;
-	virtual uint32 GetInputType(int32 InputIndex) override { return MCT_Unknown; }
-	virtual uint32 GetOutputType(int32 OutputIndex) override { return MCT_Unknown; }
+	virtual EMaterialValueType GetInputValueType(int32 InputIndex) override { return MCT_Unknown; }
+	virtual EMaterialValueType GetOutputValueType(int32 OutputIndex) override { return MCT_Unknown; }
 #endif // WITH_EDITOR
 	//~ End UMaterialExpression Interface
 };
diff --git a/Engine/Source/Runtime/Engine/Public/Materials/MaterialExpressionRerouteBase.h b/Engine/Source/Runtime/Engine/Public/Materials/MaterialExpressionRerouteBase.h
--- a/Engine/Source/Runtime/Engine/Public/Materials/MaterialExpressionRerouteBase.h
+++ b/Engine/Source/Runtime/Engine/Public/Materials/MaterialExpressionRerouteBase.h
@@ -26,8 +26,8 @@ public:
 
 	//~ Begin UMaterialExpression Interface
 #if WITH_EDITOR
-	ENGINE_API virtual uint32 GetInputType(int32 InputIndex) override;
-	ENGINE_API virtual uint32 GetOutputType(int32 OutputIndex) override;
+	ENGINE_API virtual EMaterialValueType GetInputValueType(int32 InputIndex) override;
+	ENGINE_API virtual EMaterialValueType GetOutputValueType(int32 OutputIndex) override;
 	ENGINE_API virtual bool IsResultMaterialAttributes(int32 OutputIndex) override;
 	ENGINE_API virtual bool IsResultSubstrateMaterial(int32 OutputIndex) override;
 	ENGINE_API void GatherSubstrateMaterialInfo(FSubstrateMaterialInfo& SubstrateMaterialInfo, int32 OutputIndex) override;
diff --git a/Engine/Source/Runtime/Engine/Public/Materials/MaterialExpressionRound.h b/Engine/Source/Runtime/Engine/Public/Materials/MaterialExpressionRound.h
--- a/Engine/Source/Runtime/Engine/Public/Materials/MaterialExpressionRound.h
+++ b/Engine/Source/Runtime/Engine/Public/Materials/MaterialExpressionRound.h
@@ -22,7 +22,7 @@ public:
 #if WITH_EDITOR
 	virtual void Build(MIR::FEmitter& Emitter) override;
 	virtual int32 Compile(class FMaterialCompiler* Compiler, int32 OutputIndex) override;
-	virtual uint32 GetOutputType(int32 OutputIndex) override;
+	virtual EMaterialValueType GetOutputValueType(int32 OutputIndex) override;
 	virtual void GetCaption(TArray<FString>& OutCaptions) const override;
 	virtual void GetExpressionToolTip(TArray<FString>& OutToolTip) override;
 #endif
diff --git a/Engine/Source/Runtime/Engine/Public/Materials/MaterialExpressionRuntimeVirtualTextureReplace.h b/Engine/Source/Runtime/Engine/Public/Materials/MaterialExpressionRuntimeVirtualTextureReplace.h
--- a/Engine/Source/Runtime/Engine/Public/Materials/MaterialExpressionRuntimeVirtualTextureReplace.h
+++ b/Engine/Source/Runtime/Engine/Public/Materials/MaterialExpressionRuntimeVirtualTextureReplace.h
@@ -28,8 +28,8 @@ public:
 #if WITH_EDITOR
 	virtual int32 Compile(class FMaterialCompiler* Compiler, int32 OutputIndex) override;
 	virtual bool IsResultMaterialAttributes(int32 OutputIndex) override;
-	virtual uint32 GetInputType(int32 InputIndex) override { return MCT_Unknown; }
-	virtual uint32 GetOutputType(int32 OutputIndex) override { return MCT_Unknown; }
+	virtual EMaterialValueType GetInputValueType(int32 InputIndex) override { return MCT_Unknown; }
+	virtual EMaterialValueType GetOutputValueType(int32 OutputIndex) override { return MCT_Unknown; }
 	virtual void GetCaption(TArray<FString>& OutCaptions) const override;
 #endif
 	//~ End UMaterialExpression Interface
diff --git a/Engine/Source/Runtime/Engine/Public/Materials/MaterialExpressionSetMaterialAttributes.h b/Engine/Source/Runtime/Engine/Public/Materials/MaterialExpressionSetMaterialAttributes.h
--- a/Engine/Source/Runtime/Engine/Public/Materials/MaterialExpressionSetMaterialAttributes.h
+++ b/Engine/Source/Runtime/Engine/Public/Materials/MaterialExpressionSetMaterialAttributes.h
@@ -39,7 +39,7 @@ class UMaterialExpressionSetMaterialAttributes : public UMaterialExpression
 	virtual TArrayView<FExpressionInput*> GetInputsView()override;
 	virtual FExpressionInput* GetInput(int32 InputIndex)override;
 	virtual FName GetInputName(int32 InputIndex) const override;
-	virtual uint32 GetInputType(int32 InputIndex) override;
+	virtual EMaterialValueType GetInputValueType(int32 InputIndex) override;
 	virtual bool IsInputConnectionRequired(int32 InputIndex) const override {return true;}
 	virtual bool IsResultMaterialAttributes(int32 OutputIndex) override {return true;}
 
diff --git a/Engine/Source/Runtime/Engine/Public/Materials/MaterialExpressionShadingModel.h b/Engine/Source/Runtime/Engine/Public/Materials/MaterialExpressionShadingModel.h
--- a/Engine/Source/Runtime/Engine/Public/Materials/MaterialExpressionShadingModel.h
+++ b/Engine/Source/Runtime/Engine/Public/Materials/MaterialExpressionShadingModel.h
@@ -29,7 +29,7 @@ class UMaterialExpressionShadingModel : public UMaterialExpression
 #if WITH_EDITOR
 	virtual int32 Compile(class FMaterialCompiler* Compiler, int32 OutputIndex) override;
 	virtual int32 CompilePreview(class FMaterialCompiler* Compiler, int32 OutputIndex) override;
-	virtual uint32 GetOutputType(int32 OutputIndex) override;
+	virtual EMaterialValueType GetOutputValueType(int32 OutputIndex) override;
 	virtual void GetCaption(TArray<FString>& OutCaptions) const override;
 #endif
 public:
diff --git a/Engine/Source/Runtime/Engine/Public/Materials/MaterialExpressionShadingPathSwitch.h b/Engine/Source/Runtime/Engine/Public/Materials/MaterialExpressionShadingPathSwitch.h
--- a/Engine/Source/Runtime/Engine/Public/Materials/MaterialExpressionShadingPathSwitch.h
+++ b/Engine/Source/Runtime/Engine/Public/Materials/MaterialExpressionShadingPathSwitch.h
@@ -33,8 +33,8 @@ public:
 	virtual FName GetInputName(int32 InputIndex) const override;
 	virtual bool IsInputConnectionRequired(int32 InputIndex) const override;
 	virtual bool IsResultMaterialAttributes(int32 OutputIndex) override;
-	virtual uint32 GetInputType(int32 InputIndex) override {return MCT_Unknown;}
-	virtual uint32 GetOutputType(int32 InputIndex) override {return MCT_Unknown;}
+	virtual EMaterialValueType GetInputValueType(int32 InputIndex) override {return MCT_Unknown;}
+	virtual EMaterialValueType GetOutputValueType(int32 InputIndex) override {return MCT_Unknown;}
 
 	virtual bool IsResultSubstrateMaterial(int32 OutputIndex) override;
 	virtual void GatherSubstrateMaterialInfo(FSubstrateMaterialInfo& SubstrateMaterialInfo, int32 OutputIndex) override;
diff --git a/Engine/Source/Runtime/Engine/Public/Materials/MaterialExpressionSparseVolumeTextureObject.h b/Engine/Source/Runtime/Engine/Public/Materials/MaterialExpressionSparseVolumeTextureObject.h
--- a/Engine/Source/Runtime/Engine/Public/Materials/MaterialExpressionSparseVolumeTextureObject.h
+++ b/Engine/Source/Runtime/Engine/Public/Materials/MaterialExpressionSparseVolumeTextureObject.h
@@ -23,7 +23,7 @@ class UMaterialExpressionSparseVolumeTextureObject : public UMaterialExpressionS
 #if WITH_EDITOR
 	virtual int32 Compile(class FMaterialCompiler* Compiler, int32 OutputIndex) override;
 	virtual void GetCaption(TArray<FString>& OutCaptions) const override;
-	virtual uint32 GetOutputType(int32 OutputIndex) override;
+	virtual EMaterialValueType GetOutputValueType(int32 OutputIndex) override;
 	//~ End UMaterialExpression Interface
 #endif // WITH_EDITOR
 };
@@ -37,7 +37,7 @@ class UMaterialExpressionSparseVolumeTextureObjectParameter : public UMaterialEx
 #if WITH_EDITOR
 	virtual void GetCaption(TArray<FString>& OutCaptions) const override;
 	virtual int32 Compile(class FMaterialCompiler* Compiler, int32 OutputIndex) override;
-	virtual uint32 GetOutputType(int32 OutputIndex) override;
+	virtual EMaterialValueType GetOutputValueType(int32 OutputIndex) override;
 #endif
 	//~ End UMaterialExpression Interface
 };
diff --git a/Engine/Source/Runtime/Engine/Public/Materials/MaterialExpressionSparseVolumeTextureSample.h b/Engine/Source/Runtime/Engine/Public/Materials/MaterialExpressionSparseVolumeTextureSample.h
--- a/Engine/Source/Runtime/Engine/Public/Materials/MaterialExpressionSparseVolumeTextureSample.h
+++ b/Engine/Source/Runtime/Engine/Public/Materials/MaterialExpressionSparseVolumeTextureSample.h
@@ -65,8 +65,8 @@ protected:
 	ENGINE_API virtual FName GetInputName(int32 InputIndex) const override;
 	ENGINE_API virtual int32 Compile(class FMaterialCompiler* Compiler, int32 OutputIndex) override;
 	ENGINE_API virtual void GetCaption(TArray<FString>& OutCaptions) const override;
-	ENGINE_API virtual uint32 GetOutputType(int32 OutputIndex) override;
-	ENGINE_API virtual uint32 GetInputType(int32 InputIndex) override;
+	ENGINE_API virtual EMaterialValueType GetOutputValueType(int32 OutputIndex) override;
+	ENGINE_API virtual EMaterialValueType GetInputValueType(int32 InputIndex) override;
 public:
 	ENGINE_API virtual void PostEditChangeProperty(FPropertyChangedEvent& PropertyChangedEvent) override;
 #endif
diff --git a/Engine/Source/Runtime/Engine/Public/Materials/MaterialExpressionStaticBool.h b/Engine/Source/Runtime/Engine/Public/Materials/MaterialExpressionStaticBool.h
--- a/Engine/Source/Runtime/Engine/Public/Materials/MaterialExpressionStaticBool.h
+++ b/Engine/Source/Runtime/Engine/Public/Materials/MaterialExpressionStaticBool.h
@@ -24,7 +24,7 @@ class UMaterialExpressionStaticBool : public UMaterialExpression
 	virtual int32 CompilePreview(class FMaterialCompiler* Compiler, int32 OutputIndex) override;
 	virtual int32 Compile(class FMaterialCompiler* Compiler, int32 OutputIndex) override;
 	virtual void GetCaption(TArray<FString>& OutCaptions) const override;
-	virtual uint32 GetOutputType(int32 OutputIndex) override {return MCT_StaticBool;}
+	virtual EMaterialValueType GetOutputValueType(int32 OutputIndex) override {return MCT_StaticBool;}
 
 #endif
 	//~ End UMaterialExpression Interface
diff --git a/Engine/Source/Runtime/Engine/Public/Materials/MaterialExpressionStaticBoolParameter.h b/Engine/Source/Runtime/Engine/Public/Materials/MaterialExpressionStaticBoolParameter.h
--- a/Engine/Source/Runtime/Engine/Public/Materials/MaterialExpressionStaticBoolParameter.h
+++ b/Engine/Source/Runtime/Engine/Public/Materials/MaterialExpressionStaticBoolParameter.h
@@ -28,7 +28,7 @@ public:
 	virtual int32 Compile(class FMaterialCompiler* Compiler, int32 OutputIndex) override;
 	virtual int32 CompilePreview(class FMaterialCompiler* Compiler, int32 OutputIndex) override;
 	virtual void GetCaption(TArray<FString>& OutCaptions) const override;
-	virtual uint32 GetOutputType(int32 OutputIndex) override {return DynamicBranch ? MCT_Bool : MCT_StaticBool;}
+	virtual EMaterialValueType GetOutputValueType(int32 OutputIndex) override {return DynamicBranch ? MCT_Bool : MCT_StaticBool;}
 	virtual bool GetParameterValue(FMaterialParameterMetadata& OutMeta) const override
 	{
 		OutMeta.Value = (bool)DefaultValue;
diff --git a/Engine/Source/Runtime/Engine/Public/Materials/MaterialExpressionStaticSwitch.h b/Engine/Source/Runtime/Engine/Public/Materials/MaterialExpressionStaticSwitch.h
--- a/Engine/Source/Runtime/Engine/Public/Materials/MaterialExpressionStaticSwitch.h
+++ b/Engine/Source/Runtime/Engine/Public/Materials/MaterialExpressionStaticSwitch.h
@@ -37,8 +37,8 @@ public:
 	virtual void GetCaption(TArray<FString>& OutCaptions) const override;
 	virtual FName GetInputName(int32 InputIndex) const override;
 	virtual bool IsResultMaterialAttributes(int32 OutputIndex) override;
-	virtual uint32 GetInputType(int32 InputIndex) override;
-	virtual uint32 GetOutputType(int32 OutputIndex) override {return MCT_Unknown;}
+	virtual EMaterialValueType GetInputValueType(int32 InputIndex) override;
+	virtual EMaterialValueType GetOutputValueType(int32 OutputIndex) override {return MCT_Unknown;}
 
 
 	virtual bool IsResultSubstrateMaterial(int32 OutputIndex) override;
diff --git a/Engine/Source/Runtime/Engine/Public/Materials/MaterialExpressionStaticSwitchParameter.h b/Engine/Source/Runtime/Engine/Public/Materials/MaterialExpressionStaticSwitchParameter.h
--- a/Engine/Source/Runtime/Engine/Public/Materials/MaterialExpressionStaticSwitchParameter.h
+++ b/Engine/Source/Runtime/Engine/Public/Materials/MaterialExpressionStaticSwitchParameter.h
@@ -30,8 +30,8 @@ public:
 	virtual int32 Compile(class FMaterialCompiler* Compiler, int32 OutputIndex) override;
 	virtual void GetCaption(TArray<FString>& OutCaptions) const override;
 	virtual bool IsResultMaterialAttributes(int32 OutputIndex) override;
-	virtual uint32 GetInputType(int32 InputIndex) override {return MCT_Unknown;}
-	virtual uint32 GetOutputType(int32 OutputIndex) override {return MCT_Unknown;}
+	virtual EMaterialValueType GetInputValueType(int32 InputIndex) override {return MCT_Unknown;}
+	virtual EMaterialValueType GetOutputValueType(int32 OutputIndex) override {return MCT_Unknown;}
 
 	virtual bool IsResultSubstrateMaterial(int32 OutputIndex) override;
 	virtual void GatherSubstrateMaterialInfo(FSubstrateMaterialInfo& SubstrateMaterialInfo, int32 OutputIndex) override;
diff --git a/Engine/Source/Runtime/Engine/Public/Materials/MaterialExpressionSubstrate.h b/Engine/Source/Runtime/Engine/Public/Materials/MaterialExpressionSubstrate.h
--- a/Engine/Source/Runtime/Engine/Public/Materials/MaterialExpressionSubstrate.h
+++ b/Engine/Source/Runtime/Engine/Public/Materials/MaterialExpressionSubstrate.h
@@ -214,8 +214,8 @@ class UMaterialExpressionSubstrateShadingModels : public UMaterialExpressionSubs
 	virtual FExpressionInput* GetInput(int32 InputIndex) override;
 	virtual int32 Compile(class FMaterialCompiler* Compiler, int32 OutputIndex) override;
 	virtual void GetCaption(TArray<FString>& OutCaptions) const override;
-	virtual uint32 GetOutputType(int32 OutputIndex) override;
-	virtual uint32 GetInputType(int32 InputIndex) override;
+	virtual EMaterialValueType GetOutputValueType(int32 OutputIndex) override;
+	virtual EMaterialValueType GetInputValueType(int32 InputIndex) override;
 	virtual bool IsResultSubstrateMaterial(int32 OutputIndex) override;
 	virtual void GatherSubstrateMaterialInfo(FSubstrateMaterialInfo& SubstrateMaterialInfo, int32 OutputIndex) override;
 	virtual FSubstrateOperator* SubstrateGenerateMaterialTopologyTree(class FMaterialCompiler* Compiler, class UMaterialExpression* Parent, int32 OutputIndex) override;
@@ -388,8 +388,8 @@ class UMaterialExpressionSubstrateSlabBSDF : public UMaterialExpressionSubstrate
 	virtual int32 Compile(class FMaterialCompiler* Compiler, int32 OutputIndex) override;
 	static int32 CompileDefaultSlab(class FMaterialCompiler* Compiler, FVector3f EmissiveOverride = FVector3f::Zero());
 	virtual void GetCaption(TArray<FString>& OutCaptions) const override;
-	virtual uint32 GetOutputType(int32 OutputIndex) override;
-	virtual uint32 GetInputType(int32 InputIndex) override;
+	virtual EMaterialValueType GetOutputValueType(int32 OutputIndex) override;
+	virtual EMaterialValueType GetInputValueType(int32 InputIndex) override;
 	virtual bool IsResultSubstrateMaterial(int32 OutputIndex) override;
 	virtual void GatherSubstrateMaterialInfo(FSubstrateMaterialInfo& SubstrateMaterialInfo, int32 OutputIndex) override;
 	virtual FSubstrateOperator* SubstrateGenerateMaterialTopologyTree(class FMaterialCompiler* Compiler, class UMaterialExpression* Parent, int32 OutputIndex) override;
@@ -471,8 +471,8 @@ class UMaterialExpressionSubstrateSimpleClearCoatBSDF : public UMaterialExpressi
 #if WITH_EDITOR
 	virtual int32 Compile(class FMaterialCompiler* Compiler, int32 OutputIndex) override;
 	virtual void GetCaption(TArray<FString>& OutCaptions) const override;
-	virtual uint32 GetOutputType(int32 OutputIndex) override;
-	virtual uint32 GetInputType(int32 InputIndex) override;
+	virtual EMaterialValueType GetOutputValueType(int32 OutputIndex) override;
+	virtual EMaterialValueType GetInputValueType(int32 InputIndex) override;
 	virtual bool IsResultSubstrateMaterial(int32 OutputIndex) override;
 	virtual void GatherSubstrateMaterialInfo(FSubstrateMaterialInfo& SubstrateMaterialInfo, int32 OutputIndex) override;
 	virtual FSubstrateOperator* SubstrateGenerateMaterialTopologyTree(class FMaterialCompiler* Compiler, class UMaterialExpression* Parent, int32 OutputIndex) override;
@@ -518,8 +518,8 @@ class UMaterialExpressionSubstrateVolumetricFogCloudBSDF : public UMaterialExpre
 #if WITH_EDITOR
 	virtual int32 Compile(class FMaterialCompiler* Compiler, int32 OutputIndex) override;
 	virtual void GetCaption(TArray<FString>& OutCaptions) const override;
-	virtual uint32 GetOutputType(int32 OutputIndex) override;
-	virtual uint32 GetInputType(int32 InputIndex) override;
+	virtual EMaterialValueType GetOutputValueType(int32 OutputIndex) override;
+	virtual EMaterialValueType GetInputValueType(int32 InputIndex) override;
 	virtual bool IsResultSubstrateMaterial(int32 OutputIndex) override;
 	virtual void GatherSubstrateMaterialInfo(FSubstrateMaterialInfo& SubstrateMaterialInfo, int32 OutputIndex) override;
 	virtual FSubstrateOperator* SubstrateGenerateMaterialTopologyTree(class FMaterialCompiler* Compiler, class UMaterialExpression* Parent, int32 OutputIndex) override;
@@ -564,8 +564,8 @@ class UMaterialExpressionSubstrateUnlitBSDF : public UMaterialExpressionSubstrat
 #if WITH_EDITOR
 	virtual int32 Compile(class FMaterialCompiler* Compiler, int32 OutputIndex) override;
 	virtual void GetCaption(TArray<FString>& OutCaptions) const override;
-	virtual uint32 GetOutputType(int32 OutputIndex) override;
-	virtual uint32 GetInputType(int32 InputIndex) override;
+	virtual EMaterialValueType GetOutputValueType(int32 OutputIndex) override;
+	virtual EMaterialValueType GetInputValueType(int32 InputIndex) override;
 	virtual bool IsResultSubstrateMaterial(int32 OutputIndex) override;
 	virtual void GatherSubstrateMaterialInfo(FSubstrateMaterialInfo& SubstrateMaterialInfo, int32 OutputIndex) override;
 	virtual FSubstrateOperator* SubstrateGenerateMaterialTopologyTree(class FMaterialCompiler* Compiler, class UMaterialExpression* Parent, int32 OutputIndex) override;
@@ -624,8 +624,8 @@ class UMaterialExpressionSubstrateHairBSDF : public UMaterialExpressionSubstrate
 #if WITH_EDITOR
 	virtual int32 Compile(class FMaterialCompiler* Compiler, int32 OutputIndex) override;
 	virtual void GetCaption(TArray<FString>& OutCaptions) const override;
-	virtual uint32 GetOutputType(int32 OutputIndex) override;
-	virtual uint32 GetInputType(int32 InputIndex) override;
+	virtual EMaterialValueType GetOutputValueType(int32 OutputIndex) override;
+	virtual EMaterialValueType GetInputValueType(int32 InputIndex) override;
 	virtual bool IsResultSubstrateMaterial(int32 OutputIndex) override;
 	virtual void GatherSubstrateMaterialInfo(FSubstrateMaterialInfo& SubstrateMaterialInfo, int32 OutputIndex) override;
 	virtual FSubstrateOperator* SubstrateGenerateMaterialTopologyTree(class FMaterialCompiler* Compiler, class UMaterialExpression* Parent, int32 OutputIndex) override;
@@ -694,8 +694,8 @@ class UMaterialExpressionSubstrateEyeBSDF : public UMaterialExpressionSubstrateB
 #if WITH_EDITOR
 	virtual int32 Compile(class FMaterialCompiler* Compiler, int32 OutputIndex) override;
 	virtual void GetCaption(TArray<FString>& OutCaptions) const override;
-	virtual uint32 GetOutputType(int32 OutputIndex) override;
-	virtual uint32 GetInputType(int32 InputIndex) override;
+	virtual EMaterialValueType GetOutputValueType(int32 OutputIndex) override;
+	virtual EMaterialValueType GetInputValueType(int32 InputIndex) override;
 	virtual bool IsResultSubstrateMaterial(int32 OutputIndex) override;
 	virtual void GatherSubstrateMaterialInfo(FSubstrateMaterialInfo& SubstrateMaterialInfo, int32 OutputIndex) override;
 	virtual FSubstrateOperator* SubstrateGenerateMaterialTopologyTree(class FMaterialCompiler* Compiler, class UMaterialExpression* Parent, int32 OutputIndex) override;
@@ -778,8 +778,8 @@ class UMaterialExpressionSubstrateSingleLayerWaterBSDF : public UMaterialExpress
 #if WITH_EDITOR
 	virtual int32 Compile(class FMaterialCompiler* Compiler, int32 OutputIndex) override;
 	virtual void GetCaption(TArray<FString>& OutCaptions) const override;
-	virtual uint32 GetOutputType(int32 OutputIndex) override;
-	virtual uint32 GetInputType(int32 InputIndex) override;
+	virtual EMaterialValueType GetOutputValueType(int32 OutputIndex) override;
+	virtual EMaterialValueType GetInputValueType(int32 InputIndex) override;
 	virtual bool IsResultSubstrateMaterial(int32 OutputIndex) override;
 	virtual void GatherSubstrateMaterialInfo(FSubstrateMaterialInfo& SubstrateMaterialInfo, int32 OutputIndex) override;
 	virtual FSubstrateOperator* SubstrateGenerateMaterialTopologyTree(class FMaterialCompiler* Compiler, class UMaterialExpression* Parent, int32 OutputIndex) override;
@@ -802,8 +802,8 @@ class UMaterialExpressionSubstrateLightFunction : public UMaterialExpressionSubs
 #if WITH_EDITOR
 	virtual int32 Compile(class FMaterialCompiler* Compiler, int32 OutputIndex) override;
 	virtual void GetCaption(TArray<FString>& OutCaptions) const override;
-	virtual uint32 GetOutputType(int32 OutputIndex) override;
-	virtual uint32 GetInputType(int32 InputIndex) override;
+	virtual EMaterialValueType GetOutputValueType(int32 OutputIndex) override;
+	virtual EMaterialValueType GetInputValueType(int32 InputIndex) override;
 	virtual bool IsResultSubstrateMaterial(int32 OutputIndex) override;
 	virtual void GatherSubstrateMaterialInfo(FSubstrateMaterialInfo& SubstrateMaterialInfo, int32 OutputIndex) override;
 	virtual FSubstrateOperator* SubstrateGenerateMaterialTopologyTree(class FMaterialCompiler* Compiler, class UMaterialExpression* Parent, int32 OutputIndex) override;
@@ -840,8 +840,8 @@ class UMaterialExpressionSubstratePostProcess : public UMaterialExpressionSubstr
 #if WITH_EDITOR
 	virtual int32 Compile(class FMaterialCompiler* Compiler, int32 OutputIndex) override;
 	virtual void GetCaption(TArray<FString>& OutCaptions) const override;
-	virtual uint32 GetOutputType(int32 OutputIndex) override;
-	virtual uint32 GetInputType(int32 InputIndex) override;
+	virtual EMaterialValueType GetOutputValueType(int32 OutputIndex) override;
+	virtual EMaterialValueType GetInputValueType(int32 InputIndex) override;
 	virtual bool IsResultSubstrateMaterial(int32 OutputIndex) override;
 	virtual void GatherSubstrateMaterialInfo(FSubstrateMaterialInfo& SubstrateMaterialInfo, int32 OutputIndex) override;
 	virtual FSubstrateOperator* SubstrateGenerateMaterialTopologyTree(class FMaterialCompiler* Compiler, class UMaterialExpression* Parent, int32 OutputIndex) override;
@@ -877,8 +877,8 @@ class UMaterialExpressionSubstrateUI : public UMaterialExpressionSubstrateBSDF
 #if WITH_EDITOR
 	virtual int32 Compile(class FMaterialCompiler* Compiler, int32 OutputIndex) override;
 	virtual void GetCaption(TArray<FString>& OutCaptions) const override;
-	virtual uint32 GetOutputType(int32 OutputIndex) override;
-	virtual uint32 GetInputType(int32 InputIndex) override;
+	virtual EMaterialValueType GetOutputValueType(int32 OutputIndex) override;
+	virtual EMaterialValueType GetInputValueType(int32 InputIndex) override;
 	virtual bool IsResultSubstrateMaterial(int32 OutputIndex) override;
 	virtual void GatherSubstrateMaterialInfo(FSubstrateMaterialInfo& SubstrateMaterialInfo, int32 OutputIndex) override;
 	virtual FSubstrateOperator* SubstrateGenerateMaterialTopologyTree(class FMaterialCompiler* Compiler, class UMaterialExpression* Parent, int32 OutputIndex) override;
@@ -914,8 +914,8 @@ class UMaterialExpressionSubstrateConvertToDecal : public UMaterialExpressionSub
 #if WITH_EDITOR
 	virtual int32 Compile(class FMaterialCompiler* Compiler, int32 OutputIndex) override;
 	virtual void GetCaption(TArray<FString>& OutCaptions) const override;
-	virtual uint32 GetOutputType(int32 OutputIndex) override;
-	virtual uint32 GetInputType(int32 InputIndex) override;
+	virtual EMaterialValueType GetOutputValueType(int32 OutputIndex) override;
+	virtual EMaterialValueType GetInputValueType(int32 InputIndex) override;
 	virtual bool IsResultSubstrateMaterial(int32 OutputIndex) override;
 	virtual void GatherSubstrateMaterialInfo(FSubstrateMaterialInfo& SubstrateMaterialInfo, int32 OutputIndex) override;
 	virtual FSubstrateOperator* SubstrateGenerateMaterialTopologyTree(class FMaterialCompiler* Compiler, class UMaterialExpression* Parent, int32 OutputIndex) override;
@@ -967,8 +967,8 @@ class UMaterialExpressionSubstrateConvertMaterialAttributes : public UMaterialEx
 	virtual FExpressionInput* GetInput(int32 InputIndex) override;
 	virtual int32 Compile(class FMaterialCompiler* Compiler, int32 OutputIndex) override;
 	virtual void GetCaption(TArray<FString>& OutCaptions) const override;
-	virtual uint32 GetOutputType(int32 OutputIndex) override;
-	virtual uint32 GetInputType(int32 InputIndex) override;
+	virtual EMaterialValueType GetOutputValueType(int32 OutputIndex) override;
+	virtual EMaterialValueType GetInputValueType(int32 InputIndex) override;
 	virtual FName GetInputName(int32 InputIndex) const override;
 	virtual bool IsResultSubstrateMaterial(int32 OutputIndex) override;
 	virtual bool IsResultMaterialAttributes(int32 OutputIndex) override;
@@ -1031,8 +1031,8 @@ class UMaterialExpressionSubstrateHorizontalMixing : public UMaterialExpressionS
 #if WITH_EDITOR
 	virtual int32 Compile(class FMaterialCompiler* Compiler, int32 OutputIndex) override;
 	virtual void GetCaption(TArray<FString>& OutCaptions) const override;
-	virtual uint32 GetOutputType(int32 OutputIndex) override;
-	virtual uint32 GetInputType(int32 InputIndex) override;
+	virtual EMaterialValueType GetOutputValueType(int32 OutputIndex) override;
+	virtual EMaterialValueType GetInputValueType(int32 InputIndex) override;
 	virtual bool IsResultSubstrateMaterial(int32 OutputIndex) override;
 	virtual void GatherSubstrateMaterialInfo(FSubstrateMaterialInfo& SubstrateMaterialInfo, int32 OutputIndex) override;
 	virtual FSubstrateOperator* SubstrateGenerateMaterialTopologyTree(class FMaterialCompiler* Compiler, class UMaterialExpression* Parent, int32 OutputIndex) override;
@@ -1074,8 +1074,8 @@ class UMaterialExpressionSubstrateVerticalLayering : public UMaterialExpressionS
 #if WITH_EDITOR
 	virtual int32 Compile(class FMaterialCompiler* Compiler, int32 OutputIndex) override;
 	virtual void GetCaption(TArray<FString>& OutCaptions) const override;
-	virtual uint32 GetOutputType(int32 OutputIndex) override;
-	virtual uint32 GetInputType(int32 InputIndex) override;
+	virtual EMaterialValueType GetOutputValueType(int32 OutputIndex) override;
+	virtual EMaterialValueType GetInputValueType(int32 InputIndex) override;
 	virtual bool IsResultSubstrateMaterial(int32 OutputIndex) override;
 	virtual void GatherSubstrateMaterialInfo(FSubstrateMaterialInfo& SubstrateMaterialInfo, int32 OutputIndex) override;
 	virtual FSubstrateOperator* SubstrateGenerateMaterialTopologyTree(class FMaterialCompiler* Compiler, UMaterialExpression* Parent, int32 OutputIndex) override;
@@ -1111,8 +1111,8 @@ class UMaterialExpressionSubstrateAdd : public UMaterialExpressionSubstrateBSDF
 #if WITH_EDITOR
 	virtual int32 Compile(class FMaterialCompiler* Compiler, int32 OutputIndex) override;
 	virtual void GetCaption(TArray<FString>& OutCaptions) const override;
-	virtual uint32 GetOutputType(int32 OutputIndex) override;
-	virtual uint32 GetInputType(int32 InputIndex) override;
+	virtual EMaterialValueType GetOutputValueType(int32 OutputIndex) override;
+	virtual EMaterialValueType GetInputValueType(int32 InputIndex) override;
 	virtual bool IsResultSubstrateMaterial(int32 OutputIndex) override;
 	virtual void GatherSubstrateMaterialInfo(FSubstrateMaterialInfo& SubstrateMaterialInfo, int32 OutputIndex) override;
 	virtual FSubstrateOperator* SubstrateGenerateMaterialTopologyTree(class FMaterialCompiler* Compiler, class UMaterialExpression* Parent, int32 OutputIndex) override;
@@ -1141,8 +1141,8 @@ class UMaterialExpressionSubstrateWeight : public UMaterialExpressionSubstrateBS
 #if WITH_EDITOR
 	virtual int32 Compile(class FMaterialCompiler* Compiler, int32 OutputIndex) override;
 	virtual void GetCaption(TArray<FString>& OutCaptions) const override;
-	virtual uint32 GetOutputType(int32 OutputIndex) override;
-	virtual uint32 GetInputType(int32 InputIndex) override;
+	virtual EMaterialValueType GetOutputValueType(int32 OutputIndex) override;
+	virtual EMaterialValueType GetInputValueType(int32 InputIndex) override;
 	virtual bool IsResultSubstrateMaterial(int32 OutputIndex) override;
 	virtual void GatherSubstrateMaterialInfo(FSubstrateMaterialInfo& SubstrateMaterialInfo, int32 OutputIndex) override;
 	virtual FSubstrateOperator* SubstrateGenerateMaterialTopologyTree(class FMaterialCompiler* Compiler, class UMaterialExpression* Parent, int32 OutputIndex) override;
@@ -1189,8 +1189,8 @@ class UMaterialExpressionSubstrateSelect : public UMaterialExpressionSubstrateBS
 #if WITH_EDITOR
 	virtual int32 Compile(class FMaterialCompiler* Compiler, int32 OutputIndex) override;
 	virtual void GetCaption(TArray<FString>& OutCaptions) const override;
-	virtual uint32 GetOutputType(int32 OutputIndex) override;
-	virtual uint32 GetInputType(int32 InputIndex) override;
+	virtual EMaterialValueType GetOutputValueType(int32 OutputIndex) override;
+	virtual EMaterialValueType GetInputValueType(int32 InputIndex) override;
 	virtual bool IsResultSubstrateMaterial(int32 OutputIndex) override;
 	virtual void GatherSubstrateMaterialInfo(FSubstrateMaterialInfo& SubstrateMaterialInfo, int32 OutputIndex) override;
 	virtual FSubstrateOperator* SubstrateGenerateMaterialTopologyTree(class FMaterialCompiler* Compiler, class UMaterialExpression* Parent, int32 OutputIndex) override;
@@ -1231,8 +1231,8 @@ class UMaterialExpressionSubstrateTransmittanceToMFP : public UMaterialExpressio
 #if WITH_EDITOR
 	virtual int32 Compile(class FMaterialCompiler* Compiler, int32 OutputIndex) override;
 	virtual void GetCaption(TArray<FString>& OutCaptions) const override;
-	virtual uint32 GetOutputType(int32 OutputIndex) override;
-	virtual uint32 GetInputType(int32 InputIndex) override;
+	virtual EMaterialValueType GetOutputValueType(int32 OutputIndex) override;
+	virtual EMaterialValueType GetInputValueType(int32 InputIndex) override;
 	virtual void GetConnectorToolTip(int32 InputIndex, int32 OutputIndex, TArray<FString>& OutToolTip) override;
 	virtual void GetExpressionToolTip(TArray<FString>& OutToolTip) override;
 #endif
@@ -1266,8 +1266,8 @@ class UMaterialExpressionSubstrateMetalnessToDiffuseAlbedoF0 : public UMaterialE
 #if WITH_EDITOR
 	virtual int32 Compile(class FMaterialCompiler* Compiler, int32 OutputIndex) override;
 	virtual void GetCaption(TArray<FString>& OutCaptions) const override;
-	virtual uint32 GetOutputType(int32 OutputIndex) override;
-	virtual uint32 GetInputType(int32 InputIndex) override;
+	virtual EMaterialValueType GetOutputValueType(int32 OutputIndex) override;
+	virtual EMaterialValueType GetInputValueType(int32 InputIndex) override;
 	virtual void GetConnectorToolTip(int32 InputIndex, int32 OutputIndex, TArray<FString>& OutToolTip) override;
 	virtual void GetExpressionToolTip(TArray<FString>& OutToolTip) override;
 #endif
@@ -1295,8 +1295,8 @@ class UMaterialExpressionSubstrateHazinessToSecondaryRoughness : public UMateria
 #if WITH_EDITOR
 	virtual int32 Compile(class FMaterialCompiler* Compiler, int32 OutputIndex) override;
 	virtual void GetCaption(TArray<FString>& OutCaptions) const override;
-	virtual uint32 GetOutputType(int32 OutputIndex) override;
-	virtual uint32 GetInputType(int32 InputIndex) override;
+	virtual EMaterialValueType GetOutputValueType(int32 OutputIndex) override;
+	virtual EMaterialValueType GetInputValueType(int32 InputIndex) override;
 	virtual void GetConnectorToolTip(int32 InputIndex, int32 OutputIndex, TArray<FString>& OutToolTip) override;
 	virtual void GetExpressionToolTip(TArray<FString>& OutToolTip) override;
 #endif
@@ -1342,8 +1342,8 @@ class UMaterialExpressionSubstrateThinFilm : public UMaterialExpressionSubstrate
 #if WITH_EDITOR
 	virtual int32 Compile(class FMaterialCompiler* Compiler, int32 OutputIndex) override;
 	virtual void GetCaption(TArray<FString>& OutCaptions) const override;
-	virtual uint32 GetOutputType(int32 OutputIndex) override;
-	virtual uint32 GetInputType(int32 InputIndex) override;
+	virtual EMaterialValueType GetOutputValueType(int32 OutputIndex) override;
+	virtual EMaterialValueType GetInputValueType(int32 InputIndex) override;
 	virtual void GetConnectorToolTip(int32 InputIndex, int32 OutputIndex, TArray<FString>& OutToolTip) override;
 	virtual void GetExpressionToolTip(TArray<FString>& OutToolTip) override;
 #endif
diff --git a/Engine/Source/Runtime/Engine/Public/Materials/MaterialExpressionSubsurfaceMediumMaterialOutput.h b/Engine/Source/Runtime/Engine/Public/Materials/MaterialExpressionSubsurfaceMediumMaterialOutput.h
--- a/Engine/Source/Runtime/Engine/Public/Materials/MaterialExpressionSubsurfaceMediumMaterialOutput.h
+++ b/Engine/Source/Runtime/Engine/Public/Materials/MaterialExpressionSubsurfaceMediumMaterialOutput.h
@@ -26,7 +26,7 @@ public:
 	//~ Begin UMaterialExpression Interface
 	virtual int32 Compile(class FMaterialCompiler* Compiler, int32 OutputIndex) override;
 	virtual void GetCaption(TArray<FString>& OutCaptions) const override;
-	virtual uint32 GetInputType(int32 InputIndex) override;
+	virtual EMaterialValueType GetInputValueType(int32 InputIndex) override;
 	//~ End UMaterialExpression Interface
 #endif
 
diff --git a/Engine/Source/Runtime/Engine/Public/Materials/MaterialExpressionSwitch.h b/Engine/Source/Runtime/Engine/Public/Materials/MaterialExpressionSwitch.h
--- a/Engine/Source/Runtime/Engine/Public/Materials/MaterialExpressionSwitch.h
+++ b/Engine/Source/Runtime/Engine/Public/Materials/MaterialExpressionSwitch.h
@@ -75,7 +75,7 @@ class UMaterialExpressionSwitch : public UMaterialExpression
 	virtual TArrayView<FExpressionInput*> GetInputsView() override;
 	virtual FExpressionInput* GetInput(int32 InputIndex) override;
 	virtual FName GetInputName(int32 InputIndex) const override;
-	virtual uint32 GetInputType(int32 InputIndex) override {return MCT_Unknown;}
+	virtual EMaterialValueType GetInputValueType(int32 InputIndex) override {return MCT_Unknown;}
 
 #endif // WITH_EDITOR
 	//~ End UMaterialExpression Interface
diff --git a/Engine/Source/Runtime/Engine/Public/Materials/MaterialExpressionTangentOutput.h b/Engine/Source/Runtime/Engine/Public/Materials/MaterialExpressionTangentOutput.h
--- a/Engine/Source/Runtime/Engine/Public/Materials/MaterialExpressionTangentOutput.h
+++ b/Engine/Source/Runtime/Engine/Public/Materials/MaterialExpressionTangentOutput.h
@@ -20,7 +20,7 @@ class UMaterialExpressionTangentOutput : public UMaterialExpressionCustomOutput
 #if WITH_EDITOR
 	virtual int32 Compile(class FMaterialCompiler* Compiler, int32 OutputIndex) override;
 	virtual void GetCaption(TArray<FString>& OutCaptions) const override;
-	virtual uint32 GetInputType(int32 InputIndex) override { return MCT_Float; }
+	virtual EMaterialValueType GetInputValueType(int32 InputIndex) override { return MCT_Float; }
 #endif
 
 	virtual int32 GetNumOutputs() const override { return 1; }
diff --git a/Engine/Source/Runtime/Engine/Public/Materials/MaterialExpressionTextureCollection.h b/Engine/Source/Runtime/Engine/Public/Materials/MaterialExpressionTextureCollection.h
--- a/Engine/Source/Runtime/Engine/Public/Materials/MaterialExpressionTextureCollection.h
+++ b/Engine/Source/Runtime/Engine/Public/Materials/MaterialExpressionTextureCollection.h
@@ -18,7 +18,7 @@ class UMaterialExpressionTextureCollection : public UMaterialExpression
 #if WITH_EDITOR
 	virtual int32 Compile(class FMaterialCompiler* Compiler, int32 OutputIndex) override;
 	virtual void GetCaption(TArray<FString>& OutCaptions) const override;
-	virtual uint32 GetOutputType(int32 OutputIndex) override;
+	virtual EMaterialValueType GetOutputValueType(int32 OutputIndex) override;
 
 	virtual UTextureCollection* GetReferencedTextureCollection() const { return TextureCollection; }
 #endif
diff --git a/Engine/Source/Runtime/Engine/Public/Materials/MaterialExpressionTextureObject.h b/Engine/Source/Runtime/Engine/Public/Materials/MaterialExpressionTextureObject.h
--- a/Engine/Source/Runtime/Engine/Public/Materials/MaterialExpressionTextureObject.h
+++ b/Engine/Source/Runtime/Engine/Public/Materials/MaterialExpressionTextureObject.h
@@ -28,7 +28,7 @@ class UMaterialExpressionTextureObject : public UMaterialExpressionTextureBase
 	virtual int32 Compile(class FMaterialCompiler* Compiler, int32 OutputIndex) override;
 	virtual int32 CompilePreview(class FMaterialCompiler* Compiler, int32 OutputIndex) override;
 	virtual void GetCaption(TArray<FString>& OutCaptions) const override;
-	virtual uint32 GetOutputType(int32 OutputIndex) override;
+	virtual EMaterialValueType GetOutputValueType(int32 OutputIndex) override;
 	//~ End UMaterialExpression Interface
 #endif // WITH_EDITOR
 };
diff --git a/Engine/Source/Runtime/Engine/Public/Materials/MaterialExpressionTextureObjectFromCollection.h b/Engine/Source/Runtime/Engine/Public/Materials/MaterialExpressionTextureObjectFromCollection.h
--- a/Engine/Source/Runtime/Engine/Public/Materials/MaterialExpressionTextureObjectFromCollection.h
+++ b/Engine/Source/Runtime/Engine/Public/Materials/MaterialExpressionTextureObjectFromCollection.h
@@ -58,8 +58,8 @@ public:
 	//~ Begin UMaterialExpression Interface
 #if WITH_EDITOR
 	virtual int32 Compile(class FMaterialCompiler* Compiler, int32 OutputIndex) override;
-	virtual uint32 GetInputType(int32 InputIndex) override;
-	virtual uint32 GetOutputType(int32 OutputIndex) override;
+	virtual EMaterialValueType GetInputValueType(int32 InputIndex) override;
+	virtual EMaterialValueType GetOutputValueType(int32 OutputIndex) override;
 	virtual void GetCaption(TArray<FString>& OutCaptions) const override;
 
 	virtual UTextureCollection* GetReferencedTextureCollection() const { return TextureCollectionObject; }
diff --git a/Engine/Source/Runtime/Engine/Public/Materials/MaterialExpressionTextureObjectParameter.h b/Engine/Source/Runtime/Engine/Public/Materials/MaterialExpressionTextureObjectParameter.h
--- a/Engine/Source/Runtime/Engine/Public/Materials/MaterialExpressionTextureObjectParameter.h
+++ b/Engine/Source/Runtime/Engine/Public/Materials/MaterialExpressionTextureObjectParameter.h
@@ -30,7 +30,7 @@ class UMaterialExpressionTextureObjectParameter : public UMaterialExpressionText
 	virtual void GetCaption(TArray<FString>& OutCaptions) const override;	
 	virtual int32 Compile(class FMaterialCompiler* Compiler, int32 OutputIndex) override;
 	virtual int32 CompilePreview(class FMaterialCompiler* Compiler, int32 OutputIndex) override;
-	virtual uint32 GetOutputType(int32 InputIndex) override {return MCT_Texture;}
+	virtual EMaterialValueType GetOutputValueType(int32 InputIndex) override {return MCT_Texture;}
 #endif
 	//~ End UMaterialExpression Interface
 };
diff --git a/Engine/Source/Runtime/Engine/Public/Materials/MaterialExpressionTextureProperty.h b/Engine/Source/Runtime/Engine/Public/Materials/MaterialExpressionTextureProperty.h
--- a/Engine/Source/Runtime/Engine/Public/Materials/MaterialExpressionTextureProperty.h
+++ b/Engine/Source/Runtime/Engine/Public/Materials/MaterialExpressionTextureProperty.h
@@ -41,7 +41,7 @@ class UMaterialExpressionTextureProperty : public UMaterialExpression
 	ENGINE_API virtual int32 Compile(class FMaterialCompiler* Compiler, int32 OutputIndex) override;
 	ENGINE_API virtual void GetTexturesForceMaterialRecompile(TArray<UTexture *> &Textures) const override;
 	ENGINE_API virtual void GetCaption(TArray<FString>& OutCaptions) const override;
-	ENGINE_API virtual uint32 GetInputType(int32 InputIndex) override;
+	ENGINE_API virtual EMaterialValueType GetInputValueType(int32 InputIndex) override;
 	ENGINE_API virtual bool MatchesSearchQuery(const TCHAR* SearchQuery) override;
 
 #endif //WITH_EDITOR
diff --git a/Engine/Source/Runtime/Engine/Public/Materials/MaterialExpressionTextureSample.h b/Engine/Source/Runtime/Engine/Public/Materials/MaterialExpressionTextureSample.h
--- a/Engine/Source/Runtime/Engine/Public/Materials/MaterialExpressionTextureSample.h
+++ b/Engine/Source/Runtime/Engine/Public/Materials/MaterialExpressionTextureSample.h
@@ -125,7 +125,7 @@ public:
 	virtual int32 GetLabelPadding() override { return 8; }
 	ENGINE_API virtual int32 Compile(class FMaterialCompiler* Compiler, int32 OutputIndex) override;
 	ENGINE_API virtual void GetCaption(TArray<FString>& OutCaptions) const override;
-	ENGINE_API virtual uint32 GetInputType(int32 InputIndex) override;
+	ENGINE_API virtual EMaterialValueType GetInputValueType(int32 InputIndex) override;
 	virtual bool CanIgnoreOutputIndex() { return true; }
 	#endif // WITH_EDITOR
 	//~ End UMaterialExpression Interface
diff --git a/Engine/Source/Runtime/Engine/Public/Materials/MaterialExpressionThinTranslucentMaterialOutput.h b/Engine/Source/Runtime/Engine/Public/Materials/MaterialExpressionThinTranslucentMaterialOutput.h
--- a/Engine/Source/Runtime/Engine/Public/Materials/MaterialExpressionThinTranslucentMaterialOutput.h
+++ b/Engine/Source/Runtime/Engine/Public/Materials/MaterialExpressionThinTranslucentMaterialOutput.h
@@ -26,7 +26,7 @@ public:
 	//~ Begin UMaterialExpression Interface
 	virtual int32 Compile(class FMaterialCompiler* Compiler, int32 OutputIndex) override;
 	virtual void GetCaption(TArray<FString>& OutCaptions) const override;
-	virtual uint32 GetInputType(int32 InputIndex) override;
+	virtual EMaterialValueType GetInputValueType(int32 InputIndex) override;
 	//~ End UMaterialExpression Interface
 #endif
 
diff --git a/Engine/Source/Runtime/Engine/Public/Materials/MaterialExpressionTruncate.h b/Engine/Source/Runtime/Engine/Public/Materials/MaterialExpressionTruncate.h
--- a/Engine/Source/Runtime/Engine/Public/Materials/MaterialExpressionTruncate.h
+++ b/Engine/Source/Runtime/Engine/Public/Materials/MaterialExpressionTruncate.h
@@ -22,7 +22,7 @@ public:
 #if WITH_EDITOR
 	virtual void Build(MIR::FEmitter& Emitter) override;
 	virtual int32 Compile(class FMaterialCompiler* Compiler, int32 OutputIndex) override;
-	virtual uint32 GetOutputType(int32 OutputIndex) override;
+	virtual EMaterialValueType GetOutputValueType(int32 OutputIndex) override;
 	virtual void GetCaption(TArray<FString>& OutCaptions) const override;
 	virtual void GetExpressionToolTip(TArray<FString>& OutToolTip) override;
 #endif
diff --git a/Engine/Source/Runtime/Engine/Public/Materials/MaterialExpressionVertexInterpolator.h b/Engine/Source/Runtime/Engine/Public/Materials/MaterialExpressionVertexInterpolator.h
--- a/Engine/Source/Runtime/Engine/Public/Materials/MaterialExpressionVertexInterpolator.h
+++ b/Engine/Source/Runtime/Engine/Public/Materials/MaterialExpressionVertexInterpolator.h
@@ -18,8 +18,8 @@ class UMaterialExpressionVertexInterpolator : public UMaterialExpressionCustomOu
 	virtual int32 Compile(class FMaterialCompiler* Compiler, int32 OutputIndex) override;
 	ENGINE_API int32 CompileInput(class FMaterialCompiler* Compiler, int32 AssignedInterpolatorIndex);
 	virtual void GetCaption(TArray<FString>& OutCaptions) const override;
-	virtual uint32 GetInputType(int32 InputIndex) override;
-	virtual uint32 GetOutputType(int32 OutputIndex) override;
+	virtual EMaterialValueType GetInputValueType(int32 InputIndex) override;
+	virtual EMaterialValueType GetOutputValueType(int32 OutputIndex) override;
 	virtual bool HasCustomSourceOutput() override { return true; }
 
 	virtual FExpressionInput* GetInput(int32 InputIndex) override;
diff --git a/Engine/Source/Runtime/Engine/Public/Materials/MaterialExpressionVirtualTextureFeatureSwitch.h b/Engine/Source/Runtime/Engine/Public/Materials/MaterialExpressionVirtualTextureFeatureSwitch.h
--- a/Engine/Source/Runtime/Engine/Public/Materials/MaterialExpressionVirtualTextureFeatureSwitch.h
+++ b/Engine/Source/Runtime/Engine/Public/Materials/MaterialExpressionVirtualTextureFeatureSwitch.h
@@ -29,8 +29,8 @@ public:
 #if WITH_EDITOR
 	virtual int32 Compile(class FMaterialCompiler* Compiler, int32 OutputIndex) override;
 	virtual bool IsResultMaterialAttributes(int32 OutputIndex) override;
-	virtual uint32 GetInputType(int32 InputIndex) override { return MCT_Unknown; }
-	virtual uint32 GetOutputType(int32 OutputIndex) override { return MCT_Unknown; }
+	virtual EMaterialValueType GetInputValueType(int32 InputIndex) override { return MCT_Unknown; }
+	virtual EMaterialValueType GetOutputValueType(int32 OutputIndex) override { return MCT_Unknown; }
 	virtual void GetCaption(TArray<FString>& OutCaptions) const override;
 
 #endif
diff --git a/Engine/Source/Runtime/Engine/Public/Materials/MaterialExpressionVolumetricAdvancedMaterialInput.h b/Engine/Source/Runtime/Engine/Public/Materials/MaterialExpressionVolumetricAdvancedMaterialInput.h
--- a/Engine/Source/Runtime/Engine/Public/Materials/MaterialExpressionVolumetricAdvancedMaterialInput.h
+++ b/Engine/Source/Runtime/Engine/Public/Materials/MaterialExpressionVolumetricAdvancedMaterialInput.h
@@ -29,7 +29,7 @@ class UMaterialExpressionVolumetricCloudEmptySpaceSkippingInput : public UMateri
 #if WITH_EDITOR
 	virtual void GetCaption(TArray<FString>& OutCaptions) const override;
 	virtual void GetConnectorToolTip(int32 InputIndex, int32 OutputIndex, TArray<FString>& OutToolTip) override;
-	virtual uint32 GetInputType(int32 InputIndex) override;
+	virtual EMaterialValueType GetInputValueType(int32 InputIndex) override;
 #endif
 	//~ End UMaterialExpression Interface
 };
diff --git a/Engine/Source/Runtime/Engine/Public/Materials/MaterialExpressionVolumetricAdvancedMaterialOutput.h b/Engine/Source/Runtime/Engine/Public/Materials/MaterialExpressionVolumetricAdvancedMaterialOutput.h
--- a/Engine/Source/Runtime/Engine/Public/Materials/MaterialExpressionVolumetricAdvancedMaterialOutput.h
+++ b/Engine/Source/Runtime/Engine/Public/Materials/MaterialExpressionVolumetricAdvancedMaterialOutput.h
@@ -125,7 +125,7 @@ public:
 	//~ Begin UMaterialExpression Interface
 	virtual int32 Compile(class FMaterialCompiler* Compiler, int32 OutputIndex) override;
 	virtual void GetCaption(TArray<FString>& OutCaptions) const override;
-	virtual uint32 GetOutputType(int32 OutputIndex) override;
+	virtual EMaterialValueType GetOutputValueType(int32 OutputIndex) override;
 	//~ End UMaterialExpression Interface
 #endif
 
diff --git a/Engine/Source/Runtime/Landscape/Classes/Materials/MaterialExpressionLandscapeGrassOutput.h b/Engine/Source/Runtime/Landscape/Classes/Materials/MaterialExpressionLandscapeGrassOutput.h
--- a/Engine/Source/Runtime/Landscape/Classes/Materials/MaterialExpressionLandscapeGrassOutput.h
+++ b/Engine/Source/Runtime/Landscape/Classes/Materials/MaterialExpressionLandscapeGrassOutput.h
@@ -59,7 +59,7 @@ class UMaterialExpressionLandscapeGrassOutput : public UMaterialExpressionCustom
 	LANDSCAPE_API virtual void PostEditChangeProperty(FPropertyChangedEvent& PropertyChangedEvent) override;
 	//~ End UObject Interface
 
-	virtual uint32 GetInputType(int32 InputIndex) override { return MCT_Float; }
+	virtual EMaterialValueType GetInputValueType(int32 InputIndex) override { return MCT_Float; }
 #endif
 
 	virtual int32 GetNumOutputs() const override { return GrassTypes.Num(); }
diff --git a/Engine/Source/Runtime/Landscape/Classes/Materials/MaterialExpressionLandscapeLayerBlend.h b/Engine/Source/Runtime/Landscape/Classes/Materials/MaterialExpressionLandscapeLayerBlend.h
--- a/Engine/Source/Runtime/Landscape/Classes/Materials/MaterialExpressionLandscapeLayerBlend.h
+++ b/Engine/Source/Runtime/Landscape/Classes/Materials/MaterialExpressionLandscapeLayerBlend.h
@@ -75,8 +75,8 @@ class UMaterialExpressionLandscapeLayerBlend : public UMaterialExpression
 	//~ End UObject Interface
 
 	//~ Begin UMaterialExpression Interface
-	LANDSCAPE_API virtual uint32 GetInputType(int32 InputIndex) override;
-	virtual uint32 GetOutputType(int32 InputIndex) override { return MCT_Unknown; }
+	LANDSCAPE_API virtual EMaterialValueType GetInputValueType(int32 InputIndex) override;
+	virtual EMaterialValueType GetOutputValueType(int32 InputIndex) override { return MCT_Unknown; }
 	LANDSCAPE_API virtual bool IsResultMaterialAttributes(int32 OutputIndex) override;
 	LANDSCAPE_API virtual int32 Compile(class FMaterialCompiler* Compiler, int32 OutputIndex) override;
 	LANDSCAPE_API virtual void GetCaption(TArray<FString>& OutCaptions) const override;
diff --git a/Engine/Source/Runtime/Landscape/Classes/Materials/MaterialExpressionLandscapeLayerSwitch.h b/Engine/Source/Runtime/Landscape/Classes/Materials/MaterialExpressionLandscapeLayerSwitch.h
--- a/Engine/Source/Runtime/Landscape/Classes/Materials/MaterialExpressionLandscapeLayerSwitch.h
+++ b/Engine/Source/Runtime/Landscape/Classes/Materials/MaterialExpressionLandscapeLayerSwitch.h
@@ -51,8 +51,8 @@ public:
 	LANDSCAPE_API virtual ReferencedTextureArray GetReferencedTextures() const override;
 	virtual bool CanReferenceTexture() const override { return true; }
 #if WITH_EDITOR
-	virtual uint32 GetInputType(int32 InputIndex) override {return MCT_Unknown;}
-	virtual uint32 GetOutputType(int32 InputIndex) override {return MCT_Unknown;}
+	virtual EMaterialValueType GetInputValueType(int32 InputIndex) override {return MCT_Unknown;}
+	virtual EMaterialValueType GetOutputValueType(int32 InputIndex) override {return MCT_Unknown;}
 
 	/**
 	 * Gets the landscape layer names
diff --git a/Engine/Source/Runtime/Landscape/Classes/Materials/MaterialExpressionLandscapeLayerWeight.h b/Engine/Source/Runtime/Landscape/Classes/Materials/MaterialExpressionLandscapeLayerWeight.h
--- a/Engine/Source/Runtime/Landscape/Classes/Materials/MaterialExpressionLandscapeLayerWeight.h
+++ b/Engine/Source/Runtime/Landscape/Classes/Materials/MaterialExpressionLandscapeLayerWeight.h
@@ -51,7 +51,7 @@ public:
 	LANDSCAPE_API virtual ReferencedTextureArray GetReferencedTextures() const override;
 	virtual bool CanReferenceTexture() const override { return true; }
 #if WITH_EDITOR
-	virtual uint32 GetInputType(int32 InputIndex) override {return MCT_Float | MCT_MaterialAttributes;}
+	virtual EMaterialValueType GetInputValueType(int32 InputIndex) override {return static_cast<EMaterialValueType>(MCT_Float | MCT_MaterialAttributes);}
 
 	/**
 	 * Gets the landscape layer names
diff --git a/Engine/Source/Runtime/Landscape/Classes/Materials/MaterialExpressionLandscapePhysicalMaterialOutput.h b/Engine/Source/Runtime/Landscape/Classes/Materials/MaterialExpressionLandscapePhysicalMaterialOutput.h
--- a/Engine/Source/Runtime/Landscape/Classes/Materials/MaterialExpressionLandscapePhysicalMaterialOutput.h
+++ b/Engine/Source/Runtime/Landscape/Classes/Materials/MaterialExpressionLandscapePhysicalMaterialOutput.h
@@ -51,7 +51,7 @@ class UMaterialExpressionLandscapePhysicalMaterialOutput : public UMaterialExpre
 	LANDSCAPE_API virtual TArrayView<FExpressionInput*> GetInputsView() override;
 	LANDSCAPE_API virtual FExpressionInput* GetInput(int32 InputIndex) override;
 	LANDSCAPE_API virtual FName GetInputName(int32 InputIndex) const override;
-	virtual uint32 GetInputType(int32 InputIndex) override { return MCT_Float; }
+	virtual EMaterialValueType GetInputValueType(int32 InputIndex) override { return MCT_Float; }
 	LANDSCAPE_API virtual int32 Compile(class FMaterialCompiler* Compiler, int32 OutputIndex) override;
 
 	//~ End UMaterialExpression Interface
diff --git a/Engine/Source/Runtime/Landscape/Private/Materials/MaterialExpressionLandscapeLayerBlend.cpp b/Engine/Source/Runtime/Landscape/Private/Materials/MaterialExpressionLandscapeLayerBlend.cpp
--- a/Engine/Source/Runtime/Landscape/Private/Materials/MaterialExpressionLandscapeLayerBlend.cpp
+++ b/Engine/Source/Runtime/Landscape/Private/Materials/MaterialExpressionLandscapeLayerBlend.cpp
@@ -116,14 +116,14 @@ FName UMaterialExpressionLandscapeLayerBlend::GetInputName(int32 InputIndex) con
 	return NAME_None;
 }
 
-uint32 UMaterialExpressionLandscapeLayerBlend::GetInputType(int32 InputIndex)
+EMaterialValueType UMaterialExpressionLandscapeLayerBlend::GetInputValueType(int32 InputIndex)
 {
 	int32 Idx = 0;
 	for (int32 LayerIdx = 0; LayerIdx<Layers.Num(); LayerIdx++)
 	{
 		if (InputIndex == Idx++)
 		{
-			return MCT_Float | MCT_MaterialAttributes; // can accept pretty much anything including MaterialAttributes
+			return static_cast<EMaterialValueType>(MCT_Float | MCT_MaterialAttributes); // can accept pretty much anything including MaterialAttributes
 		}
 		if (Layers[LayerIdx].BlendType == LB_HeightBlend)
 		{
