Letting FActorSpawnParameters::bNoFail upgrade the actor class's SpawnCollisionHandlingMethod when FActorSpawnParameters::SpawnCollisionHandlingOverride is Undefined. Makes bNoFail consistently ignore spawn collision, whereas before it would ignore spawn collision for some settings of SpawnCollisionHandlingOverride + SpawnCollisionHandlingMethod but not others.

diff --git a/Engine/Source/Runtime/Engine/Private/LevelActor.cpp b/Engine/Source/Runtime/Engine/Private/LevelActor.cpp
--- a/Engine/Source/Runtime/Engine/Private/LevelActor.cpp
+++ b/Engine/Source/Runtime/Engine/Private/LevelActor.cpp
@@ -619,25 +619,28 @@ AActor* UWorld::SpawnActor( UClass* Class, FTransform const* UserTransformPtr, c
 
 	FTransform const UserTransform = UserTransformPtr ? *UserTransformPtr : FTransform::Identity;
 
-	ESpawnActorCollisionHandlingMethod CollisionHandlingOverride = SpawnParameters.SpawnCollisionHandlingOverride;
+	// Choose the collision handling method. In order of increasing priority: actor template setting (often class default), collision handling override, bNoFail
+	ESpawnActorCollisionHandlingMethod CollisionHandlingMethod = Template->SpawnCollisionHandlingMethod;
 
-	// "no fail" take preedence over collision handling settings that include fails
+	// Adopt collision handling override if any was set
+	if (SpawnParameters.SpawnCollisionHandlingOverride != ESpawnActorCollisionHandlingMethod::Undefined)
+	{
+		CollisionHandlingMethod = SpawnParameters.SpawnCollisionHandlingOverride;
+	}
+
+	// If bNoFail is true, upgrade collision handling method to an always spawning one
 	if (SpawnParameters.bNoFail)
 	{
-		// maybe upgrade to disallow fail
-		if (CollisionHandlingOverride == ESpawnActorCollisionHandlingMethod::AdjustIfPossibleButDontSpawnIfColliding)
+		if (CollisionHandlingMethod == ESpawnActorCollisionHandlingMethod::AdjustIfPossibleButDontSpawnIfColliding)
 		{
-			CollisionHandlingOverride = ESpawnActorCollisionHandlingMethod::AdjustIfPossibleButAlwaysSpawn;
+			CollisionHandlingMethod = ESpawnActorCollisionHandlingMethod::AdjustIfPossibleButAlwaysSpawn;
 		}
-		else if (CollisionHandlingOverride == ESpawnActorCollisionHandlingMethod::DontSpawnIfColliding)
+		else if (CollisionHandlingMethod == ESpawnActorCollisionHandlingMethod::DontSpawnIfColliding)
 		{
-			CollisionHandlingOverride = ESpawnActorCollisionHandlingMethod::AlwaysSpawn;
+			CollisionHandlingMethod = ESpawnActorCollisionHandlingMethod::AlwaysSpawn;
 		}
 	}
 
-	// use override if set, else fall back to actor's preference
-	ESpawnActorCollisionHandlingMethod const CollisionHandlingMethod = (CollisionHandlingOverride == ESpawnActorCollisionHandlingMethod::Undefined) ? Template->SpawnCollisionHandlingMethod : CollisionHandlingOverride;
-
 	// see if we can avoid spawning altogether by checking native components
 	// note: we can't handle all cases here, since we don't know the full component hierarchy until after the actor is spawned
 	if (CollisionHandlingMethod == ESpawnActorCollisionHandlingMethod::DontSpawnIfColliding)
