[NetPhysics] Implement a custom Uninitialize_Internal callback for the NetworkPhysicsComponent to remove deletages and unregister the Input and State data linked to PhysicsObject from RewindData.
- Also implement callback OnPhysicsObjectUnregistered_Internal if the physics object gets deleted before the ISImCallbackObject.

diff --git a/Engine/Source/Runtime/Engine/Private/PhysicsEngine/NetworkPhysicsComponent.cpp b/Engine/Source/Runtime/Engine/Private/PhysicsEngine/NetworkPhysicsComponent.cpp
--- a/Engine/Source/Runtime/Engine/Private/PhysicsEngine/NetworkPhysicsComponent.cpp
+++ b/Engine/Source/Runtime/Engine/Private/PhysicsEngine/NetworkPhysicsComponent.cpp
@@ -565,6 +565,19 @@ void UNetworkPhysicsComponent::UninitializeComponent()
 		{
 			if (Chaos::FPhysicsSolver* Solver = PhysScene->GetSolver())
 			{
+				/* Run OnUninitialize_Internal on the ISimCallbackObject as a way to unregister input / state history, unsubscribe from delegates etc.
+				* After UnregisterAndFreeSimCallbackObject_External the ISimCallbackObject will not get any callbacks anymore, use this as the last safe place to use the cached FPhysicsObject for example */
+				FAsyncNetworkPhysicsComponent* AsyncNetworkPhysicsComponent = NetworkPhysicsComponent_Internal;
+				Solver->EnqueueCommandImmediate(
+					[AsyncNetworkPhysicsComponent]()
+					{
+						if (AsyncNetworkPhysicsComponent)
+						{
+							AsyncNetworkPhysicsComponent->OnUninitialize_Internal();
+						}
+					}
+				);
+
 				// Clear async component from Physics Thread and memory
 				Solver->UnregisterAndFreeSimCallbackObject_External(NetworkPhysicsComponent_Internal);
 			}
@@ -1485,7 +1498,27 @@ FAsyncNetworkPhysicsComponent::FAsyncNetworkPhysicsComponent() : TSimCallbackObj
 {
 }
 
-FAsyncNetworkPhysicsComponent::~FAsyncNetworkPhysicsComponent()
+void FAsyncNetworkPhysicsComponent::OnPostInitialize_Internal()
+{
+	if (Chaos::FPhysicsSolverBase* BaseSolver = GetSolver())
+	{
+		if (BaseSolver->IsNetworkPhysicsPredictionEnabled())
+		{
+			// Register for Pre- and Post- ProcessInputs_Internal callbacks
+			if (FNetworkPhysicsCallback* SolverCallback = static_cast<FNetworkPhysicsCallback*>(BaseSolver->GetRewindCallback()))
+			{
+				DelegateOnPreProcessInputs_Internal = SolverCallback->PreProcessInputsInternal.AddRaw(this, &FAsyncNetworkPhysicsComponent::OnPreProcessInputs_Internal);
+				DelegateOnPostProcessInputs_Internal = SolverCallback->PostProcessInputsInternal.AddRaw(this, &FAsyncNetworkPhysicsComponent::OnPostProcessInputs_Internal);
+			}
+		}
+		else
+		{
+			UE_LOG(LogChaos, Warning, TEXT("A NetworkPhysicsComponent is trying to set up but 'Project Settings -> Physics -> Physics Prediction' is not enabled. The component might not work as intended."));
+		}
+	}
+}
+
+void FAsyncNetworkPhysicsComponent::OnUninitialize_Internal()
 {
 	if (Chaos::FPhysicsSolverBase* BaseSolver = GetSolver())
 	{
@@ -1503,23 +1536,12 @@ FAsyncNetworkPhysicsComponent::~FAsyncNetworkPhysicsComponent()
 	UnregisterDataHistoryFromRewindData();
 }
 
-void FAsyncNetworkPhysicsComponent::OnPostInitialize_Internal()
+void FAsyncNetworkPhysicsComponent::OnPhysicsObjectUnregistered_Internal(Chaos::FConstPhysicsObjectHandle InPhysicsObject)
 {
-	if (Chaos::FPhysicsSolverBase* BaseSolver = GetSolver())
+	if (PhysicsObject == InPhysicsObject)
 	{
-		if (BaseSolver->IsNetworkPhysicsPredictionEnabled())
-		{
-			// Register for Pre- and Post- ProcessInputs_Internal callbacks
-			if (FNetworkPhysicsCallback* SolverCallback = static_cast<FNetworkPhysicsCallback*>(BaseSolver->GetRewindCallback()))
-			{
-				DelegateOnPreProcessInputs_Internal = SolverCallback->PreProcessInputsInternal.AddRaw(this, &FAsyncNetworkPhysicsComponent::OnPreProcessInputs_Internal);
-				DelegateOnPostProcessInputs_Internal = SolverCallback->PostProcessInputsInternal.AddRaw(this, &FAsyncNetworkPhysicsComponent::OnPostProcessInputs_Internal);
-			}
-		}
-		else
-		{
-			UE_LOG(LogChaos, Warning, TEXT("A NetworkPhysicsComponent is trying to set up but 'Project Settings -> Physics -> Physics Prediction' is not enabled. The component might not work as intended."));
-		}
+		UnregisterDataHistoryFromRewindData();
+		PhysicsObject = nullptr;
 	}
 }
 
diff --git a/Engine/Source/Runtime/Engine/Public/Physics/NetworkPhysicsComponent.h b/Engine/Source/Runtime/Engine/Public/Physics/NetworkPhysicsComponent.h
--- a/Engine/Source/Runtime/Engine/Public/Physics/NetworkPhysicsComponent.h
+++ b/Engine/Source/Runtime/Engine/Public/Physics/NetworkPhysicsComponent.h
@@ -1413,12 +1413,12 @@ struct FAsyncNetworkPhysicsComponentOutput : public Chaos::FSimCallbackOutput
 class FAsyncNetworkPhysicsComponent : public Chaos::TSimCallbackObject<
 	FAsyncNetworkPhysicsComponentInput,
 	FAsyncNetworkPhysicsComponentOutput,
-	Chaos::ESimCallbackOptions::None>
+	Chaos::ESimCallbackOptions::PhysicsObjectUnregister>
 {
 
 public:
 	FAsyncNetworkPhysicsComponent();
-	~FAsyncNetworkPhysicsComponent();
+	~FAsyncNetworkPhysicsComponent() {};
 
 	// Get reference to async output for current internal frame and initialize it if not already done
 	FAsyncNetworkPhysicsComponentOutput& GetAsyncOutput_Internal();
@@ -1454,6 +1454,11 @@ private:
 	// Initialize, bind delegates etc.
 	virtual void OnPostInitialize_Internal() override;
 
+	// Uninitialize, unbind delegates etc.
+	void OnUninitialize_Internal();
+
+	virtual void OnPhysicsObjectUnregistered_Internal(Chaos::FConstPhysicsObjectHandle PhysicsObject) override;
+
 	// Delegate ran at the start of FNetworkPhysicsCallback::ProcessInputs_Internal, used to receive and apply inputs and states if needed.
 	void OnPreProcessInputs_Internal(const int32 PhysicsStep);
 
