[ComputeFramework] Commandlet work support
- ComputeFramework will compile Synchronously when running inside a rendering commandlet
- CommandletHelpers::TickEngine will Submit compute work if any
- Fix check in ComputeKernelShaderCompilationManager

#rnx

diff --git a/Engine/Plugins/Runtime/ComputeFramework/Source/ComputeFramework/Private/ComputeFramework.cpp b/Engine/Plugins/Runtime/ComputeFramework/Source/ComputeFramework/Private/ComputeFramework.cpp
--- a/Engine/Plugins/Runtime/ComputeFramework/Source/ComputeFramework/Private/ComputeFramework.cpp
+++ b/Engine/Plugins/Runtime/ComputeFramework/Source/ComputeFramework/Private/ComputeFramework.cpp
@@ -53,7 +53,7 @@ namespace ComputeFramework
 	bool IsDeferredCompilation()
 	{
 #if WITH_EDITOR
-		return GComputeFrameworkDeferredCompilation != 0;
+		return GComputeFrameworkDeferredCompilation != 0 && !(IsRunningCommandlet() && IsAllowCommandletRendering());
 #else
 		return false;
 #endif
diff --git a/Engine/Plugins/Runtime/ComputeFramework/Source/ComputeFramework/Private/ComputeGraph.cpp b/Engine/Plugins/Runtime/ComputeFramework/Source/ComputeFramework/Private/ComputeGraph.cpp
--- a/Engine/Plugins/Runtime/ComputeFramework/Source/ComputeFramework/Private/ComputeGraph.cpp
+++ b/Engine/Plugins/Runtime/ComputeFramework/Source/ComputeFramework/Private/ComputeGraph.cpp
@@ -347,7 +347,12 @@ enum class EComputeKernelCompilationFlags
 void UComputeGraph::UpdateResources()
 {
 #if WITH_EDITOR
-	CacheResourceShadersForRendering(uint32(EComputeKernelCompilationFlags::ApplyCompletedShaderMapForRendering));
+	uint32 CompilationFlags = (uint32)EComputeKernelCompilationFlags::ApplyCompletedShaderMapForRendering;
+
+	bool bIsRenderCommandlet = IsRunningCommandlet() && IsAllowCommandletRendering();
+	CompilationFlags |= bIsRenderCommandlet ? (uint32)EComputeKernelCompilationFlags::Synchronous : (uint32)EComputeKernelCompilationFlags::None;
+
+	CacheResourceShadersForRendering(CompilationFlags);
 #endif
 
 	ReleaseRenderProxy(RenderProxy);
diff --git a/Engine/Plugins/Runtime/ComputeFramework/Source/ComputeFramework/Private/ComputeKernelShaderCompilationManager.cpp b/Engine/Plugins/Runtime/ComputeFramework/Source/ComputeFramework/Private/ComputeKernelShaderCompilationManager.cpp
--- a/Engine/Plugins/Runtime/ComputeFramework/Source/ComputeFramework/Private/ComputeKernelShaderCompilationManager.cpp
+++ b/Engine/Plugins/Runtime/ComputeFramework/Source/ComputeFramework/Private/ComputeKernelShaderCompilationManager.cpp
@@ -10,6 +10,7 @@
 #include "ShaderCompiler.h"
 
 #if WITH_EDITOR
+#include "Algo/AnyOf.h"
 #include "Interfaces/IShaderFormat.h"
 #include "Interfaces/ITargetPlatformManagerModule.h"
 #endif
@@ -411,7 +412,7 @@ void FComputeKernelShaderCompilationManager::FinishCompilation(const TCHAR* InKe
 	GShaderCompilingManager->FinishCompilation(NULL, ShaderMapIdsToFinishCompiling);
 	ProcessAsyncResults();	// grab compiled shader maps and assign them to their resources
 
-	check(ComputeKernelShaderMapJobs.Num() == 0);
+	check(!Algo::AnyOf(ShaderMapIdsToFinishCompiling, [this](int32 ShaderMapId) { return ComputeKernelShaderMapJobs.Contains(ShaderMapId); }));
 }
 
 void FComputeKernelShaderCompilationManager::CancelCompilation(const TCHAR* InKernelName, const TArray<int32>& ShaderMapIdsToCancelCompiling)
diff --git a/Engine/Source/Runtime/Engine/Private/Commandlets/Commandlet.cpp b/Engine/Source/Runtime/Engine/Private/Commandlets/Commandlet.cpp
--- a/Engine/Source/Runtime/Engine/Private/Commandlets/Commandlet.cpp
+++ b/Engine/Source/Runtime/Engine/Private/Commandlets/Commandlet.cpp
@@ -4,10 +4,14 @@
 #include "Engine/World.h"
 #include "Misc/AutomationTest.h"
 
+#include "ComputeWorkerInterface.h"
 #include "Engine/Engine.h"
 #include "EngineModule.h"
 #include "Misc/App.h"
 #include "RenderingThread.h"
+#include "RenderGraphBuilder.h"
+#include "RenderGraphEvent.h"
+#include "SceneInterface.h"
 #include "Misc/CoreDelegates.h"
 #include "RendererInterface.h"
 #include "HAL/ThreadManager.h"
@@ -118,6 +122,26 @@ void CommandletHelpers::TickEngine(UWorld* InWorld, double InDeltaTime)
 				FCoreDelegates::OnBeginFrameRT.Broadcast();
 			});
 
+			// Submit Compute Framework work (ComputeTaskExecutionGroup::EndOfFrameUpdate)
+			const ERHIFeatureLevel::Type FeatureLevel = Scene->GetFeatureLevel();
+			TArray<IComputeTaskWorker*> ComputeWorkers;
+			Scene->GetComputeTaskWorkers(ComputeWorkers);
+
+			ENQUEUE_RENDER_COMMAND(ComputeTaskWorker_EndOfFrameUpdate)([FeatureLevel, ComputeWorkers](FRHICommandListImmediate& RHICmdList)
+			{				
+				FRDGBuilder GraphBuilder(RHICmdList, FRDGEventName(TEXT("ComputeTaskWorker_EndOfFrameUpdate")));
+				
+				for (IComputeTaskWorker* ComputeTaskWorker : ComputeWorkers)
+				{
+					if (ComputeTaskWorker->HasWork(ComputeTaskExecutionGroup::EndOfFrameUpdate))
+					{
+						ComputeTaskWorker->SubmitWork(GraphBuilder, ComputeTaskExecutionGroup::EndOfFrameUpdate, FeatureLevel);
+					}
+				}
+
+				GraphBuilder.Execute();
+			});
+
 			ENQUEUE_RENDER_COMMAND(EndFrame)([](FRHICommandListImmediate& RHICmdList)
 			{
 				FCoreDelegates::OnEndFrameRT.Broadcast();
