Fixed an incorrect test that skipped updating the ID mapping for any removes that affect instances at the end of the range in the new instance data manager.

#rnx

diff --git a/Engine/Source/Runtime/Engine/Private/InstanceData/InstanceDataUpdateUtils.h b/Engine/Source/Runtime/Engine/Private/InstanceData/InstanceDataUpdateUtils.h
--- a/Engine/Source/Runtime/Engine/Private/InstanceData/InstanceDataUpdateUtils.h
+++ b/Engine/Source/Runtime/Engine/Private/InstanceData/InstanceDataUpdateUtils.h
@@ -287,7 +287,7 @@ inline void UpdateIdMapping(FInstanceUpdateChangeSet& ChangeSet, const IndexRema
 		{
 			OutInstanceIdIndexMap.RebuildFromIndexToIdMap(MoveTemp(ChangeSet.IndexToIdMapDeltaData), ChangeSet.MaxInstanceId);
 		}
-		else if (!IndexDelta.IsEmpty())
+		else
 		{
 			// General path that handles incremental removes and other updates.
 			OutInstanceIdIndexMap.ResizeExplicit(ChangeSet.NumSourceInstances, ChangeSet.MaxInstanceId);
