Fix check hit for GPU-instanced primitives

#rnx

diff --git a/Engine/Source/Runtime/Engine/Private/InstancedStaticMesh.cpp b/Engine/Source/Runtime/Engine/Private/InstancedStaticMesh.cpp
--- a/Engine/Source/Runtime/Engine/Private/InstancedStaticMesh.cpp
+++ b/Engine/Source/Runtime/Engine/Private/InstancedStaticMesh.cpp
@@ -1819,7 +1819,15 @@ void FInstancedStaticMeshSceneProxy::GetDynamicRayTracingInstances(FRayTracingIn
 	FVector ViewPosition = Collector.GetReferenceView()->ViewLocation;
 
 	const FInstanceSceneDataBuffers* InstanceSceneDataBuffers = GetInstanceSceneDataBuffers();
-	check(InstanceSceneDataBuffers && InstanceSceneDataBuffers->GetNumInstances() == InstanceCount && !InstanceSceneDataBuffers->IsInstanceDataGPUOnly());
+
+	if (InstanceSceneDataBuffers->IsInstanceDataGPUOnly())
+	{
+		// TODO: Update FInstancedStaticMeshSceneProxy to make it compatible with GPU Scene when not using WPO materials, and support fallback to static geometry
+		// when IsInstanceDataGPUOnly() instead of simply being hidden from ray tracing.
+		return;
+	}
+
+	check(InstanceSceneDataBuffers && InstanceSceneDataBuffers->GetNumInstances() == InstanceCount);
 
 	auto GetDistanceToInstance = [&ViewPosition, InstanceSceneDataBuffers](int32 InstanceIndex, float& OutInstanceRadius, float& OutDistanceToInstanceCenter, float& OutDistanceToInstanceStart)
 	{
