[HWRT] Show on-screen warnings when ray tracing geometry pool is overbudget.

diff --git a/Engine/Source/Runtime/Engine/Private/Rendering/RayTracingGeometryManager.cpp b/Engine/Source/Runtime/Engine/Private/Rendering/RayTracingGeometryManager.cpp
--- a/Engine/Source/Runtime/Engine/Private/Rendering/RayTracingGeometryManager.cpp
+++ b/Engine/Source/Runtime/Engine/Private/Rendering/RayTracingGeometryManager.cpp
@@ -19,6 +19,8 @@
 #include "ProfilingDebugging/CsvProfiler.h"
 #include "ProfilingDebugging/IoStoreTrace.h"
 
+#include "Engine/Engine.h"
+
 #if RHI_RAYTRACING
 
 static bool bHasRayTracingEnableChanged = false;
@@ -110,6 +112,14 @@ static FAutoConsoleVariableRef CVarRayTracingPendingBuildPriorityBoostPerFrame(
 	ECVF_RenderThreadSafe
 );
 
+static bool GRayTracingShowOnScreenWarnings = true;
+static FAutoConsoleVariableRef CVarRayTracingShowOnScreenWarnings(
+	TEXT("r.RayTracing.ShowOnScreenWarnings"),
+	GRayTracingShowOnScreenWarnings,
+	TEXT("Whether to show on-screen warnings related to ray tracing."),
+	ECVF_RenderThreadSafe
+);
+
 DECLARE_STATS_GROUP(TEXT("Ray Tracing Geometry"), STATGROUP_RayTracingGeometry, STATCAT_Advanced);
 
 DECLARE_DWORD_ACCUMULATOR_STAT(TEXT("Geometry Count"), STAT_RayTracingGeometryCount, STATGROUP_RayTracingGeometry);
@@ -704,9 +714,6 @@ void FRayTracingGeometryManager::Tick(FRHICommandList& RHICmdList)
 		// ensure(ReferencedGeometries.Num() == TSet(ReferencedGeometries).Num());
 #endif
 
-		SET_MEMORY_STAT(STAT_RayTracingGeometryRequestedMemory, RequestedSize);
-		CSV_CUSTOM_STAT(RayTracingGeometry, RequestedSizeMB, RequestedSize / 1024.0f / 1024.0f, ECsvCustomStatOp::Set);
-
 		const uint64 ResidentGeometryMemoryPoolSize = FUnitConversion::Convert(GRayTracingResidentGeometryMemoryPoolSizeInMB, EUnit::Megabytes, EUnit::Bytes);
 
 		// Step 3
@@ -787,6 +794,29 @@ void FRayTracingGeometryManager::Tick(FRHICommandList& RHICmdList)
 				++Index;
 			}
 		}
+
+		SET_MEMORY_STAT(STAT_RayTracingGeometryRequestedMemory, RequestedSize);
+		CSV_CUSTOM_STAT(RayTracingGeometry, RequestedSizeMB, RequestedSize / 1024.0f / 1024.0f, ECsvCustomStatOp::Set);
+
+#if !UE_BUILD_SHIPPING
+		if (GRayTracingShowOnScreenWarnings)
+		{
+			if (TotalAlwaysResidentSize > ResidentGeometryMemoryPoolSize)
+			{
+				GEngine->AddOnScreenDebugMessage(uint64(this), 1.f, FColor::Red,
+					*FString::Printf(TEXT("RAY TRACING GEOMETRY - ALWAYS RESIDENT MEMORY OVER BUDGET %s / %s"),
+						*FText::AsMemory(TotalAlwaysResidentSize).ToString(),
+						*FText::AsMemory(ResidentGeometryMemoryPoolSize).ToString()));
+			}
+			else if (RequestedSize > ResidentGeometryMemoryPoolSize)
+			{
+				GEngine->AddOnScreenDebugMessage(uint64(this), 1.f, FColor::Yellow,
+					*FString::Printf(TEXT("RAY TRACING GEOMETRY - REQUESTED MEMORY OVER BUDGET %s / %s"),
+						*FText::AsMemory(RequestedSize).ToString(),
+						*FText::AsMemory(ResidentGeometryMemoryPoolSize).ToString()));
+			}
+		}
+#endif
 	}
 	else
 	{
