[Materials] Append hash over all external code expressions to material in HLSLMaterialTranslator.

This ensures material shadermaps are invalidated automatically when any of the external code declarations content changes.

#rnx

diff --git a/Engine/Source/Runtime/Core/Private/UObject/FortniteMainBranchObjectVersion.cpp b/Engine/Source/Runtime/Core/Private/UObject/FortniteMainBranchObjectVersion.cpp
--- a/Engine/Source/Runtime/Core/Private/UObject/FortniteMainBranchObjectVersion.cpp
+++ b/Engine/Source/Runtime/Core/Private/UObject/FortniteMainBranchObjectVersion.cpp
@@ -8,7 +8,7 @@ TMap<FGuid, FGuid> FFortniteMainBranchObjectVersion::GetSystemGuids()
 
 	SystemGuids.Add(DevGuids.GLOBALSHADERMAP_DERIVEDDATA_VER, FGuid("125CF21C881744418BD4A614B3B125FF"));
 	SystemGuids.Add(DevGuids.LANDSCAPE_MOBILE_COOK_VERSION, FGuid("32D02EF867C74B71A0D4E0FA41392732"));
-	SystemGuids.Add(DevGuids.MATERIALSHADERMAP_DERIVEDDATA_VER, FGuid("68CF693DB35E4ABDA0F053D785102A99"));
+	SystemGuids.Add(DevGuids.MATERIALSHADERMAP_DERIVEDDATA_VER, FGuid("04B8C817071A453C9EBEFBFB51D5F785"));
 	SystemGuids.Add(DevGuids.NANITE_DERIVEDDATA_VER, FGuid("42C46B24E1586986880AFE62B0F575AB"));
 	SystemGuids.Add(DevGuids.NIAGARASHADERMAP_DERIVEDDATA_VER, FGuid("45D5978F33D940EEA6602B78559900DC"));
 	SystemGuids.Add(DevGuids.Niagara_LatestScriptCompileVersion, FGuid("670A2EBDEE1A4794BB77BBFF331EA7EF"));
diff --git a/Engine/Source/Runtime/Engine/Private/Materials/Material.cpp b/Engine/Source/Runtime/Engine/Private/Materials/Material.cpp
--- a/Engine/Source/Runtime/Engine/Private/Materials/Material.cpp
+++ b/Engine/Source/Runtime/Engine/Private/Materials/Material.cpp
@@ -356,6 +356,8 @@ void FMaterialResource::GetShaderMapId(EShaderPlatform Platform, const ITargetPl
 	CachedData.AppendReferencedFunctionIdsTo(OutId.ReferencedFunctions);
 	CachedData.AppendReferencedParameterCollectionIdsTo(OutId.ReferencedParameterCollections);
 
+	CachedData.GetExternalCodeReferencesHash(OutId.ExternalCodeReferencesHash);
+
 	Material->GetForceRecompileTextureIdsHash(OutId.TextureReferencesHash);
 
 	if(MaterialInstance)
diff --git a/Engine/Source/Runtime/Engine/Private/Materials/MaterialCachedData.cpp b/Engine/Source/Runtime/Engine/Private/Materials/MaterialCachedData.cpp
--- a/Engine/Source/Runtime/Engine/Private/Materials/MaterialCachedData.cpp
+++ b/Engine/Source/Runtime/Engine/Private/Materials/MaterialCachedData.cpp
@@ -32,6 +32,7 @@
 #include "Materials/MaterialExpressionMeshPaintTextureObject.h"
 #include "Materials/MaterialExpressionWorldPosition.h"
 #include "Materials/MaterialExpressionActorPositionWS.h"
+#include "Materials/MaterialExternalCodeRegistry.h"
 #include "Materials/MaterialFunctionInterface.h"
 #include "Materials/MaterialParameterCollection.h"
 #include "VT/RuntimeVirtualTexture.h"
@@ -111,6 +112,28 @@ void FMaterialCachedExpressionData::AppendReferencedParameterCollectionIdsTo(TAr
 	}
 }
 
+void FMaterialCachedExpressionData::GetExternalCodeReferencesHash(FSHAHash& OutHash) const
+{
+	FSHA1 Hasher;
+	for (const TObjectPtr<UClass>& ExternalCodeExpressionClass : ReferencedExternalCodeExpressionClasses)
+	{
+		if (const UObject* DefaultExternalCodeExpression = ExternalCodeExpressionClass->GetDefaultObject())
+		{
+			const UMaterialExpressionExternalCodeBase* ExternalCodeExpressionBase = Cast<UMaterialExpressionExternalCodeBase>(DefaultExternalCodeExpression);
+			checkf(ExternalCodeExpressionBase != nullptr, TEXT("Wrong UObject class type cached in FMaterialCachedExpressionData::ReferencedExternalCodeExpressionClasses"));
+
+			for (const FName& ExternalCodeIdentifier : ExternalCodeExpressionBase->ExternalCodeIdentifiers)
+			{
+				if (const FMaterialExternalCodeDeclaration* ExternalCodeDeclaration = MaterialExternalCodeRegistry::Get().FindExternalCode(ExternalCodeIdentifier))
+				{
+					ExternalCodeDeclaration->UpdateHash(Hasher);
+				}
+			}
+		}
+	}
+	OutHash = Hasher.Finalize();
+}
+
 #if WITH_EDITOR
 static bool TryAddParameter(FMaterialCachedExpressionData& CachedData,
 	EMaterialParameterType Type,
@@ -508,6 +531,11 @@ void FMaterialCachedExpressionData::UpdateForExpressions(const FMaterialCachedEx
 			}
 		}
 
+		if (UMaterialExpressionExternalCodeBase* ExternalCodeExpression = Cast<UMaterialExpressionExternalCodeBase>(Expression))
+		{
+			ReferencedExternalCodeExpressionClasses.AddUnique(Expression->GetClass());
+		}
+
 		if (UMaterialExpressionCollectionParameter* ExpressionCollectionParameter = Cast<UMaterialExpressionCollectionParameter>(Expression))
 		{
 			UMaterialParameterCollection* Collection = ExpressionCollectionParameter->Collection;
diff --git a/Engine/Source/Runtime/Engine/Private/Materials/MaterialExternalCodeRegistry.cpp b/Engine/Source/Runtime/Engine/Private/Materials/MaterialExternalCodeRegistry.cpp
--- a/Engine/Source/Runtime/Engine/Private/Materials/MaterialExternalCodeRegistry.cpp
+++ b/Engine/Source/Runtime/Engine/Private/Materials/MaterialExternalCodeRegistry.cpp
@@ -3,6 +3,53 @@
 #include "MaterialExternalCodeRegistry.h"
 #include "Misc/Parse.h"
 
+
+/*
+ * FMaterialExternalCodeEnvironmentDefine
+ */
+
+void FMaterialExternalCodeEnvironmentDefine::UpdateHash(FSHA1& Hasher) const
+{
+	FString NameString = Name.ToString();
+	Hasher.UpdateWithString(*NameString, NameString.Len());
+	Hasher.Update(ShaderFrequency);
+}
+
+
+/*
+ * FMaterialExternalCodeDeclaration 
+ */
+
+void FMaterialExternalCodeDeclaration::UpdateHash(FSHA1& Hasher) const
+{
+	const uint8 bIsInlinedByte = bIsInlined ? 1 : 0;
+	Hasher.Update(&bIsInlinedByte, sizeof(bIsInlinedByte));
+
+	Hasher.Update(ReturnType);
+
+	FString NameString = Name.ToString();
+	Hasher.UpdateWithString(*NameString, NameString.Len());
+
+	Hasher.UpdateWithString(*Definition, Definition.Len());
+	Hasher.Update(Derivative);
+	Hasher.Update(ShaderFrequency);
+
+	for (const TEnumAsByte<EMaterialDomain>& InDomain : Domains)
+	{
+		Hasher.Update(InDomain.GetValue());
+	}
+
+	for (const FMaterialExternalCodeEnvironmentDefine& InEnvironmentDefine : EnvironmentDefines)
+	{
+		InEnvironmentDefine.UpdateHash(Hasher);
+	}
+}
+
+
+/*
+ * MaterialExternalCodeRegistry
+ */
+
 MaterialExternalCodeRegistry& MaterialExternalCodeRegistry::Get()
 {
 	static MaterialExternalCodeRegistry Instance;
@@ -24,8 +71,7 @@ void MaterialExternalCodeRegistry::BuildMapToExternalDeclarations()
 		{
 			if (ExternalCodeDeclarationMap.Contains(InExternalCode.Name))
 			{
-				UE_LOG(LogMaterial, Error, TEXT("External HLSL code declaration in material IR must not be overloaded, but '%s' is defined more than once"), *InExternalCode.Name.ToString());
-				continue;
+				UE_LOG(LogMaterial, Fatal, TEXT("External HLSL code declarations for materials must not be overloaded, but '%s' is defined more than once"), *InExternalCode.Name.ToString());
 			}
 			ExternalCodeDeclarationMap.Add(InExternalCode.Name) = &InExternalCode;
 		}
diff --git a/Engine/Source/Runtime/Engine/Private/Materials/MaterialExternalCodeRegistry.h b/Engine/Source/Runtime/Engine/Private/Materials/MaterialExternalCodeRegistry.h
--- a/Engine/Source/Runtime/Engine/Private/Materials/MaterialExternalCodeRegistry.h
+++ b/Engine/Source/Runtime/Engine/Private/Materials/MaterialExternalCodeRegistry.h
@@ -86,6 +86,9 @@ struct FMaterialExternalCodeEnvironmentDefine
 	/** Optional shader frequency to further restrict this environment define besides its code declaration shader frequency. */
 	UPROPERTY(Config)
 	EMaterialShaderFrequency ShaderFrequency = EMaterialShaderFrequency::Any;
+
+	/** Updates the input hasher state with the content of this environemnt define. */
+	void UpdateHash(FSHA1& Hasher) const;
 };
 
 /** Declaration of external HLSL code. Such code expressions can be emitted as part of a material translation. */
@@ -130,6 +133,9 @@ public:
 	{
 		return (EMaterialValueType)ReturnType;
 	}
+
+	/** Updates the input hasher state with the content of this external code declaration. */
+	void UpdateHash(FSHA1& Hasher) const;
 };
 
 UCLASS(Config=MaterialExpressions)
diff --git a/Engine/Source/Runtime/Engine/Private/Materials/MaterialShader.cpp b/Engine/Source/Runtime/Engine/Private/Materials/MaterialShader.cpp
--- a/Engine/Source/Runtime/Engine/Private/Materials/MaterialShader.cpp
+++ b/Engine/Source/Runtime/Engine/Private/Materials/MaterialShader.cpp
@@ -1101,10 +1101,12 @@ void FMaterialShaderMapId::Serialize(FArchive& Ar, bool bLoadingCooked)
 		if (!bIsLegacyPackage)
 		{
 			Ar << bUsingNewHLSLGenerator;
+			Ar << ExternalCodeReferencesHash;
 		}
 		else
 		{
 			bUsingNewHLSLGenerator = false;
+			ExternalCodeReferencesHash = FSHAHash();
 		}
 
 		// SUBSTRATE_TODO We do not need to serialize FSubstrateCompilationConfig for now since this is only used when debugging in the editor.
@@ -1208,6 +1210,8 @@ void FMaterialShaderMapId::GetMaterialHash(FSHAHash& OutHash, bool bWithStaticPa
 
 	HashState.Update((const uint8*)&bUsingNewHLSLGenerator, sizeof(bUsingNewHLSLGenerator));
 
+	HashState.Update((const uint8*)&ExternalCodeReferencesHash, sizeof(ExternalCodeReferencesHash));
+
 	SubstrateCompilationConfig.UpdateHash(HashState);
 
 	HashState.Final();
@@ -1585,6 +1589,7 @@ void FMaterialShaderMapId::RecordAndEmit(FMaterialKeyGeneratorContext& Context)
 		Context.RecordAndEmit("TextureReferencesHash", TextureReferencesHash);
 		Context.RecordAndEmit("ExpressionIncludesHash", ExpressionIncludesHash);
 		Context.RecordAndEmit("BasePropertyOverridesHash", BasePropertyOverridesHash);
+		Context.RecordAndEmit("ExternalCodeReferencesHash", ExternalCodeReferencesHash);
 	}
 
 	Context.Record("bUsingNewHLSLGenerator", bUsingNewHLSLGenerator);
diff --git a/Engine/Source/Runtime/Engine/Private/Materials/MaterialShared.cpp b/Engine/Source/Runtime/Engine/Private/Materials/MaterialShared.cpp
--- a/Engine/Source/Runtime/Engine/Private/Materials/MaterialShared.cpp
+++ b/Engine/Source/Runtime/Engine/Private/Materials/MaterialShared.cpp
@@ -758,6 +758,7 @@ void FMaterial::GetShaderMapId(EShaderPlatform Platform, const ITargetPlatform*
 		OutId.SetShaderDependencies(ShaderTypes, ShaderPipelineTypes, VFTypes, Platform);
 		GetReferencedTexturesHash(Platform, OutId.TextureReferencesHash);
 		GetExpressionIncludesHash(Platform, OutId.ExpressionIncludesHash);
+		GetExternalCodeReferencesHash(OutId.ExternalCodeReferencesHash);
 
 		OutId.SubstrateCompilationConfig = GetSubstrateCompilationConfig();
 
@@ -4396,6 +4397,11 @@ void FMaterial::GetExpressionIncludesHash(EShaderPlatform Platform, FSHAHash& Ou
 
 	OutHash = HashState.Finalize();
 }
+
+void FMaterial::GetExternalCodeReferencesHash(FSHAHash& OutHash) const
+{
+	GetCachedExpressionData().GetExternalCodeReferencesHash(OutHash);
+}
 #endif // WITH_EDITOR
 
 bool FMaterial::GetMaterialExpressionSource(FString& OutSource)
diff --git a/Engine/Source/Runtime/Engine/Public/MaterialCachedData.h b/Engine/Source/Runtime/Engine/Public/MaterialCachedData.h
--- a/Engine/Source/Runtime/Engine/Public/MaterialCachedData.h
+++ b/Engine/Source/Runtime/Engine/Public/MaterialCachedData.h
@@ -19,6 +19,7 @@ class UCurveLinearColorAtlas;
 class UFont;
 class UMaterial;
 class UMaterialExpression;
+class UMaterialExpressionExternalCodeBase;
 class URuntimeVirtualTexture;
 class ULandscapeGrassType;
 class UMaterialFunctionInterface;
@@ -259,6 +260,9 @@ struct FMaterialCachedExpressionData
 	/** Returns an array of the guids of parameter collections used. */
 	void AppendReferencedParameterCollectionIdsTo(TArray<FGuid>& OutIds) const;
 
+	/** Returns a hash over all external code references from this cache. */
+	void GetExternalCodeReferencesHash(FSHAHash& OutHash) const;
+
 	bool IsPropertyConnected(EMaterialProperty Property) const
 	{
 		return ((PropertyConnectedMask >> (uint64)Property) & 0x1) != 0;
@@ -325,6 +329,10 @@ struct FMaterialCachedExpressionData
 	UPROPERTY()
 	TArray<TObjectPtr<UTextureCollection>> ReferencedTextureCollections;
 
+	/** Array of all material expression classes that contain external code declarations. */
+	UPROPERTY()
+	TArray<TObjectPtr<UClass>> ReferencedExternalCodeExpressionClasses;
+
 	/** Array of all functions this material depends on. */
 	UPROPERTY()
 	TArray<FMaterialFunctionInfo> FunctionInfos;
diff --git a/Engine/Source/Runtime/Engine/Public/MaterialShared.h b/Engine/Source/Runtime/Engine/Public/MaterialShared.h
--- a/Engine/Source/Runtime/Engine/Public/MaterialShared.h
+++ b/Engine/Source/Runtime/Engine/Public/MaterialShared.h
@@ -1180,6 +1180,9 @@ public:
 	/** A hash of the base property overrides for this material instance. */
 	FSHAHash BasePropertyOverridesHash;
 
+	/** A hash of all external code delcarations referenced in the shader map. */
+	FSHAHash ExternalCodeReferencesHash;
+
 	/** Is the material using the new HLSL generator? */
 	bool bUsingNewHLSLGenerator;
 
@@ -2879,6 +2882,8 @@ private:
 	void GetReferencedTexturesHash(EShaderPlatform Platform, FSHAHash& OutHash) const;
 
 	void GetExpressionIncludesHash(EShaderPlatform Platform, FSHAHash& OutHash) const;
+
+	void GetExternalCodeReferencesHash(FSHAHash& OutHash) const;
 #endif // WITH_EDITOR
 
 	friend class FMaterialShaderMap;
