Skeletal Mesh: Import paths are not needed in cooked builds

#rnx

diff --git a/Engine/Source/Runtime/ClothingSystemRuntimeInterface/Public/ClothingAssetBase.h b/Engine/Source/Runtime/ClothingSystemRuntimeInterface/Public/ClothingAssetBase.h
--- a/Engine/Source/Runtime/ClothingSystemRuntimeInterface/Public/ClothingAssetBase.h
+++ b/Engine/Source/Runtime/ClothingSystemRuntimeInterface/Public/ClothingAssetBase.h
@@ -100,9 +100,11 @@ public:
 		return AssetGuid;
 	}
 
+#if WITH_EDITORONLY_DATA
 	// If this asset was imported from a file, this will be the original path
 	UPROPERTY(VisibleAnywhere, Category = Import)
 	FString ImportedFilePath;
+#endif
 
 protected:
 
diff --git a/Engine/Source/Runtime/Engine/Classes/Engine/SkinnedAssetCommon.h b/Engine/Source/Runtime/Engine/Classes/Engine/SkinnedAssetCommon.h
--- a/Engine/Source/Runtime/Engine/Classes/Engine/SkinnedAssetCommon.h
+++ b/Engine/Source/Runtime/Engine/Classes/Engine/SkinnedAssetCommon.h
@@ -190,9 +190,11 @@ struct FSkeletalMeshLODInfo
 	UPROPERTY(EditAnywhere, Category = ReductionSettings)
 	TObjectPtr<UAnimSequence> BakePoseOverride;
 
+#if WITH_EDITORONLY_DATA
 	/** The filename of the file tha was used to import this LOD if it was not auto generated. */
 	UPROPERTY(VisibleAnywhere, Category= SkeletalMeshLODInfo, AdvancedDisplay)
 	FString SourceImportFilename;
+#endif
 
 	/**
 	 * How this LOD uses the skin cache feature. Auto will defer to the default project global option. If Support Ray Tracing is enabled on the mesh, will imply Enabled
