Removed deprecated members in FPrimitiveSceneInfoData that were causing implicit copy constructor to report deprecation errors.

diff --git a/Engine/Source/Runtime/Engine/Public/PrimitiveSceneInfoData.h b/Engine/Source/Runtime/Engine/Public/PrimitiveSceneInfoData.h
--- a/Engine/Source/Runtime/Engine/Public/PrimitiveSceneInfoData.h
+++ b/Engine/Source/Runtime/Engine/Public/PrimitiveSceneInfoData.h
@@ -17,9 +17,6 @@ struct FPrimitiveSceneInfoData
 	/** The primitive's scene info. */
 	FPrimitiveSceneProxy* SceneProxy = nullptr;
 
-	UE_DEPRECATED(5.5, "LastSubmitTime is no longer used")
-	float LastSubmitTime = -1000.0f;
-
 	/**
 	 * The value of WorldSettings->TimeSeconds for the frame when this component was last rendered.  This is written
 	 * from the render thread, which is up to a frame behind the game thread, so you should allow this time to
@@ -55,9 +52,6 @@ struct FPrimitiveSceneInfoData
 	/** Used by the renderer, to identify a primitive across re-registers. */
 	FPrimitiveComponentId PrimitiveSceneId;
 
-	UE_DEPRECATED(5.5, "RegistrationSerialNumber is no longer used")
-	int32 RegistrationSerialNumber : 31;
-
 	/** Whether the primitive is always visible. If true the last render time will be unset. */
 	int32 bAlwaysVisible : 1;
 
@@ -89,8 +83,7 @@ public:
 	static int32 GetNextRegistrationSerialNumber() { return -1; }
 
 	FPrimitiveSceneInfoData()
-		: RegistrationSerialNumber(-1)
-		, bAlwaysVisible(false)
+		: bAlwaysVisible(false)
 	{
 		PrimitiveSceneId.PrimIDValue = AutoRTFM::Open([] { return NextPrimitiveId.Increment(); });
 	}
