[Materials] Replace Cast() with CastChecked() to remove redundant assertion.

#rnx

diff --git a/Engine/Source/Runtime/Engine/Private/Materials/MaterialCachedData.cpp b/Engine/Source/Runtime/Engine/Private/Materials/MaterialCachedData.cpp
--- a/Engine/Source/Runtime/Engine/Private/Materials/MaterialCachedData.cpp
+++ b/Engine/Source/Runtime/Engine/Private/Materials/MaterialCachedData.cpp
@@ -119,9 +119,7 @@ void FMaterialCachedExpressionData::GetExternalCodeReferencesHash(FSHAHash& OutH
 	{
 		if (const UObject* DefaultExternalCodeExpression = ExternalCodeExpressionClass->GetDefaultObject())
 		{
-			const UMaterialExpressionExternalCodeBase* ExternalCodeExpressionBase = Cast<UMaterialExpressionExternalCodeBase>(DefaultExternalCodeExpression);
-			checkf(ExternalCodeExpressionBase != nullptr, TEXT("Wrong UObject class type cached in FMaterialCachedExpressionData::ReferencedExternalCodeExpressionClasses"));
-
+			const UMaterialExpressionExternalCodeBase* ExternalCodeExpressionBase = CastChecked<UMaterialExpressionExternalCodeBase>(DefaultExternalCodeExpression);
 			for (const FName& ExternalCodeIdentifier : ExternalCodeExpressionBase->ExternalCodeIdentifiers)
 			{
 				if (const FMaterialExternalCodeDeclaration* ExternalCodeDeclaration = MaterialExternalCodeRegistry::Get().FindExternalCode(ExternalCodeIdentifier))
