[Materials] Ignore null pointer to material in GetCaptions() and GetExpressionToolTip() functions.

UMaterialExpressionObjectPositionWS should not crash when the pointer to its material is not initialized in the aforementioned functions.
These functions are called in commandlets such as "DumpMaterialExpressions" which don't guarantee a link to a material.

#rnx

diff --git a/Engine/Source/Runtime/Engine/Private/Materials/MaterialExpressions.cpp b/Engine/Source/Runtime/Engine/Private/Materials/MaterialExpressions.cpp
--- a/Engine/Source/Runtime/Engine/Private/Materials/MaterialExpressions.cpp
+++ b/Engine/Source/Runtime/Engine/Private/Materials/MaterialExpressions.cpp
@@ -12130,7 +12130,7 @@ int32 UMaterialExpressionObjectPositionWS::Compile(class FMaterialCompiler* Comp
 
 void UMaterialExpressionObjectPositionWS::GetCaption(TArray<FString>& OutCaptions) const
 {
-	FString ObjectType = Material->MaterialDomain == MD_LightFunction ? "Light" : "Object";
+	FString ObjectType = Material && Material->MaterialDomain == MD_LightFunction ? "Light" : "Object";
 	switch (OriginType)
 	{
 		case EPositionOrigin::Absolute:
@@ -12155,8 +12155,12 @@ void UMaterialExpressionObjectPositionWS::GetCaption(TArray<FString>& OutCaption
 
 void UMaterialExpressionObjectPositionWS::GetExpressionToolTip(TArray<FString>& OutToolTip)
 {
+	if (!Material)
+	{
+		return;
+	}
 	FString ToolTipText;
-	switch(Material->MaterialDomain)
+	switch (Material->MaterialDomain)
 	{
 		case MD_LightFunction:
 		{
