[Navmesh] Brush Navmesh Issues
- Prevent Static Brush from registering to Nav System
- Rebuild Level BSP when sublevel transforms as a whole
- Fixed perf issue by only rebuilding BSP post edit change rather than every OnRegister because Brush Actor forces registering itself post edit move for some reasons. Nav Relevancy is only involved in Registration to Nav System. The user can't add a brushcomponent in the editor. If the user adds a Brush Actor, BSP Rebuild will be fired from FBrushBuilderDragDropOp::OnDrop

diff --git a/Engine/Source/Editor/UnrealEd/Private/LevelEditorViewport.cpp b/Engine/Source/Editor/UnrealEd/Private/LevelEditorViewport.cpp
--- a/Engine/Source/Editor/UnrealEd/Private/LevelEditorViewport.cpp
+++ b/Engine/Source/Editor/UnrealEd/Private/LevelEditorViewport.cpp
@@ -3662,7 +3662,7 @@ void FLevelEditorViewportClient::TrackingStarted( const FInputEventState& InInpu
 		FTypedElementListConstRef ElementsToManipulate = GetElementsToManipulate();
 		TypedElementListObjectUtil::ForEachObject<ABrush>(ElementsToManipulate, [this](ABrush* InBrush)
 		{
-			bIsTrackingBrushModification = InBrush && !InBrush->IsVolumeBrush() && !FActorEditorUtils::IsABuilderBrush(InBrush);
+			bIsTrackingBrushModification = InBrush && InBrush->CanEverAffectBSP();
 			return !bIsTrackingBrushModification;
 		});
 	}
diff --git a/Engine/Source/Runtime/Engine/Classes/Components/BrushComponent.h b/Engine/Source/Runtime/Engine/Classes/Components/BrushComponent.h
--- a/Engine/Source/Runtime/Engine/Classes/Components/BrushComponent.h
+++ b/Engine/Source/Runtime/Engine/Classes/Components/BrushComponent.h
@@ -41,10 +41,6 @@ class UBrushComponent : public UPrimitiveComponent
 	ENGINE_API virtual void PostEditChangeProperty(FPropertyChangedEvent& PropertyChangedEvent) override;
 	//~ End UObject Interface
 
-	//~Begin UActorComponent Interface
-	ENGINE_API virtual void OnRegister() override;
-	//~End UActorComponent Interface
-
 	void ConditionalRebuildAlteredBSP() const;
 #endif //WITH_EDITOR
 
diff --git a/Engine/Source/Runtime/Engine/Classes/Engine/Brush.h b/Engine/Source/Runtime/Engine/Classes/Engine/Brush.h
--- a/Engine/Source/Runtime/Engine/Classes/Engine/Brush.h
+++ b/Engine/Source/Runtime/Engine/Classes/Engine/Brush.h
@@ -171,6 +171,7 @@ public:
 	
 	// AActor interface
 	ENGINE_API virtual bool IsLevelBoundsRelevant() const override;
+	ENGINE_API virtual bool IsComponentRelevantForNavigation(UActorComponent* Component) const override;
 	ENGINE_API virtual void RebuildNavigationData();
 
 #if WITH_EDITOR
@@ -189,7 +190,10 @@ public:
 	ENGINE_API virtual void InitPosRotScale();
 	ENGINE_API virtual void CopyPosRotScaleFrom( ABrush* Other );
 
-	/** @return true if its Brush Component is static and should affect navmesh based on user settings*/
+	/** @return true if its Brush Component should ever participate in BSP Rebuild */
+	bool ENGINE_API CanEverAffectBSP() const;
+
+	/** @return true if its Brush Component should be considered when building the level's static navigable geometry */
 	bool ENGINE_API ShouldExportStaticNavigableGeometry() const;
 	
 	static void SetSuppressBSPRegeneration(bool bSuppress) { bSuppressBSPRegeneration = bSuppress; }
diff --git a/Engine/Source/Runtime/Engine/Private/Brush.cpp b/Engine/Source/Runtime/Engine/Private/Brush.cpp
--- a/Engine/Source/Runtime/Engine/Private/Brush.cpp
+++ b/Engine/Source/Runtime/Engine/Private/Brush.cpp
@@ -55,7 +55,7 @@ namespace BrushNavmeshGenerationCVars
 	static FAutoConsoleVariableRef CVarForceNavmeshGenerationOnStaticBrush(
 		TEXT("brush.ForceNavmeshGenerationOnStaticBrush"), 
 		bForceNavmeshGenerationOnStaticBrush, 
-		TEXT("Force exporting static brush's geometry data to Navigation System regardless of CanEverAffectNavigation and CollisionPreset values. Enabled by default to be backward compatible with brushes on legacy maps."), 
+		TEXT("Force exporting static brush to level's static navigable geometry through BSP regardless of bCanEverAffectNavigation. Enabled by default to be backward compatible with brushes on legacy maps."), 
 		ECVF_Default);
 }
 #endif
@@ -193,13 +193,36 @@ void ABrush::SetNeedRebuild(ULevel* InLevel)
 	}
 }
 
+bool ABrush::CanEverAffectBSP() const
+{
+	return !(IsVolumeBrush() || FActorEditorUtils::IsABuilderBrush(this));
+}
+
 bool ABrush::ShouldExportStaticNavigableGeometry() const
 {
 	if (BrushComponent)
 	{
-		const bool bNavRelevant = BrushComponent->IsNavigationRelevant();
+		if (!CanEverAffectBSP())
+		{
+			return false;
+		}
+
+		// non-static brush should participate in BSP Rebuild for setting up other data from the existing code.
+		// However, they will not export static navigable geometry data to the level.
+		// Instead, they set up their Body Setup and Body Instance and register to the navigation system
+		// on a per-component basis like other components.
+		if (!IsStaticBrush())
+		{
+			return false;
+		}
+
+		// Nav relevancy (i.e., IsNavigationRelevant()) only applies to registration to navigation system, so don't check here
+		if (!(BrushNavmeshGenerationCVars::bForceNavmeshGenerationOnStaticBrush || BrushComponent->CanEverAffectNavigation()))
+		{
+			return false;
+		}
 
-		return IsStaticBrush() && (BrushNavmeshGenerationCVars::bForceNavmeshGenerationOnStaticBrush || bNavRelevant);
+		return true;
 	}
 
 	return false;
@@ -418,6 +441,22 @@ bool ABrush::IsLevelBoundsRelevant() const
 	return (Level && this != Level->Actors[1].Get());
 }
 
+bool ABrush::IsComponentRelevantForNavigation(UActorComponent* Component) const
+{
+	if (BrushComponent && BrushComponent == Component)
+	{
+		// When Brush changes from movable to static, it doesn't reset its Body Setup and Body Instance,
+		// which causes its geometric data be exported twice, both BSP Level and Component Registration
+		// So we need to mark it as nav irrelevant
+		if (IsStaticBrush())
+		{
+			return false;
+		}
+	}
+
+	return Super::IsComponentRelevantForNavigation(Component);
+}
+
 void ABrush::RebuildNavigationData()
 {
 	// empty in base class
diff --git a/Engine/Source/Runtime/Engine/Private/Components/BrushComponent.cpp b/Engine/Source/Runtime/Engine/Private/Components/BrushComponent.cpp
--- a/Engine/Source/Runtime/Engine/Private/Components/BrushComponent.cpp
+++ b/Engine/Source/Runtime/Engine/Private/Components/BrushComponent.cpp
@@ -531,16 +531,19 @@ void UBrushComponent::PostEditChangeProperty(FPropertyChangedEvent& PropertyChan
 {
 	if (const ABrush* BrushOwner = Cast<ABrush>(GetOwner()))
 	{
-		if (!BrushOwner->IsVolumeBrush())
+		// todo: see if we can optimize to only rebuild BSP only once when all the Brushes finished changing. 
+		if (BrushOwner->CanEverAffectBSP())
 		{
-			if (!BrushOwner->IsStaticBrush())
+			// if the brush changes from static to non-static or vice versa, we need to update the BSP
+			if (PropertyChangedEvent.Property && PropertyChangedEvent.Property->GetFName() == GET_MEMBER_NAME_CHECKED(UBrushComponent, Mobility))
 			{
-				if (PropertyChangedEvent.Property && PropertyChangedEvent.Property->GetFName() == GET_MEMBER_NAME_CHECKED(UBrushComponent, Mobility))
-				{
-					// This is to cover a case when the user changes the non-volume brush from static type to non-static type(i.e. changes its mobility)
-					// The previously registered static navigation data for BSP didn't get updated and this component's part was staying there
-					ConditionalRebuildAlteredBSP();
-				}
+				ConditionalRebuildAlteredBSP();
+			}
+
+			// For brush, the flag also affects navigation data export via BSP
+			if (PropertyChangedEvent.Property && PropertyChangedEvent.Property->GetFName() == GET_MEMBER_NAME_CHECKED(UBrushComponent, bCanEverAffectNavigation))
+			{
+				ConditionalRebuildAlteredBSP();
 			}
 		}
 	}
@@ -548,23 +551,6 @@ void UBrushComponent::PostEditChangeProperty(FPropertyChangedEvent& PropertyChan
 	Super::PostEditChangeProperty(PropertyChangedEvent);
 }
 
-void UBrushComponent::OnRegister()
-{
-	Super::OnRegister();
-
-	if (const ABrush* BrushOwner = Cast<ABrush>(GetOwner()))
-	{
-		// Volume Brush and Non-Static Brush's Geometric Data will be exported via registering the component
-		// which is handled in UPrimitiveComponent::OnRegister
-		if (!BrushOwner->IsVolumeBrush() && BrushOwner->IsStaticBrush())
-		{
-			// Not checking NavRelavency is intentional to cover cases when it changes from being Nav Relevant to Nav Irrelevant,
-			// we need to update the registered static Navigation Data for BSP and remove this component's part.
-			ConditionalRebuildAlteredBSP();
-		}
-	}
-}
-
 void UBrushComponent::ConditionalRebuildAlteredBSP() const
 {
 	// Not entirely clear why we should limit the rebuilding to within a transaction,
diff --git a/Engine/Source/Runtime/Engine/Private/LevelUtils.cpp b/Engine/Source/Runtime/Engine/Private/LevelUtils.cpp
--- a/Engine/Source/Runtime/Engine/Private/LevelUtils.cpp
+++ b/Engine/Source/Runtime/Engine/Private/LevelUtils.cpp
@@ -359,16 +359,26 @@ void FLevelUtils::SetEditorTransform(ULevelStreaming* StreamingLevel, const FTra
 	const FScopedTransaction LevelOffsetTransaction( LOCTEXT( "ChangeEditorLevelTransform", "Edit Level Transform" ) );
 	StreamingLevel->Modify();
 
-	// Ensure that all Actors are in the transaction so that their location is restored and any construction script behaviors 
-	// based on being at a different location are correctly applied on undo/redo
+	bool bNeedsToRebuildLevelBSP = false;
 	if (ULevel* LoadedLevel = StreamingLevel->GetLoadedLevel())
 	{
 		for (AActor* Actor : LoadedLevel->Actors)
 		{
+			// Ensure that all Actors are in the transaction so that their location is restored and any construction script behaviors 
+			// based on being at a different location are correctly applied on undo/redo
 			if (Actor)
 			{
 				Actor->Modify();
 			}
+
+			// Iterate over all actors in the level to see if we need to rebuild the Level BSP(i.e., moving any Brush actor affecting BSP)
+			if (!bNeedsToRebuildLevelBSP && GEditor)
+			{
+				if (ABrush* Brush = Cast<ABrush>(Actor))
+				{
+					bNeedsToRebuildLevelBSP = Brush->CanEverAffectBSP();
+				}
+			}
 		}
 	}
 
@@ -377,6 +387,17 @@ void FLevelUtils::SetEditorTransform(ULevelStreaming* StreamingLevel, const FTra
 	StreamingLevel->LevelTransform = Transform;
 	ApplyEditorTransform(StreamingLevel, bDoPostEditMove);
 
+	if (bNeedsToRebuildLevelBSP)
+	{
+		check(GEditor);
+		check(StreamingLevel->GetLoadedLevel());
+
+		// RebuildAlteredBSP() only tracks the levels of currently selected actors if they are brushes
+		// which is not applicable when we are moving the whole level
+		// So we need to build the specific level BSP
+		GEditor->RebuildLevel(*StreamingLevel->GetLoadedLevel());
+	}
+
 	// Redraw the viewports to see this change
 	FEditorSupportDelegates::RedrawAllViewports.Broadcast();
 }
