[AssetManager] Added UE::SourceLocation to public API to let developers track loading request sources
- Also CreateCombinedHandle now let's you to pass callbacks so late binding is no longer required.

#rnx
#tests RR, local match

diff --git a/Engine/Plugins/Mutable/Source/CustomizableObject/Internal/MuCO/MutableStreamableManager.h b/Engine/Plugins/Mutable/Source/CustomizableObject/Internal/MuCO/MutableStreamableManager.h
--- a/Engine/Plugins/Mutable/Source/CustomizableObject/Internal/MuCO/MutableStreamableManager.h
+++ b/Engine/Plugins/Mutable/Source/CustomizableObject/Internal/MuCO/MutableStreamableManager.h
@@ -14,14 +14,15 @@ struct FStreamableHandle;
 class FMutableStreamableManager : public TSharedFromThis<FMutableStreamableManager>
 {
 public:
-	template<typename PathContainerType = TArray<FSoftObjectPath>, typename FuncType = FStreamableDelegate, typename DebugNameType = UE::FSourceLocation>
+	template<typename PathContainerType = TArray<FSoftObjectPath>, typename FuncType = FStreamableDelegate>
 		TSharedPtr<FStreamableHandle> RequestAsyncLoad(
 			PathContainerType&& TargetsToStream,
 			FuncType&& DelegateToCall = FStreamableDelegate(),
 			TAsyncLoadPriority Priority = FStreamableManager::DefaultAsyncLoadPriority,
 			bool bManageActiveHandle = false,
 			bool bStartStalled = false,
-			DebugNameType&& DebugNameOrLocation = UE::FSourceLocation::Current())
+			FString&& DebugName = FString(TEXT("MutableRequestAsyncLoad")),
+			UE::FSourceLocation&& Location = UE::FSourceLocation::Current())
 	{
 		check(IsInGameThread());
 		
@@ -50,7 +51,7 @@ public:
 		});
 		
 		FCookLoadScope CookLoadScope(ECookLoadType::EditorOnly);
-		TSharedPtr<FStreamableHandle> Handle = StreamableManager.RequestAsyncLoad(TargetsToStream, Delegate, Priority, bManageActiveHandle, bStartStalled, DebugNameOrLocation);
+		TSharedPtr<FStreamableHandle> Handle = StreamableManager.RequestAsyncLoad(TargetsToStream, Delegate, Priority, bManageActiveHandle, bStartStalled, MoveTemp(DebugName), MoveTemp(Location));
 		if (Handle)
 		{
 			MutableStreamableHandle->Handle = Handle;
@@ -60,14 +61,15 @@ public:
 		return MutableStreamableHandle->Handle;
 	}
 	
-	template<typename PathContainerType = TArray<FSoftObjectPath>, typename DebugNameType = UE::FSourceLocation>
+	template<typename PathContainerType = TArray<FSoftObjectPath>>
 	TSharedPtr<FStreamableHandle> RequestSyncLoad(
 		PathContainerType&& TargetsToStream,
 		bool bManageActiveHandle = false,
-		DebugNameType&& DebugNameOrLocation = UE::FSourceLocation::Current())
+		FString&& DebugName = FString(TEXT("MutableRequestAsyncLoad")),
+		UE::FSourceLocation&& Location = UE::FSourceLocation::Current())
 	{
 		FCookLoadScope CookLoadScope(ECookLoadType::EditorOnly);
-		return StreamableManager.RequestSyncLoad(TargetsToStream, bManageActiveHandle, DebugNameOrLocation);
+		return StreamableManager.RequestSyncLoad(TargetsToStream, bManageActiveHandle, MoveTemp(DebugName), MoveTemp(Location));
 	}
 
 	int32 Tick(bool bBlocking);
diff --git a/Engine/Source/Runtime/Engine/Classes/Engine/AssetManager.h b/Engine/Source/Runtime/Engine/Classes/Engine/AssetManager.h
--- a/Engine/Source/Runtime/Engine/Classes/Engine/AssetManager.h
+++ b/Engine/Source/Runtime/Engine/Classes/Engine/AssetManager.h
@@ -261,9 +261,15 @@ public:
 	 * @param LoadBundles		List of bundles to load for those assets
 	 * @param DelegateToCall	Delegate that will be called on completion
 	 * @param Priority			Async loading priority for this request
+	 * @param Location			[optional] Is not intended for direct use, the parameter catches call site source location and passes it down to the streaming manager. See FStreamingManager.
 	 * @return					Streamable Handle that can be used to poll or wait. You do not need to keep this handle to stop the assets from being unloaded
 	 */
-	virtual TSharedPtr<FStreamableHandle> LoadPrimaryAssets(const TArray<FPrimaryAssetId>& AssetsToLoad, const TArray<FName>& LoadBundles = TArray<FName>(), FStreamableDelegate DelegateToCall = FStreamableDelegate(), TAsyncLoadPriority Priority = FStreamableManager::DefaultAsyncLoadPriority);
+	virtual TSharedPtr<FStreamableHandle> LoadPrimaryAssets(
+		const TArray<FPrimaryAssetId>& AssetsToLoad,
+		const TArray<FName>& LoadBundles = TArray<FName>(),
+		FStreamableDelegate DelegateToCall = FStreamableDelegate(),
+		TAsyncLoadPriority Priority = FStreamableManager::DefaultAsyncLoadPriority,
+		UE::FSourceLocation Location = UE::FSourceLocation::Current());
 	/**
 	 * Loads a list of Primary Assets. This will start an async load of those assets, calling callback on completion.
 	 * These assets will stay in memory until explicitly unloaded.
@@ -273,19 +279,42 @@ public:
 	 * @param AssetsToLoad		List of primary assets to load
 	 * @param LoadBundles		List of bundles to load for those assets
 	 * @param LoadParams		Contains callbacks and async load priority
+	 * @param Location			[optional] Is not intended for direct use, the parameter catches call site source location and passes it down to the streaming manager. See FStreamingManager.
 	 * @return					Streamable Handle that can be used to poll or wait. You do not need to keep this handle to stop the assets from being unloaded
 	 */
-	ENGINE_API virtual TSharedPtr<FStreamableHandle> LoadPrimaryAssets(const TArray<FPrimaryAssetId>& AssetsToLoad, const TArray<FName>& LoadBundles, FAssetManagerLoadParams&& LoadParams);
+	ENGINE_API virtual TSharedPtr<FStreamableHandle> LoadPrimaryAssets(
+		const TArray<FPrimaryAssetId>& AssetsToLoad,
+		const TArray<FName>& LoadBundles,
+		FAssetManagerLoadParams&& LoadParams,
+		UE::FSourceLocation Location = UE::FSourceLocation::Current());
 
 	/** Single asset wrapper, see LoadPrimaryAssets */
-	virtual TSharedPtr<FStreamableHandle> LoadPrimaryAsset(const FPrimaryAssetId& AssetToLoad, const TArray<FName>& LoadBundles = TArray<FName>(), FStreamableDelegate DelegateToCall = FStreamableDelegate(), TAsyncLoadPriority Priority = FStreamableManager::DefaultAsyncLoadPriority);
+	virtual TSharedPtr<FStreamableHandle> LoadPrimaryAsset(
+		const FPrimaryAssetId& AssetToLoad,
+		const TArray<FName>& LoadBundles = TArray<FName>(),
+		FStreamableDelegate DelegateToCall = FStreamableDelegate(),
+		TAsyncLoadPriority Priority = FStreamableManager::DefaultAsyncLoadPriority,
+		UE::FSourceLocation Location = UE::FSourceLocation::Current());
 	/** Single asset wrapper, see LoadPrimaryAssets */
-	ENGINE_API virtual TSharedPtr<FStreamableHandle> LoadPrimaryAsset(const FPrimaryAssetId& AssetToLoad, const TArray<FName>& LoadBundles, FAssetManagerLoadParams&& LoadParams);
+	ENGINE_API virtual TSharedPtr<FStreamableHandle> LoadPrimaryAsset(
+		const FPrimaryAssetId& AssetToLoad,
+		const TArray<FName>& LoadBundles,
+		FAssetManagerLoadParams&& LoadParams,
+		UE::FSourceLocation Location = UE::FSourceLocation::Current());
 
 	/** Loads all assets of a given type, useful for cooking */
-	virtual TSharedPtr<FStreamableHandle> LoadPrimaryAssetsWithType(FPrimaryAssetType PrimaryAssetType, const TArray<FName>& LoadBundles = TArray<FName>(), FStreamableDelegate DelegateToCall = FStreamableDelegate(), TAsyncLoadPriority Priority = FStreamableManager::DefaultAsyncLoadPriority);
+	virtual TSharedPtr<FStreamableHandle> LoadPrimaryAssetsWithType(
+		FPrimaryAssetType PrimaryAssetType,
+		const TArray<FName>& LoadBundles = TArray<FName>(),
+		FStreamableDelegate DelegateToCall = FStreamableDelegate(),
+		TAsyncLoadPriority Priority = FStreamableManager::DefaultAsyncLoadPriority,
+		UE::FSourceLocation Location = UE::FSourceLocation::Current());
 	/** Loads all assets of a given type, useful for cooking */
-	ENGINE_API virtual TSharedPtr<FStreamableHandle> LoadPrimaryAssetsWithType(FPrimaryAssetType PrimaryAssetType, const TArray<FName>& LoadBundles, FAssetManagerLoadParams&& LoadParams);
+	ENGINE_API virtual TSharedPtr<FStreamableHandle> LoadPrimaryAssetsWithType(
+		FPrimaryAssetType PrimaryAssetType,
+		const TArray<FName>& LoadBundles,
+		FAssetManagerLoadParams&& LoadParams,
+		UE::FSourceLocation Location = UE::FSourceLocation::Current());
 
 	/** 
 	 * Unloads a list of Primary Assets that were previously Loaded.
@@ -311,12 +340,20 @@ public:
 	 * @param AssetsToChange	List of primary assets to change state of
 	 * @param AddBundles		List of bundles to add
 	 * @param RemoveBundles		Explicit list of bundles to remove
-	 * @param RemoveAllBundles	If true, remove all existing bundles even if not in remove list
-	 * @param DelegateToCall	Delegate that will be called on completion
-	 * @param Priority			Async loading priority for this request
+	 * @param RemoveAllBundles	[optional] If true, remove all existing bundles even if not in remove list
+	 * @param DelegateToCall	[optional] Delegate that will be called on completion
+	 * @param Priority			[optional] Async loading priority for this request
+	 * @param Location			[optional] Is not intended for direct use, the parameter catches call site source location and passes it down to the streaming manager. See FStreamingManager.
 	 * @return					Streamable Handle that can be used to poll or wait. You do not need to keep this handle to stop the assets from being unloaded
 	 */
-	virtual TSharedPtr<FStreamableHandle> ChangeBundleStateForPrimaryAssets(const TArray<FPrimaryAssetId>& AssetsToChange, const TArray<FName>& AddBundles, const TArray<FName>& RemoveBundles, bool bRemoveAllBundles = false, FStreamableDelegate DelegateToCall = FStreamableDelegate(), TAsyncLoadPriority Priority = FStreamableManager::DefaultAsyncLoadPriority);
+	virtual TSharedPtr<FStreamableHandle> ChangeBundleStateForPrimaryAssets(
+		const TArray<FPrimaryAssetId>& AssetsToChange,
+		const TArray<FName>& AddBundles,
+		const TArray<FName>& RemoveBundles,
+		bool bRemoveAllBundles = false,
+		FStreamableDelegate DelegateToCall = FStreamableDelegate(),
+		TAsyncLoadPriority Priority = FStreamableManager::DefaultAsyncLoadPriority,
+		UE::FSourceLocation Location = UE::FSourceLocation::Current());
 	/**
 	 * Changes the bundle state of a set of loaded primary assets.
 	 * You can wait on the returned streamable request or poll as needed.
@@ -327,9 +364,16 @@ public:
 	 * @param RemoveBundles		Explicit list of bundles to remove
 	 * @param RemoveAllBundles	If true, remove all existing bundles even if not in remove list
 	 * @param LoadParams		Contains callbacks and async load priority
+	 * @param Location			[optional] Is not intended for direct use, the parameter catches call site source location and passes it down to the streaming manager. See FStreamingManager.
 	 * @return					Streamable Handle that can be used to poll or wait. You do not need to keep this handle to stop the assets from being unloaded
 	 */
-	ENGINE_API virtual TSharedPtr<FStreamableHandle> ChangeBundleStateForPrimaryAssets(const TArray<FPrimaryAssetId>& AssetsToChange, const TArray<FName>& AddBundles, const TArray<FName>& RemoveBundles, bool bRemoveAllBundles, FAssetManagerLoadParams&& LoadParams);
+	ENGINE_API virtual TSharedPtr<FStreamableHandle> ChangeBundleStateForPrimaryAssets(
+		const TArray<FPrimaryAssetId>& AssetsToChange,
+		const TArray<FName>& AddBundles,
+		const TArray<FName>& RemoveBundles,
+		bool bRemoveAllBundles,
+		FAssetManagerLoadParams&& LoadParams,
+		UE::FSourceLocation Location = UE::FSourceLocation::Current());
 
 	/** 
 	 * Changes the bundle state of all loaded primary assets. Only assets matching OldBundles will be modified
@@ -339,11 +383,17 @@ public:
 	 *
 	 * @param NewBundles		New bundle state for the assets that are changed
 	 * @param OldBundles		Old bundle state, it will remove these bundles and replace with NewBundles
-	 * @param DelegateToCall	Delegate that will be called on completion
-	 * @param Priority			Async loading priority for this request
+	 * @param DelegateToCall	[optional] Delegate that will be called on completion
+	 * @param Priority			[optional] Async loading priority for this request
+	 * @param Location			[optional] Is not intended for direct use, the parameter catches call site source location and passes it down to the streaming manager. See FStreamingManager.
 	 * @return					Streamable Handle that can be used to poll or wait. You do not need to keep this handle to stop the assets from being unloaded
 	 */
-	virtual TSharedPtr<FStreamableHandle> ChangeBundleStateForMatchingPrimaryAssets(const TArray<FName>& NewBundles, const TArray<FName>& OldBundles, FStreamableDelegate DelegateToCall = FStreamableDelegate(), TAsyncLoadPriority Priority = FStreamableManager::DefaultAsyncLoadPriority);
+	virtual TSharedPtr<FStreamableHandle> ChangeBundleStateForMatchingPrimaryAssets(
+		const TArray<FName>& NewBundles,
+		const TArray<FName>& OldBundles,
+		FStreamableDelegate DelegateToCall = FStreamableDelegate(),
+		TAsyncLoadPriority Priority = FStreamableManager::DefaultAsyncLoadPriority,
+		UE::FSourceLocation Location = UE::FSourceLocation::Current());
 	/**
 	 * Changes the bundle state of all loaded primary assets. Only assets matching OldBundles will be modified
 	 * You can wait on the returned streamable request or poll as needed.
@@ -352,9 +402,14 @@ public:
 	 * @param NewBundles		New bundle state for the assets that are changed
 	 * @param OldBundles		Old bundle state, it will remove these bundles and replace with NewBundles
 	 * @param LoadParams		Contains callbacks and async load priority
+	 * @param Location			[optional] Is not intended for direct use, the parameter catches call site source location and passes it down to the streaming manager. See FStreamingManager.
 	 * @return					Streamable Handle that can be used to poll or wait. You do not need to keep this handle to stop the assets from being unloaded
 	 */
-	ENGINE_API virtual TSharedPtr<FStreamableHandle> ChangeBundleStateForMatchingPrimaryAssets(const TArray<FName>& NewBundles, const TArray<FName>& OldBundles, FAssetManagerLoadParams&& LoadParams);
+	ENGINE_API virtual TSharedPtr<FStreamableHandle> ChangeBundleStateForMatchingPrimaryAssets(
+		const TArray<FName>& NewBundles,
+		const TArray<FName>& OldBundles,
+		FAssetManagerLoadParams&& LoadParams,
+		UE::FSourceLocation Location = UE::FSourceLocation::Current());
 
 	/** 
 	 * Returns the loading handle associated with the primary asset, it can then be checked for progress or waited on
@@ -413,11 +468,18 @@ public:
 	 * @param AssetsToLoad		List of primary assets to load
 	 * @param LoadBundles		List of bundles to load for those assets
 	 * @param bLoadRecursive	If true, this will call RecursivelyExpandBundleData and recurse into sub bundles of other primary assets loaded by a bundle reference
-	 * @param DelegateToCall	Delegate that will be called on completion
-	 * @param Priority			Async loading priority for this request
+	 * @param DelegateToCall	[optional] Delegate that will be called on completion
+	 * @param Priority			[optional] Async loading priority for this request
+	 * @param Location			[optional] Is not intended for direct use, the parameter catches call site source location and passes it down to the streaming manager. See FStreamingManager.
 	 * @return					Streamable Handle that must be stored to keep the preloaded assets from being freed
 	 */
-	virtual TSharedPtr<FStreamableHandle> PreloadPrimaryAssets(const TArray<FPrimaryAssetId>& AssetsToLoad, const TArray<FName>& LoadBundles, bool bLoadRecursive, FStreamableDelegate DelegateToCall = FStreamableDelegate(), TAsyncLoadPriority Priority = FStreamableManager::DefaultAsyncLoadPriority);
+	virtual TSharedPtr<FStreamableHandle> PreloadPrimaryAssets(
+		const TArray<FPrimaryAssetId>& AssetsToLoad,
+		const TArray<FName>& LoadBundles,
+		bool bLoadRecursive,
+		FStreamableDelegate DelegateToCall = FStreamableDelegate(),
+		TAsyncLoadPriority Priority = FStreamableManager::DefaultAsyncLoadPriority,
+		UE::FSourceLocation Location = UE::FSourceLocation::Current());
 	/**
 	 * Preloads data for a set of assets in a specific bundle state, and returns a handle you must keep active.
 	 * These assets are not officially Loaded, so Unload/ChangeBundleState will not affect them and if you release the handle without otherwise loading the assets they will be freed.
@@ -426,9 +488,15 @@ public:
 	 * @param LoadBundles		List of bundles to load for those assets
 	 * @param bLoadRecursive	If true, this will call RecursivelyExpandBundleData and recurse into sub bundles of other primary assets loaded by a bundle reference
 	 * @param LoadParams		Contains callbacks and async load priority
+	 * @param Location			[optional] Is not intended for direct use, the parameter catches call site source location and passes it down to the streaming manager. See FStreamingManager.
 	 * @return					Streamable Handle that must be stored to keep the preloaded assets from being freed
 	 */
-	ENGINE_API virtual TSharedPtr<FStreamableHandle> PreloadPrimaryAssets(const TArray<FPrimaryAssetId>& AssetsToLoad, const TArray<FName>& LoadBundles, bool bLoadRecursive, FAssetManagerLoadParams&& LoadParams);
+	ENGINE_API virtual TSharedPtr<FStreamableHandle> PreloadPrimaryAssets(
+		const TArray<FPrimaryAssetId>& AssetsToLoad,
+		const TArray<FName>& LoadBundles,
+		bool bLoadRecursive,
+		FAssetManagerLoadParams&& LoadParams,
+		UE::FSourceLocation Location = UE::FSourceLocation::Current());
 
 	/**
 	 * Downloads data for a set of assets in a specific bundle state, and returns a handle with a FStreamableDownloadCachePin.
@@ -439,9 +507,15 @@ public:
 	 * @param LoadBundles		List of bundles to load for those assets
 	 * @param bLoadRecursive	If true, this will call RecursivelyExpandBundleData and recurse into sub bundles of other primary assets loaded by a bundle reference
 	 * @param LoadParams		Contains callbacks and download params. Download params are required.
+	 * @param Location			[optional] Is not intended for direct use, the parameter catches call site source location and passes it down to the streaming manager. See FStreamingManager.
 	 * @return					Streamable Handle
 	 */
-	ENGINE_API virtual TSharedPtr<FStreamableHandle> PredownloadPrimaryAssets(const TArray<FPrimaryAssetId>& AssetsToLoad, const TArray<FName>& LoadBundles, bool bLoadRecursive, FAssetManagerLoadParams&& LoadParams);
+	ENGINE_API virtual TSharedPtr<FStreamableHandle> PredownloadPrimaryAssets(
+		const TArray<FPrimaryAssetId>& AssetsToLoad,
+		const TArray<FName>& LoadBundles,
+		bool bLoadRecursive,
+		FAssetManagerLoadParams&& LoadParams,
+		UE::FSourceLocation Location = UE::FSourceLocation::Current());
 
 	/**
 	 * Load non primary assets with the primary streamable manager.
@@ -451,42 +525,56 @@ public:
 	 * @param AssetList			List of non primary assets to load
 	 * @param DelegateToCall	[optional] Delegate that will be called on completion
 	 * @param Priority			[optional] Async loading priority for this request
-	 * @param DebguName			[optional] Name of this handle, either FString or anything that can construct FString, will be reported in debug tools, will report Source Location if DebugName is not specified explicitly
+	 * @param DebugName			[optional] Name of this handle, either FString or anything that can construct FString, will be reported in debug tools, will report Source Location if DebugName is not specified explicitly
+	 * @param Location			[optional] Is not intended for direct use, the parameter catches call site source location and passes it down to the streaming manager. See FStreamingManager.
 	 * @return					Streamable Handle that must be stored to keep the preloaded assets from being freed
 	 */
-	template< typename DebugNameType = UE::FSourceLocation >
 	TSharedPtr<FStreamableHandle> LoadAssetList(
 		const TArray<FSoftObjectPath>& AssetList,
 		FStreamableDelegate DelegateToCall = FStreamableDelegate(),
 		TAsyncLoadPriority Priority = FStreamableManager::DefaultAsyncLoadPriority,
-		DebugNameType&& DebugNameOrLocation = UE::FSourceLocation::Current());
+		FString&& DebugName = FString(TEXT("LoadAssetList")),
+		UE::FSourceLocation&& Location = UE::FSourceLocation::Current());
 	/** rvalue reference overload for Asset List */
-	template< typename DebugNameType = UE::FSourceLocation >
 	TSharedPtr<FStreamableHandle> LoadAssetList(
 		TArray<FSoftObjectPath>&& AssetList,
 		FStreamableDelegate DelegateToCall = FStreamableDelegate(),
 		TAsyncLoadPriority Priority = FStreamableManager::DefaultAsyncLoadPriority,
-		DebugNameType&& DebugNameOrLocation = UE::FSourceLocation::Current());
+		FString&& DebugName = FString(TEXT("LoadAssetList")),
+		UE::FSourceLocation&& Location = UE::FSourceLocation::Current());
 	/**
-	 * Load non primary assets with the primary streamable manager.
+	 * lvalue reference overload to load non primary assets with the primary streamable manager.
 	 * This will not auto release the handle, release it if needed.
+	 * Prefer using rvalue reference version as it will not create an extra copy of your input array.
 	 * 
 	 * @param AssetList			List of non primary assets to load
 	 * @param LoadParams		Contains callbacks and async load priority
-	 * @param DebguName			[optional] Name of this handle, either FString or anything that can construct FString, will be reported in debug tools, will report Source Location if DebugName is not specified explicitly
+	 * @param DebugName			[optional] Name of this handle, either FString or anything that can construct FString, will be reported in debug tools, will report Source Location if DebugName is not specified explicitly
+	 * @param Location			[optional] Is not intended for direct use, the parameter catches call site source location and passes it down to the streaming manager. See FStreamingManager.
 	 * @return					Streamable Handle that must be stored to keep the preloaded assets from being freed
 	 */
-	template< typename DebugNameType = UE::FSourceLocation >
 	TSharedPtr<FStreamableHandle> LoadAssetList(
 		const TArray<FSoftObjectPath>& AssetList,
 		FAssetManagerLoadParams&& LoadParams,
-		DebugNameType&& DebugNameOrLocation = UE::FSourceLocation::Current());
-	/** rvalue reference overload for Asset List */
-	template< typename DebugNameType = UE::FSourceLocation >
-	TSharedPtr<FStreamableHandle> LoadAssetList(
+		FString&& DebugName = FString(TEXT("LoadAssetList")),
+		UE::FSourceLocation&& Location = UE::FSourceLocation::Current());
+
+	/**
+	 * Load non primary assets with the primary streamable manager.
+	 * This will not auto release the handle, release it if needed.
+	 * This is a entry point for the set of overloads. Overide it to have your custom behavior.
+	 * 
+	 * @param AssetList			List of non primary assets to load
+	 * @param LoadParams		Contains callbacks and async load priority
+	 * @param DebugName			[optional] Name of this handle, either FString or anything that can construct FString, will be reported in debug tools, will report Source Location if DebugName is not specified explicitly
+	 * @param Location			[optional] Is not intended for direct use, the parameter catches call site source location and passes it down to the streaming manager. See FStreamingManager.
+	 * @return					Streamable Handle that must be stored to keep the preloaded assets from being freed
+	 */
+	ENGINE_API virtual TSharedPtr<FStreamableHandle> LoadAssetList(
 		TArray<FSoftObjectPath>&& AssetList,
 		FAssetManagerLoadParams&& LoadParams,
-		DebugNameType&& DebugNameOrLocation = UE::FSourceLocation::Current());
+		FString&& DebugName = FString(TEXT("LoadAssetList")),
+		UE::FSourceLocation&& Location = UE::FSourceLocation::Current());
 
 	/**
 	 * Download non primary assets with the primary streamable manager.
@@ -495,20 +583,21 @@ public:
 	 * 
 	 * @param AssetList			List of non primary assets to load
 	 * @param LoadParams		Contains callbacks and async load priority
-	 * @param DebguName			[optional] Name of this handle, either FString or anything that can construct FString, will be reported in debug tools, will report Source Location if DebugName is not specified explicitly
+	 * @param DebugName			[optional] Name of this handle, either FString or anything that can construct FString, will be reported in debug tools, will report Source Location if DebugName is not specified explicitly
+	 * @param Location			[optional] Is not intended for direct use, the parameter catches call site source location and passes it down to the streaming manager. See FStreamingManager.
 	 * @return					Streamable Handle that must be stored to keep the preloaded assets from being freed
 	 */
-	template< typename DebugNameType = UE::FSourceLocation >
 	TSharedPtr<FStreamableHandle> DownloadAssetList(
 		const TArray<FSoftObjectPath>& AssetList,
 		FAssetManagerLoadParams&& LoadParams,
-		DebugNameType&& DebugNameOrLocation = UE::FSourceLocation::Current());
+		FString&& DebugName = FString(TEXT("DownloadAssetList")),
+		UE::FSourceLocation&& Location = UE::FSourceLocation::Current());
 	/** rvalue reference overload for Asset List */
-	template< typename DebugNameType = UE::FSourceLocation >
-	TSharedPtr<FStreamableHandle> DownloadAssetList(
+	ENGINE_API TSharedPtr<FStreamableHandle> DownloadAssetList(
 		TArray<FSoftObjectPath>&& AssetList,
 		FAssetManagerLoadParams&& LoadParams,
-		DebugNameType&& DebugNameOrLocation = UE::FSourceLocation::Current());
+		FString&& DebugName = FString(TEXT("DownloadAssetList")),
+		UE::FSourceLocation&& Location = UE::FSourceLocation::Current());
 
 	/** Returns a single AssetBundleInfo, matching Scope and Name */
 	ENGINE_API virtual FAssetBundleEntry GetAssetBundleEntry(const FPrimaryAssetId& BundleScope, FName BundleName) const;
@@ -907,17 +996,12 @@ protected:
 	ENGINE_API virtual void OnChunkDownloaded(uint32 ChunkId, bool bSuccess);
 
 	/** Called to load asset list */
-	UE_DEPRECATED(5.6, "Function was split up into two overloads. Use those if you need a hook in the same place.")
+	UE_DEPRECATED(5.6, "Function was renamed to LoadAssetList that accepts FAssetManagerLoadParams struct.")
 	virtual TSharedPtr<FStreamableHandle> LoadAssetListInternal(TArray<FSoftObjectPath>&& AssetList, FStreamableDelegateWithHandle&& DelegateToCall, TAsyncLoadPriority Priority, FString&& DebugName)
 	{
 		checkf(false, TEXT("Deprecated"));
 		return nullptr;
 	}
-	/** Called to load asset list */
-	ENGINE_API virtual TSharedPtr<FStreamableHandle> LoadAssetListInternal(TArray<FSoftObjectPath>&& AssetList, FAssetManagerLoadParams&& LoadParams, FString&& DebugName);
-
-	/** Called to download asset list. Does not load assets. */
-	ENGINE_API TSharedPtr<FStreamableHandle> DownloadAssetListInternal(TArray<FSoftObjectPath>&& AssetList, FAssetManagerLoadParams&& LoadParams, FString&& DebugName);
 
 #if WITH_EDITOR
 	/** Function used during creating Management references to decide when to recurse and set references */
@@ -1131,142 +1215,102 @@ private:
 	friend struct FCompiledAssetManagerSearchRules;
 };
 
-inline TSharedPtr<FStreamableHandle> UAssetManager::LoadPrimaryAssets(const TArray<FPrimaryAssetId>& AssetsToLoad, const TArray<FName>& LoadBundles, FStreamableDelegate DelegateToCall, TAsyncLoadPriority Priority)
+inline TSharedPtr<FStreamableHandle> UAssetManager::LoadPrimaryAssets(const TArray<FPrimaryAssetId>& AssetsToLoad, const TArray<FName>& LoadBundles, FStreamableDelegate DelegateToCall, TAsyncLoadPriority Priority, UE::FSourceLocation Location)
 {
 	FAssetManagerLoadParams LoadParams;
 	LoadParams.OnComplete = UE::StreamableManager::Private::WrapDelegate(DelegateToCall);
 	LoadParams.Priority = Priority;
 
-	return LoadPrimaryAssets(AssetsToLoad, LoadBundles, MoveTemp(LoadParams));
+	return LoadPrimaryAssets(AssetsToLoad, LoadBundles, MoveTemp(LoadParams), MoveTemp(Location));
 }
 
-inline TSharedPtr<FStreamableHandle> UAssetManager::LoadPrimaryAsset(const FPrimaryAssetId& AssetToLoad, const TArray<FName>& LoadBundles, FStreamableDelegate DelegateToCall, TAsyncLoadPriority Priority)
+inline TSharedPtr<FStreamableHandle> UAssetManager::LoadPrimaryAsset(const FPrimaryAssetId& AssetToLoad, const TArray<FName>& LoadBundles, FStreamableDelegate DelegateToCall, TAsyncLoadPriority Priority, UE::FSourceLocation Location)
 {
 	FAssetManagerLoadParams LoadParams;
 	LoadParams.OnComplete = UE::StreamableManager::Private::WrapDelegate(DelegateToCall);
 	LoadParams.Priority = Priority;
 
-	return LoadPrimaryAsset(AssetToLoad, LoadBundles, MoveTemp(LoadParams));
+	return LoadPrimaryAsset(AssetToLoad, LoadBundles, MoveTemp(LoadParams), MoveTemp(Location));
 }
 
-inline TSharedPtr<FStreamableHandle> UAssetManager::LoadPrimaryAssetsWithType(FPrimaryAssetType PrimaryAssetType, const TArray<FName>& LoadBundles, FStreamableDelegate DelegateToCall, TAsyncLoadPriority Priority)
+inline TSharedPtr<FStreamableHandle> UAssetManager::LoadPrimaryAssetsWithType(FPrimaryAssetType PrimaryAssetType, const TArray<FName>& LoadBundles, FStreamableDelegate DelegateToCall, TAsyncLoadPriority Priority, UE::FSourceLocation Location)
 {
 	FAssetManagerLoadParams LoadParams;
 	LoadParams.OnComplete = UE::StreamableManager::Private::WrapDelegate(DelegateToCall);
 	LoadParams.Priority = Priority;
 
-	return LoadPrimaryAssetsWithType(PrimaryAssetType, LoadBundles, MoveTemp(LoadParams));
+	return LoadPrimaryAssetsWithType(PrimaryAssetType, LoadBundles, MoveTemp(LoadParams), MoveTemp(Location));
 }
 
-inline TSharedPtr<FStreamableHandle> UAssetManager::ChangeBundleStateForPrimaryAssets(const TArray<FPrimaryAssetId>& AssetsToChange, const TArray<FName>& AddBundles, const TArray<FName>& RemoveBundles, bool bRemoveAllBundles, FStreamableDelegate DelegateToCall, TAsyncLoadPriority Priority)
+inline TSharedPtr<FStreamableHandle> UAssetManager::ChangeBundleStateForPrimaryAssets(const TArray<FPrimaryAssetId>& AssetsToChange, const TArray<FName>& AddBundles, const TArray<FName>& RemoveBundles, bool bRemoveAllBundles, FStreamableDelegate DelegateToCall, TAsyncLoadPriority Priority, UE::FSourceLocation Location)
 {
 	FAssetManagerLoadParams LoadParams;
 	LoadParams.OnComplete = UE::StreamableManager::Private::WrapDelegate(MoveTemp(DelegateToCall));
 	LoadParams.Priority = Priority;
 
-	return ChangeBundleStateForPrimaryAssets(AssetsToChange, AddBundles, RemoveBundles, bRemoveAllBundles, MoveTemp(LoadParams));
+	return ChangeBundleStateForPrimaryAssets(AssetsToChange, AddBundles, RemoveBundles, bRemoveAllBundles, MoveTemp(LoadParams), MoveTemp(Location));
 }
 
-inline TSharedPtr<FStreamableHandle> UAssetManager::ChangeBundleStateForMatchingPrimaryAssets(const TArray<FName>& NewBundles, const TArray<FName>& OldBundles, FStreamableDelegate DelegateToCall, TAsyncLoadPriority Priority)
+inline TSharedPtr<FStreamableHandle> UAssetManager::ChangeBundleStateForMatchingPrimaryAssets(const TArray<FName>& NewBundles, const TArray<FName>& OldBundles, FStreamableDelegate DelegateToCall, TAsyncLoadPriority Priority, UE::FSourceLocation Location)
 {
 	FAssetManagerLoadParams LoadParams;
 	LoadParams.OnComplete = UE::StreamableManager::Private::WrapDelegate(DelegateToCall);
 	LoadParams.Priority = Priority;
 
-	return ChangeBundleStateForMatchingPrimaryAssets(NewBundles, OldBundles, MoveTemp(LoadParams));
+	return ChangeBundleStateForMatchingPrimaryAssets(NewBundles, OldBundles, MoveTemp(LoadParams), MoveTemp(Location));
 }
 
-inline TSharedPtr<FStreamableHandle> UAssetManager::PreloadPrimaryAssets(const TArray<FPrimaryAssetId>& AssetsToLoad, const TArray<FName>& LoadBundles, bool bLoadRecursive, FStreamableDelegate DelegateToCall, TAsyncLoadPriority Priority)
+inline TSharedPtr<FStreamableHandle> UAssetManager::PreloadPrimaryAssets(const TArray<FPrimaryAssetId>& AssetsToLoad, const TArray<FName>& LoadBundles, bool bLoadRecursive, FStreamableDelegate DelegateToCall, TAsyncLoadPriority Priority, UE::FSourceLocation Location)
 {
 	FAssetManagerLoadParams LoadParams;
 	LoadParams.OnComplete = UE::StreamableManager::Private::WrapDelegate(MoveTemp(DelegateToCall));
 	LoadParams.Priority = Priority;
 
-	return PreloadPrimaryAssets(AssetsToLoad, LoadBundles, bLoadRecursive, MoveTemp(LoadParams));
+	return PreloadPrimaryAssets(AssetsToLoad, LoadBundles, bLoadRecursive, MoveTemp(LoadParams), MoveTemp(Location));
 }
 
-template< typename DebugNameType >
-TSharedPtr<FStreamableHandle> UAssetManager::LoadAssetList(
+inline TSharedPtr<FStreamableHandle> UAssetManager::LoadAssetList(
 	const TArray<FSoftObjectPath>& AssetList,
 	FStreamableDelegate DelegateToCall,
 	TAsyncLoadPriority Priority,
-	DebugNameType&& DebugNameOrLocation)
+	FString&& DebugName,
+	UE::FSourceLocation&& Location)
 {
 	FAssetManagerLoadParams LoadParams;
 	LoadParams.OnComplete = UE::StreamableManager::Private::WrapDelegate(MoveTemp(DelegateToCall));
 	LoadParams.Priority = Priority;
 
-	return LoadAssetList(CopyTemp(AssetList), MoveTemp(LoadParams), Forward<DebugNameType>(DebugNameOrLocation));
+	return LoadAssetList(CopyTemp(AssetList), MoveTemp(LoadParams), MoveTemp(DebugName), MoveTemp(Location));
 }
 
-template< typename DebugNameType >
-TSharedPtr<FStreamableHandle> UAssetManager::LoadAssetList(
+inline TSharedPtr<FStreamableHandle> UAssetManager::LoadAssetList(
 	TArray<FSoftObjectPath>&& AssetList,
 	FStreamableDelegate DelegateToCall,
 	TAsyncLoadPriority Priority,
-	DebugNameType&& DebugNameOrLocation)
+	FString&& DebugName,
+	UE::FSourceLocation&& Location)
 {
 	FAssetManagerLoadParams LoadParams;
 	LoadParams.OnComplete = UE::StreamableManager::Private::WrapDelegate(MoveTemp(DelegateToCall));
 	LoadParams.Priority = Priority;
 
-	return LoadAssetList(MoveTemp(AssetList), MoveTemp(LoadParams), Forward<DebugNameType>(DebugNameOrLocation));
+	return LoadAssetList(MoveTemp(AssetList), MoveTemp(LoadParams), MoveTemp(DebugName), MoveTemp(Location));
 }
 
-template< typename DebugNameType >
-TSharedPtr<FStreamableHandle> UAssetManager::LoadAssetList(
+inline TSharedPtr<FStreamableHandle> UAssetManager::LoadAssetList(
 	const TArray<FSoftObjectPath>& AssetList,
 	FAssetManagerLoadParams&& LoadParams,
-	DebugNameType&& DebugNameOrLocation)
-{
-	return LoadAssetList(CopyTemp(AssetList), MoveTemp(LoadParams), Forward<DebugNameType>(DebugNameOrLocation));
-}
-
-template< typename DebugNameType >
-TSharedPtr<FStreamableHandle> UAssetManager::LoadAssetList(
-	TArray<FSoftObjectPath>&& AssetList,
-	FAssetManagerLoadParams&& LoadParams,
-	DebugNameType&& DebugNameOrLocation)
+	FString&& DebugName,
+	UE::FSourceLocation&& Location)
 {
-	if constexpr (std::is_same_v<std::decay_t<DebugNameType>, UE::FSourceLocation>)
-	{
-		return LoadAssetListInternal(
-			MoveTemp(AssetList), MoveTemp(LoadParams),
-			UE::SourceLocation::ToFileAndLineString(DebugNameOrLocation));
-	}
-	else
-	{
-		return LoadAssetListInternal(
-			MoveTemp(AssetList), MoveTemp(LoadParams),
-			FString{ Forward<DebugNameType>(DebugNameOrLocation) });
-	}
+	return LoadAssetList(CopyTemp(AssetList), MoveTemp(LoadParams), MoveTemp(DebugName), MoveTemp(Location));
 }
 
-template< typename DebugNameType >
-TSharedPtr<FStreamableHandle> UAssetManager::DownloadAssetList(
+inline TSharedPtr<FStreamableHandle> UAssetManager::DownloadAssetList(
 	const TArray<FSoftObjectPath>& AssetList,
 	FAssetManagerLoadParams&& LoadParams,
-	DebugNameType&& DebugNameOrLocation)
-{
-	return DownloadAssetList(CopyTemp(AssetList), MoveTemp(LoadParams), Forward<DebugNameType>(DebugNameOrLocation));
-}
-
-template< typename DebugNameType >
-TSharedPtr<FStreamableHandle> UAssetManager::DownloadAssetList(
-	TArray<FSoftObjectPath>&& AssetList,
-	FAssetManagerLoadParams&& LoadParams,
-	DebugNameType&& DebugNameOrLocation)
+	FString&& DebugName,
+	UE::FSourceLocation&& Location)
 {
-	if constexpr (std::is_same_v<std::decay_t<DebugNameType>, UE::FSourceLocation>)
-	{
-		return DownloadAssetListInternal(
-			MoveTemp(AssetList), MoveTemp(LoadParams),
-			UE::SourceLocation::ToFileAndLineString(DebugNameOrLocation));
-	}
-	else
-	{
-		return DownloadAssetListInternal(
-			MoveTemp(AssetList), MoveTemp(LoadParams),
-			FString{ Forward<DebugNameType>(DebugNameOrLocation) });
-	}
+	return DownloadAssetList(CopyTemp(AssetList), MoveTemp(LoadParams), MoveTemp(DebugName), MoveTemp(Location));
 }
diff --git a/Engine/Source/Runtime/Engine/Classes/Engine/StreamableManager.h b/Engine/Source/Runtime/Engine/Classes/Engine/StreamableManager.h
--- a/Engine/Source/Runtime/Engine/Classes/Engine/StreamableManager.h
+++ b/Engine/Source/Runtime/Engine/Classes/Engine/StreamableManager.h
@@ -549,7 +549,7 @@ private:
 	struct FPrivateToken { explicit FPrivateToken() = default; };
 	FStreamableHandle() = delete;
 public:
-	explicit FStreamableHandle(FPrivateToken);
+	explicit FStreamableHandle(FPrivateToken, UE::FSourceLocation&& InLocation);
 
 	/** Destructor */
 	ENGINE_API ~FStreamableHandle();
@@ -608,6 +608,9 @@ private:
 	/** Name of this handle, passed in by caller to help in debugging */
 	FString DebugName;
 
+	/** Holds a location where this asset list was requested from */
+	UE::FSourceLocation Location;
+
 	/** Error details in case of an error */
 	TOptional<FStreamableError> Error;
 
@@ -741,65 +744,87 @@ struct FStreamableManager : public FGCObject
 	 * @param Priority				[optional] Priority to pass to the streaming system, higher priority will be loaded first
 	 * @param bManageActiveHandle	[optional] If true, the manager will keep the streamable handle active until explicitly released
 	 * @param bStartStalled			[optional] If true, the handle will start in a stalled state and will not attempt to actually async load until StartStalledHandle is called on it
-	 * @param DebugName				[optional] Name of this handle, either FString or anything that can construct FString, will be reported in debug tools, will report Source Location if DebugName is not specified explicitly
+	 * @param DebugName				[optional] Name of this handle, either FString or anything that can construct FString
+	 * @param Location				[optional] Is not intended for direct use, the parameter catches call site source location. That location is used for debug tools output.
 	 */
-	template< typename PathContainerType = TArray<FSoftObjectPath>, typename FuncType = FStreamableDelegateWithHandle, typename DebugNameType = UE::FSourceLocation >
+	template< typename PathContainerType = TArray<FSoftObjectPath>, typename FuncType = FStreamableDelegateWithHandle, typename = std::enable_if_t<!std::is_same_v<std::decay_t<PathContainerType>, FStreamableAsyncLoadParams>, void> >
 	TSharedPtr<FStreamableHandle> RequestAsyncLoad(
 		PathContainerType&& TargetsToStream,
 		FuncType&& DelegateToCall = FStreamableDelegateWithHandle(),
 		TAsyncLoadPriority Priority = DefaultAsyncLoadPriority,
 		bool bManageActiveHandle = false,
 		bool bStartStalled = false,
-		DebugNameType&& DebugNameOrLocation = UE::FSourceLocation::Current());
+		FString DebugName = FString(),
+		UE::FSourceLocation Location = UE::FSourceLocation::Current());
 
 	/** 
 	 * This is the primary streamable operation. Requests streaming of one or more target objects. When complete, a delegate function is called. Returns a Streamable Handle.
 	 *
 	 * @param Params				Load Parameters
-	 * @param DebugName				[optional] Name of this handle, either FString or anything that can construct FString, will be reported in debug tools, will report Source Location if DebugName is not specified explicitly
+	 * @param DebugName				[optional] Name of this handle, either FString or anything that can construct FString, will be reported in debug tools
+	 * @param Location				[optional] Is not intended for direct use, the parameter catches call site source location. That location is used for debug tools output.
 	 */ 
-	template< typename DebugNameType = UE::FSourceLocation >
-	TSharedPtr<FStreamableHandle> RequestAsyncLoad(
+	ENGINE_API TSharedPtr<FStreamableHandle> RequestAsyncLoad(
 		FStreamableAsyncLoadParams&& Params,
-		DebugNameType&& DebugNameOrLocation = UE::FSourceLocation::Current());
+		FString DebugName = FString(),
+		UE::FSourceLocation Location = UE::FSourceLocation::Current());
 
 	/** 
 	 * Synchronously load a set of assets, and return a handle. This can be very slow and may stall the game thread for several seconds.
 	 * 
 	 * @param TargetsToStream		Assets to load off disk
 	 * @param bManageActiveHandle	[optional] If true, the manager will keep the streamable handle active until explicitly released
-	 * @param DebugName				[optional] Name of this handle, either FString or anything that can construct FString, will be reported in debug tools, will report Source Location if DebugName is not specified explicitly
+	 * @param DebugName				[optional] Name of this handle, either FString or anything that can construct FString
+	 * @param Location				[optional] Is not intended for direct use, the parameter catches call site source location. That location is used for debug tools output.
 	 */
-	template< typename PathContainerType = TArray<FSoftObjectPath>, typename DebugNameType = UE::FSourceLocation >
+	template< typename PathContainerType = TArray<FSoftObjectPath>>
 	TSharedPtr<FStreamableHandle> RequestSyncLoad(
 		PathContainerType&& TargetsToStream,
 		bool bManageActiveHandle = false,
-		DebugNameType&& DebugNameOrLocation = UE::FSourceLocation::Current());
+		FString DebugName = FString(),
+		UE::FSourceLocation Location = UE::FSourceLocation::Current());
 
 	/** 
 	 * Synchronously load the referred asset and return the loaded object, or nullptr if it can't be found. This can be very slow and may stall the game thread for several seconds.
 	 * 
 	 * @param Target				Specific asset to load off disk
-	 * @param bManageActiveHandle	If true, the manager will keep the streamable handle active until explicitly released
-	 * @param RequestHandlePointer	If non-null, this will set the handle to the handle used to make this request. This useful for later releasing the handle
+	 * @param bManageActiveHandle	[optional] If true, the manager will keep the streamable handle active until explicitly released
+	 * @param RequestHandlePointer	[optional] If non-null, this will set the handle to the handle used to make this request. This useful for later releasing the handle
+	 * @param Location				[optional] Is not intended for direct use, the parameter catches call site source location. That location is used for debug tools output.
 	 */
-	ENGINE_API UObject* LoadSynchronous(const FSoftObjectPath& Target, bool bManageActiveHandle = false, TSharedPtr<FStreamableHandle>* RequestHandlePointer = nullptr, UE::FSourceLocation Location = UE::FSourceLocation::Current());
+	ENGINE_API UObject* LoadSynchronous(
+		const FSoftObjectPath& Target,
+		bool bManageActiveHandle = false,
+		TSharedPtr<FStreamableHandle>* RequestHandlePointer = nullptr,
+		UE::FSourceLocation Location = UE::FSourceLocation::Current());
 
 	/** Typed wrappers */
 	template< typename T >
-	T* LoadSynchronous(const FSoftObjectPath& Target, bool bManageActiveHandle = false, TSharedPtr<FStreamableHandle>* RequestHandlePointer = nullptr, UE::FSourceLocation Location = UE::FSourceLocation::Current())
+	T* LoadSynchronous(
+		const FSoftObjectPath& Target,
+		bool bManageActiveHandle = false,
+		TSharedPtr<FStreamableHandle>* RequestHandlePointer = nullptr,
+		UE::FSourceLocation Location = UE::FSourceLocation::Current())
 	{
 		return Cast<T>(LoadSynchronous(Target, bManageActiveHandle, RequestHandlePointer, MoveTemp(Location)) );
 	}
 
 	template< typename T >
-	T* LoadSynchronous(const TSoftObjectPtr<T>& Target, bool bManageActiveHandle = false, TSharedPtr<FStreamableHandle>* RequestHandlePointer = nullptr, UE::FSourceLocation Location = UE::FSourceLocation::Current())
+	T* LoadSynchronous(
+		const TSoftObjectPtr<T>& Target,
+		bool bManageActiveHandle = false,
+		TSharedPtr<FStreamableHandle>* RequestHandlePointer = nullptr,
+		UE::FSourceLocation Location = UE::FSourceLocation::Current())
 	{
 		return Cast<T>(LoadSynchronous(Target.ToSoftObjectPath(), bManageActiveHandle, RequestHandlePointer, MoveTemp(Location)));
 	}
 
 	template< typename T >
-	TSubclassOf<T> LoadSynchronous(const TSoftClassPtr<T>& Target, bool bManageActiveHandle = false, TSharedPtr<FStreamableHandle>* RequestHandlePointer = nullptr, UE::FSourceLocation Location = UE::FSourceLocation::Current())
+	TSubclassOf<T> LoadSynchronous(
+		const TSoftClassPtr<T>& Target,
+		bool bManageActiveHandle = false,
+		TSharedPtr<FStreamableHandle>* RequestHandlePointer = nullptr,
+		UE::FSourceLocation Location = UE::FSourceLocation::Current())
 	{
 		TSubclassOf<T> ReturnClass;
 		ReturnClass = Cast<UClass>(LoadSynchronous(Target.ToSoftObjectPath(), bManageActiveHandle, RequestHandlePointer, MoveTemp(Location)));
@@ -811,11 +836,18 @@ struct FStreamableManager : public FGCObject
 	 * The child handles will be held as hard references as long as this handle is active.
 	 * 
 	 * @param ChildHandles			List of handles to wrap into this one
-	 * @param DebugName				Name of this handle, will be reported in debug tools
+	 * @param DebugName				[optional] Name of this handle, will be reported in debug tools
+	 * @param Options				[optional] Additional flags that let you specify how the combined handle will be constructed
+	 * @param Params				[optional] Load Parameters. TargetsToStream, Priority, bManageActiveHandle, bStartStalled and DownloadParams are ignored
+	 * @param Location				[optional] Is not intended for direct use, the parameter catches call site source location. That location is used for debug tools output.
 	 */
-	ENGINE_API TSharedPtr<FStreamableHandle> CreateCombinedHandle(const TConstArrayView<TSharedPtr<FStreamableHandle>>& ChildHandles,
-		const FString& DebugName = FStreamableHandle::HandleDebugName_CombinedHandle,
-		EStreamableManagerCombinedHandleOptions Options = EStreamableManagerCombinedHandleOptions::None);
+	ENGINE_API TSharedPtr<FStreamableHandle> CreateCombinedHandle(
+		TConstArrayView<TSharedPtr<FStreamableHandle>> ChildHandles,
+		FString DebugName = FStreamableHandle::HandleDebugName_CombinedHandle,
+		EStreamableManagerCombinedHandleOptions Options = EStreamableManagerCombinedHandleOptions::None,
+		FStreamableAsyncLoadParams&& Params = FStreamableAsyncLoadParams{},
+		UE::FSourceLocation Location = UE::FSourceLocation::Current()
+	);
 
 	/**
 	 * Creates a canceled handle with an error set. Does not do any loading. Does not manage the handle.
@@ -824,11 +856,14 @@ struct FStreamableManager : public FGCObject
 	 * @param Error					The error to set on the handle
 	 * @param Params				Load Parameters
 	 * @param DownloadParams		Additional Parameters to control downloading
+	 * @param DebugName				[optional] Name of this handle, will be reported in debug tools
+	 * @param Location				[optional] Is not intended for direct use, the parameter catches call site source location. That location is used for debug tools output.
 	 */
 	ENGINE_API TSharedPtr<FStreamableHandle> CreateErrorHandle(
 		FStreamableError&& Error,
 		FStreamableAsyncLoadParams&& Params,
-		const FString& DebugName = FStreamableHandle::HandleDebugName_CombinedHandle);
+		FString DebugName = FStreamableHandle::HandleDebugName_CombinedHandle,
+		UE::FSourceLocation Location = UE::FSourceLocation::Current());
 
 	/** 
 	 * Gets list of handles that are directly referencing this asset, returns true if any found.
@@ -887,8 +922,7 @@ private:
 	void AsyncLoadCallback(FSoftObjectPath Request, UPackage* Package, EAsyncLoadingResult::Type Result);
 
 	ENGINE_API static bool ShouldStripDebugName();
-	ENGINE_API TSharedPtr<FStreamableHandle> RequestAsyncLoadInternal(FStreamableAsyncLoadParams&& Params, FString&& DebugName);
-	ENGINE_API TSharedPtr<FStreamableHandle> RequestSyncLoadInternal(TArray<FSoftObjectPath>&& TargetsToStream, bool bManageActiveHandle, FString&& DebugName);
+	ENGINE_API TSharedPtr<FStreamableHandle> RequestSyncLoadInternal(TArray<FSoftObjectPath>&& TargetsToStream, bool bManageActiveHandle, FString&& DebugName, UE::FSourceLocation&& Location);
 
 	/** Map of paths to streamable objects, this will be the post-redirector name */
 	typedef TMap<FSoftObjectPath, struct FStreamable*> TStreamableMap;
@@ -959,14 +993,15 @@ void FStreamableHandle::ForEachLoadedAsset(CallableT Callable) const
 	}
 }
 
-template< typename PathContainerType, typename FuncType, typename DebugNameType >
+template< typename PathContainerType, typename FuncType, typename >
 TSharedPtr<FStreamableHandle> FStreamableManager::RequestAsyncLoad(
 	PathContainerType&& TargetsToStream,
 	FuncType&& Callback,
 	TAsyncLoadPriority Priority,
 	bool bManageActiveHandle,
 	bool bStartStalled,
-	DebugNameType&& DebugNameOrLocation)
+	FString DebugName,
+	UE::FSourceLocation Location)
 {
 	FStreamableAsyncLoadParams Params;
 	Params.TargetsToStream = TArray<FSoftObjectPath>{ Forward<PathContainerType>(TargetsToStream) };
@@ -999,57 +1034,22 @@ TSharedPtr<FStreamableHandle> FStreamableManager::RequestAsyncLoad(
 	Params.bManageActiveHandle = bManageActiveHandle;
 	Params.bStartStalled = bStartStalled;
 
-	if constexpr (std::is_same_v<std::decay_t<DebugNameType>, UE::FSourceLocation>)
-	{
-		return RequestAsyncLoadInternal(
-			MoveTemp(Params),
-			ShouldStripDebugName() ? FString() : UE::SourceLocation::ToFileAndLineString(DebugNameOrLocation));
-	}
-	else
-	{
-		return RequestAsyncLoadInternal(
-			MoveTemp(Params),
-			ShouldStripDebugName() ? FString() : FString{ Forward<DebugNameType>(DebugNameOrLocation) });
-	}
-}
-
-template< typename DebugNameType>
-TSharedPtr<FStreamableHandle> FStreamableManager::RequestAsyncLoad(
-	FStreamableAsyncLoadParams&& Params,
-	DebugNameType&& DebugNameOrLocation)
-{
-	if constexpr (std::is_same_v<std::decay_t<DebugNameType>, UE::FSourceLocation>)
-	{
-		return RequestAsyncLoadInternal(
-			MoveTemp(Params),
-			ShouldStripDebugName() ? FString() : UE::SourceLocation::ToFileAndLineString(DebugNameOrLocation));
-	}
-	else
-	{
-		return RequestAsyncLoadInternal(
-			MoveTemp(Params),
-			ShouldStripDebugName() ? FString() : FString{ Forward<DebugNameType>(DebugNameOrLocation) });
-	}
+	return RequestAsyncLoad(
+		MoveTemp(Params),
+		MoveTemp(DebugName),
+		MoveTemp(Location));
 }
 
-template< typename PathContainerType, typename DebugNameType >
+template< typename PathContainerType>
 TSharedPtr<FStreamableHandle> FStreamableManager::RequestSyncLoad(
 	PathContainerType&& TargetsToStream,
 	bool bManageActiveHandle,
-	DebugNameType&& DebugNameOrLocation)
+	FString DebugName,
+	UE::FSourceLocation Location)
 {
-	if constexpr (std::is_same_v<std::decay_t<DebugNameType>, UE::FSourceLocation>)
-	{
-		return RequestSyncLoadInternal(
-			TArray<FSoftObjectPath>{ Forward<PathContainerType>(TargetsToStream) },
-			bManageActiveHandle,
-			ShouldStripDebugName() ? FString() : UE::SourceLocation::ToFileAndLineString(DebugNameOrLocation));
-	}
-	else
-	{
-		return RequestSyncLoadInternal(
-			TArray<FSoftObjectPath>{ Forward<PathContainerType>(TargetsToStream) },
-			bManageActiveHandle,
-			ShouldStripDebugName() ? FString() : FString{ Forward<DebugNameType>(DebugNameOrLocation) });
-	}
+	return RequestSyncLoadInternal(
+		TArray<FSoftObjectPath>{ Forward<PathContainerType>(TargetsToStream) },
+		bManageActiveHandle,
+		MoveTemp(DebugName),
+		MoveTemp(Location));
 }
diff --git a/Engine/Source/Runtime/Engine/Private/AssetManager.cpp b/Engine/Source/Runtime/Engine/Private/AssetManager.cpp
--- a/Engine/Source/Runtime/Engine/Private/AssetManager.cpp
+++ b/Engine/Source/Runtime/Engine/Private/AssetManager.cpp
@@ -374,7 +374,7 @@ namespace UE::AssetManager::Private
 		FStringBuilderBase& DebugValid, FStringBuilderBase& DebugInvalid)
 	{
 		TArray<FSoftObjectPath> PathsToLoad;
-		bool ExeededMaxLength = false;
+		bool ExceededMaxLength = false;
 
 		for (const FPrimaryAssetId& PrimaryAssetId : AssetsToLoad)
 		{
@@ -386,7 +386,7 @@ namespace UE::AssetManager::Private
 				}
 				else
 				{
-					ExeededMaxLength = true;
+					ExceededMaxLength = true;
 				}
 			}
 			else
@@ -397,7 +397,7 @@ namespace UE::AssetManager::Private
 				}
 				else
 				{
-					ExeededMaxLength = true;
+					ExceededMaxLength = true;
 				}
 			}
 		}
@@ -409,7 +409,7 @@ namespace UE::AssetManager::Private
 				DebugInvalid << TEXT(" (Valid: ") << DebugValid << TEXT(")");
 			}
 
-			DebugInvalid << (ExeededMaxLength ? TEXT("...") : TEXT(""));
+			DebugInvalid << (ExceededMaxLength ? TEXT("...") : TEXT(""));
 
 			UE_LOG(LogAssetManager, Warning, TEXT("Requested preload of some Primary Assets failed: %s"), *DebugInvalid);
 
@@ -1805,7 +1805,8 @@ TSharedPtr<FStreamableHandle> UAssetManager::ChangeBundleStateForPrimaryAssets(
 	const TArray<FName>& AddBundles,
 	const TArray<FName>& RemoveBundles,
 	bool bRemoveAllBundles,
-	FAssetManagerLoadParams&& LoadParams)
+	FAssetManagerLoadParams&& LoadParams,
+	UE::FSourceLocation Location)
 {
 	TArray<TSharedPtr<FStreamableHandle> > NewHandles, ExistingHandles;
 	TArray<FPrimaryAssetId> NewAssets;
@@ -1938,11 +1939,11 @@ TSharedPtr<FStreamableHandle> UAssetManager::ChangeBundleStateForPrimaryAssets(
 			if (LoadParams.DownloadParams)
 			{
 				LoadParamsWithoutCallbacks.DownloadParams = CopyTemp(*LoadParams.DownloadParams);
-				NewHandle = LoadAssetList(MoveTemp(PathsToLoad), MoveTemp(LoadParamsWithoutCallbacks), FString(DebugName));
+				NewHandle = LoadAssetList(MoveTemp(PathsToLoad), MoveTemp(LoadParamsWithoutCallbacks), FString(DebugName), CopyTemp(Location));
 			}
 			else if (!bDownloadRequired)
 			{
-				NewHandle = LoadAssetList(MoveTemp(PathsToLoad), MoveTemp(LoadParamsWithoutCallbacks), FString(DebugName));
+				NewHandle = LoadAssetList(MoveTemp(PathsToLoad), MoveTemp(LoadParamsWithoutCallbacks), FString(DebugName), CopyTemp(Location));
 			}
 			else
 			{
@@ -1952,7 +1953,7 @@ TSharedPtr<FStreamableHandle> UAssetManager::ChangeBundleStateForPrimaryAssets(
 				StreamableLoadParams.TargetsToStream = MoveTemp(PathsToLoad);
 				NewHandle = StreamableManager.CreateErrorHandle(
 					FStreamableError(EStreamableErrorType::DownloadError, FString::Printf(TEXT("Expected DownloadParams for %s"), *PrimaryAssetId.ToString())),
-					MoveTemp(StreamableLoadParams), FString(DebugName));
+					MoveTemp(StreamableLoadParams), FString(DebugName), CopyTemp(Location));
 			}
 
 			if (!NewHandle.IsValid())
@@ -2004,7 +2005,12 @@ TSharedPtr<FStreamableHandle> UAssetManager::ChangeBundleStateForPrimaryAssets(
 		// If multiple handles or we have an old handle, need to make wrapper handle
 		NewHandles.Append(ExistingHandles);
 
-		ReturnHandle = StreamableManager.CreateCombinedHandle(NewHandles, FString::Printf(TEXT("%s CreateCombinedHandle"), *GetName()));
+		ReturnHandle = StreamableManager.CreateCombinedHandle(
+			NewHandles,
+			FStreamableHandle::HandleDebugName_CombinedHandle,
+			EStreamableManagerCombinedHandleOptions::None,
+			FStreamableAsyncLoadParams{},
+			CopyTemp(Location));
 	}
 	else if (NewHandles.Num() == 1)
 	{
@@ -2051,14 +2057,14 @@ TSharedPtr<FStreamableHandle> UAssetManager::ChangeBundleStateForPrimaryAssets(
 	return ReturnHandle;
 }
 
-TSharedPtr<FStreamableHandle> UAssetManager::ChangeBundleStateForMatchingPrimaryAssets(const TArray<FName>& NewBundles, const TArray<FName>& OldBundles, FAssetManagerLoadParams&& LoadParams)
+TSharedPtr<FStreamableHandle> UAssetManager::ChangeBundleStateForMatchingPrimaryAssets(const TArray<FName>& NewBundles, const TArray<FName>& OldBundles, FAssetManagerLoadParams&& LoadParams, UE::FSourceLocation Location)
 {
 	TArray<FPrimaryAssetId> AssetsToChange;
 
 	if (GetPrimaryAssetsWithBundleState(AssetsToChange, TArray<FPrimaryAssetType>(), OldBundles))
 	{
 		// This will call delegate when done
-		return ChangeBundleStateForPrimaryAssets(AssetsToChange, NewBundles, OldBundles, false, MoveTemp(LoadParams));
+		return ChangeBundleStateForPrimaryAssets(AssetsToChange, NewBundles, OldBundles, false, MoveTemp(LoadParams), MoveTemp(Location));
 	}
 
 	// Nothing to transition, call delegate now
@@ -2133,28 +2139,28 @@ bool UAssetManager::GetPrimaryAssetLoadList(TArray<FSoftObjectPath>& OutAssetLoa
 	return NameData != nullptr;
 }
 
-TSharedPtr<FStreamableHandle> UAssetManager::PreloadPrimaryAssets(const TArray<FPrimaryAssetId>& AssetsToLoad, const TArray<FName>& LoadBundles, bool bLoadRecursive, FAssetManagerLoadParams&& LoadParams)
+TSharedPtr<FStreamableHandle> UAssetManager::PreloadPrimaryAssets(const TArray<FPrimaryAssetId>& AssetsToLoad, const TArray<FName>& LoadBundles, bool bLoadRecursive, FAssetManagerLoadParams&& LoadParams, UE::FSourceLocation Location)
 {
 	TStringBuilder<256> DebugValid;
 	TStringBuilder<256> DebugInvalid;
 
 	TArray<FSoftObjectPath> PathsToLoad = UE::AssetManager::Private::GetPreloadPrimaryAssetsLoadList(this, AssetsToLoad, LoadBundles, bLoadRecursive, DebugValid, DebugInvalid);
 
-	TSharedPtr<FStreamableHandle> ReturnHandle = LoadAssetList(MoveTemp(PathsToLoad), MoveTemp(LoadParams), FString(*DebugValid));
+	TSharedPtr<FStreamableHandle> ReturnHandle = LoadAssetList(MoveTemp(PathsToLoad), MoveTemp(LoadParams), FString(*DebugValid), MoveTemp(Location));
 
 	ensureMsgf(ReturnHandle.IsValid(), TEXT("Requested preload of Primary Assets failed: %s"), *DebugInvalid);
 
 	return ReturnHandle;
 }
 
-TSharedPtr<FStreamableHandle> UAssetManager::PredownloadPrimaryAssets(const TArray<FPrimaryAssetId>& AssetsToLoad, const TArray<FName>& LoadBundles, bool bLoadRecursive, FAssetManagerLoadParams&& LoadParams)
+TSharedPtr<FStreamableHandle> UAssetManager::PredownloadPrimaryAssets(const TArray<FPrimaryAssetId>& AssetsToLoad, const TArray<FName>& LoadBundles, bool bLoadRecursive, FAssetManagerLoadParams&& LoadParams, UE::FSourceLocation Location)
 {
 	TStringBuilder<256> DebugValid;
 	TStringBuilder<256> DebugInvalid;
 
 	TArray<FSoftObjectPath> PathsToLoad = UE::AssetManager::Private::GetPreloadPrimaryAssetsLoadList(this, AssetsToLoad, LoadBundles, bLoadRecursive, DebugValid, DebugInvalid);
 
-	TSharedPtr<FStreamableHandle> ReturnHandle = DownloadAssetList(MoveTemp(PathsToLoad), MoveTemp(LoadParams), FString(*DebugValid));
+	TSharedPtr<FStreamableHandle> ReturnHandle = DownloadAssetList(MoveTemp(PathsToLoad), MoveTemp(LoadParams), FString(*DebugValid), MoveTemp(Location));
 
 	ensureMsgf(ReturnHandle.IsValid(), TEXT("Requested predownload of Primary Assets failed: %s"), *DebugInvalid);
 
@@ -2198,21 +2204,21 @@ void UAssetManager::OnAssetStateChangeCompleted(const TSharedPtr<FStreamableHand
 	}
 }
 
-TSharedPtr<FStreamableHandle> UAssetManager::LoadPrimaryAssets(const TArray<FPrimaryAssetId>& AssetsToLoad, const TArray<FName>& LoadBundles, FAssetManagerLoadParams&& LoadParams)
+TSharedPtr<FStreamableHandle> UAssetManager::LoadPrimaryAssets(const TArray<FPrimaryAssetId>& AssetsToLoad, const TArray<FName>& LoadBundles, FAssetManagerLoadParams&& LoadParams, UE::FSourceLocation Location)
 {
-	return ChangeBundleStateForPrimaryAssets(AssetsToLoad, LoadBundles, TArray<FName>(), true, MoveTemp(LoadParams));
+	return ChangeBundleStateForPrimaryAssets(AssetsToLoad, LoadBundles, TArray<FName>(), true, MoveTemp(LoadParams), MoveTemp(Location));
 }
 
-TSharedPtr<FStreamableHandle> UAssetManager::LoadPrimaryAsset(const FPrimaryAssetId& AssetToLoad, const TArray<FName>& LoadBundles, FAssetManagerLoadParams&& LoadParams)
+TSharedPtr<FStreamableHandle> UAssetManager::LoadPrimaryAsset(const FPrimaryAssetId& AssetToLoad, const TArray<FName>& LoadBundles, FAssetManagerLoadParams&& LoadParams, UE::FSourceLocation Location)
 {
-	return LoadPrimaryAssets(TArray<FPrimaryAssetId>{AssetToLoad}, LoadBundles, MoveTemp(LoadParams));
+	return LoadPrimaryAssets(TArray<FPrimaryAssetId>{AssetToLoad}, LoadBundles, MoveTemp(LoadParams), MoveTemp(Location));
 }
 
-TSharedPtr<FStreamableHandle> UAssetManager::LoadPrimaryAssetsWithType(FPrimaryAssetType PrimaryAssetType, const TArray<FName>& LoadBundles, FAssetManagerLoadParams&& LoadParams)
+TSharedPtr<FStreamableHandle> UAssetManager::LoadPrimaryAssetsWithType(FPrimaryAssetType PrimaryAssetType, const TArray<FName>& LoadBundles, FAssetManagerLoadParams&& LoadParams, UE::FSourceLocation Location)
 {
 	TArray<FPrimaryAssetId> Assets;
 	GetPrimaryAssetIdList(PrimaryAssetType, Assets);
-	return LoadPrimaryAssets(Assets, LoadBundles, MoveTemp(LoadParams));
+	return LoadPrimaryAssets(Assets, LoadBundles, MoveTemp(LoadParams), MoveTemp(Location));
 }
 
 TSharedPtr<FStreamableHandle> UAssetManager::GetPrimaryAssetHandle(const FPrimaryAssetId& PrimaryAssetId, bool bForceCurrent, TArray<FName>* Bundles) const
@@ -2377,7 +2383,7 @@ namespace UE::AssetManager::Private
 	}
 }
 
-TSharedPtr<FStreamableHandle> UAssetManager::LoadAssetListInternal(TArray<FSoftObjectPath>&& AssetList, FAssetManagerLoadParams&& LoadParams, FString&& DebugName)
+TSharedPtr<FStreamableHandle> UAssetManager::LoadAssetList(TArray<FSoftObjectPath>&& AssetList, FAssetManagerLoadParams&& LoadParams, FString&& DebugName, UE::FSourceLocation&& Location)
 {
 	TArray<int32> MissingChunks, ErrorChunks;
 
@@ -2399,7 +2405,7 @@ TSharedPtr<FStreamableHandle> UAssetManager::LoadAssetListInternal(TArray<FSoftO
 	// SynchronousLoad doesn't make sense if we are going to download packages	
 	if (!LoadParams.DownloadParams && bShouldUseSynchronousLoad && MissingChunks.Num() == 0)
 	{
-		NewHandle = StreamableManager.RequestSyncLoad(MoveTemp(AssetList), false, MoveTemp(DebugName));
+		NewHandle = StreamableManager.RequestSyncLoad(MoveTemp(AssetList), false, MoveTemp(DebugName), MoveTemp(Location));
 		FStreamableHandle::ExecuteDelegate(MoveTemp(LoadParams.OnComplete));
 	}
 	else
@@ -2409,7 +2415,7 @@ TSharedPtr<FStreamableHandle> UAssetManager::LoadAssetListInternal(TArray<FSoftO
 		Params.bManageActiveHandle = false;
 		Params.bStartStalled = MissingChunks.Num() > 0;
 
-		NewHandle = StreamableManager.RequestAsyncLoad(MoveTemp(Params), MoveTemp(DebugName));
+		NewHandle = StreamableManager.RequestAsyncLoad(MoveTemp(Params), MoveTemp(DebugName), MoveTemp(Location));
 
 		if (MissingChunks.Num() > 0 && NewHandle.IsValid())
 		{
@@ -2420,7 +2426,7 @@ TSharedPtr<FStreamableHandle> UAssetManager::LoadAssetListInternal(TArray<FSoftO
 	return NewHandle;
 }
 
-TSharedPtr<FStreamableHandle> UAssetManager::DownloadAssetListInternal(TArray<FSoftObjectPath>&& AssetList, FAssetManagerLoadParams&& LoadParams, FString&& DebugName)
+TSharedPtr<FStreamableHandle> UAssetManager::DownloadAssetList(TArray<FSoftObjectPath>&& AssetList, FAssetManagerLoadParams&& LoadParams, FString&& DebugName, UE::FSourceLocation&& Location)
 {
 	if (!LoadParams.DownloadParams)
 	{
@@ -2433,7 +2439,7 @@ TSharedPtr<FStreamableHandle> UAssetManager::DownloadAssetListInternal(TArray<FS
 	Params.bManageActiveHandle = false;
 	Params.DownloadParams->bDownloadOnly = true;
 
-	return StreamableManager.RequestAsyncLoad(MoveTemp(Params), MoveTemp(DebugName));
+	return StreamableManager.RequestAsyncLoad(MoveTemp(Params), MoveTemp(DebugName), MoveTemp(Location));
 }
 
 FAssetBundleEntry UAssetManager::GetAssetBundleEntry(const FPrimaryAssetId& BundleScope, FName BundleName) const
diff --git a/Engine/Source/Runtime/Engine/Private/StreamableManager.cpp b/Engine/Source/Runtime/Engine/Private/StreamableManager.cpp
--- a/Engine/Source/Runtime/Engine/Private/StreamableManager.cpp
+++ b/Engine/Source/Runtime/Engine/Private/StreamableManager.cpp
@@ -542,8 +542,9 @@ namespace UE::StreamableManager::Private
 	}
 }
 
-FStreamableHandle::FStreamableHandle(FPrivateToken)
-	: Priority(0)
+FStreamableHandle::FStreamableHandle(FPrivateToken, UE::FSourceLocation&& InLocation)
+	: Location(MoveTemp(InLocation))
+	, Priority(0)
 	, StreamablesLoading(0)
 	, OwningManager(nullptr)
 #if WITH_EDITOR
@@ -1589,6 +1590,10 @@ void FStreamableManager::SetManagerName(FString InName)
 
 bool FStreamableManager::GetReferencerPropertyName(UObject* Object, FString& OutPropertyName) const
 {
+	if (ShouldStripDebugName())
+	{
+		return false;
+	}
 	if (Object)
 	{
 		// Report the first handle
@@ -1598,15 +1603,26 @@ bool FStreamableManager::GetReferencerPropertyName(UObject* Object, FString& Out
 			check(Existing);
 			if (Existing->Target == Object)
 			{
+				TStringBuilder<1024> PropName;
+
 				if (Existing->ActiveHandles.Num() > 0)
 				{
-					OutPropertyName = Existing->ActiveHandles[0]->GetDebugName();
-					return true;
+					PropName.Appendf(TEXT("Active[1/%d]:"), Existing->ActiveHandles.Num());
+					PropName << Existing->ActiveHandles[0]->GetDebugName();
+					PropName << TEXT(" @ ");
+					PropName << UE::SourceLocation::ToFileAndLineString(Existing->ActiveHandles[0]->Location);
+				}
+				else if (Existing->LoadingHandles.Num() > 0)
+				{
+					PropName.Appendf(TEXT("Loading[1/%d]:"), Existing->LoadingHandles.Num());
+					PropName << Existing->LoadingHandles[0]->GetDebugName();
+					PropName << TEXT(" @ ");
+					PropName << UE::SourceLocation::ToFileAndLineString(Existing->LoadingHandles[0]->Location);
 				}
 
-				if (Existing->LoadingHandles.Num() > 0)
+				if (PropName.Len() > 0)
 				{
-					OutPropertyName = Existing->LoadingHandles[0]->GetDebugName();
+					OutPropertyName = PropName;
 					return true;
 				}
 			}
@@ -1748,18 +1764,21 @@ bool FStreamableManager::ShouldStripDebugName()
 	return bShouldStripDebugName;
 }
 
-TSharedPtr<FStreamableHandle> FStreamableManager::RequestAsyncLoadInternal(FStreamableAsyncLoadParams&& Params, FString&& DebugName)
+TSharedPtr<FStreamableHandle> FStreamableManager::RequestAsyncLoad(FStreamableAsyncLoadParams&& Params, FString DebugName, UE::FSourceLocation Location)
 {
 	LLM_SCOPE(ELLMTag::StreamingManager);
 
 	// Schedule a new callback, this will get called when all related async loads are completed
-	TSharedRef<FStreamableHandle> NewRequest = MakeShared<FStreamableHandle>(FStreamableHandle::FPrivateToken());
+	TSharedRef<FStreamableHandle> NewRequest = MakeShared<FStreamableHandle>(FStreamableHandle::FPrivateToken(), MoveTemp(Location));
 	NewRequest->CompleteDelegate = MoveTemp(Params.OnComplete);
 	NewRequest->CancelDelegate = MoveTemp(Params.OnCancel);
 	NewRequest->UpdateDelegate = MoveTemp(Params.OnUpdate);
 	NewRequest->OwningManager = this;
 	NewRequest->RequestedAssets = MoveTemp(Params.TargetsToStream);
-	NewRequest->DebugName = MoveTemp(DebugName);
+	if (!ShouldStripDebugName())
+	{
+		NewRequest->DebugName = MoveTemp(DebugName);
+	}
 
 	NewRequest->Priority = Params.Priority;
 #if UE_WITH_PACKAGE_ACCESS_TRACKING
@@ -1880,7 +1899,7 @@ TSharedPtr<FStreamableHandle> FStreamableManager::RequestAsyncLoadInternal(FStre
 	return NewRequest;
 }
 
-TSharedPtr<FStreamableHandle> FStreamableManager::RequestSyncLoadInternal(TArray<FSoftObjectPath>&& TargetsToStream, bool bManageActiveHandle, FString&& DebugName)
+TSharedPtr<FStreamableHandle> FStreamableManager::RequestSyncLoadInternal(TArray<FSoftObjectPath>&& TargetsToStream, bool bManageActiveHandle, FString&& DebugName, UE::FSourceLocation&& Location)
 {
 	// If in async loading thread or from callback always do sync as recursive tick is unsafe
 	// If in EDL always do sync as EDL internally avoids flushing
@@ -1893,7 +1912,7 @@ TSharedPtr<FStreamableHandle> FStreamableManager::RequestSyncLoadInternal(TArray
 	AsyncParams.bManageActiveHandle = bManageActiveHandle;
 
 	// Do an async load and wait to complete. In some cases this will do a sync load due to safety issues
-	TSharedPtr<FStreamableHandle> Request = RequestAsyncLoadInternal(MoveTemp(AsyncParams), MoveTemp(DebugName));
+	TSharedPtr<FStreamableHandle> Request = RequestAsyncLoad(MoveTemp(AsyncParams), MoveTemp(DebugName), MoveTemp(Location));
 
 	bForceSynchronousLoads = false;
 
@@ -1997,7 +2016,8 @@ UObject* FStreamableManager::LoadSynchronous(const FSoftObjectPath& Target, bool
 	TSharedPtr<FStreamableHandle> Request = RequestSyncLoadInternal(
 		TArray<FSoftObjectPath>{Target},
 		bManageActiveHandle,
-		FString::Printf(TEXT("LoadSynchronous from %s (%d)"), StringCast<TCHAR>(Location.GetFileName()).Get(), Location.GetLine()));
+		FString(TEXT("LoadSynchronous")),
+		MoveTemp(Location));
 
 	if (RequestHandlePointer)
 	{
@@ -2270,17 +2290,30 @@ void FStreamableManager::Unload(const FSoftObjectPath& Target)
 	}
 }
 
-TSharedPtr<FStreamableHandle> FStreamableManager::CreateCombinedHandle(const TConstArrayView<TSharedPtr<FStreamableHandle>>& ChildHandles, const FString& DebugName, EStreamableManagerCombinedHandleOptions Options)
+TSharedPtr<FStreamableHandle> FStreamableManager::CreateCombinedHandle(
+	TConstArrayView<TSharedPtr<FStreamableHandle>> ChildHandles,
+	FString DebugName,
+	EStreamableManagerCombinedHandleOptions Options,
+	FStreamableAsyncLoadParams&& Params,
+	UE::FSourceLocation Location)
 {
 	if (!ensure(ChildHandles.Num() > 0))
 	{
+		// Complete immediately if nothing to load
+		Params.OnComplete.ExecuteIfBound(nullptr);
 		return nullptr;
 	}
 
-	TSharedRef<FStreamableHandle> NewRequest = MakeShared<FStreamableHandle>(FStreamableHandle::FPrivateToken());
+	TSharedRef<FStreamableHandle> NewRequest = MakeShared<FStreamableHandle>(FStreamableHandle::FPrivateToken(), MoveTemp(Location));
 	NewRequest->OwningManager = this;
 	NewRequest->bIsCombinedHandle = true;
-	NewRequest->DebugName = DebugName;
+	NewRequest->CompleteDelegate = MoveTemp(Params.OnComplete);
+	NewRequest->CancelDelegate = MoveTemp(Params.OnCancel);
+	NewRequest->UpdateDelegate = MoveTemp(Params.OnUpdate);
+	if (!ShouldStripDebugName())
+	{
+		NewRequest->DebugName = MoveTemp(DebugName);
+	}
 
 	static const auto RemapHandleParentRelationship = [](TSharedPtr<FStreamableHandle>& IndividualHandle, TSharedRef<FStreamableHandle>& MergedHandle)
 	{
@@ -2397,13 +2430,16 @@ TSharedPtr<FStreamableHandle> FStreamableManager::CreateCombinedHandle(const TCo
 	return NewRequest;
 }
 
-TSharedPtr<FStreamableHandle> FStreamableManager::CreateErrorHandle(FStreamableError&& Error, FStreamableAsyncLoadParams&& Params, const FString& DebugName)
+TSharedPtr<FStreamableHandle> FStreamableManager::CreateErrorHandle(FStreamableError&& Error, FStreamableAsyncLoadParams&& Params, FString DebugName, UE::FSourceLocation Location)
 {
-	TSharedRef<FStreamableHandle> NewRequest = MakeShared<FStreamableHandle>(FStreamableHandle::FPrivateToken());
+	TSharedRef<FStreamableHandle> NewRequest = MakeShared<FStreamableHandle>(FStreamableHandle::FPrivateToken(), MoveTemp(Location));
 	NewRequest->OwningManager = this;
 	NewRequest->RequestedAssets = MoveTemp(Params.TargetsToStream);
-	NewRequest->DebugName = DebugName;
 	NewRequest->Priority = Params.Priority;
+	if (!ShouldStripDebugName())
+	{
+		NewRequest->DebugName = MoveTemp(DebugName);
+	}
 
 	if (Params.DownloadParams)
 	{
