Iris
- Fixed missing case for UE-184828

diff --git a/Engine/Source/Runtime/Engine/Private/Net/Iris/ReplicationSystem/EngineReplicationBridge.cpp b/Engine/Source/Runtime/Engine/Private/Net/Iris/ReplicationSystem/EngineReplicationBridge.cpp
--- a/Engine/Source/Runtime/Engine/Private/Net/Iris/ReplicationSystem/EngineReplicationBridge.cpp
+++ b/Engine/Source/Runtime/Engine/Private/Net/Iris/ReplicationSystem/EngineReplicationBridge.cpp
@@ -507,8 +507,8 @@ void UEngineReplicationBridge::StopReplicatingActor(AActor* Actor, EEndPlayReaso
 	{
 		Flags |= EEndReplicationFlags::SkipPendingEndReplicationValidation;
 	}
-					
-	const bool bShouldCreateDestructionInfo = RefHandle.IsStatic() && bShouldDestroyObject && GetReplicationSystem()->IsServer();
+
+	const bool bShouldCreateDestructionInfo = RefHandle.IsStatic() && bShouldDestroyObject && GetReplicationSystem()->IsServer() && !GIsReconstructingBlueprintInstances && !UE::Net::ShouldIgnoreStaticActorDestruction();
 	if (bShouldCreateDestructionInfo)
 	{
 		const UObjectReplicationBridge::FDestructionParameters EndReplicationParameters
