When r.InstancedSkinnedMeshes.SampledBounds = 1, UInstancedSkinnedMeshComponent::CalcBounds could return a zeroed bounds in the case where FBoxSphereBounds::Builder is invalid. Instead, fallback to the default implementation.

diff --git a/Engine/Source/Runtime/Engine/Private/Components/InstancedSkinnedMeshComponent.cpp b/Engine/Source/Runtime/Engine/Private/Components/InstancedSkinnedMeshComponent.cpp
--- a/Engine/Source/Runtime/Engine/Private/Components/InstancedSkinnedMeshComponent.cpp
+++ b/Engine/Source/Runtime/Engine/Private/Components/InstancedSkinnedMeshComponent.cpp
@@ -1477,20 +1477,22 @@ FBoxSphereBounds UInstancedSkinnedMeshComponent::CalcBounds(const FTransform& Lo
 				}
 			}
 
-			return BoundsBuilder;
+			// Only use bounds if valid, else continue with implementation not using AnimBankItems
+			if (BoundsBuilder.IsValid())
+			{
+				return BoundsBuilder;
+			}
 		}
-		else
+		
+		const FBox InstanceBounds = SkinnedAssetPtr->GetBounds().GetBox();
+		if (InstanceBounds.IsValid)
 		{
-			const FBox InstanceBounds = SkinnedAssetPtr->GetBounds().GetBox();
-			if (InstanceBounds.IsValid)
+			for (int32 InstanceIndex = 0; InstanceIndex < InstanceData.Num(); InstanceIndex++)
 			{
-				for (int32 InstanceIndex = 0; InstanceIndex < InstanceData.Num(); InstanceIndex++)
-				{
-					BoundsBuilder += InstanceBounds.TransformBy(FTransform(InstanceData[InstanceIndex].Transform) * LocalToWorld);
-				}
-
-				return BoundsBuilder;
+				BoundsBuilder += InstanceBounds.TransformBy(FTransform(InstanceData[InstanceIndex].Transform) * LocalToWorld);
 			}
+
+			return BoundsBuilder;
 		}
 	}
 
