Modify some virtual texture default settings in line with the work to enable by default.
Seperated out streaming VT budget.
Set a value for throttling release of allocated VTs.

diff --git a/Engine/Source/Runtime/Engine/Private/VT/VirtualTextureScalability.cpp b/Engine/Source/Runtime/Engine/Private/VT/VirtualTextureScalability.cpp
--- a/Engine/Source/Runtime/Engine/Private/VT/VirtualTextureScalability.cpp
+++ b/Engine/Source/Runtime/Engine/Private/VT/VirtualTextureScalability.cpp
@@ -26,16 +26,16 @@ namespace VirtualTextureScalability
 
 	static TAutoConsoleVariable<int32> CVarVTMaxUploadsPerFrameStreaming(
 		TEXT("r.VT.MaxUploadsPerFrame.Streaming"),
-		0,
+		32,
 		TEXT("If positive, max number of page uploads per frame in game for streaming VT. Negative means no limit.\n")
-		TEXT("If zero, SVTs won't be budgeted separately. They will be limited by r.VT.MaxUploadsPerFrame along with other types of VTs. This is the old behavior.\n")
+		TEXT("If zero, SVTs won't be budgeted separately. They will be limited by r.VT.MaxUploadsPerFrame along with other types of VTs. This is the legacy behavior.\n")
 		TEXT("This limit should be high if streaming pages is slow so that I/O requests are not throttled which can cause long delays to acquire page data."),
 		ECVF_RenderThreadSafe | ECVF_Scalability
 	);
 
 	static TAutoConsoleVariable<int32> CVarMaxPagesProducedPerFrame(
 		TEXT("r.VT.MaxTilesProducedPerFrame"),
-		30,
+		32,
 		TEXT("Max number of pages that can be produced per frame"),
 		ECVF_RenderThreadSafe | ECVF_Scalability
 	);
@@ -58,8 +58,8 @@ namespace VirtualTextureScalability
 
 	static TAutoConsoleVariable<int32> CVarVTMaxReleasedPerFrame(
 		TEXT("r.VT.MaxReleasedPerFrame"),
-		0,
-		TEXT("Max number of allocated virtual textures to release per frame"),
+		5,
+		TEXT("Max number of allocated virtual textures to release per frame in non editor builds."),
 		ECVF_RenderThreadSafe | ECVF_Scalability
 	);
 
