Moved IsExpressionClassPermitted function from HLSLMaterialTranslator.cpp to MaterialExpressionSettings.cpp, a more appropriate location.

diff --git a/Engine/Source/Runtime/Engine/Private/Materials/HLSLMaterialTranslator.cpp b/Engine/Source/Runtime/Engine/Private/Materials/HLSLMaterialTranslator.cpp
--- a/Engine/Source/Runtime/Engine/Private/Materials/HLSLMaterialTranslator.cpp
+++ b/Engine/Source/Runtime/Engine/Private/Materials/HLSLMaterialTranslator.cpp
@@ -143,28 +143,6 @@ static int32 GetLWCTruncateMode()
 	return LWCTruncateValue;
 }
 
-bool IsExpressionClassPermitted(const UClass* const Class)
-{
-	if (Class)
-	{
-		// Always allow cooked content
-		if (Class->bCooked)
-		{
-			return true;
-		}
-
-#if WITH_EDITOR
-		// Uncooked content must be explicitly allowed if in restrictive mode
-		const bool bIsClassPathAllowed = FMaterialExpressionSettings::Get()->IsClassPathAllowed(FTopLevelAssetPath(Class->GetPathName()));
-		if (!bIsClassPathAllowed)
-		{
-			return false;
-		}
-#endif
-	}
-	return true;
-}
-
 #if WITH_EDITORONLY_DATA
 
 static int32 GAnalyticDerivEnabled = 1;
diff --git a/Engine/Source/Runtime/Engine/Private/Materials/HLSLMaterialTranslator.h b/Engine/Source/Runtime/Engine/Private/Materials/HLSLMaterialTranslator.h
--- a/Engine/Source/Runtime/Engine/Private/Materials/HLSLMaterialTranslator.h
+++ b/Engine/Source/Runtime/Engine/Private/Materials/HLSLMaterialTranslator.h
@@ -64,12 +64,6 @@ namespace UE::DerivedData
 	class FRequestOwner;
 }
 
-/**
- * Returns whether the specified class of material expression is permitted.
- * For instance, custom expressions are not permitted in certain UE editor configurations for client generated materials.
- */
-bool IsExpressionClassPermitted(const UClass* const Class);
-
 #if WITH_EDITORONLY_DATA
 
 enum EMaterialExpressionVisitResult
diff --git a/Engine/Source/Runtime/Engine/Private/Materials/MaterialExpressionSettings.cpp b/Engine/Source/Runtime/Engine/Private/Materials/MaterialExpressionSettings.cpp
--- a/Engine/Source/Runtime/Engine/Private/Materials/MaterialExpressionSettings.cpp
+++ b/Engine/Source/Runtime/Engine/Private/Materials/MaterialExpressionSettings.cpp
@@ -1,6 +1,7 @@
 // Copyright Epic Games, Inc. All Rights Reserved.
 
 #include "MaterialExpressionSettings.h"
+#include "UObject/Class.h"
 
 #if WITH_EDITOR
 
@@ -37,5 +38,26 @@ bool FMaterialExpressionSettings::HasClassPathFiltering() const
 	return IsClassPathAllowedDelegates.Num() > 0;
 }
 
+bool IsExpressionClassPermitted(const UClass* const Class)
+{
+	if (Class)
+	{
+		// Always allow cooked content
+		if (Class->bCooked)
+		{
+			return true;
+		}
+
+#if WITH_EDITOR
+		// Uncooked content must be explicitly allowed if in restrictive mode
+		const bool bIsClassPathAllowed = FMaterialExpressionSettings::Get()->IsClassPathAllowed(FTopLevelAssetPath(Class->GetPathName()));
+		if (!bIsClassPathAllowed)
+		{
+			return false;
+		}
+#endif
+	}
+	return true;
+}
 
 #endif // WITH_EDITOR
diff --git a/Engine/Source/Runtime/Engine/Private/Materials/MaterialExpressions.cpp b/Engine/Source/Runtime/Engine/Private/Materials/MaterialExpressions.cpp
--- a/Engine/Source/Runtime/Engine/Private/Materials/MaterialExpressions.cpp
+++ b/Engine/Source/Runtime/Engine/Private/Materials/MaterialExpressions.cpp
@@ -317,6 +317,7 @@
 #include "Materials/MaterialInstanceConstant.h"
 #include "Curves/CurveLinearColorAtlas.h"
 #include "Curves/CurveLinearColor.h"
+#include "MaterialExpressionSettings.h"
 
 #define LOCTEXT_NAMESPACE "MaterialExpression"
 
diff --git a/Engine/Source/Runtime/Engine/Public/MaterialExpressionSettings.h b/Engine/Source/Runtime/Engine/Public/MaterialExpressionSettings.h
--- a/Engine/Source/Runtime/Engine/Public/MaterialExpressionSettings.h
+++ b/Engine/Source/Runtime/Engine/Public/MaterialExpressionSettings.h
@@ -32,4 +32,10 @@ private:
 	TMap<FName, FOnIsClassPathAllowed> IsClassPathAllowedDelegates;
 };
 
+/**
+ * Returns whether the specified class of material expression is permitted.
+ * For instance, custom expressions are not permitted in certain UE editor configurations for client generated materials.
+ */
+bool IsExpressionClassPermitted(const UClass* const Class);
+
 #endif // WITH_EDITOR
diff --git a/Engine/Source/Runtime/Engine/Public/Materials/MaterialExpression.h b/Engine/Source/Runtime/Engine/Public/Materials/MaterialExpression.h
--- a/Engine/Source/Runtime/Engine/Public/Materials/MaterialExpression.h
+++ b/Engine/Source/Runtime/Engine/Public/Materials/MaterialExpression.h
@@ -290,10 +290,10 @@ class UMaterialExpression : public UObject
 #endif // WITH_EDITOR
 
 	/**
-	* Returns whether the creation of this expression should be allowed in a given material or material function
+	* Returns whether the creation of this expression is allowed in a given material or material function
 	*
 	* @param MaterialOrFunction - The material or function for this expression to test against
-	* @return True if this expression should be created in this material or function, false if it should not.
+	* @return True if this expression can be created in this material or function, false otherwise.
 	*/
 	ENGINE_API virtual bool IsAllowedIn(const UObject* MaterialOrFunction) const;
 
