Fix for bad behavior of Layered Blend Per Bone in root space on first frame. Nand propagation

diff --git a/Engine/Source/Runtime/Engine/Private/Animation/AnimationRuntime.cpp b/Engine/Source/Runtime/Engine/Private/Animation/AnimationRuntime.cpp
--- a/Engine/Source/Runtime/Engine/Private/Animation/AnimationRuntime.cpp
+++ b/Engine/Source/Runtime/Engine/Private/Animation/AnimationRuntime.cpp
@@ -1981,6 +1981,9 @@ void FAnimationRuntime::BlendPosesPerBoneFilter(FCompactPose& BasePose, const TA
 			else
 			{
 				AccumulateMeshSpaceScale(PoseIndex, BoneIndex, FVector(1.0f), FVector(1.0f));
+				// Ignore the root bone for component-space calculation. This forces our rotations to be in root-space.
+				SourceRotations[BoneIndex] = FQuat::Identity;
+				TargetRotations[BoneIndex] = FQuat::Identity;
 			}
 
 			const FTransform& BaseAtom = BasePose[BoneIndex];
@@ -2030,6 +2033,12 @@ void FAnimationRuntime::BlendPosesPerBoneFilter(FCompactPose& BasePose, const TA
 			{
 				AccumulateMeshSpaceRotation(PoseIndex, BoneIndex, SourceRotations[ParentIndex], TargetRotations[ParentIndex]);
 			}
+			else
+			{
+				// Ignore the root bone for component-space calculation. This forces our rotations to be in root-space.
+				SourceRotations[BoneIndex] = FQuat::Identity;
+				TargetRotations[BoneIndex] = FQuat::Identity;
+			}
 
 			const FTransform& BaseAtom = BasePose[BoneIndex];
 			const FTransform& TargetAtom = BlendPoses[PoseIndex][BoneIndex];
