[UObject]

- Remove the deprecated UMetadata export from its package upon deserialization. Without this, users would often assert in loading code when iterating over the metadata export since it may have already been deleted by the GC (since we were removing the RF_Standalone flag). After conversion to FMetaData the UMetaData is no longer needed so we can garbage the object and remove the export from the linker.

#rnx

diff --git a/Engine/Source/Runtime/CoreUObject/Private/Serialization/AsyncLoading2.cpp b/Engine/Source/Runtime/CoreUObject/Private/Serialization/AsyncLoading2.cpp
--- a/Engine/Source/Runtime/CoreUObject/Private/Serialization/AsyncLoading2.cpp
+++ b/Engine/Source/Runtime/CoreUObject/Private/Serialization/AsyncLoading2.cpp
@@ -6719,17 +6719,7 @@ bool FAsyncPackage2::PreloadLinkerLoadExports(FAsyncLoadingThreadState2& ThreadS
 		FExportObject& ExportObject = Data.Exports[ExportIndex];
 		FObjectExport& LinkerExport = LinkerLoadState->Linker->ExportMap[ExportIndex];
 
-		// The linker export table can be patched during reinstantiation. We need to adjust our own export table if needed.
-		if (!LinkerExport.bExportLoadFailed)
-		{
-			if (ExportObject.Object != LinkerExport.Object)
-			{
-				UE_ASYNC_PACKAGE_LOG(Verbose, Desc, TEXT("PreloadLinkerLoadExports"), TEXT("Patching export %d: %s -> %s"), ExportIndex, *GetPathNameSafe(ExportObject.Object), *GetPathNameSafe(LinkerExport.Object));
-				ExportObject.Object = LinkerExport.Object;
-			}
-		}
-
-		if (UObject* Object = ExportObject.Object)
+		if (UObject* Object = LinkerExport.Object)
 		{
 			if (Object->HasAnyFlags(RF_NeedLoad))
 			{
@@ -6738,6 +6728,16 @@ bool FAsyncPackage2::PreloadLinkerLoadExports(FAsyncLoadingThreadState2& ThreadS
 				LinkerLoadState->Linker->Preload(Object);
 			}
 		}
+
+		// The linker export table can be patched during reinstantiation/serialization. We need to adjust our own export table if needed.
+		if (ExportObject.Object != LinkerExport.Object)
+		{
+			UE_ASYNC_PACKAGE_LOG(Verbose, Desc, TEXT("PreloadLinkerLoadExports"), TEXT("Patching export %d: %s -> %s"), ExportIndex, *GetPathNameSafe(ExportObject.Object), *GetPathNameSafe(LinkerExport.Object));
+			ExportObject.Object = LinkerExport.Object;
+			ExportObject.bWasFoundInMemory = !!ExportObject.Object;
+			ExportObject.bExportLoadFailed = LinkerExport.bExportLoadFailed;
+		}
+
 		if (ThreadState.IsTimeLimitExceeded(TEXT("SerializeLinkerLoadExports")))
 		{
 			return false;
diff --git a/Engine/Source/Runtime/CoreUObject/Private/UObject/MetaData.cpp b/Engine/Source/Runtime/CoreUObject/Private/UObject/MetaData.cpp
--- a/Engine/Source/Runtime/CoreUObject/Private/UObject/MetaData.cpp
+++ b/Engine/Source/Runtime/CoreUObject/Private/UObject/MetaData.cpp
@@ -5,6 +5,7 @@
 #include "HAL/IConsoleManager.h"
 #include "Algo/Transform.h"
 #include "UObject/Linker.h"
+#include "UObject/LinkerLoad.h"
 #include "UObject/Package.h"
 #include "UObject/EditorObjectVersion.h"
 #include "UObject/UObjectHash.h"
@@ -229,6 +230,15 @@ void UDEPRECATED_MetaData::Serialize(FArchive& Ar)
 			MetaData.RootMetaDataMap.Append(RootMetaDataMap);
 
 			ClearFlags(RF_Standalone);
+			MarkAsGarbage();
+
+			// This deprecated object is no longer needed now that we have converted it to FMetaData. However,
+			// if loading from disk, the linker might still be referring to the deprecated metadata export which
+			// may be GC'd during load. Thus, we remove the export to avoid issues during loading.
+			if (FLinkerLoad* Linker = GetLinker())
+			{
+				Linker->InvalidateExport(this, true /*bHideGarbageObjects*/);
+			}
 		}
 #endif
 	}
diff --git a/Engine/Source/Runtime/Engine/Private/Tests/Loading/AsyncLoadingTests.cpp b/Engine/Source/Runtime/Engine/Private/Tests/Loading/AsyncLoadingTests.cpp
--- a/Engine/Source/Runtime/Engine/Private/Tests/Loading/AsyncLoadingTests.cpp
+++ b/Engine/Source/Runtime/Engine/Private/Tests/Loading/AsyncLoadingTests.cpp
@@ -8,6 +8,7 @@
 #include "Async/ParallelFor.h"
 #include "AsyncLoadingTests_Shared.h"
 #include "StructUtils/UserDefinedStruct.h"
+#include "UObject/LinkerLoad.h"
 #include "UObject/UObjectHash.h"
 
 #if WITH_DEV_AUTOMATION_TESTS
@@ -63,5 +64,57 @@ bool FThreadSafeAsyncLoadingTest::RunTest(const FString& Parameters)
 	return true;
 }
 
+/**
+ * Ensure we can properly handle Serialize implementations that might invalidate exports during preload.
+ */
+IMPLEMENT_SIMPLE_AUTOMATION_TEST(FAsyncLoadingTestInvalidateExportDuringPreload, TEXT(TEST_NAME_ROOT ".InvalidateExportDuringPreload"), EAutomationTestFlags::EditorContext | EAutomationTestFlags::EngineFilter)
+bool FAsyncLoadingTestInvalidateExportDuringPreload::RunTest(const FString& Parameters)
+{
+	auto VerifyLoad = [this](FLoadingTestsScope& LoadingTestScope, bool bExpectToFindObject)
+		{
+			UPackage* Package = LoadPackage(nullptr, LoadingTestScope.PackagePath1, LOAD_None);
+			TestTrue(TEXT("The pacakge should load successfully"), Package != nullptr);
+			
+			// Exclude garbage objects as the GC won't have run yet but invalidated objects shdould be marked as garbage by this point
+			UAsyncLoadingTests_Shared* Object1 = FindObjectFast<UAsyncLoadingTests_Shared>(Package, LoadingTestScope.ObjectName, true/*bExactClass*/, EObjectFlags::RF_MirroredGarbage);
+			if (bExpectToFindObject)
+			{
+				TestTrue(TEXT("The object should have been loaded"), Object1 != nullptr);
+			}
+			else
+			{
+				TestTrue(TEXT("The object should not have been loaded"), Object1 == nullptr);
+			}
+		};
+
+	{
+		FLoadingTestsScope LoadingTestScope(this);
+
+
+		VerifyLoad(LoadingTestScope, true /*bExpectToFindObject*/);
+	}
+	
+	{
+		FLoadingTestsScope LoadingTestScope(this);
+
+		UAsyncLoadingTests_Shared::OnSerialize.BindLambda(
+			[this](FArchive& Ar, UAsyncLoadingTests_Shared* Object)
+			{
+				if (Ar.IsLoading())
+				{
+					if (FLinkerLoad* Linker = Object->GetLinker())
+					{
+						Object->MarkAsGarbage();
+						Linker->InvalidateExport(Object, true /*bHidesGarbageObjects*/);
+					}
+				}
+			}
+		);
+
+		VerifyLoad(LoadingTestScope, false /*bExpectToFindObject*/);
+	}
+	return true;
+}
+
 #undef TEST_NAME_ROOT
-#endif // WITH_DEV_AUTOMATION_TESTS
\ No newline at end of file
+#endif // WITH_DEV_AUTOMATION_TESTS
