Add float suffix to float literals generated by the Material Translator to workaround FXC internal error.

diff --git a/Engine/Source/Runtime/Engine/Private/Materials/HLSLMaterialTranslator.cpp b/Engine/Source/Runtime/Engine/Private/Materials/HLSLMaterialTranslator.cpp
--- a/Engine/Source/Runtime/Engine/Private/Materials/HLSLMaterialTranslator.cpp
+++ b/Engine/Source/Runtime/Engine/Private/Materials/HLSLMaterialTranslator.cpp
@@ -5647,25 +5647,25 @@ int32 FHLSLMaterialTranslator::NumericParameter(EMaterialParameterType Parameter
 int32 FHLSLMaterialTranslator::Constant(float X)
 {
 	FAddUniformExpressionScope Scope(this);
-	return AddUniformExpression(Scope, new FMaterialUniformExpressionConstant(FLinearColor(X,X,X,X),MCT_Float),MCT_Float,TEXT("%0.8f"),X);
+	return AddUniformExpression(Scope, new FMaterialUniformExpressionConstant(FLinearColor(X,X,X,X),MCT_Float),MCT_Float,TEXT("%0.8ff"),X);
 }
 
 int32 FHLSLMaterialTranslator::Constant2(float X,float Y)
 {
 	FAddUniformExpressionScope Scope(this);
-	return AddUniformExpression(Scope, new FMaterialUniformExpressionConstant(FLinearColor(X,Y,0,0),MCT_Float2),MCT_Float2,TEXT("MaterialFloat2(%0.8f,%0.8f)"),X,Y);
+	return AddUniformExpression(Scope, new FMaterialUniformExpressionConstant(FLinearColor(X,Y,0,0),MCT_Float2),MCT_Float2,TEXT("MaterialFloat2(%0.8ff,%0.8ff)"),X,Y);
 }
 
 int32 FHLSLMaterialTranslator::Constant3(float X,float Y,float Z)
 {
 	FAddUniformExpressionScope Scope(this);
-	return AddUniformExpression(Scope, new FMaterialUniformExpressionConstant(FLinearColor(X,Y,Z,0),MCT_Float3),MCT_Float3,TEXT("MaterialFloat3(%0.8f,%0.8f,%0.8f)"),X,Y,Z);
+	return AddUniformExpression(Scope, new FMaterialUniformExpressionConstant(FLinearColor(X,Y,Z,0),MCT_Float3),MCT_Float3,TEXT("MaterialFloat3(%0.8ff,%0.8ff,%0.8ff)"),X,Y,Z);
 }
 
 int32 FHLSLMaterialTranslator::Constant4(float X,float Y,float Z,float W)
 {
 	FAddUniformExpressionScope Scope(this);
-	return AddUniformExpression(Scope, new FMaterialUniformExpressionConstant(FLinearColor(X,Y,Z,W),MCT_Float4),MCT_Float4,TEXT("MaterialFloat4(%0.8f,%0.8f,%0.8f,%0.8f)"),X,Y,Z,W);
+	return AddUniformExpression(Scope, new FMaterialUniformExpressionConstant(FLinearColor(X,Y,Z,W),MCT_Float4),MCT_Float4,TEXT("MaterialFloat4(%0.8ff,%0.8ff,%0.8ff,%0.8ff)"),X,Y,Z,W);
 }
 
 int32 FHLSLMaterialTranslator::GenericConstant(const UE::Shader::FValue& Value)
