World Partition: move world external actor references gathering outside of GetAssetRegistryTags into PreSave to avoid invalid calls to StaticFindObjects resulting from serializing the world during save.
This data is only required when parsing data on disk, not for the world in-memory FAssetData creation, so it also reduces the time spent creating in-memory asset datas for the world.

diff --git a/Engine/Source/Runtime/Engine/Private/WorldPartition/WorldPartition.cpp b/Engine/Source/Runtime/Engine/Private/WorldPartition/WorldPartition.cpp
--- a/Engine/Source/Runtime/Engine/Private/WorldPartition/WorldPartition.cpp
+++ b/Engine/Source/Runtime/Engine/Private/WorldPartition/WorldPartition.cpp
@@ -40,7 +40,6 @@
 #include "LocationVolume.h"
 #include "Engine/ExternalObjectAndActorDependencyGatherer.h"
 #include "Engine/LevelScriptBlueprint.h"
-#include "ActorReferencesUtils.h"
 #include "WorldPartition/DataLayer/WorldDataLayersActorDesc.h"
 #include "WorldPartition/IWorldPartitionEditorModule.h"
 #include "WorldPartition/WorldPartitionLevelHelper.h"
@@ -599,6 +598,16 @@ void UWorldPartition::OnCookEvent(UE::Cook::ECookEvent CookEvent, UE::Cook::FCoo
 	}
 }
 
+void UWorldPartition::PreSave(FObjectPreSaveContext SaveContext)
+{
+	Super::PreSave(SaveContext);
+
+	// Gather world references outside of the save package scope to avoid invalid calls to StaticFindObject from serialization.
+	const ActorsReferencesUtils::FGetActorReferencesParams Params = ActorsReferencesUtils::FGetActorReferencesParams(GetWorld())
+		.SetRequiredFlags(RF_HasExternalPackage);
+	WorldExternalActorReferences = ActorsReferencesUtils::GetActorReferences(Params);
+}
+
 FName UWorldPartition::GetWorldPartitionEditorName() const
 {
 	if (SupportsStreaming())
@@ -2134,10 +2143,6 @@ void UWorldPartition::AppendAssetRegistryTags(FAssetRegistryTagsContext Context)
 	}
 
 	// Append world references so we can perform changelists validations without loading it
-	const ActorsReferencesUtils::FGetActorReferencesParams Params = ActorsReferencesUtils::FGetActorReferencesParams(GetWorld())
-		.SetRequiredFlags(RF_HasExternalPackage);
-	TArray<ActorsReferencesUtils::FActorReference> WorldExternalActorReferences = ActorsReferencesUtils::GetActorReferences(Params);
-		
 	if (WorldExternalActorReferences.Num())
 	{
 		FStringBuilderBase StringBuilder;
@@ -2150,6 +2155,8 @@ void UWorldPartition::AppendAssetRegistryTags(FAssetRegistryTagsContext Context)
 
 		static const FName NAME_WorldExternalActorsReferences(TEXT("WorldExternalActorsReferences"));
 		Context.AddTag(FAssetRegistryTag(NAME_WorldExternalActorsReferences, StringBuilder.ToString(), FAssetRegistryTag::TT_Hidden));
+
+		WorldExternalActorReferences.Empty();
 	}
 }
 
diff --git a/Engine/Source/Runtime/Engine/Public/WorldPartition/WorldPartition.h b/Engine/Source/Runtime/Engine/Public/WorldPartition/WorldPartition.h
--- a/Engine/Source/Runtime/Engine/Public/WorldPartition/WorldPartition.h
+++ b/Engine/Source/Runtime/Engine/Public/WorldPartition/WorldPartition.h
@@ -23,6 +23,7 @@
 #include "WorldPartition/WorldPartitionActorLoaderInterface.h"
 #include "WorldPartition/WorldPartitionEditorLoaderAdapter.h"
 #include "WorldPartition/WorldPartitionRuntimeCellTransformer.h"
+#include "ActorReferencesUtils.h"
 #include "ExternalDirtyActorsTracker.h"
 #include "PackageSourceControlHelper.h"
 #include "CookPackageSplitter.h"
@@ -210,6 +211,7 @@ public:
 	ENGINE_API virtual void PostEditChangeChainProperty(FPropertyChangedChainEvent& PropertyChangedEvent) override;
 	ENGINE_API virtual void PostDuplicate(EDuplicateMode::Type DuplicateMode) override;
 	ENGINE_API virtual void OnCookEvent(UE::Cook::ECookEvent CookEvent, UE::Cook::FCookEventContext& CookContext) override;
+	ENGINE_API virtual void PreSave(FObjectPreSaveContext SaveContext) override;
 #endif //WITH_EDITOR
 	ENGINE_API virtual void Serialize(FArchive& Ar) override;
 	ENGINE_API virtual UWorld* GetWorld() const override;
@@ -550,6 +552,11 @@ private:
 
 	/** Used to know if we need to recheck if the user should enable streaming based on world size. */
 	bool bShouldCheckEnableStreamingWarning;
+
+#if WITH_EDITOR
+	/** Holds the world actor references, filled in PreSave and used in AppendAssetRegistryTags. */
+	mutable TArray<ActorsReferencesUtils::FActorReference> WorldExternalActorReferences;
+#endif
 #endif
 
 private:
