World Partition: non-editor-only actors having only editor-only components are now treated as a single point for streaming, since we can't retrieve a runtime streaming bounds for them.

diff --git a/Engine/Source/Runtime/Engine/Private/Actor.cpp b/Engine/Source/Runtime/Engine/Private/Actor.cpp
--- a/Engine/Source/Runtime/Engine/Private/Actor.cpp
+++ b/Engine/Source/Runtime/Engine/Private/Actor.cpp
@@ -1458,10 +1458,14 @@ static bool HasComponentForceActorNonSpatiallyLoaded(const AActor* InActor)
 void AActor::GetStreamingBounds(FBox& OutRuntimeBounds, FBox& OutEditorBounds) const
 {
 	OutRuntimeBounds.Init();
-	ForEachStreamingRelevantComponent(this, false, [&OutRuntimeBounds](const UActorComponent* Component, const FBox& StreamingBound)
+	if (!ForEachStreamingRelevantComponent(this, false, [&OutRuntimeBounds](const UActorComponent* Component, const FBox& StreamingBound)
 	{
 		OutRuntimeBounds += StreamingBound;
-	});
+	}))
+	{
+		const FVector ActorLocation = GetActorLocation();
+		OutRuntimeBounds = FBox(ActorLocation, ActorLocation);
+	}
 
 	OutEditorBounds.Init();
 	ForEachStreamingRelevantComponent(this, true, [&OutEditorBounds](const UActorComponent* Component, const FBox& StreamingBound)
