Suppress StaticAnalysis warning which does not understand FReferenceFinder

diff --git a/Engine/Source/Runtime/Engine/Private/Blueprint.cpp b/Engine/Source/Runtime/Engine/Private/Blueprint.cpp
--- a/Engine/Source/Runtime/Engine/Private/Blueprint.cpp
+++ b/Engine/Source/Runtime/Engine/Private/Blueprint.cpp
@@ -1352,7 +1352,7 @@ void UBlueprint::ForceLoadGraph(UObject* Object)
 		ComponentCollector.FindReferences(InnerObject);
 	}
 
-	for(UObject* Obj : ExternalObjects)
+	for(UObject* Obj : ExternalObjects) //-V1078
 	{
 		ForceLoad(Obj);
 	}
