Better assert for when adding an editor-only actor descriptors to a level streaming cell.

diff --git a/Engine/Source/Runtime/Engine/Private/WorldPartition/WorldPartitionRuntimeLevelStreamingCell.cpp b/Engine/Source/Runtime/Engine/Private/WorldPartition/WorldPartitionRuntimeLevelStreamingCell.cpp
--- a/Engine/Source/Runtime/Engine/Private/WorldPartition/WorldPartitionRuntimeLevelStreamingCell.cpp
+++ b/Engine/Source/Runtime/Engine/Private/WorldPartition/WorldPartitionRuntimeLevelStreamingCell.cpp
@@ -291,7 +291,7 @@ void UWorldPartitionRuntimeLevelStreamingCell::AddActorToCell(const FStreamingGe
 {
 	// Non-spatially loaded actors coming from level instances are not moved into the persistent level in PIE but rather placed into an always loaded cell,
 	// because that would imply loading them through an instancing context.
-	check(!ActorDescView.GetActorIsEditorOnly() || ActorDescView.GetActorIsEditorOnlyLoadedInPIE());
+	checkf(!ActorDescView.GetActorIsEditorOnly() || ActorDescView.GetActorIsEditorOnlyLoadedInPIE(), TEXT("Invalid editor-only actor descriptor:\n\t%s"), *ActorDescView.ToString(FWorldPartitionActorDesc::EToStringMode::ForDiff));
 
 	const UActorDescContainerInstance* ContainerInstance = ActorDescView.GetContainerInstance();
 	check(ContainerInstance);
