Substrate - reorderd node when converting decal.

#rnx
#tests Material Editor

diff --git a/Engine/Source/Runtime/Engine/Private/Materials/Material.cpp b/Engine/Source/Runtime/Engine/Private/Materials/Material.cpp
--- a/Engine/Source/Runtime/Engine/Private/Materials/Material.cpp
+++ b/Engine/Source/Runtime/Engine/Private/Materials/Material.cpp
@@ -4069,7 +4069,7 @@ bool UMaterial::ConvertMaterialToSubstrateMaterial(bool bAllowEmptyMaterialUpdat
 					UMaterialExpressionSubstrateConvertToDecal* ConvertToDecalNode= NewObject<UMaterialExpressionSubstrateConvertToDecal>(this);
 					EditorOnly->ExpressionCollection.AddExpression(ConvertToDecalNode);
 					ConvertToDecalNode->Material = this;
-					ReplaceNodeAndMoveToTheRight(ConvertNode, ConvertToDecalNode);
+					ReplaceNodeAndMoveToTheRight(ConvertToDecalNode, ConvertNode);
 					ConvertToDecalNode->DecalMaterial.Connect(0, ConvertNode);
 
 					EditorOnly->FrontMaterial.Connect(0, ConvertToDecalNode);
