ue5 - Increased PawnRemoteViewPitch precision from 8 bits to 16 bits. Removed some replay-related code that could have been removed when replay interpolation mode was removed (CL 5316124 and a later one that removed ENetworkSmoothingMode::Replay entirely).

#tests recorded replay without changes, loaded old replay and new replay successfully with changes. PIE with multiple clients looked good.

diff --git a/Engine/Source/Runtime/Core/Public/Misc/EngineNetworkCustomVersion.h b/Engine/Source/Runtime/Core/Public/Misc/EngineNetworkCustomVersion.h
--- a/Engine/Source/Runtime/Core/Public/Misc/EngineNetworkCustomVersion.h
+++ b/Engine/Source/Runtime/Core/Public/Misc/EngineNetworkCustomVersion.h
@@ -94,6 +94,7 @@ struct FEngineNetworkCustomVersion
 		BeaconNetIDVariant = 39,				// Bump version for online beacon's BeaconAssignGUID control message switching to use FNetIDVariant
 		JoinNoPawn = 40,						// Bump version to support NMT_JoinNoPawn and NMT_JoinNoPawnSplit control messages.
 		ClientHandshakeId = 41,					// Bump version to support serialization of APlayerController::ClientHandeshakeId.
+		PawnRemoteViewPitchTo16Bit = 42,		// Bump version to support higher precision PawnRemoteViewPitch
 
 		// -----<new versions can be added above this line>-------------------------------------------------
 		VersionPlusOne,
diff --git a/Engine/Source/Runtime/Engine/Classes/GameFramework/CharacterMovementComponent.h b/Engine/Source/Runtime/Engine/Classes/GameFramework/CharacterMovementComponent.h
--- a/Engine/Source/Runtime/Engine/Classes/GameFramework/CharacterMovementComponent.h
+++ b/Engine/Source/Runtime/Engine/Classes/GameFramework/CharacterMovementComponent.h
@@ -3060,23 +3060,6 @@ inline FSavedMove_Character& FSavedMove_Character::operator=(FSavedMove_Characte
 PRAGMA_ENABLE_DEPRECATION_WARNINGS
 
 
-class FCharacterReplaySample
-{
-public:
-	FCharacterReplaySample() : RemoteViewPitch( 0 ), Time( 0.0f )
-	{
-	}
-
-	friend ENGINE_API FArchive& operator<<( FArchive& Ar, FCharacterReplaySample& V );
-
-	FVector			Location;
-	FRotator		Rotation;
-	FVector			Velocity;
-	FVector			Acceleration;
-	uint8			RemoteViewPitch;
-	float			Time;					// This represents time since replay started
-};
-
 class FNetworkPredictionData_Client_Character : public FNetworkPredictionData_Client, protected FNoncopyable
 {
 	using Super = FNetworkPredictionData_Client;
@@ -3175,9 +3158,6 @@ public:
 	FVector LastServerLocation;
 	float	SimulatedDebugDrawTime;
 
-	/** Array of replay samples that we use to interpolate between to get smooth location/rotation/velocity/ect */
-	TArray< FCharacterReplaySample > ReplaySamples;
-
 	/** Finds SavedMove index for given TimeStamp. Returns INDEX_NONE if not found (move has been already Acked or cleared). */
 	ENGINE_API int32 GetSavedMoveIndex(float TimeStamp) const;
 
diff --git a/Engine/Source/Runtime/Engine/Classes/GameFramework/Pawn.h b/Engine/Source/Runtime/Engine/Classes/GameFramework/Pawn.h
--- a/Engine/Source/Runtime/Engine/Classes/GameFramework/Pawn.h
+++ b/Engine/Source/Runtime/Engine/Classes/GameFramework/Pawn.h
@@ -112,12 +112,17 @@ public:
 	UPROPERTY(EditAnywhere, Category=Pawn)
 	EAutoPossessAI AutoPossessAI;
 
+protected:
 	/** Replicated so we can see where remote clients are looking. */
 	UPROPERTY(replicated)
-	uint8 RemoteViewPitch;
+	uint16 RemoteViewPitch16;
+
+public:
 
-	/** Playback of replays writes blended pitch to this, rather than the RemoteViewPitch. This is to avoid having to compress and interpolated value. */
-	float BlendedReplayViewPitch;
+	/** DEPRECTED, use RemoteViewPitch16 instead. */
+	UE_DEPRECATED(5.6, "This property is deprecated, use GetRemoteViewPitch()/RemoteViewPitch16 instead.")
+	UPROPERTY(replicated)
+	uint8 RemoteViewPitch;
 
 	/** Default class to use when pawn is controlled by AI. */
 	UPROPERTY(EditAnywhere, BlueprintReadWrite, meta=(DisplayName="AI Controller Class"), Category=Pawn)
@@ -206,14 +211,14 @@ public:
 	ENGINE_API virtual void PawnStartFire(uint8 FireModeNum = 0);
 
 	/** Replicated pitch value so we can see where remote clients are looking. */
-	FORCEINLINE uint8 GetRemoteViewPitch() const { return RemoteViewPitch; }
+	FORCEINLINE uint16 GetRemoteViewPitch() const { return RemoteViewPitch16; }
 
 	/**
 	 * Set Pawn ViewPitch, so we can see where remote clients are looking.
-	 * Maps 360.0 degrees into a byte
+	 * Maps 360.0 degrees into 16 bits
 	 * @param	NewRemoteViewPitch	Pitch component to replicate to remote (non owned) clients.
 	 */
-	ENGINE_API void SetRemoteViewPitch(float NewRemoteViewPitch);
+	ENGINE_API void SetRemoteViewPitch(double NewRemoteViewPitch);
 
 	/** Return Physics Volume for this Pawn */
 	UE_DEPRECATED(5.0, "GetPawnPhysicsVolume is deprecated. Please use GetPhysicsVolume instead.")
diff --git a/Engine/Source/Runtime/Engine/Private/Components/CharacterMovementComponent.cpp b/Engine/Source/Runtime/Engine/Private/Components/CharacterMovementComponent.cpp
--- a/Engine/Source/Runtime/Engine/Private/Components/CharacterMovementComponent.cpp
+++ b/Engine/Source/Runtime/Engine/Private/Components/CharacterMovementComponent.cpp
@@ -8148,22 +8148,6 @@ void UCharacterMovementComponent::SmoothCorrection(const FVector& OldLocation, c
 	}
 }
 
-FArchive& operator<<( FArchive& Ar, FCharacterReplaySample& V )
-{
-	SerializePackedVector<10, 24>( V.Location, Ar );
-	SerializePackedVector<10, 24>( V.Velocity, Ar );
-	SerializePackedVector<10, 24>( V.Acceleration, Ar );
-	V.Rotation.SerializeCompressed( Ar );
-	Ar << V.RemoteViewPitch;
-
-	if (Ar.IsSaving() || (!Ar.AtEnd() && !Ar.IsError()))
-	{
-		Ar << V.Time;
-	}
-
-	return Ar;
-}
-
 void UCharacterMovementComponent::SmoothClientPosition(float DeltaSeconds)
 {
 	if (!HasValidData() || NetworkSmoothingMode == ENetworkSmoothingMode::Disabled)
@@ -11725,11 +11709,6 @@ void UCharacterMovementComponent::ApplyWorldOffset(const FVector& InOffset, bool
 				PendingMove->StartLocation += InOffset;
 				PendingMove->SavedLocation += InOffset;
 			}
-
-			for (int32 i = 0; i < ClientData->ReplaySamples.Num(); i++)
-			{
-				ClientData->ReplaySamples[i].Location += InOffset;
-			}
 		}
 	}
 }
diff --git a/Engine/Source/Runtime/Engine/Private/Pawn.cpp b/Engine/Source/Runtime/Engine/Private/Pawn.cpp
--- a/Engine/Source/Runtime/Engine/Private/Pawn.cpp
+++ b/Engine/Source/Runtime/Engine/Private/Pawn.cpp
@@ -87,7 +87,6 @@ APawn::APawn(const FObjectInitializer& ObjectInitializer)
 	SetReplicatingMovement(true);
 	BaseEyeHeight = 64.0f;
 	AllowedYawError = 10.99f;
-	BlendedReplayViewPitch = 0.0f;
 	bCollideWhenPlacing = true;
 	SpawnCollisionHandlingMethod = ESpawnActorCollisionHandlingMethod::AdjustIfPossibleButDontSpawnIfColliding;
 	bGenerateOverlapEventsDuringLevelStreaming = true;
@@ -299,10 +298,10 @@ float APawn::GetDefaultHalfHeight() const
 	}
 }
 
-void APawn::SetRemoteViewPitch(float NewRemoteViewPitch)
+void APawn::SetRemoteViewPitch(double NewRemoteViewPitch)
 {
-	// Compress pitch to 1 byte
-	RemoteViewPitch = FRotator::CompressAxisToByte(NewRemoteViewPitch);
+	// Compress pitch to a 16-bit short
+	RemoteViewPitch16 = FRotator::CompressAxisToShort(NewRemoteViewPitch);
 }
 
 
@@ -1009,26 +1008,27 @@ FRotator APawn::GetBaseAimRotation() const
 	// If our Pitch is 0, then use a replicated view pitch
 	if( FMath::IsNearlyZero(POVRot.Pitch) )
 	{
-		if (BlendedReplayViewPitch != 0.0f)
+		const UWorld* World = GetWorld();
+		const UDemoNetDriver* DemoNetDriver = World ? World->GetDemoNetDriver() : nullptr;
+
+		if (DemoNetDriver && DemoNetDriver->IsPlaying() && (DemoNetDriver->GetPlaybackEngineNetworkProtocolVersion() < FEngineNetworkCustomVersion::PawnRemoteViewPitch))
+		{
+			// Use deprecated 8-bit value serialized from old replays
+			PRAGMA_DISABLE_DEPRECATION_WARNINGS
+			POVRot.Pitch = RemoteViewPitch;
+			PRAGMA_ENABLE_DEPRECATION_WARNINGS
+			POVRot.Pitch = POVRot.Pitch * 360.0f / 255.0f;
+		}
+		else if (DemoNetDriver && DemoNetDriver->IsPlaying() && (DemoNetDriver->GetPlaybackEngineNetworkProtocolVersion() < FEngineNetworkCustomVersion::PawnRemoteViewPitchTo16Bit))
 		{
-			// If we are in a replay and have a blended value for playback, use that
-			POVRot.Pitch = BlendedReplayViewPitch;
+			// Use deprecated 8-bit value serialized from old replays
+			PRAGMA_DISABLE_DEPRECATION_WARNINGS
+			POVRot.Pitch = FRotator::DecompressAxisFromByte(RemoteViewPitch);
+			PRAGMA_ENABLE_DEPRECATION_WARNINGS
 		}
 		else
 		{
-			// Else use the RemoteViewPitch
-			const UWorld* World = GetWorld();
-			const UDemoNetDriver* DemoNetDriver = World ? World->GetDemoNetDriver() : nullptr;
-
-			if (DemoNetDriver && DemoNetDriver->IsPlaying() && (DemoNetDriver->GetPlaybackEngineNetworkProtocolVersion() < FEngineNetworkCustomVersion::PawnRemoteViewPitch))
-			{
-				POVRot.Pitch = GetRemoteViewPitch();
-				POVRot.Pitch = POVRot.Pitch * 360.0f / 255.0f;
-			}
-			else
-			{
-				POVRot.Pitch = FRotator::DecompressAxisFromByte(GetRemoteViewPitch());
-			}
+			POVRot.Pitch = FRotator::DecompressAxisFromShort(RemoteViewPitch16);
 		}
 	}
 
@@ -1290,7 +1290,12 @@ void APawn::GetLifetimeReplicatedProps( TArray< FLifetimeProperty > & OutLifetim
 	DOREPLIFETIME( APawn, PlayerState ); 
 	DOREPLIFETIME( APawn, Controller );
 
-	DOREPLIFETIME_CONDITION( APawn, RemoteViewPitch, 	COND_SkipOwner );
+	// Deprecated RemoteViewPitch, still loaded for old replays
+	PRAGMA_DISABLE_DEPRECATION_WARNINGS
+	DOREPLIFETIME_CONDITION( APawn, RemoteViewPitch, COND_ReplayOnly );
+	PRAGMA_ENABLE_DEPRECATION_WARNINGS
+
+	DOREPLIFETIME_CONDITION( APawn, RemoteViewPitch16, 	COND_SkipOwner );
 }
 
 void APawn::SetController(AController* InController)
