Fix static analysis warning C6297: arithmetic overflow following CL 40952317.

diff --git a/Engine/Source/Runtime/Engine/Private/WorldPartition/WorldPartitionRuntimeSpatialHash.cpp b/Engine/Source/Runtime/Engine/Private/WorldPartition/WorldPartitionRuntimeSpatialHash.cpp
--- a/Engine/Source/Runtime/Engine/Private/WorldPartition/WorldPartitionRuntimeSpatialHash.cpp
+++ b/Engine/Source/Runtime/Engine/Private/WorldPartition/WorldPartitionRuntimeSpatialHash.cpp
@@ -1953,7 +1953,7 @@ ENGINE_API EWorldPartitionStreamingPerformance UWorldPartitionRuntimeSpatialHash
 		EWorldPartitionStreamingPerformance NonBlockingPerformance = EWorldPartitionStreamingPerformance::Good;
 		if (Cell->GetBlockOnSlowLoading() && CellDataSpatialHash->bCachedWasRequestedByBlockingSource)
 		{
-			const double Distance = FMath::Sqrt(CellDataSpatialHash->CachedMinSquareDistanceToBlockingSource) - (double(StreamingGrid->CellSize << CellDataSpatialHash->HierarchicalLevel) / 2);
+			const double Distance = FMath::Sqrt(CellDataSpatialHash->CachedMinSquareDistanceToBlockingSource) - (double((int64)StreamingGrid->CellSize << (int64)CellDataSpatialHash->HierarchicalLevel) / 2);
 
 			const double Ratio = Distance / LoadingRange;
 			if (Ratio <= 0.0)
@@ -1970,7 +1970,7 @@ ENGINE_API EWorldPartitionStreamingPerformance UWorldPartitionRuntimeSpatialHash
 			}
 		}
 
-		const double Distance = FMath::Sqrt(CellDataSpatialHash->CachedMinSquareDistanceToSource) - (double(StreamingGrid->CellSize << CellDataSpatialHash->HierarchicalLevel) / 2);
+		const double Distance = FMath::Sqrt(CellDataSpatialHash->CachedMinSquareDistanceToSource) - (double((int64)StreamingGrid->CellSize << (int64)CellDataSpatialHash->HierarchicalLevel) / 2);
 
 		const double Ratio = Distance / LoadingRange;
 		if (Ratio <= 0.0)
