Adding check in sound wave proxy constructor to validate that a procedural sound wave isn't being turned into a rpoxy

diff --git a/Engine/Source/Runtime/Engine/Private/SoundWave.cpp b/Engine/Source/Runtime/Engine/Private/SoundWave.cpp
--- a/Engine/Source/Runtime/Engine/Private/SoundWave.cpp
+++ b/Engine/Source/Runtime/Engine/Private/SoundWave.cpp
@@ -4763,6 +4763,8 @@ FSoundWaveProxy::FSoundWaveProxy(USoundWave* InWave)
 	// this should have been allocated by the USoundWave and should always be valid
 	check(SoundWaveDataPtr);
 
+	checkf(!InWave->bProcedural, TEXT("FSoundWaveProxy can not  be made from a procedural sound wave."));
+
 	bool bIsStreaming = InWave->IsStreaming(nullptr);
 
 	// non-streaming sources need resource data initialized before the FSoundWaveProxy
