[Backout] - CL40976944
#rnx
Original CL Desc
-----------------------------------------------------------------
Added UnifiedError and integrated with FStreamableManager

diff --git a/Engine/Source/Runtime/Core/Private/Experimental/IO/IoStatusError.cpp b/Engine/Source/Runtime/Core/Private/Experimental/IO/IoStatusError.cpp
deleted file mode 100644
diff --git a/Engine/Source/Runtime/Core/Private/Experimental/UnifiedError/UnifiedError.cpp b/Engine/Source/Runtime/Core/Private/Experimental/UnifiedError/UnifiedError.cpp
deleted file mode 100644
diff --git a/Engine/Source/Runtime/Core/Private/IO/IoStatus.cpp b/Engine/Source/Runtime/Core/Private/IO/IoStatus.cpp
--- a/Engine/Source/Runtime/Core/Private/IO/IoStatus.cpp
+++ b/Engine/Source/Runtime/Core/Private/IO/IoStatus.cpp
@@ -6,8 +6,6 @@
 #include "HAL/UnrealMemory.h"
 #include "Logging/LogCategory.h"
 #include "Logging/LogMacros.h"
-#include "Experimental/UnifiedError/UnifiedError.h"
-#include "Experimental/IO/IoStatusError.h"
 #include "Math/UnrealMathUtility.h"
 
 DEFINE_LOG_CATEGORY_STATIC(LogIoStatus, Log, All);
@@ -112,4 +110,3 @@ FIoStatusBuilder operator<<(const FIoStatus& Status, FStringView String)
 { 
 	return FIoStatusBuilder(Status, String);
 }
-
diff --git a/Engine/Source/Runtime/Core/Public/Experimental/IO/IoStatusError.h b/Engine/Source/Runtime/Core/Public/Experimental/IO/IoStatusError.h
deleted file mode 100644
diff --git a/Engine/Source/Runtime/Core/Public/Experimental/UnifiedError/UnifiedError.h b/Engine/Source/Runtime/Core/Public/Experimental/UnifiedError/UnifiedError.h
deleted file mode 100644
diff --git a/Engine/Source/Runtime/Core/Public/IO/IoStatus.h b/Engine/Source/Runtime/Core/Public/IO/IoStatus.h
--- a/Engine/Source/Runtime/Core/Public/IO/IoStatus.h
+++ b/Engine/Source/Runtime/Core/Public/IO/IoStatus.h
@@ -4,11 +4,9 @@
 
 #include "Containers/StringView.h"
 #include "Containers/UnrealString.h"
-#include "Experimental/UnifiedError/UnifiedError.h"
 
 template <typename CharType> class TStringBuilderBase;
 
-
 ///////////////////////////////////////////////////////////////////////////////
 
 /**
@@ -72,11 +70,10 @@ public:
 	CORE_API static const FIoStatus Unknown;
 	CORE_API static const FIoStatus Invalid;
 
+private:
 	static constexpr int32 MaxErrorMessageLength = 128;
 	using FErrorMessage = TCHAR[MaxErrorMessageLength];
 
-	CORE_API const FErrorMessage& GetErrorMessage() const { return ErrorMessage; }
-private:
 	EIoErrorCode	ErrorCode = EIoErrorCode::Ok;
 	FErrorMessage	ErrorMessage;
 
diff --git a/Engine/Source/Runtime/Engine/Classes/Engine/Experimental/StreamableManagerError.h b/Engine/Source/Runtime/Engine/Classes/Engine/Experimental/StreamableManagerError.h
deleted file mode 100644
diff --git a/Engine/Source/Runtime/Engine/Classes/Engine/StreamableManager.h b/Engine/Source/Runtime/Engine/Classes/Engine/StreamableManager.h
--- a/Engine/Source/Runtime/Engine/Classes/Engine/StreamableManager.h
+++ b/Engine/Source/Runtime/Engine/Classes/Engine/StreamableManager.h
@@ -15,7 +15,6 @@
 #include "UObject/GCObject.h"
 #include "UObject/ObjectMacros.h"
 #include "UObject/SoftObjectPtr.h"
-#include "Experimental/StreamableManagerError.h"
 
 enum class EStreamableProgressType : uint8
 {
@@ -27,6 +26,31 @@ enum class EStreamableProgressType : uint8
 	DownloadAbsolute,
 };
 
+enum class EStreamableErrorType : uint8
+{
+	Unknown = 0,
+	/** Low-level package loading error */
+	PackageLoadError,
+	/** Low-level package loading was canceled. This is different than canceling a FStreamableHandle which just detaches from the low level request. */
+	PackageLoadCanceled,
+	/** Error downloading assets */
+	DownloadError,
+};
+
+ENGINE_API const TCHAR* LexToString(EStreamableErrorType ErrorType);
+
+struct FStreamableError
+{
+	FStreamableError(EStreamableErrorType InErrorType, FString InErrorMessage)
+		: ErrorMessage(MoveTemp(InErrorMessage))
+		, ErrorType(InErrorType)
+	{}
+
+	FString ErrorMessage;
+	EStreamableErrorType ErrorType = EStreamableErrorType::Unknown;
+};
+
+ENGINE_API FString LexToString(const FStreamableError Error);
 
 /** Defines FStreamableDelegate delegate interface */
 using FStreamableDelegate = TDelegate<void()>;
@@ -236,7 +260,7 @@ struct FStreamableHandle : public TSharedFromThis<FStreamableHandle>
 	}
 
 	/** Get error details */
-	const TOptional<UE::UnifiedError::FError>& GetError() const
+	const TOptional<FStreamableError>& GetError() const
 	{
 		return Error;
 	}
@@ -569,8 +593,8 @@ private:
 	void UnbindDelegates();
 
 	/** Record any errors that occured with this handle */
-	void SetError(UE::UnifiedError::FError&& Error);
-	void SetErrorOnParents(const UE::UnifiedError::FError& Error);
+	void SetError(FStreamableError&& Error);
+	void SetErrorOnParents(const FStreamableError& Error);
 
 	/** Delegate to call when streaming is completed */
 	FStreamableDelegateWithHandle CompleteDelegate;
@@ -588,7 +612,7 @@ private:
 	UE::FSourceLocation Location;
 
 	/** Error details in case of an error */
-	TOptional<UE::UnifiedError::FError> Error;
+	TOptional<FStreamableError> Error;
 
 	/** Handle that holds packages in the download cache */
 	FStreamableDownloadCachePin DownloadCachePin;
@@ -836,7 +860,7 @@ struct FStreamableManager : public FGCObject
 	 * @param Location				[optional] Is not intended for direct use, the parameter catches call site source location. That location is used for debug tools output.
 	 */
 	ENGINE_API TSharedPtr<FStreamableHandle> CreateErrorHandle(
-		UE::UnifiedError::FError&& Error,
+		FStreamableError&& Error,
 		FStreamableAsyncLoadParams&& Params,
 		FString DebugName = FStreamableHandle::HandleDebugName_CombinedHandle,
 		UE::FSourceLocation Location = UE::FSourceLocation::Current());
diff --git a/Engine/Source/Runtime/Engine/Private/AssetManager.cpp b/Engine/Source/Runtime/Engine/Private/AssetManager.cpp
--- a/Engine/Source/Runtime/Engine/Private/AssetManager.cpp
+++ b/Engine/Source/Runtime/Engine/Private/AssetManager.cpp
@@ -20,7 +20,6 @@
 #include "Interfaces/IPluginManager.h"
 #include "Internationalization/PackageLocalizationManager.h"
 #include "IPlatformFilePak.h"
-#include "Logging/StructuredLog.h"
 #include "Misc/CoreDelegates.h"
 #include "Misc/FileHelper.h"
 #include "Misc/DelayedAutoRegister.h"
@@ -1953,7 +1952,7 @@ TSharedPtr<FStreamableHandle> UAssetManager::ChangeBundleStateForPrimaryAssets(
 				FStreamableAsyncLoadParams StreamableLoadParams;
 				StreamableLoadParams.TargetsToStream = MoveTemp(PathsToLoad);
 				NewHandle = StreamableManager.CreateErrorHandle(
-					UE::UnifiedError::StreamableManager::MissingDownloadParams::MakeError(*PrimaryAssetId.ToString()),
+					FStreamableError(EStreamableErrorType::DownloadError, FString::Printf(TEXT("Expected DownloadParams for %s"), *PrimaryAssetId.ToString())),
 					MoveTemp(StreamableLoadParams), FString(DebugName), CopyTemp(Location));
 			}
 
@@ -2181,17 +2180,7 @@ void UAssetManager::OnAssetStateChangeCompleted(const TSharedPtr<FStreamableHand
 			// If an error occured, don't update the current state with a bad state
 			if (BoundHandle->HasError())
 			{
-				TRefCountPtr<const UE::UnifiedError::StreamableManager::FStreamableManagerErrorDetails> ErrorDetails = BoundHandle->GetError()->GetErrorDetails<UE::UnifiedError::StreamableManager::FStreamableManagerErrorDetails>();
-
-				if (ErrorDetails.IsValid())
-				{
-					UE_LOG(LogAssetManager, Warning, TEXT("OnAssetStateChangeCompleted: Received error (%s) loading asset! Primary asset id: (%s) erroring asset:%s"), *LexToString(*BoundHandle->GetError()), *PrimaryAssetId.ToString(), *ErrorDetails->GetSoftObjectPath());
-				}
-				else
-				{
-					UE_LOG(LogAssetManager, Warning, TEXT("OnAssetStateChangeCompleted: Received error (%s) loading asset! Primary asset id:(%s)"), *LexToString(*BoundHandle->GetError()), *PrimaryAssetId.ToString());
-				}
-				UE_LOGFMT(LogAssetManager, Warning, "OnAssetStateChangeCompleted: Received error ({Error}) loading asset! Primary asset id:({PrimaryAssetId})", ("Error", BoundHandle->GetError()), ("PrimaryAssetId", *PrimaryAssetId.ToString()));
+				UE_LOG(LogAssetManager, Warning, TEXT("OnAssetStateChangeCompleted: Received error (%s) loading asset! (%s)"), *LexToString(*BoundHandle->GetError()), *PrimaryAssetId.ToString());
 			}
 			else
 			{
diff --git a/Engine/Source/Runtime/Engine/Private/Experimental/StreamableManagerError.cpp b/Engine/Source/Runtime/Engine/Private/Experimental/StreamableManagerError.cpp
deleted file mode 100644
diff --git a/Engine/Source/Runtime/Engine/Private/StreamableManager.cpp b/Engine/Source/Runtime/Engine/Private/StreamableManager.cpp
--- a/Engine/Source/Runtime/Engine/Private/StreamableManager.cpp
+++ b/Engine/Source/Runtime/Engine/Private/StreamableManager.cpp
@@ -1,8 +1,6 @@
 // Copyright Epic Games, Inc. All Rights Reserved.
 
 #include "Engine/StreamableManager.h"
-#include "Experimental/UnifiedError/UnifiedError.h"
-#include "Experimental/IO/IoStatusError.h"
 #include "UObject/ICookInfo.h"
 #include "UObject/ObjectRedirector.h"
 #include "UObject/Package.h"
@@ -22,9 +20,6 @@
 
 DEFINE_LOG_CATEGORY_STATIC(LogStreamableManager, Log, All);
 
-
-// DEFINE_ERROR_DETAILS(StreamableManager, FStreamableManagerErrorDetails);
-
 // Default to 1 frame, this will cause the delegates to go off on the next tick to avoid recursion issues. Set higher to fake disk lag
 static int32 GStreamableDelegateDelayFrames = 1;
 static FAutoConsoleVariableRef CVarStreamableDelegateDelayFrames(
@@ -51,6 +46,26 @@ static FAutoConsoleVariableRef CVarStreamableStripDebugNameInShipping(
 	ECVF_Default
 );
 
+const TCHAR* LexToString(EStreamableErrorType ErrorType)
+{
+	switch (ErrorType)
+	{
+	default:
+	case EStreamableErrorType::Unknown:
+		return TEXT("Unknown");
+	case EStreamableErrorType::PackageLoadError:
+		return TEXT("PackageLoadError");
+	case EStreamableErrorType::PackageLoadCanceled:
+		return TEXT("PackageLoadCanceled");
+	case EStreamableErrorType::DownloadError:
+		return TEXT("DownloadError");
+	}
+}
+
+FString LexToString(const FStreamableError Error)
+{
+	return FString::Printf(TEXT("%s - %s"), LexToString(Error.ErrorType), *Error.ErrorMessage);
+}
 
 FStreamableDownloadCachePin::FStreamableDownloadCachePin() = default;
 FStreamableDownloadCachePin::~FStreamableDownloadCachePin() = default;
@@ -407,7 +422,7 @@ namespace UE::StreamableManager::Private
 			}
 		}
 
-		void CompleteLoad(const TSharedRef<FStreamableHandle>& OwningHandle, UE::UnifiedError::FError&& Error)
+		void CompleteLoad(const TSharedRef<FStreamableHandle>& OwningHandle, FStreamableError&& Error)
 		{
 			OwningHandle->SetError(MoveTemp(Error));
 			CompleteLoad(OwningHandle);
@@ -424,7 +439,7 @@ namespace UE::StreamableManager::Private
 			UE::IoStore::IOnDemandIoStore* IoStore = UE::IoStore::TryGetOnDemandIoStore();
 			if (!IoStore)
 			{
-				CompleteLoad(OwningHandle, UE::UnifiedError::StreamableManager::IoStoreNotFound::MakeError());
+				CompleteLoad(OwningHandle, FStreamableError{ EStreamableErrorType::DownloadError, TEXT("IoStore not found") });
 				OwningHandle->DownloadContext = nullptr; // Destroy self
 				return;
 			}
@@ -481,7 +496,7 @@ namespace UE::StreamableManager::Private
 				}
 				else
 				{
-					CompleteLoad(OwningHandle, UE::UnifiedError::IoStore::ConvertError(OnDemandInstallResult.Status));
+					CompleteLoad(OwningHandle, FStreamableError{ EStreamableErrorType::DownloadError, OnDemandInstallResult.Status.ToString() });
 				}
 
 				return;
@@ -510,6 +525,21 @@ namespace UE::StreamableManager::Private
 			}
 		}
 	};
+
+	static FStreamableError GetStreamableError(EAsyncLoadingResult::Type Result)
+	{
+		switch (Result)
+		{
+		case EAsyncLoadingResult::Failed:
+			// We could supply an error string if async loading bubbled one up...
+			// Possibly GetExplanationForUnavailablePackage?
+			return FStreamableError{ EStreamableErrorType::PackageLoadError, TEXT("Async load failed") };
+		case EAsyncLoadingResult::Canceled:
+			return FStreamableError{ EStreamableErrorType::PackageLoadCanceled, TEXT("Async load canceled") };
+		}
+
+		return FStreamableError{ EStreamableErrorType::Unknown, FString() };
+	}
 }
 
 FStreamableHandle::FStreamableHandle(FPrivateToken, UE::FSourceLocation&& InLocation)
@@ -1053,10 +1083,6 @@ FStreamableHandle::~FStreamableHandle()
 void FStreamableHandle::CompleteLoad()
 {
 	// Only complete if it's still active
-	if (Error.IsSet())
-	{
-		UE_LOG(LogStreamableManager, Verbose, TEXT("FStreamableHandle::CompleteLoad failed load %s, with error %s!"), *DebugName, *Error->GetErrorMessage().ToString());
-	}
 	if (IsActive())
 	{
 		bLoadCompleted = true;
@@ -1191,17 +1217,16 @@ void FStreamableHandle::UnbindDelegates()
 	CompleteDelegate.Unbind();
 }
 
-void FStreamableHandle::SetError(UE::UnifiedError::FError&& InError)
+void FStreamableHandle::SetError(FStreamableError&& InError)
 {
 	if (!Error)
 	{
 		SetErrorOnParents(InError);
 		Error = MoveTemp(InError);
-		
 	}
 }
 
-void FStreamableHandle::SetErrorOnParents(const UE::UnifiedError::FError& InError)
+void FStreamableHandle::SetErrorOnParents(const FStreamableError& InError)
 {
 	// Update any meta handles that are still active
 	for (TWeakPtr<FStreamableHandle>& WeakHandle : ParentHandles)
@@ -1224,12 +1249,7 @@ void FStreamableHandle::AsyncLoadCallbackWrapper(const FName& PackageName, UPack
 	{
 		if (Result != EAsyncLoadingResult::Succeeded)
 		{
-			UE::UnifiedError::FError NewError = UE::UnifiedError::StreamableManager::GetStreamableError(Result);
-
-			TRefCountPtr<UE::UnifiedError::IErrorDetails> NewErrorDetails( new UE::UnifiedError::StreamableManager::FStreamableManagerErrorDetails(TargetName.ToString()));
-
-			NewError.PushErrorDetails(NewErrorDetails);
-			SetError(MoveTemp(NewError));
+			SetError(UE::StreamableManager::Private::GetStreamableError(Result));
 		}
 
 		OwningManager->AsyncLoadCallback(MoveTemp(TargetName), Package, Result);
@@ -1687,7 +1707,7 @@ FStreamable* FStreamableManager::StreamInternal(const FSoftObjectPath& InTargetN
 				Existing->bLoadFailed = true;
 				UE_LOG(LogStreamableManager, Log, TEXT("Failed attempt to load %s"), *TargetName.ToString());
 
-				Handle->SetError(UE::UnifiedError::StreamableManager::PackageLoadFailed::MakeError(*TargetName.ToString()));
+				Handle->SetError(FStreamableError{ EStreamableErrorType::PackageLoadError, TEXT("Static load failed") });
 			}
 			Existing->bAsyncLoadRequestOutstanding = false;
 		}
@@ -1708,7 +1728,7 @@ FStreamable* FStreamableManager::StreamInternal(const FSoftObjectPath& InTargetN
 				Existing->bLoadFailed = true;
 				Existing->bAsyncLoadRequestOutstanding = false;
 
-				Handle->SetError(UE::UnifiedError::StreamableManager::PackageNameInvalid::MakeError(Package));
+				Handle->SetError(FStreamableError{ EStreamableErrorType::PackageLoadError, TEXT("Unkown package")});
 			}
 			else
 			{
@@ -1910,7 +1930,7 @@ TSharedPtr<FStreamableHandle> FStreamableManager::RequestSyncLoadInternal(TArray
 
 		if (bBadPackageState || bLoadNotComplete)
 		{
-			Request->SetError(UE::UnifiedError::StreamableManager::SyncLoadIncomplete::MakeError(Request->DebugName));
+			Request->SetError(FStreamableError{ EStreamableErrorType::PackageLoadError, TEXT("Sync load incomplete") });
 		}
 	}
 
@@ -2119,15 +2139,14 @@ void FStreamableManager::AsyncLoadCallback(FSoftObjectPath TargetName, UPackage*
 		{
 			for (const TSharedRef<FStreamableHandle>& Handle : Existing->LoadingHandles)
 			{
-				// @Per @Justin In my testing this has already been set by the AsyncLoadCallbackWrapper.  does this need to happen twice?
-				Handle->SetError(UE::UnifiedError::StreamableManager::GetStreamableError(Result));
+				Handle->SetError(UE::StreamableManager::Private::GetStreamableError(Result));
 			}
 		}
 		else if (!Existing->Target)
 		{
 			for (const TSharedRef<FStreamableHandle>& Handle : Existing->LoadingHandles)
 			{
-				Handle->SetError(UE::UnifiedError::StreamableManager::PackageLoadFailed::MakeError(TargetName.ToString()));
+				Handle->SetError(FStreamableError{ EStreamableErrorType::PackageLoadError, TEXT("Failed to find target asset in memory") });
 			}
 		}
 
@@ -2138,7 +2157,6 @@ void FStreamableManager::AsyncLoadCallback(FSoftObjectPath TargetName, UPackage*
 	}
 	else
 	{
-		// should this be an error?  Is it an error if no handles are waiting for it? 
 		UE_LOG(LogStreamableManager, Error, TEXT("Can't find streamable for %s in AsyncLoadCallback!"), *TargetName.ToString());
 	}
 }
@@ -2412,7 +2430,7 @@ TSharedPtr<FStreamableHandle> FStreamableManager::CreateCombinedHandle(
 	return NewRequest;
 }
 
-TSharedPtr<FStreamableHandle> FStreamableManager::CreateErrorHandle(UE::UnifiedError::FError&& Error, FStreamableAsyncLoadParams&& Params, FString DebugName, UE::FSourceLocation Location)
+TSharedPtr<FStreamableHandle> FStreamableManager::CreateErrorHandle(FStreamableError&& Error, FStreamableAsyncLoadParams&& Params, FString DebugName, UE::FSourceLocation Location)
 {
 	TSharedRef<FStreamableHandle> NewRequest = MakeShared<FStreamableHandle>(FStreamableHandle::FPrivateToken(), MoveTemp(Location));
 	NewRequest->OwningManager = this;
