World Partition: don't timestamp streaming generation logs on build machines so we get a single log per map.

diff --git a/Engine/Source/Runtime/Engine/Private/WorldPartition/WorldPartitionStreamingGeneration.cpp b/Engine/Source/Runtime/Engine/Private/WorldPartition/WorldPartitionStreamingGeneration.cpp
--- a/Engine/Source/Runtime/Engine/Private/WorldPartition/WorldPartitionStreamingGeneration.cpp
+++ b/Engine/Source/Runtime/Engine/Private/WorldPartition/WorldPartitionStreamingGeneration.cpp
@@ -1825,8 +1825,8 @@ public:
 	static TUniquePtr<FArchive> CreateDumpStateLogArchive(const TCHAR* Suffix)
 	{
 		const FString LogDir = GIsBuildMachine ? FPaths::Combine(*FPaths::EngineDir(), TEXT("Programs"), TEXT("AutomationTool"), TEXT("Saved"), TEXT("Logs")) : FPaths::ProjectLogDir();
-		const FString StateLogOutputFilename = FPaths::Combine(LogDir, TEXT("WorldPartition"), 
-			FString::Printf(TEXT("StreamingGeneration-%s-%08x-%s.log"), Suffix, FPlatformProcess::GetCurrentProcessId(), *FDateTime::Now().ToIso8601().Replace(TEXT(":"), TEXT("."))));
+		const FString LogTimeStamp = GIsBuildMachine ? TEXT("") : FString::Printf(TEXT("-%08x-%s"), FPlatformProcess::GetCurrentProcessId(), *FDateTime::Now().ToIso8601().Replace(TEXT(":"), TEXT(".")));
+		const FString StateLogOutputFilename = FPaths::Combine(LogDir, TEXT("WorldPartition"), FString::Printf(TEXT("StreamingGeneration-%s%s.log"), Suffix, *LogTimeStamp));
 		return TUniquePtr<FArchive>(IFileManager::Get().CreateFileWriter(*StateLogOutputFilename));
 	}
 
