Add new bCanSplit flag to virtual texture physical spaces.
Only allow the "r.VT.SplitPhysicalPoolSize" setting to split pools that have this flag set. Previously all pools could be split.
The flag is set by the virtual texture producer and defaults to true.
MeshPaintVirtualTexture sets the flag to false because that system can't support splitting of pools. We require that all mesh paint textures end up in the same physical texture.

diff --git a/Engine/Source/Runtime/Engine/Classes/Engine/Texture2D.h b/Engine/Source/Runtime/Engine/Classes/Engine/Texture2D.h
--- a/Engine/Source/Runtime/Engine/Classes/Engine/Texture2D.h
+++ b/Engine/Source/Runtime/Engine/Classes/Engine/Texture2D.h
@@ -388,4 +388,10 @@ public:
 	 */
 	virtual bool IsVirtualTexturedWithSinglePhysicalSpace() const { return false;  }
 
+	/**
+	 * Returns true if this virtual texture requires that it's physical pool is never split by the r.VT.SplitPhysicalPoolSize setting.
+	 * This is necessary if the texture is part of a system that assumes that all textures in the system are in the same physical pool.
+	 * Note that textures with this setting will necessarily be put in different physical spaces to those without the setting.
+	 */
+	virtual bool IsVirtualTexturedWithSinglePhysicalPool() const { return false; }
 };
diff --git a/Engine/Source/Runtime/Engine/Classes/VT/MeshPaintVirtualTexture.h b/Engine/Source/Runtime/Engine/Classes/VT/MeshPaintVirtualTexture.h
--- a/Engine/Source/Runtime/Engine/Classes/VT/MeshPaintVirtualTexture.h
+++ b/Engine/Source/Runtime/Engine/Classes/VT/MeshPaintVirtualTexture.h
@@ -25,6 +25,7 @@ class UMeshPaintVirtualTexture : public UTexture2D
 	//~ Begin UTexture Interface.
 	virtual void GetVirtualTextureBuildSettings(FVirtualTextureBuildSettings& OutSettings) const override;
 	virtual void UpdateResourceWithParams(EUpdateResourceFlags InFlags) override;
+	virtual bool IsVirtualTexturedWithSinglePhysicalPool() const override { return true; }
 
 #if WITH_EDITOR
 	virtual void BeginCacheForCookedPlatformData(const ITargetPlatform* TargetPlatform) override;
diff --git a/Engine/Source/Runtime/Engine/Private/Texture2D.cpp b/Engine/Source/Runtime/Engine/Private/Texture2D.cpp
--- a/Engine/Source/Runtime/Engine/Private/Texture2D.cpp
+++ b/Engine/Source/Runtime/Engine/Private/Texture2D.cpp
@@ -1523,6 +1523,7 @@ FVirtualTexture2DResource::FVirtualTexture2DResource(const UTexture2D* InOwner,
 	TexCreateFlags |= InOwner->bNoTiling ? ETextureCreateFlags::NoTiling : ETextureCreateFlags::None;
 
 	bSinglePhysicalSpace = InOwner->IsVirtualTexturedWithSinglePhysicalSpace();
+	bRequiresSinglePhysicalPool = InOwner->IsVirtualTexturedWithSinglePhysicalPool();
 
 	check(InVTData);
 	VTData = InVTData;
@@ -1604,6 +1605,7 @@ void FVirtualTexture2DResource::InitRHI(FRHICommandListBase& RHICmdList)
 		ProducerDesc.PhysicalGroupIndex[LayerIndex] = bSinglePhysicalSpace ? 0 : LayerIndex;
 		ProducerDesc.bIsLayerSRGB[LayerIndex] = bSRGB;
 	}
+	ProducerDesc.bRequiresSinglePhysicalPool = bRequiresSinglePhysicalPool;
 	// TODO [jonathan.bard] : For now, let's use the normal priority, but that could come from the texture asset and/or the texture group too : 
 	ProducerDesc.Priority = EVTProducerPriority::Normal;
 
diff --git a/Engine/Source/Runtime/Engine/Private/VT/VirtualTextureAdapter.cpp b/Engine/Source/Runtime/Engine/Private/VT/VirtualTextureAdapter.cpp
--- a/Engine/Source/Runtime/Engine/Private/VT/VirtualTextureAdapter.cpp
+++ b/Engine/Source/Runtime/Engine/Private/VT/VirtualTextureAdapter.cpp
@@ -474,7 +474,7 @@ class FVirtualTextureAdapterRenderResource : public FVirtualTexture2DResource
 	uint32 NumSourceMips = 1;
 
 public:
-	FVirtualTextureAdapterRenderResource(UVirtualTextureAdapter const* InOwner, UTexture* InTexture, int32 InTileSize, int32 InTileBorderSize, EPixelFormat InFinalPixelFormat)
+	FVirtualTextureAdapterRenderResource(UVirtualTextureAdapter const* InOwner, UTexture* InTexture, int32 InTileSize, int32 InTileBorderSize, EPixelFormat InFinalPixelFormat, bool bInRequiresSinglePhysicalPool)
 	{
 		TextureName = InOwner->GetFName();
 		PackageName = InOwner->GetOutermost()->GetFName();
@@ -494,6 +494,7 @@ public:
 			bSRGB = RenderTarget2D->SRGB;
 		}
 		Format = VirtualTextureAdapter::GetFinalFormat(SourcePixelFormat, InFinalPixelFormat);
+		bRequiresSinglePhysicalPool = bInRequiresSinglePhysicalPool;
 
 		TileSize = InTileSize;
 		TileBorderSize = InTileBorderSize;
@@ -528,6 +529,7 @@ public:
 		ProducerDesc.Name = TextureName;
 		ProducerDesc.FullNameHash = GetTypeHash(TextureName);
 		ProducerDesc.bContinuousUpdate = false;
+		ProducerDesc.bRequiresSinglePhysicalPool = bRequiresSinglePhysicalPool;
 		ProducerDesc.Dimensions = 2;
 		ProducerDesc.TileSize = TileSize;
 		ProducerDesc.TileBorderSize = TileBorderSize;
@@ -572,8 +574,9 @@ FTextureResource* UVirtualTextureAdapter::CreateResource()
 	const uint32 FinalTileBorderSize = bUseDefaultTileSizes ? DefaultSettings.TileBorderSize : FVirtualTextureBuildSettings::ClampAndAlignTileBorderSize(TileBorderSize);
 	
 	const EPixelFormat FinalPixelFormat = OverrideWithTextureFormat ? OverrideWithTextureFormat->GetPixelFormat(0) : PF_Unknown;
+	const bool bRequiresSinglePhysicalPool = OverrideWithTextureFormat ? OverrideWithTextureFormat->IsVirtualTexturedWithSinglePhysicalPool() : false;
 
-	return new FVirtualTextureAdapterRenderResource(this, Texture, FinalTileSize, FinalTileBorderSize, FinalPixelFormat);
+	return new FVirtualTextureAdapterRenderResource(this, Texture, FinalTileSize, FinalTileBorderSize, FinalPixelFormat, bRequiresSinglePhysicalPool);
 }
 
 
diff --git a/Engine/Source/Runtime/Engine/Public/TextureResource.h b/Engine/Source/Runtime/Engine/Public/TextureResource.h
--- a/Engine/Source/Runtime/Engine/Public/TextureResource.h
+++ b/Engine/Source/Runtime/Engine/Public/TextureResource.h
@@ -229,8 +229,9 @@ protected:
 	TEnumAsByte<ESamplerAddressMode> AddressV = AM_Wrap;
 	/** Cached flags for texture creation. */
 	ETextureCreateFlags TexCreateFlags = ETextureCreateFlags::None;
-	/** Cached runtime virtual texture settings */
+	/** Cached owner settings */
 	bool bSinglePhysicalSpace = false;
+	bool bRequiresSinglePhysicalPool = false;
 	/** Mip offset */
 	int32 FirstMipToUse = 0;
 	/** Built data owned by texture asset */
diff --git a/Engine/Source/Runtime/RenderCore/Public/VirtualTexturing.h b/Engine/Source/Runtime/RenderCore/Public/VirtualTexturing.h
--- a/Engine/Source/Runtime/RenderCore/Public/VirtualTexturing.h
+++ b/Engine/Source/Runtime/RenderCore/Public/VirtualTexturing.h
@@ -149,6 +149,7 @@ struct FVTProducerDescription
 	
 	bool bPersistentHighestMip = true;
 	bool bContinuousUpdate = false;
+	bool bRequiresSinglePhysicalPool = false;
 	bool bNotifyCompleted = false; /** Producer will receive OnRequestsCompleted() callbacks every frame when enabled. */
 	
 	uint32 TileSize = 0u;
diff --git a/Engine/Source/Runtime/Renderer/Private/VT/VirtualTexturePhysicalSpace.h b/Engine/Source/Runtime/Renderer/Private/VT/VirtualTexturePhysicalSpace.h
--- a/Engine/Source/Runtime/Renderer/Private/VT/VirtualTexturePhysicalSpace.h
+++ b/Engine/Source/Runtime/Renderer/Private/VT/VirtualTexturePhysicalSpace.h
@@ -20,13 +20,15 @@ struct FVTPhysicalSpaceDescription
 	uint8 NumLayers;
 	TEnumAsByte<EPixelFormat> Format[VIRTUALTEXTURE_SPACE_MAXLAYERS];
 	bool bHasLayerSrgbView[VIRTUALTEXTURE_SPACE_MAXLAYERS];
+	bool bCanSplit;
 };
 
 inline bool operator==(const FVTPhysicalSpaceDescription& Lhs, const FVTPhysicalSpaceDescription& Rhs)
 {
 	if (Lhs.TileSize != Rhs.TileSize ||
 		Lhs.NumLayers != Rhs.NumLayers || 
-		Lhs.Dimensions != Rhs.Dimensions)
+		Lhs.Dimensions != Rhs.Dimensions ||
+		Lhs.bCanSplit != Rhs.bCanSplit)
 	{
 		return false;
 	}
@@ -55,6 +57,7 @@ inline uint32 GetTypeHash(const FVTPhysicalSpaceDescription& Desc)
 	uint32 Hash = Desc.TileSize;
 	Hash = HashCombine(Hash, GetTypeHash(Desc.Dimensions));
 	Hash = HashCombine(Hash, GetTypeHash(Desc.NumLayers));
+	Hash = HashCombine(Hash, GetTypeHash(Desc.bCanSplit));
 	for (int32 Layer = 0; Layer < Desc.NumLayers; ++Layer)
 	{
 		Hash = HashCombine(Hash, GetTypeHash(Desc.Format[Layer]));
diff --git a/Engine/Source/Runtime/Renderer/Private/VT/VirtualTextureProducer.cpp b/Engine/Source/Runtime/Renderer/Private/VT/VirtualTextureProducer.cpp
--- a/Engine/Source/Runtime/Renderer/Private/VT/VirtualTextureProducer.cpp
+++ b/Engine/Source/Runtime/Renderer/Private/VT/VirtualTextureProducer.cpp
@@ -74,6 +74,7 @@ FVirtualTextureProducerHandle FVirtualTextureProducerCollection::RegisterProduce
 	FVTPhysicalSpaceDescription PhysicalSpaceDesc;
 	PhysicalSpaceDesc.Dimensions = InDesc.Dimensions;
 	PhysicalSpaceDesc.TileSize = InDesc.TileSize + InDesc.TileBorderSize * 2u;
+	PhysicalSpaceDesc.bCanSplit = !InDesc.bRequiresSinglePhysicalPool;
 
 	Entry.Producer.PhysicalGroups.SetNum(InDesc.NumPhysicalGroups);
 	for (uint32 PhysicalGroupIndex = 0u; PhysicalGroupIndex < InDesc.NumPhysicalGroups; ++PhysicalGroupIndex)
diff --git a/Engine/Source/Runtime/Renderer/Private/VT/VirtualTextureSystem.cpp b/Engine/Source/Runtime/Renderer/Private/VT/VirtualTextureSystem.cpp
--- a/Engine/Source/Runtime/Renderer/Private/VT/VirtualTextureSystem.cpp
+++ b/Engine/Source/Runtime/Renderer/Private/VT/VirtualTextureSystem.cpp
@@ -984,7 +984,7 @@ void GetPhysicalSpaceExtraDescription(FVTPhysicalSpaceDescription const& InDesc,
 
 		// If we allow splitting one pool confing into multiple physical pools by size then keep incrementing pool count until we fit.
 		const int32 SplitPhysicalPoolSize = VirtualTexturePool::GetSplitPhysicalPoolSize();
-		if (SplitPhysicalPoolSize <= 0 || TileWidthHeight <= SplitPhysicalPoolSize)
+		if (!InDesc.bCanSplit || SplitPhysicalPoolSize <= 0 || TileWidthHeight <= SplitPhysicalPoolSize)
 		{
 			break;
 		}
