New material translator changes:
- Added support for the "select" intrinsic.
- Changed And and Or to output as "and" and "or" in HLSL.
- Added FEmitter::EmitConstantScalarOne() for consistency.

diff --git a/Engine/Source/Runtime/Engine/Private/Materials/MaterialIR.cpp b/Engine/Source/Runtime/Engine/Private/Materials/MaterialIR.cpp
--- a/Engine/Source/Runtime/Engine/Private/Materials/MaterialIR.cpp
+++ b/Engine/Source/Runtime/Engine/Private/Materials/MaterialIR.cpp
@@ -189,6 +189,44 @@ bool FValue::IsFalse() const
 	return Constant && Constant->IsBool() && Constant->Boolean == false;
 }
 
+bool FValue::AreAllTrue() const
+{
+	if (const MIR::FDimensional* Dimensional = As<MIR::FDimensional>())
+	{
+		for (const MIR::FValue* Component : Dimensional->GetComponents())
+		{
+			if (!Component->IsTrue())
+			{
+				return false;
+			}
+		}
+	}
+	else
+	{
+		return IsTrue();
+	}
+	return false;
+}
+
+bool FValue::AreAllFalse() const
+{
+	if (const MIR::FDimensional* Dimensional = As<MIR::FDimensional>())
+	{
+		for (const MIR::FValue* Component : Dimensional->GetComponents())
+		{
+			if (!Component->IsFalse())
+			{
+				return false;
+			}
+		}
+	}
+	else
+	{
+		return IsFalse();
+	}
+	return false;
+}
+
 bool FValue::AreAllExactlyZero() const
 {
 	if (const MIR::FDimensional* Dimensional = As<MIR::FDimensional>())
@@ -497,6 +535,7 @@ const TCHAR* LexToString(EOperator Op)
 		case BO_Xor: return TEXT("Xor");
 		case TO_Clamp: return TEXT("Clamp");
 		case TO_Lerp: return TEXT("Lerp");
+		case TO_Select: return TEXT("Select");
 		case TO_Smoothstep: return TEXT("Smoothstep");
 		case UO_Abs: return TEXT("Abs");
 		case UO_ACos: return TEXT("ACos");
diff --git a/Engine/Source/Runtime/Engine/Private/Materials/MaterialIREmitter.cpp b/Engine/Source/Runtime/Engine/Private/Materials/MaterialIREmitter.cpp
--- a/Engine/Source/Runtime/Engine/Private/Materials/MaterialIREmitter.cpp
+++ b/Engine/Source/Runtime/Engine/Private/Materials/MaterialIREmitter.cpp
@@ -448,6 +448,17 @@ FValue* FEmitter::EmitConstantScalarZero(EScalarKind Kind)
 	}
 }
 
+FValue* FEmitter::EmitConstantScalarOne(EScalarKind Kind)
+{
+	switch (Kind)
+	{
+		case ScalarKind_Bool: return EmitConstantTrue();
+		case ScalarKind_Int: return EmitConstantIntScalar(1);
+		case ScalarKind_Float: return EmitConstantFloatScalar(1.0f);
+		default: UE_MIR_UNREACHABLE();
+	}
+}
+
 FValue* FEmitter::EmitConstantTrue()
 {
 	return ConstantTrue;
@@ -602,7 +613,7 @@ FValue* FEmitter::EmitSubscript(FValue* Value, int32 Index)
 	const FPrimitiveType* PrimitiveType = Value->Type->AsPrimitive();
 	if (!PrimitiveType)
 	{
-		Errorf(TEXT("Value of type `%s` cannot be subscripted."), Value->Type->GetSpelling().GetData());
+		Errorf(TEXT("Value of type '%s' cannot be subscripted."), Value->Type->GetSpelling().GetData());
 		return MIR::FPoison::Get();
 	}
 
@@ -614,7 +625,7 @@ FValue* FEmitter::EmitSubscript(FValue* Value, int32 Index)
 
 	if (Index >= PrimitiveType->GetNumComponents())
 	{
-		Errorf(TEXT("Value of type `%s` has fewer dimensions than subscript index `%d`."), Value->Type->GetSpelling().GetData(), Index);
+		Errorf(TEXT("Value of type '%s' has fewer dimensions than subscript index `%d`."), Value->Type->GetSpelling().GetData(), Index);
 		return MIR::FPoison::Get();
 	}
 
@@ -653,7 +664,7 @@ FValue* FEmitter::EmitSwizzle(FValue* Value, FSwizzleMask Mask)
 	const FPrimitiveType* PrimitiveType = Value->Type->AsVector();
 	if (!PrimitiveType || PrimitiveType->IsMatrix())
 	{
-		Errorf(TEXT("Cannot swizzle a `%s` value."), Value->Type->GetSpelling().GetData());
+		Errorf(TEXT("Cannot swizzle a '%s' value."), Value->Type->GetSpelling().GetData());
 		return MIR::FPoison::Get();
 	}
 
@@ -662,7 +673,7 @@ FValue* FEmitter::EmitSwizzle(FValue* Value, FSwizzleMask Mask)
 	{
 		if ((int32)Component >= PrimitiveType->NumRows)
 		{
-			Errorf(TEXT("Value of type `%s` has no component `%s`."), PrimitiveType->Spelling.GetData(), VectorComponentToString(Component));
+			Errorf(TEXT("Value of type '%s' has no component '%s'."), PrimitiveType->Spelling.GetData(), VectorComponentToString(Component));
 			return MIR::FPoison::Get();
 		}
 	}
@@ -1043,6 +1054,17 @@ static FValue* TrySimplifyOperator(FEmitter* Emitter, EOperator& Op, FValue*& A,
 			}
 			break;
 
+		case TO_Select:
+			if (A->AreAllTrue())
+			{
+				return B;
+			}
+			else if (A->AreAllFalse())
+			{
+				return C;
+			}
+			break;
+
 		default:
 			break;
 	}
@@ -1062,6 +1084,12 @@ static FValue* TryFoldOperatorScalar(FEmitter* Emitter, EOperator Op, FValue* A,
 		return Simplified;
 	}
 
+	// If TrySimplifyOperator did not already fold the `select` operator, there is nothing else to do.
+	if (Op == TO_Select)
+	{
+		return nullptr;
+	}
+
 	// Verify that both lhs and rhs are constants, otherwise we cannot fold the operation.
 	const FConstant* AConstant = As<FConstant>(A);
 	const FConstant* BConstant = As<FConstant>(B);
@@ -1158,7 +1186,8 @@ enum EOperatorReturnType
 	ORT_Unknown,							///< Unspecified
 	ORT_FirstArgumentType,					///< The same type as the first argument
 	ORT_BooleanWithFirstArgumentDimensions,	///< A boolean primitive type with the same dimensions (rows and columns) as the first argument type
-	ORT_FirstArgumentTypeToScalar			///< A scalar primitive type with the same kind as the scalar type of the first argument
+	ORT_FirstArgumentTypeToScalar,			///< A scalar primitive type with the same kind as the scalar type of the first argument
+	ORT_SecondArgumentType,					///< The same type as the second argument
 };
 
 /// The signature of an operator consisting of its parameter and return type information.
@@ -1222,6 +1251,7 @@ static const FOperatorSignature* GetOperatorSignature(EOperator Op)
 		S[BO_Cross] 				= { { OPTF_Arithmetic | OPTF_Vector3, OPTF_AsFirstArgument }, ORT_FirstArgumentType };
 		S[TO_Clamp]					= TernaryArithmetic;
 		S[TO_Lerp] 					= TernaryArithmetic;
+		S[TO_Select]				= { { OPTF_Boolean | OPTF_NotMatrix, OPTF_NotMatrix, OPTF_NotMatrix}, ORT_SecondArgumentType }; // Note: this is a special operator, which is handled manually in the Validate function
 		S[TO_Smoothstep]			= TernaryArithmetic;
 		return S;
 	} ();
@@ -1260,7 +1290,7 @@ static const FPrimitiveType* ValidateOperatorAndGetResultType(FEmitter* Emitter,
 		// All operators have primitive parameter types.
 		if (!ArgumentPrimitiveType)
 		{
-			Emitter->Errorf(TEXT("Operator %s's %s argument does not have primitive type."), LexToString(Op), ArgumentsStr[i]);
+			Emitter->Errorf(TEXT("Operator '%s' %s argument does not have primitive type."), LexToString(Op), ArgumentsStr[i]);
 			return nullptr;
 		}
 		
@@ -1280,7 +1310,7 @@ static const FPrimitiveType* ValidateOperatorAndGetResultType(FEmitter* Emitter,
 		{
 			if (ArgumentPrimitiveType->ScalarKind != ScalarKind_Bool)
 			{
-				Emitter->Errorf(TEXT("Operator %s's %s argument component types are not boolean."), LexToString(Op), ArgumentsStr[i]);
+				Emitter->Errorf(TEXT("Operator '%s' %s argument component types are not boolean."), LexToString(Op), ArgumentsStr[i]);
 				bValid = false;
 			}
 		}
@@ -1289,7 +1319,7 @@ static const FPrimitiveType* ValidateOperatorAndGetResultType(FEmitter* Emitter,
 		{
 			if (!ScalarKindIsArithmetic(ArgumentPrimitiveType->ScalarKind))
 			{
-				Emitter->Errorf(TEXT("Operator %s's %s argument component types expected to be arithmetic (int, float, double, etc)."), LexToString(Op), ArgumentsStr[i]);
+				Emitter->Errorf(TEXT("Operator '%s' %s argument component types expected to be arithmetic (int, float, double, etc)."), LexToString(Op), ArgumentsStr[i]);
 				bValid = false;
 			}
 		}
@@ -1298,7 +1328,7 @@ static const FPrimitiveType* ValidateOperatorAndGetResultType(FEmitter* Emitter,
 		{
 			if (!ScalarKindIsAnyFloat(ArgumentPrimitiveType->ScalarKind))
 			{
-				Emitter->Errorf(TEXT("Operator %s's %s argument component types expected to be floating point (float, double, etc)."), LexToString(Op), ArgumentsStr[i]);
+				Emitter->Errorf(TEXT("Operator '%s' %s argument component types expected to be floating point (float, double, etc)."), LexToString(Op), ArgumentsStr[i]);
 				bValid = false;
 			}
 		}
@@ -1307,7 +1337,7 @@ static const FPrimitiveType* ValidateOperatorAndGetResultType(FEmitter* Emitter,
 		{
 			if (ArgumentPrimitiveType->IsMatrix())
 			{
-				Emitter->Errorf(TEXT("Operator %s's %s argument type expected to be a scalar or vector."), LexToString(Op), ArgumentsStr[i]);
+				Emitter->Errorf(TEXT("Operator '%s' %s argument type expected to be a scalar or vector."), LexToString(Op), ArgumentsStr[i]);
 				bValid = false;
 			}
 		}
@@ -1316,7 +1346,7 @@ static const FPrimitiveType* ValidateOperatorAndGetResultType(FEmitter* Emitter,
 		{
 			if (!ArgumentPrimitiveType->IsVector() || ArgumentPrimitiveType->GetNumComponents() != 3)
 			{
-				Emitter->Errorf(TEXT("Operator %s's %s argument type expected to be a 3D vector."), LexToString(Op), ArgumentsStr[i]);
+				Emitter->Errorf(TEXT("Operator '%s' %s argument type expected to be a 3D vector."), LexToString(Op), ArgumentsStr[i]);
 				bValid = false;
 			}
 		}
@@ -1328,15 +1358,46 @@ static const FPrimitiveType* ValidateOperatorAndGetResultType(FEmitter* Emitter,
 		}
 	}
 	
-	// Now that we know the common type, cast each argument that requires it to it.
-	for (int32 i = 0; Arguments[i]; ++i)
+	// The select operator is special insofar as its first argument is a boolean, while the second and third can be any primitive type.
+	if (Op == TO_Select)
 	{
-		EOperatorParameterTypeFilter Filter = Signature->ParameterTypes[i];
-		if (Filter & OPTF_CastToCommon)
+		// Cast the second and third argument types to primitive. This is safe as it was already checked earlier.
+		const FPrimitiveType* BPrimitiveType = static_cast<const FPrimitiveType*>(Arguments[1]->Type);
+		const FPrimitiveType* CPrimitiveType = static_cast<const FPrimitiveType*>(Arguments[2]->Type);
+
+		// Compute the maximum number of vector components between all arguments. We know they're scalar or vectors, as it was checked before.
+		int32 MaxNumComponents = FMath::Max3(FirstArgumentPrimitiveType->GetNumComponents(), BPrimitiveType->GetNumComponents(), CPrimitiveType->GetNumComponents());
+
+		// Cast the first argument (the boolean condition) to a bool vector of the maximum number of components.
+		Arguments[0] = Emitter->EmitConstruct(FPrimitiveType::Get(ScalarKind_Bool, MaxNumComponents, 1), Arguments[0]);
+
+		// Compute the common type between the second and third argument types with a number of components equal to the max of all three.
+		const FType* CommonTypeBetweenSecondAndThirdArguments = Emitter->TryGetCommonType(
+			FPrimitiveType::Get(BPrimitiveType->ScalarKind, MaxNumComponents, 1),
+			FPrimitiveType::Get(CPrimitiveType->ScalarKind, MaxNumComponents, 1));
+		
+		// Getting the common type should always be possible.
+		check(CommonTypeBetweenSecondAndThirdArguments);
+
+		// Cast second and third arguments to their common type.
+		Arguments[1] = Emitter->EmitConstruct(CommonTypeBetweenSecondAndThirdArguments, Arguments[1]);
+		Arguments[2] = Emitter->EmitConstruct(CommonTypeBetweenSecondAndThirdArguments, Arguments[2]);
+
+		// Update the valid flag to the results of the conditions above.
+		bValid &= !Arguments[0]->IsPoison() && !Arguments[1]->IsPoison() && !Arguments[2]->IsPoison();
+	}
+	else
+	{
+		// Cast every argument with the `CastToCommon` to the common type, if necessary.
+		for (int32 i = 0; Arguments[i]; ++i)
 		{
-			check(ArgumentsCommonType && ArgumentsCommonType->AsPrimitive());
-			Arguments[i] = Emitter->EmitConstruct(ArgumentsCommonType, Arguments[i]);
-			bValid &= !Arguments[i]->IsPoison();
+			EOperatorParameterTypeFilter Filter = Signature->ParameterTypes[i];
+			if (Filter & OPTF_CastToCommon)
+			{
+				check(ArgumentsCommonType && ArgumentsCommonType->AsPrimitive());
+				Arguments[i] = Emitter->EmitConstruct(ArgumentsCommonType, Arguments[i]);
+				bValid &= !Arguments[i]->IsPoison();
+			}
 		}
 	}
 
@@ -1367,6 +1428,9 @@ static const FPrimitiveType* ValidateOperatorAndGetResultType(FEmitter* Emitter,
 
 		case ORT_FirstArgumentTypeToScalar:
 			return FPrimitiveType::GetScalar(FirstArgumentPrimitiveType->ScalarKind);
+
+		case ORT_SecondArgumentType:
+			return static_cast<const FPrimitiveType*>(B->Type);
 	}
 
 	UE_MIR_UNREACHABLE();
@@ -1523,7 +1587,10 @@ static FValue* TryFoldOperator(FEmitter* Emitter, EOperator Op, FValue* A, FValu
 	FConstant* AComponents[4];
 	int32      ANumComponents;
 	
-	// Try to unpack the first argument to an array of constants.
+	// Some operations like Length, Dot and Cross are not defined on individual scalar components.
+	// For instance length(V) is not the same as [length(V.x), ..., length(V.z)]. These operations
+	// folding is handled here as special cases.
+	// First, try to unpack the first argument to an array of constants.
 	if (TryUnpackConstantScalarOrVector(A, AComponents, ANumComponents))
 	{
 		FConstant* BComponents[4];
@@ -1713,7 +1780,7 @@ static FValue* ConstructPrimitiveValue(FEmitter* Emitter, const FPrimitiveType*
 	// Construct a scalar from another scalar.
 	if (TargetPrimitiveType->IsScalar())
 	{
-		//
+		// 
 		Initializer = Emitter->EmitSubscript(Initializer, 0);
 		InitializerPrimitiveType = Initializer->Type->AsPrimitive();
 		
@@ -1730,13 +1797,16 @@ static FValue* ConstructPrimitiveValue(FEmitter* Emitter, const FPrimitiveType*
 		}
 		else
 		{
-			/* todo: emit conversion unary operator */
-			UE_MIR_TODO();
+			// Emit the cast to the target type of the subscript value
+			FCast Prototype = MakePrototype<FCast>(TargetPrimitiveType);
+			Prototype.Arg = Initializer;
+			return EmitPrototype(Emitter, Prototype);
 		}
 	}
 
-	// Construct a vector or matrix from a scalar. E.g. float4(3.14f)
-	if (!TargetPrimitiveType->IsScalar() && InitializerPrimitiveType->IsScalar())
+	// Construct a vector or matrix from a scalar. E.g. 3.14f -> float3(3.14f, 3.14f, 3.14f)
+	// Note: we know target isn't scalar as it's been handled above.
+	if (InitializerPrimitiveType->IsScalar())
 	{
 		// Create the result dimensional value.
 		FDimensional* Result = NewDimensionalValue(Emitter, TargetPrimitiveType);
@@ -1762,6 +1832,8 @@ static FValue* ConstructPrimitiveValue(FEmitter* Emitter, const FPrimitiveType*
 	// the necessary components.
 	if (TargetPrimitiveType->IsVector() && InitializerPrimitiveType->IsVector())
 	{
+		// #todo-massimo.tristano Use swizzle when scalartypes are the same, and target num components is less than initializer's.
+
 		int32 TargetNumComponents = TargetPrimitiveType->GetNumComponents();
 		int32 InitializerNumComponents = InitializerPrimitiveType->GetNumComponents();
 
@@ -1774,7 +1846,7 @@ static FValue* ConstructPrimitiveValue(FEmitter* Emitter, const FPrimitiveType*
 
 		// For iterating over the components of the result dimensional value.
 		int32 Index = 0;
-
+		
 		// Convert components from the initializer vector.
 		const int32 MinNumComponents = FMath::Min(TargetNumComponents, InitializerNumComponents);
 		for (; Index < MinNumComponents; ++Index)
@@ -1797,7 +1869,7 @@ static FValue* ConstructPrimitiveValue(FEmitter* Emitter, const FPrimitiveType*
 	{
 		check(TargetPrimitiveType->IsMatrix());
 
-		// 
+		//  If initializer is a dimensional, we can subscript components of interest.
 		if (FDimensional* DimensionalInitializer = As<FDimensional>(Initializer))
 		{
 			// Create the result dimensional value.
@@ -1851,7 +1923,7 @@ FValue* FEmitter::EmitConstruct(const FType* Type, FValue* Initializer)
 	// No other legal conversions applicable. Report error if we haven't converted the value.
 	if (!Result)
 	{
-		Errorf(TEXT("Cannot construct a '%s' from a '%s'."), TEXT(""), TEXT(""));
+		Errorf(TEXT("Cannot construct a '%s' from a '%s'."), Type->GetSpelling().GetData(), Initializer->Type->GetSpelling().GetData());
 		return FPoison::Get();
 	}
 	
@@ -1935,7 +2007,7 @@ const FType* FEmitter::GetCommonType(const FType* A, const FType* B)
 		return CommonType;
 	}
 
-	Errorf(TEXT("No common type between `%s` and `%s`."), A->GetSpelling().GetData(), B->GetSpelling().GetData());
+	Errorf(TEXT("No common type between '%s' and '%s'."), A->GetSpelling().GetData(), B->GetSpelling().GetData());
 	return nullptr;
 }
 
diff --git a/Engine/Source/Runtime/Engine/Private/Materials/MaterialIRToHLSLTranslator.cpp b/Engine/Source/Runtime/Engine/Private/Materials/MaterialIRToHLSLTranslator.cpp
--- a/Engine/Source/Runtime/Engine/Private/Materials/MaterialIRToHLSLTranslator.cpp
+++ b/Engine/Source/Runtime/Engine/Private/Materials/MaterialIRToHLSLTranslator.cpp
@@ -600,6 +600,7 @@ struct FTranslator : FMaterialIRToHLSLTranslation
 				case MIR::BO_Max: OpString = TEXT("max"); break;
 				case MIR::BO_Min: OpString = TEXT("min"); break;
 				case MIR::TO_Lerp: OpString = TEXT("lerp"); break;
+				case MIR::TO_Select: OpString = TEXT("select"); break;
 				case MIR::TO_Smoothstep: OpString = TEXT("smoothstep"); break;
 				case MIR::UO_Abs: OpString = TEXT("abs"); break;
 				case MIR::UO_ACos: OpString = TEXT("acos"); break;
diff --git a/Engine/Source/Runtime/Engine/Public/Materials/MaterialIR.h b/Engine/Source/Runtime/Engine/Public/Materials/MaterialIR.h
--- a/Engine/Source/Runtime/Engine/Public/Materials/MaterialIR.h
+++ b/Engine/Source/Runtime/Engine/Public/Materials/MaterialIR.h
@@ -198,16 +198,22 @@ struct FValue
 	/** @return Whether this value is a constant boolean with value false. */
 	bool IsFalse() const;
 
-	/** @return Whether this value is arithmetic and exactly zero. */
+	/** @return Whether this value is boolean and all components are true. */
+	bool AreAllTrue() const;
+
+	/** @return Whether this value is boolean and all components are false. */
+	bool AreAllFalse() const;
+
+	/** @return Whether this value is arithmetic and all components are exactly zero. */
 	bool AreAllExactlyZero() const;
 
-	/** @return Whether this value is arithmetic and approximately zero. */
+	/** @return Whether this value is arithmetic and all components are approximately zero. */
 	bool AreAllNearlyZero() const;
 
-	/** @return Whether this value is arithmetic and exactly one. */
+	/** @return Whether this value is arithmetic and all components are exactly one. */
 	bool AreAllExactlyOne() const;
 
-	/** @return Whether this value is arithmetic and approximately one. */
+	/** @return Whether this value is arithmetic and all components are approximately one. */
 	bool AreAllNearlyOne() const;
 	
 	/** Tries to cast this value to specified type T and returns the casted pointer, if possible (nullptr otherwise). */
@@ -642,6 +648,7 @@ enum EOperator
 
 	TO_Clamp = TO_FirstTernaryOperator,
 	TO_Lerp,
+	TO_Select,
 	TO_Smoothstep,
 
 	OperatorCount,
diff --git a/Engine/Source/Runtime/Engine/Public/Materials/MaterialIREmitter.h b/Engine/Source/Runtime/Engine/Public/Materials/MaterialIREmitter.h
--- a/Engine/Source/Runtime/Engine/Public/Materials/MaterialIREmitter.h
+++ b/Engine/Source/Runtime/Engine/Public/Materials/MaterialIREmitter.h
@@ -136,6 +136,7 @@ public:
 
 	FValue* EmitConstantFromShaderValue(const UE::Shader::FValue& InValue);
 	FValue* EmitConstantScalarZero(EScalarKind Kind);
+	FValue* EmitConstantScalarOne(EScalarKind Kind);
 	FValue* EmitConstantTrue();
 	FValue* EmitConstantFalse();
 	FValue* EmitConstantBoolScalar(bool InX);
