[MetaSounds][Wave Player] - Loop regions do not trigger as Cue Points

diff --git a/Engine/Plugins/Runtime/Metasound/Source/MetasoundEngine/Private/MetasoundWavePlayerNode.cpp b/Engine/Plugins/Runtime/Metasound/Source/MetasoundEngine/Private/MetasoundWavePlayerNode.cpp
--- a/Engine/Plugins/Runtime/Metasound/Source/MetasoundEngine/Private/MetasoundWavePlayerNode.cpp
+++ b/Engine/Plugins/Runtime/Metasound/Source/MetasoundEngine/Private/MetasoundWavePlayerNode.cpp
@@ -801,8 +801,8 @@ namespace Metasound
 			{
 				UE_LOG(LogMetaSound, Verbose, TEXT("Starting Sound: '%s'"), *CurrentWaveAsset->GetFName().ToString());
 
-				// Create local sorted copy of cue points.
-				SortedCuePoints = WaveProxy->GetCuePoints();
+				// Create local sorted copy of cue points and loop regions.
+				SortedCuePoints = WaveProxy->GetCuePointsAndLoopRegions();
 				Algo::SortBy(SortedCuePoints, WavePlayerNodePrivate::GetCuePointFrame);
 				
 				// Create the wave proxy reader.
diff --git a/Engine/Source/Runtime/Engine/Classes/Sound/SoundWave.h b/Engine/Source/Runtime/Engine/Classes/Sound/SoundWave.h
--- a/Engine/Source/Runtime/Engine/Classes/Sound/SoundWave.h
+++ b/Engine/Source/Runtime/Engine/Classes/Sound/SoundWave.h
@@ -1711,6 +1711,7 @@ public:
 	ENGINE_API uint32 GetSizeOfChunk(uint32 ChunkIndex) const;
 	ENGINE_API const TArray<FSoundWaveCuePoint>& GetCuePoints() const;
 	ENGINE_API const TArray<FSoundWaveCuePoint>& GetLoopRegions() const;
+	ENGINE_API const TArray<FSoundWaveCuePoint> GetCuePointsAndLoopRegions() const;
 
 	ENGINE_API bool IsLooping() const;
 	ENGINE_API bool IsTemplate() const;
diff --git a/Engine/Source/Runtime/Engine/Private/SoundWave.cpp b/Engine/Source/Runtime/Engine/Private/SoundWave.cpp
--- a/Engine/Source/Runtime/Engine/Private/SoundWave.cpp
+++ b/Engine/Source/Runtime/Engine/Private/SoundWave.cpp
@@ -4929,6 +4929,23 @@ const TArray<FSoundWaveCuePoint>& FSoundWaveProxy::GetLoopRegions() const
 	return SoundWaveDataPtr->GetLoopRegions();
 }
 
+const TArray<FSoundWaveCuePoint> FSoundWaveProxy::GetCuePointsAndLoopRegions() const
+{
+	check(SoundWaveDataPtr);
+	
+	TArray<FSoundWaveCuePoint> CuePointsAndLoopRegions;
+	CuePointsAndLoopRegions.Reserve(SoundWaveDataPtr->GetLoopRegions().Num() + SoundWaveDataPtr->GetCuePoints().Num());
+
+	// Previously loops had two cue points (one marked as a cue point, the other a loop region),
+	// now they only have one that is marked as a loop region so we need to append Loop Regions.
+	// Append must happen in this order to maintain the label for loops that were imported before the fix because
+	// meta sounds will pick only one cue point to render if they have matching frames.
+	CuePointsAndLoopRegions.Append(SoundWaveDataPtr->GetLoopRegions());
+	CuePointsAndLoopRegions.Append(SoundWaveDataPtr->GetCuePoints());
+	
+	return CuePointsAndLoopRegions;
+}
+
 bool FSoundWaveProxy::IsLooping() const
 {
 	check(SoundWaveDataPtr);
