[AutoRTFM] Replace includes of "AutoRTFM/AutoRTFMTesting.h" with "AutoRTFMTesting.h"

Delete the now unused "AutoRTFM/AutoRTFMTesting.h".
Part of the clean-up post move of AutoRTFM from Core to its own module.

diff --git a/Engine/Source/Runtime/AutoRTFM/Public/AutoRTFM/AutoRTFMTesting.h b/Engine/Source/Runtime/AutoRTFM/Public/AutoRTFM/AutoRTFMTesting.h
deleted file mode 100644
diff --git a/Engine/Source/Runtime/Engine/Private/Tests/AutoRTFM/AutoRTFMActorComponentTests.cpp b/Engine/Source/Runtime/Engine/Private/Tests/AutoRTFM/AutoRTFMActorComponentTests.cpp
--- a/Engine/Source/Runtime/Engine/Private/Tests/AutoRTFM/AutoRTFMActorComponentTests.cpp
+++ b/Engine/Source/Runtime/Engine/Private/Tests/AutoRTFM/AutoRTFMActorComponentTests.cpp
@@ -3,7 +3,7 @@
 #include "Misc/AutomationTest.h"
 #include "Animation/AnimInstance.h"
 #include "Animation/Skeleton.h"
-#include "AutoRTFM/AutoRTFMTesting.h"
+#include "AutoRTFMTesting.h"
 #include "AutoRTFMTestActor.h"
 #include "AutoRTFMTestAnotherActor.h"
 #include "AutoRTFMTestBodySetup.h"
diff --git a/Engine/Source/Runtime/Engine/Private/Tests/AutoRTFM/AutoRTFMNetworkProfilerTests.cpp b/Engine/Source/Runtime/Engine/Private/Tests/AutoRTFM/AutoRTFMNetworkProfilerTests.cpp
--- a/Engine/Source/Runtime/Engine/Private/Tests/AutoRTFM/AutoRTFMNetworkProfilerTests.cpp
+++ b/Engine/Source/Runtime/Engine/Private/Tests/AutoRTFM/AutoRTFMNetworkProfilerTests.cpp
@@ -1,6 +1,6 @@
 // Copyright Epic Games, Inc. All Rights Reserved.
 
-#include "AutoRTFM/AutoRTFMTesting.h"
+#include "AutoRTFMTesting.h"
 #include "AutoRTFMTestActor.h"
 #include "Engine/ActorChannel.h"
 #include "Engine/DemoNetDriver.h"
diff --git a/Engine/Source/Runtime/Engine/Private/Tests/AutoRTFM/AutoRTFMSlateTests.cpp b/Engine/Source/Runtime/Engine/Private/Tests/AutoRTFM/AutoRTFMSlateTests.cpp
--- a/Engine/Source/Runtime/Engine/Private/Tests/AutoRTFM/AutoRTFMSlateTests.cpp
+++ b/Engine/Source/Runtime/Engine/Private/Tests/AutoRTFM/AutoRTFMSlateTests.cpp
@@ -1,6 +1,6 @@
 // Copyright Epic Games, Inc. All Rights Reserved.
 
-#include "AutoRTFM/AutoRTFMTesting.h"
+#include "AutoRTFMTesting.h"
 #include "Misc/AutomationTest.h"
 #include "AutoRTFM/AutoRTFM.h"
 
diff --git a/Engine/Source/Runtime/Engine/Private/Tests/AutoRTFM/AutoRTFMTestingAssertionFailure.cpp b/Engine/Source/Runtime/Engine/Private/Tests/AutoRTFM/AutoRTFMTestingAssertionFailure.cpp
--- a/Engine/Source/Runtime/Engine/Private/Tests/AutoRTFM/AutoRTFMTestingAssertionFailure.cpp
+++ b/Engine/Source/Runtime/Engine/Private/Tests/AutoRTFM/AutoRTFMTestingAssertionFailure.cpp
@@ -1,6 +1,6 @@
 // Copyright Epic Games, Inc. All Rights Reserved.
 
-#include "AutoRTFM/AutoRTFMTesting.h"
+#include "AutoRTFMTesting.h"
 #include "Misc/AssertionMacros.h"
 
 #if WITH_DEV_AUTOMATION_TESTS
