EngineTest tests that more than 4 clients can connect to a MultiServerProxy and the correct client handshake ids are propagated to connections and controllers.

diff --git a/Engine/Plugins/Runtime/MultiServerReplication/Source/MultiServerReplication/Private/MultiServerProxy.cpp b/Engine/Plugins/Runtime/MultiServerReplication/Source/MultiServerReplication/Private/MultiServerProxy.cpp
--- a/Engine/Plugins/Runtime/MultiServerReplication/Source/MultiServerReplication/Private/MultiServerProxy.cpp
+++ b/Engine/Plugins/Runtime/MultiServerReplication/Source/MultiServerReplication/Private/MultiServerProxy.cpp
@@ -1062,6 +1062,7 @@ void UProxyNetDriver::AttachPlayerController(APlayerController* PlayerController
 		   *Player->GetName());
 
 	Connection->SetConnectionState(EConnectionState::USOCK_Open);
+	Connection->SetClientHandshakeId(PlayerController->GetClientHandshakeId());
 
 	// It's assumed the connection is not associated with a player controller.	
 	ensure(Connection->PlayerController == nullptr);
diff --git a/Engine/Source/Runtime/Engine/Private/NetConnection.cpp b/Engine/Source/Runtime/Engine/Private/NetConnection.cpp
--- a/Engine/Source/Runtime/Engine/Private/NetConnection.cpp
+++ b/Engine/Source/Runtime/Engine/Private/NetConnection.cpp
@@ -4984,6 +4984,7 @@ void UNetConnection::HandleClientPlayer( APlayerController *PC, UNetConnection*
 	UE_LOG(LogNet, Verbose, TEXT("%s setplayer %s"),*PC->GetName(),*LocalPlayer->GetName());
 	LastReceiveTime = Driver->GetElapsedTime();
 	SetConnectionState(EConnectionState::USOCK_Open);
+	SetClientHandshakeId(PC->GetClientHandshakeId());
 	PlayerController = PC;
 	OwningActor = PC;
 
@@ -5109,6 +5110,7 @@ void UChildConnection::HandleClientPlayer(APlayerController* PC, UNetConnection*
 	UE_LOG(LogNet, Verbose, TEXT("%s setplayer %s"), *PC->GetName(), *NewPlayer->GetName());
 	PlayerController = PC;
 	OwningActor = PC;
+	SetClientHandshakeId(PC->GetClientHandshakeId());
 
 	NetConnection->GetOnConnectionOpened().Broadcast(this, PC, NetConnection);
 
