World Partition HLOD: handle actor set instances with invalid bounds, can happen if non-editor-only actors only have editor-only components and weren't resaved since 40945935.

diff --git a/Engine/Source/Runtime/Engine/Private/WorldPartition/RuntimeSpatialHash/RuntimeSpatialHashHLOD.cpp b/Engine/Source/Runtime/Engine/Private/WorldPartition/RuntimeSpatialHash/RuntimeSpatialHashHLOD.cpp
--- a/Engine/Source/Runtime/Engine/Private/WorldPartition/RuntimeSpatialHash/RuntimeSpatialHashHLOD.cpp
+++ b/Engine/Source/Runtime/Engine/Private/WorldPartition/RuntimeSpatialHash/RuntimeSpatialHashHLOD.cpp
@@ -195,9 +195,12 @@ static TArray<FGuid> GenerateHLODActorsForGrid(UWorldPartition* WorldPartition,
 					{
 						const IStreamingGenerationContext::FActorInstance& ActorInstance = ActorInstances.Emplace_GetRef(ActorGuid, ActorSetInstance);
 						const FBox RuntimeBounds = ActorInstance.GetActorDescView().GetRuntimeBounds();
-						check(RuntimeBounds.IsValid);
-						CellBounds.Min.Z = FMath::Min(CellBounds.Min.Z, RuntimeBounds.Min.Z);
-						CellBounds.Max.Z = FMath::Max(CellBounds.Max.Z, RuntimeBounds.Max.Z);
+
+						if (RuntimeBounds.IsValid)
+						{
+							CellBounds.Min.Z = FMath::Min(CellBounds.Min.Z, RuntimeBounds.Min.Z);
+							CellBounds.Max.Z = FMath::Max(CellBounds.Max.Z, RuntimeBounds.Max.Z);
+						}
 					});
 				}
 
