[Source Location] Update usage of UE::SourceLocation functions to use the new adapter API. The other functions will be deprecated in favour of this.

diff --git a/Engine/Plugins/Experimental/PythonScriptPlugin/Source/PythonScriptPlugin/Private/PythonScriptPlugin.cpp b/Engine/Plugins/Experimental/PythonScriptPlugin/Source/PythonScriptPlugin/Private/PythonScriptPlugin.cpp
--- a/Engine/Plugins/Experimental/PythonScriptPlugin/Source/PythonScriptPlugin/Private/PythonScriptPlugin.cpp
+++ b/Engine/Plugins/Experimental/PythonScriptPlugin/Source/PythonScriptPlugin/Private/PythonScriptPlugin.cpp
@@ -705,7 +705,7 @@ bool FPythonScriptPlugin::IsPythonInitialized() const
 
 bool FPythonScriptPlugin::ForceEnablePythonAtRuntime(UE::FSourceLocation Location)
 {
-	return ForceEnablePythonAtRuntime(UE::SourceLocation::ToFullString(Location));
+	return ForceEnablePythonAtRuntime(UE::SourceLocation::Full(Location).ToString());
 }
 
 bool FPythonScriptPlugin::ForceEnablePythonAtRuntime(const FString& Location)
diff --git a/Engine/Source/Runtime/Engine/Classes/Engine/StreamableManager.h b/Engine/Source/Runtime/Engine/Classes/Engine/StreamableManager.h
--- a/Engine/Source/Runtime/Engine/Classes/Engine/StreamableManager.h
+++ b/Engine/Source/Runtime/Engine/Classes/Engine/StreamableManager.h
@@ -121,7 +121,7 @@ struct FStreamableDownloadCachePin
 		if constexpr (std::is_same_v<std::decay_t<DebugNameType>, UE::FSourceLocation>)
 		{
 			Result.ContentHandle = UE::IoStore::FOnDemandContentHandle::Create(
-				UE::SourceLocation::ToFileAndLineString(DebugNameOrLocation));
+				UE::SourceLocation::FileAndLine(DebugNameOrLocation).ToString());
 		}
 		else
 		{
diff --git a/Engine/Source/Runtime/Engine/Private/StreamableManager.cpp b/Engine/Source/Runtime/Engine/Private/StreamableManager.cpp
--- a/Engine/Source/Runtime/Engine/Private/StreamableManager.cpp
+++ b/Engine/Source/Runtime/Engine/Private/StreamableManager.cpp
@@ -1610,14 +1610,14 @@ bool FStreamableManager::GetReferencerPropertyName(UObject* Object, FString& Out
 					PropName.Appendf(TEXT("Active[1/%d]:"), Existing->ActiveHandles.Num());
 					PropName << Existing->ActiveHandles[0]->GetDebugName();
 					PropName << TEXT(" @ ");
-					PropName << UE::SourceLocation::ToFileAndLineString(Existing->ActiveHandles[0]->Location);
+					PropName << UE::SourceLocation::FileAndLine(Existing->ActiveHandles[0]->Location);
 				}
 				else if (Existing->LoadingHandles.Num() > 0)
 				{
 					PropName.Appendf(TEXT("Loading[1/%d]:"), Existing->LoadingHandles.Num());
 					PropName << Existing->LoadingHandles[0]->GetDebugName();
 					PropName << TEXT(" @ ");
-					PropName << UE::SourceLocation::ToFileAndLineString(Existing->LoadingHandles[0]->Location);
+					PropName << UE::SourceLocation::FileAndLine(Existing->LoadingHandles[0]->Location);
 				}
 
 				if (PropName.Len() > 0)
