Enable sample accurate seeking in streaming sound waves
Enable additional tests for audio decoding
#rnx

diff --git a/Engine/Source/Developer/AudioFormatOpus/Private/AudioFormatOpus.cpp b/Engine/Source/Developer/AudioFormatOpus/Private/AudioFormatOpus.cpp
--- a/Engine/Source/Developer/AudioFormatOpus/Private/AudioFormatOpus.cpp
+++ b/Engine/Source/Developer/AudioFormatOpus/Private/AudioFormatOpus.cpp
@@ -28,7 +28,7 @@ class FAudioFormatOpus : public IAudioFormat
 	enum
 	{
 		/** Version for OPUS format, this becomes part of the DDC key. */
-		UE_AUDIO_OPUS_VER = 12,
+		UE_AUDIO_OPUS_VER = 13,
 	};
 
 public:
diff --git a/Engine/Source/Developer/AudioFormatRad/Private/AudioFormatRad.cpp b/Engine/Source/Developer/AudioFormatRad/Private/AudioFormatRad.cpp
--- a/Engine/Source/Developer/AudioFormatRad/Private/AudioFormatRad.cpp
+++ b/Engine/Source/Developer/AudioFormatRad/Private/AudioFormatRad.cpp
@@ -41,7 +41,7 @@ class FAudioFormatRad : public IAudioFormat
 	enum
 	{
 		/** Version for RAD Audio format, this becomes part of the DDC key. */
-		UE_AUDIO_RAD_VER = 4,
+		UE_AUDIO_RAD_VER = 5,
 	};
 
 public:
diff --git a/Engine/Source/Runtime/Engine/Private/AudioDecompress.cpp b/Engine/Source/Runtime/Engine/Private/AudioDecompress.cpp
--- a/Engine/Source/Runtime/Engine/Private/AudioDecompress.cpp
+++ b/Engine/Source/Runtime/Engine/Private/AudioDecompress.cpp
@@ -322,8 +322,9 @@ bool IStreamedCompressedInfo::StreamCompressedData(uint8* Destination, bool bLoo
 	if (StreamSeekBlockIndex != INDEX_NONE)
 	{
 		UE_CLOG(WantsFilteredLogging(), LogAudioStreamCaching, Display, TEXT("[%s.%p] - Seek Requesting Block = %d"), *WriteToString<128>(StreamingSoundWave->GetFName()), this, StreamSeekBlockIndex.load(std::memory_order_relaxed));
+
 		uint32 ChunkSize = 0;
-		const uint8* NewlySeekedChunk = GetLoadedChunk(StreamingSoundWave, StreamSeekBlockIndex, ChunkSize);
+		const uint8* NewlySeekedChunk = GetLoadedChunk(StreamingSoundWave, StreamSeekBlockIndex, ChunkSize); 
 		UE_LOG(LogAudio, Log, TEXT("Seek request for (%s): (%s) Chunk=%d (%s), Offset=%d, OffsetInAudioFrames=%u"),
 			*StreamingSoundWave->GetFName().ToString(),
 			StreamSeekToAudioFrames != INDEX_NONE ? TEXT("Using streaming seek-tables") : TEXT("Using chunk/offset pair"),
@@ -355,19 +356,28 @@ bool IStreamedCompressedInfo::StreamCompressedData(uint8* Destination, bool bLoo
 		else
 		{
 			// If we are using a chunked seek-table, we need to parse the table first.
-			if (StreamingSoundWave->GetSoundWaveData()->HasChunkSeekTable(CurrentChunkIndex))
+			if (StreamingSoundWave->GetSoundWaveData()->HasChunkSeekTable(StreamSeekBlockIndex))
 			{
 				uint32 TableOffset = 0;
 				if (ensureMsgf(FStreamedAudioChunkSeekTable::Parse(NewlySeekedChunk, ChunkSize, TableOffset, GetCurrentSeekTable()),
-					TEXT("Failed to parse seektable in '%s', chunk=%d"), *StreamingSoundWave->GetFName().ToString(), CurrentChunkIndex))
+					TEXT("Failed to parse seektable in '%s', chunk=%d"), *StreamingSoundWave->GetFName().ToString(), StreamSeekBlockIndex.load(std::memory_order_relaxed)))
 				{
 					if (StreamSeekToAudioFrames != INDEX_NONE)
 					{
-						const FStreamedAudioChunk& Chunk = StreamingSoundWave->GetSoundWaveData()->GetChunk(CurrentChunkIndex);
-						int32 SeekOffsetInChunkSpace = StreamSeekToAudioFrames - Chunk.SeekOffsetInAudioFrames;
-						uint32 Offset = GetCurrentSeekTable().FindOffset(StreamSeekToAudioFrames);
+						const uint32 MaxSeekableFrame = (TrueSampleCount / FMath::Max((uint32)NumChannels, uint32(1))) - 1;
+						StreamSeekToAudioFrames = FMath::Min(StreamSeekToAudioFrames, MaxSeekableFrame); // protect against seek past end of audio.
+						const FStreamedAudioChunkSeekTable& SeekTable = GetCurrentSeekTable();
+
+						uint32 Offset = SeekTable.FindOffset(StreamSeekToAudioFrames);
 						if (Offset != INDEX_NONE)
 						{
+							uint32 ActualSeekFrame = SeekTable.FindTime(Offset);
+							if (ActualSeekFrame == INDEX_NONE)
+							{
+								// Fallback to requested frame position if the actual frame could not be found.
+								ActualSeekFrame = StreamSeekToAudioFrames;
+								UE_LOG(LogAudioStreamCaching, Warning, TEXT("Failed to seek to frame %d in sound %s due to bad seek table"), StreamSeekToAudioFrames, *StreamingSoundWave->GetFName().ToString());
+							}
 							// All looks good, commit this as our current chunk.
 							CurrentChunkIndex = StreamSeekBlockIndex;
 							SrcBufferDataSize = ChunkSize;
@@ -378,12 +388,11 @@ bool IStreamedCompressedInfo::StreamCompressedData(uint8* Destination, bool bLoo
 							SrcBufferOffset += CurrentChunkIndex == 0 ? AudioDataOffset : 0;
 
 							// Convert frames to samples and update "current sample count"
-							CurrentSampleCount = StreamSeekToAudioFrames * NumChannels;
+							CurrentSampleCount = FMath::Min(ActualSeekFrame * NumChannels, TrueSampleCount); // Safeguard against corrupt seektable or bad implementation
 							LastPCMByteSize = 0;
 							LastPCMOffset = 0;
 							bStoringEndOfFile = false;
 
-
 							UE_CLOG(WantsFilteredLogging(), LogAudioStreamCaching, Display, TEXT("[%s.%p] - Committing Seek Chunk=%d Size=%u"), *WriteToString<128>(StreamingSoundWave->GetFName()), this, CurrentChunkIndex, ChunkSize);
 						}
 						else // Seek failed (off the end of the chunk).
@@ -396,9 +405,9 @@ bool IStreamedCompressedInfo::StreamCompressedData(uint8* Destination, bool bLoo
 			}
 		}
 
-		// Seeking logic over, turn off seeking state
+		// Chunk seeking logic over, turn off seeking state. 
+		// StreamSeekToAudioFrames remains unchanged. It will be used to achieve sample accurate seek. 
 		StreamSeekBlockIndex = INDEX_NONE;
-		StreamSeekToAudioFrames = INDEX_NONE;
 		StreamSeekBlockOffset = INDEX_NONE;
 	}
 	
@@ -482,11 +491,23 @@ bool IStreamedCompressedInfo::StreamCompressedData(uint8* Destination, bool bLoo
 		}
 		else
 		{
-			LastPCMByteSize = IncrementCurrentSampleCount(DecodedFrames * NumChannels) * MONO_PCM_SAMPLE_SIZE;
+			if (DecodedFrames > 0)
+			{
+				LastPCMByteSize = IncrementCurrentSampleCount(DecodedFrames * NumChannels) * MONO_PCM_SAMPLE_SIZE;
+
+				if (StreamSeekToAudioFrames != INDEX_NONE)
+				{
+					// Perform frame accurate seek by discarding decoded audio.
+					if (SeekToDecodedPCMFrame(StreamSeekToAudioFrames))
+					{
+						StreamSeekToAudioFrames = INDEX_NONE;
+					}
+				}
 
-			// update OutNumBytesStreamed as we write out data
-			RawPCMOffset += WriteFromDecodedPCM(Destination + RawPCMOffset, BufferSize - RawPCMOffset);
-			OutNumBytesStreamed = RawPCMOffset;
+				// update OutNumBytesStreamed as we write out data
+				RawPCMOffset += WriteFromDecodedPCM(Destination + RawPCMOffset, BufferSize - RawPCMOffset);
+				OutNumBytesStreamed = RawPCMOffset;
+			}
 
 			const int32 PreviousChunkIndex = CurrentChunkIndex;
 
@@ -610,6 +631,7 @@ int32 IStreamedCompressedInfo::DecompressToPCMBuffer(uint16 /*Unused*/ )
 
 uint32 IStreamedCompressedInfo::IncrementCurrentSampleCount(uint32 NewSamples)
 {
+	check(CurrentSampleCount <= TrueSampleCount);
 	if (CurrentSampleCount + NewSamples > TrueSampleCount)
 	{
 		NewSamples = TrueSampleCount - CurrentSampleCount;
@@ -644,6 +666,56 @@ uint32 IStreamedCompressedInfo::WriteFromDecodedPCM(uint8* Destination, uint32 B
 	return BytesToCopy;
 }
 
+bool IStreamedCompressedInfo::SeekToDecodedPCMFrame(uint32 InTargetAudioFrame)
+{	
+	// The MaxAvailablePCMBytePos and MinAvailablePCMBytePos are global byte positions
+	// with respect to the entire sound wave. Conceptually, it is as though the entire
+	// sound wave was decoded into PCM, and these values represent the portion that is
+	// in the `LastDecodedPCM` buffer. In reality, because we are streaming sounds,
+	// we only have the PCM data in `LastDecodedPCM` available for rendering.
+	uint32 MaxAvailablePCMBytePos = CurrentSampleCount * MONO_PCM_SAMPLE_SIZE; 
+	check(LastPCMByteSize <= MaxAvailablePCMBytePos);
+	uint32 MinAvailablePCMBytePos = MaxAvailablePCMBytePos - LastPCMByteSize;
+
+	uint32 TargetPCMBytePos = InTargetAudioFrame * NumChannels * MONO_PCM_SAMPLE_SIZE;
+
+	if (TargetPCMBytePos >= MinAvailablePCMBytePos) 
+	{
+		LastPCMOffset = TargetPCMBytePos - MinAvailablePCMBytePos;
+		if (LastPCMOffset < LastPCMByteSize)
+		{
+			// LastPCMOffset is set to the correct location which corresponds
+			// to the target audio frame. 
+			return true;
+		}
+		else
+		{
+			// More audio needs to be decoded before reaching the desired audio frame
+			LastPCMOffset = 0;
+			LastPCMByteSize = 0;
+			return false;
+		}
+	}
+	else
+	{
+		// The target audio frame is in the past. Rewinding is not supported when 
+		// performing a sample accurate seek. 
+		uint32 CurrentFrameOffset = MinAvailablePCMBytePos / FMath::Max(static_cast<uint32>(NumChannels * MONO_PCM_SAMPLE_SIZE), uint32(1));
+		UE_LOG(LogAudio, Warning, TEXT("Failed to seek to decoded PCM frame. Seek request for (%s): (%s) Chunk=%d, Offset=%d, OffsetInAudioFrames=%u, SeekInAudioFrames=%u"),
+			*StreamingSoundWave->GetFName().ToString(),
+			StreamSeekToAudioFrames != INDEX_NONE ? TEXT("Using streaming seek-tables") : TEXT("Using chunk/offset pair"),
+			CurrentChunkIndex,
+			SrcBufferOffset,
+			CurrentFrameOffset,
+			StreamSeekToAudioFrames
+		);
+
+		// Return true to signal that we should no longer attempt seeks into the 
+		// decoded buffer. The seek is invalid. 
+		return true;
+	}
+}
+
 uint32 IStreamedCompressedInfo::ZeroBuffer(uint8* Destination, uint32 BufferSize)
 {
 	check(Destination);
diff --git a/Engine/Source/Runtime/Engine/Private/StreamedAudioChunkSeekTable.cpp b/Engine/Source/Runtime/Engine/Private/StreamedAudioChunkSeekTable.cpp
--- a/Engine/Source/Runtime/Engine/Private/StreamedAudioChunkSeekTable.cpp
+++ b/Engine/Source/Runtime/Engine/Private/StreamedAudioChunkSeekTable.cpp
@@ -158,12 +158,24 @@ struct FConstantRateSeekTable : public FStreamedAudioChunkSeekTable::ISeekTableI
 		}
 		return INDEX_NONE;
 	}
+
 	uint32 FindTime(uint32 InOffset) const override
 	{
 		if (!Table.IsEmpty() && AudioFramesPerEntry > 0 && StartTimeOffset >= 0)
 		{
-			uint32 TimeInSamples = StartTimeOffset + (InOffset * AudioFramesPerEntry);
-			return TimeInSamples;
+			if (InOffset == 0)
+			{
+				return StartTimeOffset;
+			}
+			else
+			{
+				int32 Index = Table.Find(InOffset);
+				if (INDEX_NONE != Index)
+				{
+					uint32 TimeInSamples = StartTimeOffset + (Index * AudioFramesPerEntry);
+					return TimeInSamples;
+				}
+			}
 		}
 		return INDEX_NONE;
 	}
@@ -244,14 +256,12 @@ struct FVariableRateSeekTable : public FStreamedAudioChunkSeekTable::ISeekTableI
 	uint32 FindOffset(uint32 InTimeInAudioFrames) const override
 	{
 		check(Offsets.Num() == Times.Num());
-		const int32 Index = Algo::LowerBound(Times, InTimeInAudioFrames);
+		const int32 Index = Algo::UpperBound(Times, InTimeInAudioFrames) - 1;
 		if (Offsets.IsValidIndex(Index))
 		{
-			// Minor hack here. RADA is the only customer of this currently and needs
-			// the one block behind the one we've found. 
 			if (Index > 0)
 			{
-				return Offsets[Index-1];
+				return Offsets[Index];
 			}
 			return Offsets[0];
 		}
diff --git a/Engine/Source/Runtime/Engine/Public/AudioDecompress.h b/Engine/Source/Runtime/Engine/Public/AudioDecompress.h
--- a/Engine/Source/Runtime/Engine/Public/AudioDecompress.h
+++ b/Engine/Source/Runtime/Engine/Public/AudioDecompress.h
@@ -362,6 +362,11 @@ protected:
 
 	/** If using the Chunked seek-tables, we request the seek in samples so we can resolve after the chunk table loads. */
 	uint32 StreamSeekToAudioFrames = INDEX_NONE;	
+
+private:
+	/** Attempts a seek to the frame by consuming decoded PCM data. Returns true
+	 * if future seeks should no longer be attempted. */
+	bool SeekToDecodedPCMFrame(uint32 InTargetAudioFrame);
 };
 
 struct IAudioInfoFactory 
