Add trace scope events to Actor Descriptor Container Instance initialization hot path

diff --git a/Engine/Source/Runtime/AssetRegistry/Private/AssetRegistry.cpp b/Engine/Source/Runtime/AssetRegistry/Private/AssetRegistry.cpp
--- a/Engine/Source/Runtime/AssetRegistry/Private/AssetRegistry.cpp
+++ b/Engine/Source/Runtime/AssetRegistry/Private/AssetRegistry.cpp
@@ -2449,6 +2449,7 @@ bool UAssetRegistryImpl::GetAssets(const FARFilter& InFilter, TArray<FAssetData>
 bool UAssetRegistryImpl::GetAssets(const FARCompiledFilter& CompiledFilter, TArray<FAssetData>& OutAssetData,
 	bool bSkipARFilteredAssets) const
 {
+	TRACE_CPUPROFILER_EVENT_SCOPE(UAssetRegistryImpl::GetAssets);
 	using namespace UE::AssetRegistry::Utils;
 
 	TSet<FName> PackagesToSkip;
diff --git a/Engine/Source/Runtime/Engine/Private/WorldPartition/ActorDescContainerInstance.cpp b/Engine/Source/Runtime/Engine/Private/WorldPartition/ActorDescContainerInstance.cpp
--- a/Engine/Source/Runtime/Engine/Private/WorldPartition/ActorDescContainerInstance.cpp
+++ b/Engine/Source/Runtime/Engine/Private/WorldPartition/ActorDescContainerInstance.cpp
@@ -71,6 +71,8 @@ void UActorDescContainerInstance::OnContainerReplaced(UActorDescContainer* InOld
 
 void UActorDescContainerInstance::Initialize(const FInitializeParams& InParams)
 {
+	TRACE_CPUPROFILER_EVENT_SCOPE(UActorDescContainerInstance::Initialize);
+
 	FName OuterWorldContainerPackageName;
 
 	Transform = InParams.Transform;
diff --git a/Engine/Source/Runtime/Engine/Private/WorldPartition/WorldPartitionActorDesc.cpp b/Engine/Source/Runtime/Engine/Private/WorldPartition/WorldPartitionActorDesc.cpp
--- a/Engine/Source/Runtime/Engine/Private/WorldPartition/WorldPartitionActorDesc.cpp
+++ b/Engine/Source/Runtime/Engine/Private/WorldPartition/WorldPartitionActorDesc.cpp
@@ -234,6 +234,8 @@ void FWorldPartitionActorDesc::InitTransientProperties(const FWorldPartitionActo
 
 void FWorldPartitionActorDesc::Init(const FWorldPartitionActorDescInitData& DescData)
 {
+	TRACE_CPUPROFILER_EVENT_SCOPE(FWorldPartitionActorDesc::Init);
+
 	InitTransientProperties(DescData);
 
 	auto DeprecateClass = [this](FArchive& Archive)
