World Partition: debug display of speeds now uses standard unit conversions for speed unit and display string.
Added wp.Runtime.DebugDisplaySpeedUnit console command to change the speed debug display unit (0=cm/s, 1=m/s, 2=km/h, 3=mi/h), defaults to mi/h.

diff --git a/Engine/Source/Runtime/Engine/Private/WorldPartition/WorldPartitionStreamingSource.cpp b/Engine/Source/Runtime/Engine/Private/WorldPartition/WorldPartitionStreamingSource.cpp
--- a/Engine/Source/Runtime/Engine/Private/WorldPartition/WorldPartitionStreamingSource.cpp
+++ b/Engine/Source/Runtime/Engine/Private/WorldPartition/WorldPartitionStreamingSource.cpp
@@ -2,6 +2,7 @@
 
 #include "WorldPartition/WorldPartitionStreamingSource.h"
 #include "WorldPartition/WorldPartitionStreamingPolicy.h"
+#include "Math/UnitConversion.h"
 #include "Misc/HashBuilder.h"
 
 #include UE_INLINE_GENERATED_CPP_BY_NAME(WorldPartitionStreamingSource)
@@ -20,6 +21,13 @@ FAutoConsoleVariableRef FWorldPartitionStreamingSource::CVarRotationQuantization
 	TEXT("Angle (in degrees) used to quantize the streaming sources rotation to determine if a world partition streaming update is necessary."),
 	ECVF_Default);
 
+int32 FWorldPartitionStreamingSource::DebugDisplaySpeedUnit = 3;
+FAutoConsoleVariableRef FWorldPartitionStreamingSource::CVarDebugDisplaySpeedUnit(
+	TEXT("wp.Runtime.DebugDisplaySpeedUnit"),
+	FWorldPartitionStreamingSource::DebugDisplaySpeedUnit,
+	TEXT("Unit used for debug display to show speeds (0=cm/s, 1=m/s, 2=km/h, 3=mi/h), defaults to mi/h."),
+	ECVF_Default);
+
 FString FStreamingSourceShape::ToString() const
 {
 	FStringBuilderBase StringBuilder;
@@ -45,16 +53,22 @@ FString FStreamingSourceShape::ToString() const
 FString FWorldPartitionStreamingSource::ToString() const
 {
 	FStringBuilderBase StringBuilder;
+
+	static EUnit SupportedUnits[] = { EUnit::CentimetersPerSecond, EUnit::MetersPerSecond, EUnit::KilometersPerHour, EUnit::MilesPerHour };
+	const EUnit VelocityUnitEnum = SupportedUnits[FMath::Min<int32>(DebugDisplaySpeedUnit, (sizeof(SupportedUnits) / sizeof((SupportedUnits)[0])) - 1)];
+	const int32 VelocityUnitValue = (int32)FUnitConversion::Convert(Velocity.Size(), EUnit::CentimetersPerSecond, VelocityUnitEnum);
+	const TCHAR* VelocityUnitStr = FUnitConversion::GetUnitDisplayString(VelocityUnitEnum);
+
 	StringBuilder.Appendf(
-		TEXT("Priority: %d | %s | %s | %s | Pos: X=%lld,Y=%lld,Z=%lld | Rot: %s | Vel: %3.2f m/s (%d mph)"),
+		TEXT("Priority: %d | %s | %s | %s | Pos: X=%lld,Y=%lld,Z=%lld | Rot: %s | Vel: %d %s"),
 		Priority,
 		bRemote ? TEXT("Remote") : TEXT("Local"),
 		GetStreamingSourceTargetStateName(TargetState),
 		bBlockOnSlowLoading ? TEXT("Blocking") : TEXT("NonBlocking"),
 		(int64)Location.X, (int64)Location.Y, (int64)Location.Z,
 		*Rotation.ToCompactString(),
-		Velocity.Size() * 0.01,
-		(int32)(Velocity.Size() * 0.0223694f)
+		VelocityUnitValue,
+		VelocityUnitStr
 	);
 
 	if (bForce2D)
diff --git a/Engine/Source/Runtime/Engine/Public/WorldPartition/WorldPartitionStreamingSource.h b/Engine/Source/Runtime/Engine/Public/WorldPartition/WorldPartitionStreamingSource.h
--- a/Engine/Source/Runtime/Engine/Public/WorldPartition/WorldPartitionStreamingSource.h
+++ b/Engine/Source/Runtime/Engine/Public/WorldPartition/WorldPartitionStreamingSource.h
@@ -492,6 +492,9 @@ private:
 	static ENGINE_API FAutoConsoleVariableRef CVarLocationQuantization;
 	static ENGINE_API FAutoConsoleVariableRef CVarRotationQuantization;
 
+	static ENGINE_API int32 DebugDisplaySpeedUnit;
+	static ENGINE_API FAutoConsoleVariableRef CVarDebugDisplaySpeedUnit;
+
 	/** Hash of streaming source (used to detect changes) */
 	uint32 Hash2D;
 	uint32 Hash3D;
