PR #13029: Removed additional multiply from URectLightComponent::ComputeLightBrightness

When light units are set to EV and computing the light brightness , the RectLightComponent should be using the return of EV100ToLuminance(LightBrightness) but is instead multiplying LightBrightness by the return value instead. This results in the LightBrightness value being multiplied twice.

PointLightComponent and SpotLightComponent do not share the same issue as they assign the result of EV100ToLuminance to LightBrightness rather than multiply against it.

diff --git a/Engine/Source/Runtime/Engine/Private/Components/RectLightComponent.cpp b/Engine/Source/Runtime/Engine/Private/Components/RectLightComponent.cpp
--- a/Engine/Source/Runtime/Engine/Private/Components/RectLightComponent.cpp
+++ b/Engine/Source/Runtime/Engine/Private/Components/RectLightComponent.cpp
@@ -186,7 +186,7 @@ float URectLightComponent::ComputeLightBrightness() const
 	}
 	else if (IntensityUnits == ELightUnits::EV)
 	{
-		LightBrightness *= (100.f * 100.f) * EV100ToLuminance(LightBrightness);
+		LightBrightness = EV100ToLuminance(LightBrightness) * (100.f * 100.f);
 	}
 	else
 	{
