[AssetManager] Refactor AssetManager code.
- Simplified ChangeBundleStateForPrimaryAssets by returning CombinedHandle even for 1 asset requests
- Fixed a state transition logic to correctly handle bundle states when asset is already loaded
- Fixed a state transition logic to correctly handle canceletion of streaming handle

#tests RR, editor

diff --git a/Engine/Source/Runtime/Engine/Classes/Engine/AssetManager.h b/Engine/Source/Runtime/Engine/Classes/Engine/AssetManager.h
--- a/Engine/Source/Runtime/Engine/Classes/Engine/AssetManager.h
+++ b/Engine/Source/Runtime/Engine/Classes/Engine/AssetManager.h
@@ -966,7 +966,7 @@ protected:
 	virtual void NotifyOnAssetStateChangeCompleted(const FPrimaryAssetId& PrimaryAssetId) {}
 
 	/** Called when an internal load handle finishes, handles setting to pending state */
-	ENGINE_API void OnAssetStateChangeCompleted(const TSharedPtr<FStreamableHandle>& BoundHandle, const FPrimaryAssetId& PrimaryAssetId, FStreamableDelegateWithHandle&& WrappedDelegate);
+	ENGINE_API void OnAssetStateChangeCompleted(TSharedPtr<FStreamableHandle> BoundHandle, FPrimaryAssetId PrimaryAssetId);
 
 	/** Scans all asset types specified in DefaultGame */
 	ENGINE_API virtual void ScanPrimaryAssetTypesFromConfig();
diff --git a/Engine/Source/Runtime/Engine/Private/AssetManager.cpp b/Engine/Source/Runtime/Engine/Private/AssetManager.cpp
--- a/Engine/Source/Runtime/Engine/Private/AssetManager.cpp
+++ b/Engine/Source/Runtime/Engine/Private/AssetManager.cpp
@@ -1808,8 +1808,13 @@ TSharedPtr<FStreamableHandle> UAssetManager::ChangeBundleStateForPrimaryAssets(
 	FAssetManagerLoadParams&& LoadParams,
 	UE::FSourceLocation Location)
 {
-	TArray<TSharedPtr<FStreamableHandle> > NewHandles, ExistingHandles;
-	TArray<FPrimaryAssetId> NewAssets;
+	if (AssetsToChange.IsEmpty())
+	{
+		ensureAlwaysMsgf(!AddBundles.IsEmpty() || !RemoveBundles.IsEmpty() || bRemoveAllBundles,
+			TEXT("Attempted to change state on empty primary assets set."));
+		return nullptr;
+	}
+	TArray<TSharedPtr<FStreamableHandle> > AllHandles;
 	int32 MoviePlayerNumAssets = 0;
 	for (const FPrimaryAssetId& PrimaryAssetId : AssetsToChange)
 	{
@@ -1863,7 +1868,7 @@ TSharedPtr<FStreamableHandle> UAssetManager::ChangeBundleStateForPrimaryAssets(
 				if (Algo::Compare(NameData->PendingState.BundleNames, NewBundleState))
 				{
 					// This will wait on any existing handles to finish
-					ExistingHandles.Add(NameData->PendingState.Handle);
+					AllHandles.Add(NameData->PendingState.Handle);
 					continue;
 				}
 
@@ -1933,19 +1938,7 @@ TSharedPtr<FStreamableHandle> UAssetManager::ChangeBundleStateForPrimaryAssets(
 				DebugName << TEXT(")");
 			}
 
-			// Callbacks will be bound to combined handle below.
-			FAssetManagerLoadParams LoadParamsWithoutCallbacks;
-			LoadParamsWithoutCallbacks.Priority = LoadParams.Priority;
-			if (LoadParams.DownloadParams)
-			{
-				LoadParamsWithoutCallbacks.DownloadParams = CopyTemp(*LoadParams.DownloadParams);
-				NewHandle = LoadAssetList(MoveTemp(PathsToLoad), MoveTemp(LoadParamsWithoutCallbacks), FString(DebugName), CopyTemp(Location));
-			}
-			else if (!bDownloadRequired)
-			{
-				NewHandle = LoadAssetList(MoveTemp(PathsToLoad), MoveTemp(LoadParamsWithoutCallbacks), FString(DebugName), CopyTemp(Location));
-			}
-			else
+			if (!LoadParams.DownloadParams && bDownloadRequired)
 			{
 				// The user previously use download params for requests with this primary asset, so we expect them going forward.
 				// See comment above where bDownloadRequired is delcared.
@@ -1955,43 +1948,37 @@ TSharedPtr<FStreamableHandle> UAssetManager::ChangeBundleStateForPrimaryAssets(
 					FStreamableError(EStreamableErrorType::DownloadError, FString::Printf(TEXT("Expected DownloadParams for %s"), *PrimaryAssetId.ToString())),
 					MoveTemp(StreamableLoadParams), FString(DebugName), CopyTemp(Location));
 			}
-
-			if (!NewHandle.IsValid())
+			else
 			{
-				// LoadAssetList already logs an error, no need to do it here as well
-				continue;
-			}
+				// Set pending bundle state before issuing the request, because it can be resolved immediately.
+				// OnAssetStateChangeCompleted requires this data to handle state transition correctly.
+				NameData->PendingState.BundleNames = NewBundleState;
 
-			if (NewHandle->HasLoadCompleted())
-			{
-				// If an error occured, don't update the current state with a bad state
-				if (NewHandle->HasError())
+				FAssetManagerLoadParams AssetListLoadParams;
+				AssetListLoadParams.Priority = LoadParams.Priority;
+				if (AssetListLoadParams.DownloadParams)
 				{
-					UE_LOG(LogAssetManager, Warning, TEXT("ChangeBundleStateForPrimaryAssetsInternal: Error (%s) loading asset! (%s)"), *LexToString(*NewHandle->GetError()), *PrimaryAssetId.ToString());
+					AssetListLoadParams.DownloadParams = CopyTemp(*LoadParams.DownloadParams);
 				}
-				else
+				AssetListLoadParams.OnComplete.BindUObject(this, &UAssetManager::OnAssetStateChangeCompleted, PrimaryAssetId);
+				AssetListLoadParams.OnCancel.BindUObject(this, &UAssetManager::OnAssetStateChangeCompleted, PrimaryAssetId);
+	
+				NewHandle = LoadAssetList(MoveTemp(PathsToLoad), MoveTemp(AssetListLoadParams), FString(DebugName), CopyTemp(Location));
+				// save handle so possible duplicate requests will not override this one
+				// set only if loading is not completed yet because completion callback might already reset the PendingState
+				if (!NewHandle->HasLoadCompleted())
 				{
-					// Copy right into active
-					NameData->CurrentState.BundleNames = NewBundleState;
-					NameData->CurrentState.Handle = NewHandle;
+					NameData->PendingState.Handle = NewHandle;
 				}
 			}
-			else
-			{
-				// Copy into pending and set delegate
-				NameData->PendingState.BundleNames = NewBundleState;
-				NameData->PendingState.Handle = NewHandle;
 
-				NewHandle->BindCompleteDelegate(FStreamableDelegateWithHandle::CreateWeakLambda(this,
-					[this, PrimaryAssetId](const TSharedPtr<FStreamableHandle>& Handle)
-					{
-						OnAssetStateChangeCompleted(Handle, PrimaryAssetId, FStreamableDelegateWithHandle());
-					}
-				));
+			if (!NewHandle.IsValid())
+			{
+				// LoadAssetList already logs an error, no need to do it here as well
+				continue;
 			}
 
-			NewHandles.Add(NewHandle);
-			NewAssets.Add(PrimaryAssetId);
+			AllHandles.Add(NewHandle);
 		}
 		else
 		{
@@ -1999,62 +1986,17 @@ TSharedPtr<FStreamableHandle> UAssetManager::ChangeBundleStateForPrimaryAssets(
 		}
 	}
 
-	TSharedPtr<FStreamableHandle> ReturnHandle;
-	if (NewHandles.Num() > 1 || ExistingHandles.Num() > 0)
-	{
-		// If multiple handles or we have an old handle, need to make wrapper handle
-		NewHandles.Append(ExistingHandles);
-
-		ReturnHandle = StreamableManager.CreateCombinedHandle(
-			NewHandles,
-			FStreamableHandle::HandleDebugName_CombinedHandle,
-			EStreamableManagerCombinedHandleOptions::None,
-			FStreamableAsyncLoadParams{},
-			CopyTemp(Location));
-	}
-	else if (NewHandles.Num() == 1)
-	{
-		ReturnHandle = NewHandles[0];
-		ensure(NewAssets.Num() == 1);
-
-		if (!ReturnHandle->HasLoadCompleted()) // OnAssetStateChangeCompleted should only be called for pending handles
-		{
-			// If only one handle, return it and add callback
-			// Call internal callback and external callback when it finishes
-			LoadParams.OnComplete = FStreamableDelegateWithHandle::CreateWeakLambda(this,
-				[this, PrimaryAssetId = NewAssets[0], DelegateToCall = MoveTemp(LoadParams.OnComplete)](const TSharedPtr<FStreamableHandle>& Handle) mutable
-				{
-					OnAssetStateChangeCompleted(Handle, PrimaryAssetId, MoveTemp(DelegateToCall));
-				});
-		}
-	}
-
-	if (ReturnHandle)
-	{
-		// Call delegate or bind to handle
-		if (ReturnHandle->HasLoadCompleted())
-		{
-			FStreamableHandle::ExecuteDelegate(MoveTemp(LoadParams.OnComplete), ReturnHandle);
-		}
-		else if (ReturnHandle->WasCanceled())
-		{
-			FStreamableHandle::ExecuteDelegate(MoveTemp(LoadParams.OnCancel), ReturnHandle);
-		}
-		else
-		{
-			// Call external callback when completed
-			ReturnHandle->BindCompleteDelegate(MoveTemp(LoadParams.OnComplete));
-			ReturnHandle->BindCancelDelegate(MoveTemp(LoadParams.OnCancel));
-			ReturnHandle->BindUpdateDelegate(MoveTemp(LoadParams.OnUpdate));
-		}
-	}
-	else
-	{
-		// Call completion callback, nothing to do
-		FStreamableHandle::ExecuteDelegate(MoveTemp(LoadParams.OnComplete));
-	}
+	FStreamableAsyncLoadParams CombinedHandleParams;
+	CombinedHandleParams.OnComplete = MoveTemp(LoadParams.OnComplete);
+	CombinedHandleParams.OnCancel = MoveTemp(LoadParams.OnCancel);
+	CombinedHandleParams.OnUpdate = MoveTemp(LoadParams.OnUpdate);
 
-	return ReturnHandle;
+	return StreamableManager.CreateCombinedHandle(
+		AllHandles,
+		FStreamableHandle::HandleDebugName_CombinedHandle,
+		EStreamableManagerCombinedHandleOptions::None,
+		MoveTemp(CombinedHandleParams),
+		CopyTemp(Location));
 }
 
 TSharedPtr<FStreamableHandle> UAssetManager::ChangeBundleStateForMatchingPrimaryAssets(const TArray<FName>& NewBundles, const TArray<FName>& OldBundles, FAssetManagerLoadParams&& LoadParams, UE::FSourceLocation Location)
@@ -2167,33 +2109,34 @@ TSharedPtr<FStreamableHandle> UAssetManager::PredownloadPrimaryAssets(const TArr
 	return TSharedPtr<FStreamableHandle>();
 }
 
-void UAssetManager::OnAssetStateChangeCompleted(const TSharedPtr<FStreamableHandle>& BoundHandle, const FPrimaryAssetId& PrimaryAssetId, FStreamableDelegateWithHandle&& WrappedDelegate)
+void UAssetManager::OnAssetStateChangeCompleted(TSharedPtr<FStreamableHandle> BoundHandle, FPrimaryAssetId PrimaryAssetId)
 {
 	FPrimaryAssetData* NameData = GetNameData(PrimaryAssetId);
 
 	if (NameData)
 	{
-		NotifyOnAssetStateChangeCompleted(PrimaryAssetId);
-
-		if (NameData->PendingState.Handle == BoundHandle)
+		if (!NameData->PendingState.Handle.IsValid() || NameData->PendingState.Handle == BoundHandle)
 		{
 			// If an error occured, don't update the current state with a bad state
 			if (BoundHandle->HasError())
 			{
 				UE_LOG(LogAssetManager, Warning, TEXT("OnAssetStateChangeCompleted: Received error (%s) loading asset! (%s)"), *LexToString(*BoundHandle->GetError()), *PrimaryAssetId.ToString());
 			}
+			else if (BoundHandle->WasCanceled())
+			{
+				UE_LOG(LogAssetManager, Verbose, TEXT("OnAssetStateChangeCompleted: Request was canceled (%s)"), *PrimaryAssetId.ToString());
+			}
 			else
 			{
-				NameData->CurrentState.Handle = NameData->PendingState.Handle;
+				NameData->CurrentState.Handle = BoundHandle;
 				NameData->CurrentState.BundleNames = NameData->PendingState.BundleNames;
+				
+				NotifyOnAssetStateChangeCompleted(PrimaryAssetId);
 			}
 
-			WrappedDelegate.ExecuteIfBound(BoundHandle);
-
-			// Clear old state, but don't cancel handle as we just copied it into current
 			NameData->PendingState.Reset(false);
 		}
-		else
+		else if (NameData->PendingState.Handle.IsValid())
 		{
 			UE_LOG(LogAssetManager, Verbose, TEXT("OnAssetStateChangeCompleted: Received after pending data changed, ignoring (%s)"), *PrimaryAssetId.ToString());
 		}
diff --git a/Engine/Source/Runtime/Engine/Private/StreamableManager.cpp b/Engine/Source/Runtime/Engine/Private/StreamableManager.cpp
--- a/Engine/Source/Runtime/Engine/Private/StreamableManager.cpp
+++ b/Engine/Source/Runtime/Engine/Private/StreamableManager.cpp
@@ -2297,7 +2297,7 @@ TSharedPtr<FStreamableHandle> FStreamableManager::CreateCombinedHandle(
 	FStreamableAsyncLoadParams&& Params,
 	UE::FSourceLocation Location)
 {
-	if (!ensure(ChildHandles.Num() > 0))
+	if (ChildHandles.IsEmpty())
 	{
 		// Complete immediately if nothing to load
 		Params.OnComplete.ExecuteIfBound(nullptr);
