Move logic for process touch hit results to separate protected method

#rnx

diff --git a/Engine/Source/Runtime/Engine/Classes/GameFramework/PlayerController.h b/Engine/Source/Runtime/Engine/Classes/GameFramework/PlayerController.h
--- a/Engine/Source/Runtime/Engine/Classes/GameFramework/PlayerController.h
+++ b/Engine/Source/Runtime/Engine/Classes/GameFramework/PlayerController.h
@@ -1958,6 +1958,7 @@ protected:
 	ENGINE_API virtual bool IsInViewportClient(UGameViewportClient* ViewportClient) const;
 	ENGINE_API virtual int32 GetInputIndex() const;
 	ENGINE_API virtual ACameraActor* GetAutoActivateCameraForPlayer() const;
+	ENGINE_API virtual void ProcessTouchHitResult(const FInputDeviceId DeviceId, uint32 Handle, ETouchType::Type Type, const FVector2D& TouchLocation, const FHitResult& HitResult, const bool bHit);
 
 	/** Allows the PlayerController to set up custom input bindings. */
 	ENGINE_API virtual void SetupInputComponent();
diff --git a/Engine/Source/Runtime/Engine/Private/PlayerController.cpp b/Engine/Source/Runtime/Engine/Private/PlayerController.cpp
--- a/Engine/Source/Runtime/Engine/Private/PlayerController.cpp
+++ b/Engine/Source/Runtime/Engine/Private/PlayerController.cpp
@@ -2492,72 +2492,76 @@ bool APlayerController::InputTouch(
 
 		if (bEnableTouchEvents || bEnableTouchOverEvents)
 		{
-			const ETouchIndex::Type FingerIndex = ETouchIndex::Type(Handle);
-
 			FHitResult HitResult;
 			const bool bHit = GetHitResultAtScreenPosition(TouchLocation, CurrentClickTraceChannel, true, HitResult);
 
-			UPrimitiveComponent* PreviousComponent = CurrentTouchablePrimitives[Handle].Get();
-			UPrimitiveComponent* CurrentComponent = (bHit ? HitResult.Component.Get() : NULL);
+			ProcessTouchHitResult(DeviceId, Handle, Type, TouchLocation, HitResult, bHit);
+		}
+	}
+
+	return bResult;
+}
+
+void APlayerController::ProcessTouchHitResult(const FInputDeviceId DeviceId, uint32 Handle, ETouchType::Type Type, const FVector2D& TouchLocation, const FHitResult& HitResult, const bool bHit)
+{
+	const ETouchIndex::Type FingerIndex = ETouchIndex::Type(Handle);
+	UPrimitiveComponent* PreviousComponent = CurrentTouchablePrimitives[Handle].Get();
+	UPrimitiveComponent* CurrentComponent = (bHit ? HitResult.Component.Get() : NULL);
 
-			if (GetHUD())
+	if (GetHUD())
+	{
+		if (Type == ETouchType::Began || Type == ETouchType::Ended)
+		{
+			if (GetHUD()->UpdateAndDispatchHitBoxClickEvents(TouchLocation, (Type == ETouchType::Began ? EInputEvent::IE_Pressed : EInputEvent::IE_Released)))
 			{
-				if (Type == ETouchType::Began || Type == ETouchType::Ended)
-				{
-					if (GetHUD()->UpdateAndDispatchHitBoxClickEvents(TouchLocation, (Type == ETouchType::Began ? EInputEvent::IE_Pressed : EInputEvent::IE_Released)))
-					{
-						CurrentComponent = NULL;
-					}
-				}
+				CurrentComponent = NULL;
 			}
+		}
+	}
 
-			switch(Type)
-			{
-			case ETouchType::Began:
-				// Give it a begin touch
-				if (bEnableTouchEvents && (CurrentComponent != NULL))
-				{
-					CurrentComponent->DispatchOnInputTouchBegin(FingerIndex);
-				}
+	switch (Type)
+	{
+	case ETouchType::Began:
+		// Give it a begin touch
+		if (bEnableTouchEvents && (CurrentComponent != NULL))
+		{
+			CurrentComponent->DispatchOnInputTouchBegin(FingerIndex);
+		}
 
-				// Give a touch enter event
-				if (bEnableTouchOverEvents)
-				{
-					UPrimitiveComponent::DispatchTouchOverEvents(FingerIndex, PreviousComponent, CurrentComponent);
-					CurrentTouchablePrimitives[Handle] = CurrentComponent;
-				}
-				break;
-			case ETouchType::Ended:
-				// Give it a touch exit
-				if (bEnableTouchEvents && (CurrentComponent != NULL))
-				{
-					CurrentComponent->DispatchOnInputTouchEnd(FingerIndex);
-				}
+		// Give a touch enter event
+		if (bEnableTouchOverEvents)
+		{
+			UPrimitiveComponent::DispatchTouchOverEvents(FingerIndex, PreviousComponent, CurrentComponent);
+			CurrentTouchablePrimitives[Handle] = CurrentComponent;
+		}
+		break;
+	case ETouchType::Ended:
+		// Give it a touch exit
+		if (bEnableTouchEvents && (CurrentComponent != NULL))
+		{
+			CurrentComponent->DispatchOnInputTouchEnd(FingerIndex);
+		}
 
-				// Give it a end touch
-				if (bEnableTouchOverEvents)
-				{
-					// Handle the case where the finger moved faster than tick, and is being released over a different component than it was last dragged over
-					if ((PreviousComponent != CurrentComponent) && (PreviousComponent != NULL))
-					{
-						// First notify the old component that the touch left it to go to the current component
-						UPrimitiveComponent::DispatchTouchOverEvents(FingerIndex, PreviousComponent, CurrentComponent);
-					}
+		// Give it a end touch
+		if (bEnableTouchOverEvents)
+		{
+			// Handle the case where the finger moved faster than tick, and is being released over a different component than it was last dragged over
+			if ((PreviousComponent != CurrentComponent) && (PreviousComponent != NULL))
+			{
+				// First notify the old component that the touch left it to go to the current component
+				UPrimitiveComponent::DispatchTouchOverEvents(FingerIndex, PreviousComponent, CurrentComponent);
+			}
 
-					// Now notify that the current component is being released and thus the touch is leaving it
-					PreviousComponent = CurrentComponent;
-					CurrentComponent = NULL;
-					UPrimitiveComponent::DispatchTouchOverEvents(FingerIndex, PreviousComponent, CurrentComponent);
-					CurrentTouchablePrimitives[Handle] = CurrentComponent;
-				}
-				break;
-			default:
-				break;
-			};
+			// Now notify that the current component is being released and thus the touch is leaving it
+			PreviousComponent = CurrentComponent;
+			CurrentComponent = NULL;
+			UPrimitiveComponent::DispatchTouchOverEvents(FingerIndex, PreviousComponent, CurrentComponent);
+			CurrentTouchablePrimitives[Handle] = CurrentComponent;
 		}
-	}
-
-	return bResult;
+		break;
+	default:
+		break;
+	};
 }
 
 PRAGMA_DISABLE_DEPRECATION_WARNINGS
