Removed unnecessary allocation in LineBatchComponent

diff --git a/Engine/Source/Runtime/Engine/Private/Components/LineBatchComponent.cpp b/Engine/Source/Runtime/Engine/Private/Components/LineBatchComponent.cpp
--- a/Engine/Source/Runtime/Engine/Private/Components/LineBatchComponent.cpp
+++ b/Engine/Source/Runtime/Engine/Private/Components/LineBatchComponent.cpp
@@ -433,8 +433,6 @@ void ULineBatchComponent::DrawSphere(FVector const& Center, float Radius, int32
 	float Latitude = AngleInc;
 	float SinY1 = 0.0f, CosY1 = 1.0f;
 
-	TArray<FBatchedLine> Lines;
-	Lines.Empty(NumSegmentsY * Segments * 2);
 	while (NumSegmentsY--)
 	{
 		const float SinY2 = FMath::Sin(Latitude);
