Move ActorDescContainerInitParams into a separate file

diff --git a/Engine/Source/Runtime/Engine/Public/WorldPartition/ActorDescContainer.h b/Engine/Source/Runtime/Engine/Public/WorldPartition/ActorDescContainer.h
--- a/Engine/Source/Runtime/Engine/Public/WorldPartition/ActorDescContainer.h
+++ b/Engine/Source/Runtime/Engine/Public/WorldPartition/ActorDescContainer.h
@@ -8,6 +8,7 @@
 #include "WorldPartition/ActorDescList.h"
 #include "WorldPartition/WorldPartitionHandle.h"
 #include "WorldPartition/WorldPartitionActorDesc.h"
+#include "ActorDescContainerInitParams.h"
 #include "AssetRegistry/AssetData.h"
 #include "ActorDescContainer.generated.h"
 
@@ -48,45 +49,7 @@ class UActorDescContainer : public UObject, public FActorDescList
 
 public:
 	/* Struct of parameters passed to Initialize function. */
-	struct FInitializeParams
-	{
-		UE_DEPRECATED(5.4, "Use constructor with no UWorld param instead")
-		FInitializeParams(UWorld* InWorld, FName InPackageName)
-			: PackageName(InPackageName)
-			, ContainerName(InPackageName.ToString())
-		{}
-
-		FInitializeParams(FName InPackageName)
-			: PackageName(InPackageName)
-			, ContainerName(InPackageName.ToString())
-		{}
-
-		FInitializeParams(const FString& InContainerName, FName InPackageName)
-			: PackageName(InPackageName)
-			, ContainerName(InContainerName)
-		{}
-		
-		/* The long package name of the container package on disk. */
-		FName PackageName;
-
-		/* The unique name for the container : defaults to PackageName */
-		FString ContainerName;
-
-		/** The associated Content Bundle Guid */
-		FGuid ContentBundleGuid;
-
-		/** If the container should bind to editor events */
-		bool bShouldRegisterEditorDeletages = true;
-
-		/** The associated External Data Layer Asset */
-		const UExternalDataLayerAsset* ExternalDataLayerAsset = nullptr;
-
-		/* Custom pre-init function that is called before calling Initialize on the new container */
-		TUniqueFunction<void(UActorDescContainer*)> PreInitialize;
-
-		/* Custom filter function used to filter actors descriptors. */
-		TUniqueFunction<bool(const FWorldPartitionActorDesc*)> FilterActorDesc;
-	};
+	using FInitializeParams = FActorDescContainerInitParams;
 
 	ENGINE_API virtual void Initialize(const FInitializeParams& InitParams);
 	ENGINE_API virtual void Uninitialize();
diff --git a/Engine/Source/Runtime/Engine/Public/WorldPartition/ActorDescContainerInitParams.h b/Engine/Source/Runtime/Engine/Public/WorldPartition/ActorDescContainerInitParams.h
new file mode 100644
--- /dev/null
+++ b/Engine/Source/Runtime/Engine/Public/WorldPartition/ActorDescContainerInitParams.h
@@ -0,0 +1,52 @@
+// Copyright Epic Games, Inc. All Rights Reserved.
+
+#pragma once
+
+#include "Containers/UnrealString.h"
+#include "Misc/Guid.h"
+#include "Templates/Function.h"
+#include "UObject/NameTypes.h"
+
+class UExternalDataLayerAsset;
+class UActorDescContainer;
+class FWorldPartitionActorDesc;
+
+struct FActorDescContainerInitParams
+{
+	UE_DEPRECATED(5.4, "Use constructor with no UWorld param instead")
+	FActorDescContainerInitParams(UWorld* InWorld, FName InPackageName)
+		: PackageName(InPackageName)
+		, ContainerName(InPackageName.ToString())
+	{}
+
+	FActorDescContainerInitParams(FName InPackageName)
+		: PackageName(InPackageName)
+		, ContainerName(InPackageName.ToString())
+	{}
+
+	FActorDescContainerInitParams(const FString& InContainerName, FName InPackageName)
+		: PackageName(InPackageName)
+		, ContainerName(InContainerName)
+	{}
+
+	/* The long package name of the container package on disk. */
+	FName PackageName;
+
+	/* The unique name for the container : defaults to PackageName */
+	FString ContainerName;
+
+	/** The associated Content Bundle Guid */
+	FGuid ContentBundleGuid;
+
+	/** If the container should bind to editor events */
+	bool bShouldRegisterEditorDeletages = true;
+
+	/** The associated External Data Layer Asset */
+	const UExternalDataLayerAsset* ExternalDataLayerAsset = nullptr;
+
+	/* Custom pre-init function that is called before calling Initialize on the new container */
+	TUniqueFunction<void(UActorDescContainer*)> PreInitialize;
+
+	/* Custom filter function used to filter actors descriptors. */
+	TUniqueFunction<bool(const FWorldPartitionActorDesc*)> FilterActorDesc;
+};
