Allow uber graph function handlers to run in editor, this allows the delay node to function in editor utility actors that run code at edit time

[RN] The delay node now works correctly when executed by editor functions

diff --git a/Engine/Source/Runtime/Engine/Private/Actor.cpp b/Engine/Source/Runtime/Engine/Private/Actor.cpp
--- a/Engine/Source/Runtime/Engine/Private/Actor.cpp
+++ b/Engine/Source/Runtime/Engine/Private/Actor.cpp
@@ -1362,7 +1362,9 @@ void AActor::ProcessEvent(UFunction* Function, void* Parameters)
 
 	#if WITH_EDITOR
 	static const FName CallInEditorMeta(TEXT("CallInEditor"));
-	const bool bAllowScriptExecution = GAllowActorScriptExecutionInEditor || Function->GetBoolMetaData(CallInEditorMeta);
+	const bool bAllowScriptExecution = GAllowActorScriptExecutionInEditor || 
+		Function->GetBoolMetaData(CallInEditorMeta) ||
+		Function->HasAnyFunctionFlags(FUNC_UbergraphFunction); // uber graph frame is only used internally, e.g. by timer handlers
 	#else
 	const bool bAllowScriptExecution = GAllowActorScriptExecutionInEditor;
 	#endif
