RemoteObjects:  Component->IsActive() is not the same as having its Component Tick Enabled.  We are now saving off another variable to differentiate between the two when compiled in this mode.

diff --git a/Engine/Source/Runtime/Engine/Classes/Components/ActorComponent.h b/Engine/Source/Runtime/Engine/Classes/Components/ActorComponent.h
--- a/Engine/Source/Runtime/Engine/Classes/Components/ActorComponent.h
+++ b/Engine/Source/Runtime/Engine/Classes/Components/ActorComponent.h
@@ -390,6 +390,10 @@ private:
 	/** Whether the component is creating/destroying its physics state asynchronously (or not) */
 	EPhysicsStateAsyncCreationState PhysicsStateAsyncCreationState : 2;
 
+#if UE_WITH_REMOTE_OBJECT_HANDLE
+	uint8 bComponentTickEnabled : 1 = false;
+#endif
+
 public:
 
 	/** Describes how a component instance will be created */
diff --git a/Engine/Source/Runtime/Engine/Private/Components/ActorComponent.cpp b/Engine/Source/Runtime/Engine/Private/Components/ActorComponent.cpp
--- a/Engine/Source/Runtime/Engine/Private/Components/ActorComponent.cpp
+++ b/Engine/Source/Runtime/Engine/Private/Components/ActorComponent.cpp
@@ -741,7 +741,14 @@ void UActorComponent::PostMigrate(const struct FUObjectMigrationContext& Migrati
 
 		if (RecvType == EObjectMigrationRecvType::AssignedOwnership)
 		{
-			SetComponentTickEnabled(bIsActive);
+			// If we're receiving ownership, we are now in charge of its Tick
+			if (bTickFunctionsRegistered)
+			{
+				bTickFunctionsRegistered = false;
+				RegisterAllComponentTickFunctions(true);
+			}
+
+			SetComponentTickEnabled(bComponentTickEnabled);
 		}
 	}
 }
@@ -3104,6 +3111,7 @@ void UActorComponent::Serialize(FArchive& Ar)
 			ComponentFlags |= (bHasBegunPlay & 1)				<< 3;
 			ComponentFlags |= (bIsReadyForReplication & 1) 		<< 4;
 			ComponentFlags |= (bPhysicsStateCreated & 1) 		<< 5;
+			ComponentFlags |= (IsComponentTickEnabled() & 1)	<< 6;
 		}
 
 		Ar << ComponentFlags;
@@ -3116,6 +3124,7 @@ void UActorComponent::Serialize(FArchive& Ar)
 			bHasBegunPlay				= (ComponentFlags >> 3) & 1;
 			bIsReadyForReplication 		= (ComponentFlags >> 4) & 1;
 			bPhysicsStateCreated 		= (ComponentFlags >> 5) & 1;
+			bComponentTickEnabled 		= (ComponentFlags >> 6) & 1;
 		}
 	}
 #endif // UE_WITH_REMOTE_OBJECT_HANDLE
