RemoteObjects:  Do not deregister the Actor's Tick Functions if we're simply getting our object returned to us from a previous loan.

#rnx

diff --git a/Engine/Source/Runtime/Engine/Private/Actor.cpp b/Engine/Source/Runtime/Engine/Private/Actor.cpp
--- a/Engine/Source/Runtime/Engine/Private/Actor.cpp
+++ b/Engine/Source/Runtime/Engine/Private/Actor.cpp
@@ -1333,9 +1333,9 @@ void AActor::PostMigrate(const FUObjectMigrationContext& MigrationContext)
 
 			World->AddNetworkActor(this);
 		}
-		else
+		else if (MigrationRecvType == EObjectMigrationRecvType::Borrowed)
 		{
-			// Make sure that all tick functions have been deregistered if we aren't acquring ownership.
+			// Make sure that all tick functions have been deregistered if we don't have ownership.
 			RegisterAllActorTickFunctions(false, false);
 		}
 
