Add Used With Static Mesh material flag

#tests local Lyra, Horde cook and ReplayRun

diff --git a/Engine/Source/Editor/MaterialEditor/Private/MaterialEditorDetailCustomization.cpp b/Engine/Source/Editor/MaterialEditor/Private/MaterialEditorDetailCustomization.cpp
--- a/Engine/Source/Editor/MaterialEditor/Private/MaterialEditorDetailCustomization.cpp
+++ b/Engine/Source/Editor/MaterialEditor/Private/MaterialEditorDetailCustomization.cpp
@@ -922,7 +922,26 @@ void FMaterialDetailCustomization::CustomizeDetails( IDetailLayoutBuilder& Detai
 		DetailLayout.HideCategory( TEXT("Thumbnail") );
 		DetailLayout.HideCategory( TEXT("MaterialInterface") );
 		DetailLayout.HideCategory( TEXT("PhysicalMaterial") );
-		DetailLayout.HideCategory( TEXT("Usage") );
+
+		// Usage Category
+		{
+			IDetailCategoryBuilder& UsageCategory = DetailLayout.EditCategory(TEXT("Usage"));
+
+			TArray<TSharedRef<IPropertyHandle>> AllProperties;
+			UsageCategory.GetDefaultProperties(AllProperties);
+
+			// Hide all Usage except bUsedWithStaticMesh on UI
+			for (TSharedRef<IPropertyHandle>& PropertyHandle : AllProperties)
+			{
+				FProperty* Property = PropertyHandle->GetProperty();
+				FName PropertyName = Property->GetFName();
+
+				if (PropertyName != GET_MEMBER_NAME_CHECKED(UMaterial, bUsedWithStaticMesh))
+				{
+					DetailLayout.HideProperty(PropertyHandle);
+				}
+			}
+		}
 
 		// Mobile category
 		{
diff --git a/Engine/Source/Runtime/Engine/Private/LocalVertexFactory.cpp b/Engine/Source/Runtime/Engine/Private/LocalVertexFactory.cpp
--- a/Engine/Source/Runtime/Engine/Private/LocalVertexFactory.cpp
+++ b/Engine/Source/Runtime/Engine/Private/LocalVertexFactory.cpp
@@ -247,19 +247,44 @@ void FLocalVertexFactoryShaderParameters::GetElementShaderBindings(
 	}
 }
 
+// CVarIgnoreLVFUsageMode settings:
+//     0: ignored only in editor (Default)
+//     1: ignored in editor and on PC >= SM5
+//     2: ignored always
+TAutoConsoleVariable<int32> CVarIgnoreLVFUsageMode(
+	TEXT("r.Material.IgnoreLVFUsageMode"), 0,
+	TEXT("Setting to adjust how UsedWithStaticMesh material flag is interpreted.\n"),
+	ECVF_ReadOnly);
+
+static bool ShouldIgnore(const EShaderPlatform Platform)
+{
+	switch (CVarIgnoreLVFUsageMode.GetValueOnAnyThread())
+	{
+	case 0:
+		return false;
+	case 1:
+		return IsPCPlatform(Platform) && GetMaxSupportedFeatureLevel(Platform) >= ERHIFeatureLevel::SM5;
+	case 2:
+		return true;
+	}
+
+	return false;
+}
+
 /**
  * Should we cache the material's shadertype on this platform with this vertex factory? 
  */
 bool FLocalVertexFactory::ShouldCompilePermutation(const FVertexFactoryShaderPermutationParameters& Parameters)
 {
-	// Compile local vertex factory permutations of materials in the UI domain to maintain
-	// our historical (but unintentional) support for using them in this domain.
-	if (Parameters.MaterialParameters.MaterialDomain == MD_UI)
+	// Compile local vertex factory permutations while in the editor, always. Important for editor features like preview 
+	// Additionally, allow ignoring via CVar to change what scenarios respect the bIsUsedWithStaticMesh flag
+	if (EnumHasAnyFlags(Parameters.Flags, EShaderPermutationFlags::HasEditorOnlyData) || ShouldIgnore(Parameters.Platform))
 	{
 		return true;
 	}
 
-	return true; 
+	// Otherwise compile when the Material has selected that it can be used with StaticMesh
+	return Parameters.MaterialParameters.bIsUsedWithStaticMesh == 1;
 }
 
 void FLocalVertexFactory::ModifyCompilationEnvironment(const FVertexFactoryShaderPermutationParameters& Parameters, FShaderCompilerEnvironment& OutEnvironment)
diff --git a/Engine/Source/Runtime/Engine/Private/Materials/Material.cpp b/Engine/Source/Runtime/Engine/Private/Materials/Material.cpp
--- a/Engine/Source/Runtime/Engine/Private/Materials/Material.cpp
+++ b/Engine/Source/Runtime/Engine/Private/Materials/Material.cpp
@@ -1094,6 +1094,7 @@ void UMaterial::SetInitialValues()
 	bEnableStencilTest = false;
 	bUsedWithVolumetricCloud = false;
 	bUsedWithHeterogeneousVolumes = false;
+	bUsedWithStaticMesh = true;
 
 	// We default to false=incompatible, and rely onto the automatic detection executed from the HLSLTranslator.
 	bForceCompatibleWithLightFunctionAtlas = false;
@@ -1498,6 +1499,7 @@ bool UMaterial::GetUsageByFlag(EMaterialUsage Usage) const
 		case MATUSAGE_VolumetricCloud: UsageValue = bUsedWithVolumetricCloud; break;
 		case MATUSAGE_HeterogeneousVolumes: UsageValue = bUsedWithHeterogeneousVolumes; break;
 		case MATUSAGE_MaterialCache: UsageValue = bUsedWithMaterialCache; break;
+		case MATUSAGE_StaticMesh: UsageValue = bUsedWithStaticMesh; break;
 		default: UE_LOG(LogMaterial, Fatal,TEXT("Unknown material usage: %u"), (int32)Usage);
 	};
 	return UsageValue;
@@ -1753,6 +1755,10 @@ void UMaterial::SetUsageByFlag(EMaterialUsage Usage, bool NewValue)
 		{
 			bUsedWithMaterialCache = NewValue; break;
 		}
+		case MATUSAGE_StaticMesh:
+		{
+			bUsedWithStaticMesh = NewValue; break;
+		}
 		default: UE_LOG(LogMaterial, Fatal, TEXT("Unknown material usage: %u"), (int32)Usage);
 	};
 #if WITH_EDITOR
@@ -1788,6 +1794,7 @@ FString UMaterial::GetUsageName(EMaterialUsage Usage) const
 		case MATUSAGE_VolumetricCloud: UsageName = TEXT("bUsedWithVolumetricCloud"); break;
 		case MATUSAGE_HeterogeneousVolumes: UsageName = TEXT("bUsedWithHeterogeneousVolumes"); break;
 		case MATUSAGE_MaterialCache: UsageName = TEXT("bUsedWithMaterialCache"); break;
+		case MATUSAGE_StaticMesh: UsageName = TEXT("bUsedWithStaticMesh"); break;
 		default: UE_LOG(LogMaterial, Fatal,TEXT("Unknown material usage: %u"), (int32)Usage);
 	};
 	return UsageName;
@@ -5136,7 +5143,12 @@ bool UMaterial::CanEditChange(const FProperty* InProperty) const
 			return bSubstrateEnabled && MaterialDomain == MD_Surface;
 		}
 
-		if (FCString::Strncmp(*PropertyName, TEXT("bUsedWith"), 9) == 0)
+		if (PropertyName == GET_MEMBER_NAME_STRING_CHECKED(UMaterial, bUsedWithStaticMesh))
+		{
+			// Don't bother to allow changes on DeferredDecal or Volume, they have shaders mixed into LocalVertexFactory which are needed
+			return MaterialDomain == MD_Surface || MaterialDomain == MD_UI || MaterialDomain == MD_PostProcess || MaterialDomain == MD_LightFunction;
+		}
+		else if (FCString::Strncmp(*PropertyName, TEXT("bUsedWith"), 9) == 0)
 		{
 			return MaterialDomain == MD_DeferredDecal || MaterialDomain == MD_Surface || MaterialDomain == MD_Volume;
 		}
diff --git a/Engine/Source/Runtime/Engine/Private/Materials/MaterialShader.cpp b/Engine/Source/Runtime/Engine/Private/Materials/MaterialShader.cpp
--- a/Engine/Source/Runtime/Engine/Private/Materials/MaterialShader.cpp
+++ b/Engine/Source/Runtime/Engine/Private/Materials/MaterialShader.cpp
@@ -411,6 +411,9 @@ bool LoadFromCompactBinary(FCbFieldView Field,
 namespace UE::MaterialInterface::Private
 {
 
+// DDC Key for Material struct updates, bump when changing related classes which can be loaded/saved to CompactBinary
+#define MATERIAL_CLASS_DERIVEDDATA_VER TEXT("6623D4F4-1CB7-45E7-A809-4D4660884D49")
+
 void HashMaterialStaticClassDependenciesForCook(FAppendToClassSchemaContext& Context)
 {
 	Context.Update(&GetMaterialShaderMapDDCGuid(), sizeof(GetMaterialShaderMapDDCGuid()));
@@ -443,6 +446,9 @@ void HashMaterialStaticClassDependenciesForCook(FAppendToClassSchemaContext& Con
 		TStringBuilder<32> FeatureLevelNameStr(InPlace, FeatureLevelName);
 		Context.Update(*FeatureLevelNameStr, FeatureLevelNameStr.Len() * sizeof((*FeatureLevelNameStr)[0]));
 	}
+
+	FStringView VersionString(MATERIAL_CLASS_DERIVEDDATA_VER);
+	Context.Update(VersionString.GetData(), VersionString.Len() * sizeof(VersionString[0]));
 }
 
 constexpr int32 HashMaterialDependenciesForCookArgsVersion = 1;
diff --git a/Engine/Source/Runtime/Engine/Private/Materials/MaterialShared.cpp b/Engine/Source/Runtime/Engine/Private/Materials/MaterialShared.cpp
--- a/Engine/Source/Runtime/Engine/Private/Materials/MaterialShared.cpp
+++ b/Engine/Source/Runtime/Engine/Private/Materials/MaterialShared.cpp
@@ -1422,7 +1422,7 @@ void FMaterial::SerializeInlineShaderMap(FArchive& Ar, const FName& SerializingA
 #if WITH_EDITOR
 			FinishCompilation();
 
-			bool bValid = GameThreadShaderMap != nullptr && GameThreadShaderMap->CompiledSuccessfully();
+			bool bValid = GameThreadShaderMap != nullptr && GameThreadShaderMap->CompiledSuccessfully() && (GameThreadShaderMap->GetShaderNum() > 0);
 			
 			Ar << bValid;
 
@@ -1435,7 +1435,10 @@ void FMaterial::SerializeInlineShaderMap(FArchive& Ar, const FName& SerializingA
 			{
 				UE_LOG(LogMaterial, Warning, TEXT("Cooking a material resource (in %s hierarchy) that doesn't have a valid ShaderMap! %s"),
 					*GetFriendlyName(),
-					(GameThreadShaderMap != nullptr) ? TEXT("Shadermap exists but wasn't compiled successfully (yet?)") : TEXT("Shadermap pointer is null.")
+					(GameThreadShaderMap == nullptr) ? TEXT("Shadermap pointer is null.") :
+						!GameThreadShaderMap->CompiledSuccessfully() ? 
+							TEXT("Shadermap exists but wasn't compiled successfully (yet?)") :
+							TEXT("Shadermap exists but has no shaders")
 					);
 			}
 #else
@@ -1779,6 +1782,11 @@ bool FMaterialResource::IsUsedWithHeterogeneousVolumes() const
 	return Material->bUsedWithHeterogeneousVolumes;
 }
 
+bool FMaterialResource::IsUsedWithStaticMesh() const
+{
+	return Material->bUsedWithStaticMesh;
+}
+
 bool FMaterialResource::SupportsMaterialCache() const
 {
 	return Material->GetCachedExpressionData().bHasMaterialCacheOutput || Material->bUsedWithMaterialCache;
@@ -5643,6 +5651,7 @@ FMaterialShaderParametersBase::FMaterialShaderParametersBase(const FMaterial* In
 	bDisablePreExposureScale = InMaterial->GetDisablePreExposureScale();
 	bAllowVariableRateShading = InMaterial->IsVariableRateShadingAllowed();
 	bSupportsMaterialCache = InMaterial->SupportsMaterialCache();
+	bIsUsedWithStaticMesh = InMaterial->IsUsedWithStaticMesh();
 }
 
 FMaterialShaderParameters::FMaterialShaderParameters(const FMaterial* InMaterial)
diff --git a/Engine/Source/Runtime/Engine/Public/MaterialShared.h b/Engine/Source/Runtime/Engine/Public/MaterialShared.h
--- a/Engine/Source/Runtime/Engine/Public/MaterialShared.h
+++ b/Engine/Source/Runtime/Engine/Public/MaterialShared.h
@@ -997,6 +997,15 @@ struct FMaterialShaderParametersBase
 {
 	EMaterialDomain MaterialDomain;
 	FMaterialShadingModelField ShadingModels;
+	union
+	{
+		// Room for up to 16 flags in previous padding space
+		uint16 PackedFlags16;
+		struct
+		{
+			uint16 bIsUsedWithStaticMesh : 1;
+		};
+	};
 	EBlendMode BlendMode;
 	ERHIFeatureLevel::Type FeatureLevel;
 	EMaterialQualityLevel::Type QualityLevel;
@@ -1006,6 +1015,7 @@ struct FMaterialShaderParametersBase
 	uint32 StencilCompare;
 	union
 	{
+		// These flags are full, add more flags above
 		uint64 PackedFlags;
 		struct
 		{
@@ -2206,6 +2216,7 @@ public:
 	virtual bool IsUsedWithLidarPointCloud() const { return false; }
 	virtual bool IsUsedWithVirtualHeightfieldMesh() const { return false; }
 	virtual bool IsUsedWithNeuralNetworks() const { return false; }
+	virtual bool IsUsedWithStaticMesh() const { return true; }
 	virtual bool IsFullyRough() const { return false; }
 	virtual bool GetForceCompatibleWithLightFunctionAtlas() const { return false; }
 	virtual bool UseNormalCurvatureToRoughness() const { return false; }
@@ -2993,6 +3004,7 @@ public:
 	ENGINE_API virtual bool IsUsedWithNanite() const override;
 	ENGINE_API virtual bool IsUsedWithVolumetricCloud() const override;
 	ENGINE_API virtual bool IsUsedWithHeterogeneousVolumes() const override;
+	ENGINE_API virtual bool IsUsedWithStaticMesh() const override;
 	ENGINE_API virtual bool SupportsMaterialCache() const override;
 	ENGINE_API virtual bool IsFullyRough() const override;
 	ENGINE_API virtual bool GetForceCompatibleWithLightFunctionAtlas() const override;
diff --git a/Engine/Source/Runtime/Engine/Public/Materials/Material.h b/Engine/Source/Runtime/Engine/Public/Materials/Material.h
--- a/Engine/Source/Runtime/Engine/Public/Materials/Material.h
+++ b/Engine/Source/Runtime/Engine/Public/Materials/Material.h
@@ -836,7 +836,14 @@ public:
 	UPROPERTY(EditAnywhere, BlueprintReadOnly, Category = Usage)
 	uint32 bUsedWithHeterogeneousVolumes : 1;
 
-	/** 
+	/**
+	 * Indicates that the material and its instances can be used on StaticMeshes, which is true by default on all Materials
+	 * This can be used to remove shaders which are not required, reducing shader compile time and memory usage
+	 */
+	UPROPERTY(EditAnywhere, BlueprintReadOnly, Category = Usage, AdvancedDisplay)
+	uint32 bUsedWithStaticMesh : 1;
+
+	/**
 	 * Indicates that the material and its instances can be used with Slate UI and UMG
 	 * This will result in the shaders required to support UI materials being compiled which will increase shader compile time and memory usage.
 	 */
diff --git a/Engine/Source/Runtime/Engine/Public/Materials/MaterialInterface.h b/Engine/Source/Runtime/Engine/Public/Materials/MaterialInterface.h
--- a/Engine/Source/Runtime/Engine/Public/Materials/MaterialInterface.h
+++ b/Engine/Source/Runtime/Engine/Public/Materials/MaterialInterface.h
@@ -93,6 +93,7 @@ enum EMaterialUsage : int
 	MATUSAGE_VolumetricCloud,
 	MATUSAGE_HeterogeneousVolumes,
 	MATUSAGE_MaterialCache,
+	MATUSAGE_StaticMesh,
 
 	MATUSAGE_MAX,
 };
