stat unit - display current render resolution
* If DynRes is enabled then render resolution is added to that line like this: "DynRes: 50% (1920x1080)
* If DynRes is disabled then instead of "DynRes: Unsupported" we now print out "RenderRes: 50% (1920x1080)"

diff --git a/Engine/Source/Runtime/Engine/Private/UnrealClient.cpp b/Engine/Source/Runtime/Engine/Private/UnrealClient.cpp
--- a/Engine/Source/Runtime/Engine/Private/UnrealClient.cpp
+++ b/Engine/Source/Runtime/Engine/Private/UnrealClient.cpp
@@ -797,7 +797,9 @@ int32 FStatUnitData::DrawStat(FViewport* InViewport, FCanvas* InCanvas, int32 In
 		}
 
 		// Dynamic resolution
+		if (DynamicResolutionStateInfos.Status == EDynamicResolutionStatus::Enabled || DynamicResolutionStateInfos.Status == EDynamicResolutionStatus::DebugForceEnabled)
 		{
+			const FIntPoint RenderResolution = GPixelRenderCounters.GetRenderResolution();
 			float ResolutionFraction = DynamicResolutionStateInfos.ResolutionFractionApproximations[GDynamicPrimaryResolutionFraction];
 			float ScreenPercentage = ResolutionFraction * 100.0f;
 
@@ -805,64 +807,58 @@ int32 FStatUnitData::DrawStat(FViewport* InViewport, FCanvas* InCanvas, int32 In
 			if (DynamicResolutionStateInfos.Status == EDynamicResolutionStatus::Enabled)
 			{
 				FColor Color = (ResolutionFraction < AlertResolutionFraction) ? StatRed : ((ResolutionFraction < FMath::Min(ResolutionFraction * 0.97f, 1.0f)) ? StatOrange : StatGreen);
-				DrawDefaultAvgCell(FString::Printf(TEXT("%3.1f%% x %3.1f%%"), ScreenPercentage, ScreenPercentage), Color);
+				DrawDefaultAvgCell(FString::Printf(TEXT("%3.1f%% (%dx%d)"), ScreenPercentage, RenderResolution.X, RenderResolution.Y), Color);
 			}
 			else if (DynamicResolutionStateInfos.Status == EDynamicResolutionStatus::DebugForceEnabled)
 			{
-				DrawDefaultAvgCell(FString::Printf(TEXT("%3.1f%% x %3.1f%%"), ScreenPercentage, ScreenPercentage), StatMagenda);
-			}
-			else if (DynamicResolutionStateInfos.Status == EDynamicResolutionStatus::Paused)
-			{
-				DrawDefaultAvgCell(TEXT("Paused"), StatMagenda);
-			}
-			else if (DynamicResolutionStateInfos.Status == EDynamicResolutionStatus::Disabled)
-			{
-				DrawDefaultAvgCell(TEXT("OFF"), FColor(160, 160, 160));
-			}
-			else if (DynamicResolutionStateInfos.Status == EDynamicResolutionStatus::Unsupported)
-			{
-				DrawDefaultAvgCell(TEXT("Unsupported"), FColor(160, 160, 160));
+				DrawDefaultAvgCell(FString::Printf(TEXT("%3.1f%% (%dx%d)"), ScreenPercentage, RenderResolution.X, RenderResolution.Y), StatMagenda);
 			}
 			else
 			{
 				check(0);
 			}
 			PushRows(/* RowsCount = */ 1);
-		}
 
-		// Other dynamic render scalings
-		if (DynamicResolutionStateInfos.Status == EDynamicResolutionStatus::Enabled && DynamicRenderScaling::IsSupported())
-		{
-			for (TLinkedList<DynamicRenderScaling::FBudget*>::TIterator BudgetIt(DynamicRenderScaling::FBudget::GetGlobalList()); BudgetIt; BudgetIt.Next())
+			// Other dynamic render scalings
+			if (DynamicResolutionStateInfos.Status == EDynamicResolutionStatus::Enabled && DynamicRenderScaling::IsSupported())
 			{
-				const DynamicRenderScaling::FBudget& Budget = **BudgetIt;
-				const DynamicRenderScaling::FHeuristicSettings& HeuristicSettings = Budget.GetSettings();
-				if (Budget == GDynamicPrimaryResolutionFraction || !HeuristicSettings.IsEnabled())
+				for (TLinkedList<DynamicRenderScaling::FBudget*>::TIterator BudgetIt(DynamicRenderScaling::FBudget::GetGlobalList()); BudgetIt; BudgetIt.Next())
 				{
-					continue;
-				}
-
-				float ResolutionFraction = DynamicResolutionStateInfos.ResolutionFractionApproximations[Budget];
-				float ScreenPercentage = ResolutionFraction * 100.0f;
+					const DynamicRenderScaling::FBudget& Budget = **BudgetIt;
+					const DynamicRenderScaling::FHeuristicSettings& HeuristicSettings = Budget.GetSettings();
+					if (Budget == GDynamicPrimaryResolutionFraction || !HeuristicSettings.IsEnabled())
+					{
+						continue;
+					}
 
-				FString DisplayName = Budget.GetName();
-				DisplayName.ReplaceInline(TEXT("Dynamic"), TEXT("Dyn"));
-				DisplayName.ReplaceInline(TEXT("Resolution"), TEXT("Res"));
+					FString DisplayName = Budget.GetName();
+					DisplayName.ReplaceInline(TEXT("Dynamic"), TEXT("Dyn"));
+					DisplayName.ReplaceInline(TEXT("Resolution"), TEXT("Res"));
 
-				FColor Color = (ResolutionFraction < AlertResolutionFraction) ? StatRed : ((ResolutionFraction < FMath::Min(ResolutionFraction * 0.97f, 1.0f)) ? StatOrange : StatGreen);
+					FColor Color = (ResolutionFraction < AlertResolutionFraction) ? StatRed : ((ResolutionFraction < FMath::Min(ResolutionFraction * 0.97f, 1.0f)) ? StatOrange : StatGreen);
 
-				DrawTitleString(*DisplayName, /* UnitGraphColor = */ FColor::White);
-				if (HeuristicSettings.Model == DynamicRenderScaling::EHeuristicModel::Quadratic)
-				{
-					DrawDefaultAvgCell(FString::Printf(TEXT("%3.1f%% x %3.1f%%"), ScreenPercentage, ScreenPercentage), Color);
-				}
-				else
-				{
-					DrawDefaultAvgCell(FString::Printf(TEXT("%3.1f%%"), ScreenPercentage), Color);
+					DrawTitleString(*DisplayName, /* UnitGraphColor = */ FColor::White);
+					if (HeuristicSettings.Model == DynamicRenderScaling::EHeuristicModel::Quadratic)
+					{
+						DrawDefaultAvgCell(FString::Printf(TEXT("%3.1f%% x %3.1f%%"), ScreenPercentage, ScreenPercentage), Color);
+					}
+					else
+					{
+						DrawDefaultAvgCell(FString::Printf(TEXT("%3.1f%%"), ScreenPercentage), Color);
+					}
+					PushRows(/* RowsCount = */ 1);
 				}
-				PushRows(/* RowsCount = */ 1);
 			}
 		}
+		else
+		{
+			// Render Resolution
+			const float ScreenPercentage = GPixelRenderCounters.GetResolutionFraction() * 100.0f;
+			const FIntPoint RenderResolution = GPixelRenderCounters.GetRenderResolution();
+			DrawTitleString(TEXT("RenderRes"), NoUnitGraphColor);
+			DrawDefaultAvgCell(FString::Printf(TEXT("%3.1f%% (%dx%d)"), ScreenPercentage, RenderResolution.X, RenderResolution.Y), StatGreen);
+			PushRows(/* RowsCount = */ 1);
+		}
 
 		// Draw calls
 		{
diff --git a/Engine/Source/Runtime/RenderCore/Public/RenderCounters.h b/Engine/Source/Runtime/RenderCore/Public/RenderCounters.h
--- a/Engine/Source/Runtime/RenderCore/Public/RenderCounters.h
+++ b/Engine/Source/Runtime/RenderCore/Public/RenderCounters.h
@@ -2,13 +2,21 @@
 
 #pragma once
 
-#include "CoreGlobals.h"
-
-
+#include "CoreMinimal.h"
 
 class FPixelRenderCounters
 {
 public:
+	float GetResolutionFraction() const
+	{
+		return ResolutionFraction;
+	}
+
+	FIntPoint GetRenderResolution() const
+	{
+		return RenderResolution;
+	}
+
 	uint32 GetPixelRenderCount() const
 	{
 		return PrevPixelRenderCount;
@@ -19,13 +27,17 @@ public:
 		return PrevPixelDisplayCount;
 	}
 
-	void AddViewStatistics(uint32 PixelRenderCount, uint32 PixelDisplayCount)
+	void AddViewStatistics(float InResolutionFraction, FIntPoint InRenderResolution, uint32 PixelDisplayCount)
 	{
-		CurrentPixelRenderCount += PixelRenderCount;
+		ResolutionFraction = InResolutionFraction;
+		RenderResolution = InRenderResolution;
+		CurrentPixelRenderCount += InRenderResolution.X * InRenderResolution.Y;
 		CurrentPixelDisplayCount += PixelDisplayCount;
 	}
 
 private:
+	float ResolutionFraction = 0.0f;
+	FIntPoint RenderResolution = FIntPoint(0, 0);
 	uint32 PrevPixelRenderCount = 0;
 	uint32 PrevPixelDisplayCount = 0;
 	uint32 CurrentPixelRenderCount = 0;
diff --git a/Engine/Source/Runtime/Renderer/Private/SceneRendering.cpp b/Engine/Source/Runtime/Renderer/Private/SceneRendering.cpp
--- a/Engine/Source/Runtime/Renderer/Private/SceneRendering.cpp
+++ b/Engine/Source/Runtime/Renderer/Private/SceneRendering.cpp
@@ -3306,7 +3306,7 @@ void FSceneRenderer::PrepareViewRectsForRendering()
 		if (!ViewFamily.EngineShowFlags.HitProxies && View.AntiAliasingMethod == AAM_TSR)
 		{
 			FIntPoint SecondaryViewRect = View.GetSecondaryViewRectSize();
-			GPixelRenderCounters.AddViewStatistics(View.ViewRect.Area(), SecondaryViewRect.X * SecondaryViewRect.Y);
+			GPixelRenderCounters.AddViewStatistics(ViewResolutionFraction, View.ViewRect.Size(), SecondaryViewRect.X * SecondaryViewRect.Y);
 		}
 	}
 
