Fixed merge issue with previous cl.

diff --git a/Engine/Source/Runtime/Engine/Private/StreamableManager.cpp b/Engine/Source/Runtime/Engine/Private/StreamableManager.cpp
--- a/Engine/Source/Runtime/Engine/Private/StreamableManager.cpp
+++ b/Engine/Source/Runtime/Engine/Private/StreamableManager.cpp
@@ -2283,7 +2283,7 @@ TSharedPtr<FStreamableHandle> FStreamableManager::CreateCombinedHandle(
 	{
 		// Complete immediately if nothing to load
 		TSharedPtr<FStreamableHandle> ErrorHandle = CreateErrorHandle(
-			FStreamableError{EStreamableErrorType::PackageLoadCanceled, TEXT("Nothing to load")},
+			UE::UnifiedError::StreamableManager::PackageLoadCancelled::MakeError(),
 			FStreamableAsyncLoadParams{});
 		Params.OnComplete.ExecuteIfBound(ErrorHandle);
 		return nullptr;
