Add cvar to control if duplicated vertices need to be stored for skin cache recompute normals - by default duplicated vertices are not stored anymore to save cooked package memory

diff --git a/Engine/Source/Runtime/Engine/Private/GPUSkinCache.cpp b/Engine/Source/Runtime/Engine/Private/GPUSkinCache.cpp
--- a/Engine/Source/Runtime/Engine/Private/GPUSkinCache.cpp
+++ b/Engine/Source/Runtime/Engine/Private/GPUSkinCache.cpp
@@ -127,7 +127,7 @@ TAutoConsoleVariable<int32> CVarGPUSkinCacheRecomputeTangents(
 	TEXT("r.SkinCache.RecomputeTangents"),
 	2,
 	TEXT("This option enables recomputing the vertex tangents on the GPU.\n")
-	TEXT("Can be changed at runtime, requires both r.SkinCache.CompileShaders=1, r.SkinCache.Mode=1 and r.SkinCache.Allow=1\n")
+	TEXT("Can be changed at runtime, requires both r.SkinCache.CompileShaders=1, r.SkinCache.Mode=1, r.SkinCache.Allow=1 and r.SkinCache.StoreDuplicatedVerticesForRecomputeTangents\n")
 	TEXT(" 0: off\n")
 	TEXT(" 1: on, forces all skinned object to Recompute Tangents\n")
 	TEXT(" 2: on, only recompute tangents on skinned objects who ticked the Recompute Tangents checkbox(default)\n"),
@@ -158,12 +158,21 @@ static TAutoConsoleVariable<float> CVarGPUSkinCacheSceneMemoryLimitInMB(
 	ECVF_RenderThreadSafe
 );
 
-static int32 GAllowDupedVertsForRecomputeTangents = 0;
+static int32 GStoreDuplicatedVerticesForRecomputeTangents = 0;
+FAutoConsoleVariableRef CVarGPUSkinCacheStoreDuplicatedVerticesForRecomputeTangents(
+	TEXT("r.SkinCache.StoreDuplicatedVerticesForRecomputeTangents"),
+	GStoreDuplicatedVerticesForRecomputeTangents,
+	TEXT("0: Don't store duplicated vertices for all skeletal mesh render sections. It will still be stored if the render section has bRecomputeTangent set. (default)\n")
+	TEXT("1: Store duplicated vertices for all skeletal mesh render sections.\n"),
+	ECVF_ReadOnly
+);
+
+static int32 GUseDuplicatedVerticesForRecomputeTangents = 1;
 FAutoConsoleVariableRef CVarGPUSkinCacheAllowDupedVertesForRecomputeTangents(
-	TEXT("r.SkinCache.AllowDupedVertsForRecomputeTangents"),
-	GAllowDupedVertsForRecomputeTangents,
-	TEXT("0: off (default)\n")
-	TEXT("1: Forces that vertices at the same position will be treated differently and has the potential to cause seams when verts are split.\n"),
+	TEXT("r.SkinCache.UseDuplicatedVerticesForRecomputeTangents"),
+	GUseDuplicatedVerticesForRecomputeTangents,
+	TEXT("0: Disable usage of duplicated vertices for runtime tangent recomputation/\n")
+	TEXT("1: Use stored duplicated vertices if they are available (default).\n"),
 	ECVF_RenderThreadSafe
 );
 
@@ -265,10 +274,15 @@ ENGINE_API bool GPUSkinCacheNeedsDuplicatedVertices()
 #if WITH_EDITOR // Duplicated vertices are used in the editor when merging meshes
 	return true;
 #else
-	return GAllowDupedVertsForRecomputeTangents == 0;
+	return GPUSkinCacheStoreDuplicatedVertices();
 #endif
 }
 
+ENGINE_API bool GPUSkinCacheStoreDuplicatedVertices()
+{
+	return GStoreDuplicatedVerticesForRecomputeTangents > 0;
+}
+
 RDG_REGISTER_BLACKBOARD_STRUCT(FGPUSkinCache::FTaskData);
 
 // determine if during DispatchUpdateSkinning caching should occur
@@ -1101,6 +1115,11 @@ FGPUSkinCache::FGPUSkinCache(ERHIFeatureLevel::Type InFeatureLevel, bool bInRequ
 	, World(InWorld)
 {
 	check(World);
+
+	if (GSkinCacheRecomputeTangents == 1 && GPUSkinCacheStoreDuplicatedVertices == 0)
+	{
+		UE_LOG(LogSkinCache, Warning, TEXT("r.SkinCache.RecomputeTangents is set to 1 to update all skinned objects but duplicated vertices are not are not always stored. Set r.SkinCache.RecomputeTangents to 2 or r.SkinCache.StoreDuplicatedVerticesForRecomputeTangents to 1."));
+	}
 }
 
 FGPUSkinCache::~FGPUSkinCache()
@@ -1185,7 +1204,7 @@ public:
 		SetUAVParameter(BatchedParameters, IntermediateAccumBufferUAV, StagingBuffer.UAV);
 		SetShaderValue(BatchedParameters, IntermediateAccumBufferOffset, GRecomputeTangentsParallelDispatch * DispatchData.IntermediateAccumulatedTangentBufferOffset);
 
-        if (!GAllowDupedVertsForRecomputeTangents)
+        if (DispatchData.DuplicatedIndices)
         {
 		    SetSRVParameter(BatchedParameters, DuplicatedIndices, DispatchData.DuplicatedIndices);
             SetSRVParameter(BatchedParameters, DuplicatedIndicesIndices, DispatchData.DuplicatedIndicesIndices);
@@ -1431,17 +1450,19 @@ void FGPUSkinCache::DispatchUpdateSkinTangentsTrianglePass(FRHICommandList& RHIC
 
 	bool bFullPrecisionUV = LodData.StaticVertexBuffers.StaticMeshVertexBuffer.GetUseFullPrecisionUVs();
 
+	bool bUseDuplicatedVertices = GUseDuplicatedVerticesForRecomputeTangents && LodData.RenderSections[SectionIndex].DuplicatedVerticesBuffer.DuplicatedVerticesIndexBuffer.VertexBufferSRV;
+
 	TShaderRef<FBaseRecomputeTangentsPerTriangleShader> Shader;
 
 	if (bFullPrecisionUV)
 	{
-		if (GAllowDupedVertsForRecomputeTangents) Shader = ComputeShader01;
-		else Shader = ComputeShader11;
+		if (bUseDuplicatedVertices) Shader = ComputeShader11;
+		else Shader = ComputeShader01;
 	}
 	else
 	{
-		if (GAllowDupedVertsForRecomputeTangents) Shader = ComputeShader00;
-		else Shader = ComputeShader10;
+		if (bUseDuplicatedVertices) Shader = ComputeShader10;
+		else Shader = ComputeShader00;
 	}
 
 	check(Shader.IsValid());
@@ -1460,7 +1481,7 @@ void FGPUSkinCache::DispatchUpdateSkinTangentsTrianglePass(FRHICommandList& RHIC
 		, bFullPrecisionUV
 	);
 
-	if (!GAllowDupedVertsForRecomputeTangents)
+	if (bUseDuplicatedVertices)
 	{
 #if WITH_EDITOR
 		check(LodData.RenderSections[SectionIndex].DuplicatedVerticesBuffer.DupVertData.Num() && LodData.RenderSections[SectionIndex].DuplicatedVerticesBuffer.DupVertIndexData.Num());
@@ -2628,7 +2649,12 @@ void FGPUSkinCache::CVarSinkFunction()
 
 	if (NewGPUSkinCacheValue != GEnableGPUSkinCache || NewRecomputeTangentsValue != GSkinCacheRecomputeTangents
 		|| NewSceneMaxSizeInMb != GSkinCacheSceneMemoryLimitInMB || NewNumTangentIntermediateBuffers != GNumTangentIntermediateBuffers)
-	{
+	{		
+		if (NewRecomputeTangentsValue == 1 && GPUSkinCacheStoreDuplicatedVertices == 0)
+		{
+			UE_LOG(LogSkinCache, Warning, TEXT("r.SkinCache.RecomputeTangents is set to 1 to update all skinned objects but duplicated vertices are not are not always stored. Set r.SkinCache.RecomputeTangents to 2 or r.SkinCache.StoreDuplicatedVerticesForRecomputeTangents to 1."));
+		}
+
 		ENQUEUE_RENDER_COMMAND(DoEnableSkinCaching)(UE::RenderCommandPipe::SkeletalMesh,
 			[NewRecomputeTangentsValue, NewGPUSkinCacheValue, NewSceneMaxSizeInMb, NewNumTangentIntermediateBuffers](FRHICommandList& RHICmdList)
 		{
diff --git a/Engine/Source/Runtime/Engine/Private/SkeletalMesh.cpp b/Engine/Source/Runtime/Engine/Private/SkeletalMesh.cpp
--- a/Engine/Source/Runtime/Engine/Private/SkeletalMesh.cpp
+++ b/Engine/Source/Runtime/Engine/Private/SkeletalMesh.cpp
@@ -56,6 +56,7 @@
 #include "UObject/FortniteMainBranchObjectVersion.h"
 #include "UObject/ReleaseObjectVersion.h"
 #include "GPUSkinCacheVisualizationData.h"
+#include "GPUSkinCache.h"
 #include "MeshPaintVisualize.h"
 
 #if WITH_EDITOR
@@ -5427,6 +5428,9 @@ FString USkeletalMesh::BuildDerivedDataKey(const ITargetPlatform* TargetPlatform
 		static FString CachedNaniteVersion = FDevSystemGuids::GetSystemGuid(FDevSystemGuids::Get().NANITE_DERIVEDDATA_VER).ToString();
 		KeySuffix += *CachedNaniteVersion;
 	}
+	
+	const bool bStoreDuplicatedVertices = GPUSkinCacheStoreDuplicatedVertices();
+	KeySuffix += bStoreDuplicatedVertices ? TEXT("_SDV_1") : TEXT("_SDV_0");
 
 #if PLATFORM_CPU_ARM_FAMILY
 	// Separate out arm keys as x64 and arm64 clang do not generate the same data for a given
diff --git a/Engine/Source/Runtime/Engine/Private/SkeletalMeshLODRenderData.cpp b/Engine/Source/Runtime/Engine/Private/SkeletalMeshLODRenderData.cpp
--- a/Engine/Source/Runtime/Engine/Private/SkeletalMeshLODRenderData.cpp
+++ b/Engine/Source/Runtime/Engine/Private/SkeletalMeshLODRenderData.cpp
@@ -90,11 +90,8 @@ static bool IsMeshDeformerAvailable(EShaderPlatform InPlatform)
  */
 static bool RequiresDuplicateVerticesInCook(EShaderPlatform InPlatform, FSkelMeshRenderSection const& RenderSection)
 {
-	// DuplicatedVertices are cooked if we have GPUSkinCache or MeshDeformer systems for this platform.
-	// We always cook when GPUSkinCache is available so that we can support runtime switch of r.SkinCache.RecomputeTangents.
-	// For MeshDeformer we only cook if the section was marked for bRecomputeTangent.
-	const bool bMeshDeformersAvailable = IsMeshDeformerAvailable(InPlatform);
-	return (RenderSection.bRecomputeTangent && bMeshDeformersAvailable) || IsGPUSkinCacheAvailable(InPlatform);
+	// DuplicatedVertices are cooked if we have GPUSkinCache and r.SkinCache.StoreDuplicatedVerticesForRecomputeTangents is set or if the section was marked for bRecomputeTangent.
+	return IsGPUSkinCacheAvailable(InPlatform) && (RenderSection.bRecomputeTangent || GPUSkinCacheStoreDuplicatedVertices());
 }
 
 static bool RequiresDuplicateVerticesInCook(const ITargetPlatform& InTargetPlatform, FSkelMeshRenderSection const& RenderSection)
@@ -112,22 +109,6 @@ static bool RequiresDuplicateVerticesInCook(const ITargetPlatform& InTargetPlatf
 	return false;
 }
 
-/** 
- * This function returns true if the duplicate vertices should be kept at runtime initialization. 
- * We may cook the data, but use this runtime check to discard it and save memory.
- */
-static bool RequiresDuplicateVertices()
-{
-	// Never drop at runtime if the data was cooked for MeshDeformers.
-	if (IsMeshDeformerAvailable(GMaxRHIShaderPlatform))
-	{
-		return true;
-	}
-
-	// Can assume data is for GPUSkinCache. Respect the CVar option to discard all data at runtime.
-	return GPUSkinCacheNeedsDuplicatedVertices();
-}
-
 static bool IsMeshDeformerSupportedByCookTargetPlatform(const ITargetPlatform& InTargetPlatform)
 {
 	TArray<FName> TargetedShaderFormats;
@@ -260,25 +241,14 @@ void FSkeletalMeshLODRenderData::InitResources(bool bNeedsVertexColors, int32 LO
 		BeginInitResource(OwnerName, &ClothVertexBuffer, &UE::RenderCommandPipe::SkeletalMesh);
 	}
 
-	// We can discard any the cooked DuplicatedVertices here based on runtime settings.
-	bool bDiscardDuplicatedVertices = !RequiresDuplicateVertices();
+	// If we have duplicate vertex data then initialize them
 	for (FSkelMeshRenderSection& RenderSection : RenderSections)
 	{
 		if(RenderSection.DuplicatedVerticesBuffer.DupVertData.Num())
 		{
-			if (bDiscardDuplicatedVertices)
-			{
-#if !WITH_EDITOR
-				// Discard CPU data in cooked builds. Keep CPU data when in the editor for geometry operations.
-				RenderSection.DuplicatedVerticesBuffer.ReleaseCPUResources();
-#endif
-			}
-			else
-			{
-				// No need to discard CPU data in cooked builds as bNeedsCPUAccess is false (see FDuplicatedVerticesBuffer constructor), 
-				// so it'd be auto-discarded after the RHI has copied the resource data. Keep CPU data when in the editor for geometry operations.
-				BeginInitResource(OwnerName, &RenderSection.DuplicatedVerticesBuffer, &UE::RenderCommandPipe::SkeletalMesh);
-			}
+			// No need to discard CPU data in cooked builds as bNeedsCPUAccess is false (see FDuplicatedVerticesBuffer constructor), 
+			// so it'd be auto-discarded after the RHI has copied the resource data. Keep CPU data when in the editor for geometry operations.
+			BeginInitResource(OwnerName, &RenderSection.DuplicatedVerticesBuffer, &UE::RenderCommandPipe::SkeletalMesh);
 		}
 	}
 
diff --git a/Engine/Source/Runtime/Engine/Public/GPUSkinCache.h b/Engine/Source/Runtime/Engine/Public/GPUSkinCache.h
--- a/Engine/Source/Runtime/Engine/Public/GPUSkinCache.h
+++ b/Engine/Source/Runtime/Engine/Public/GPUSkinCache.h
@@ -65,8 +65,11 @@ struct FRayTracingGeometrySegment;
 
 extern bool ShouldWeCompileGPUSkinVFShaders(EShaderPlatform Platform, ERHIFeatureLevel::Type FeatureLevel);
 
+UE_DEPRECATED(5.6, "Use GPUSkinCacheStoreDuplicatedVertices instead to check if duplicate vertices need to be stored.")
 extern ENGINE_API bool GPUSkinCacheNeedsDuplicatedVertices();
 
+extern ENGINE_API bool GPUSkinCacheStoreDuplicatedVertices();
+
 extern ENGINE_API ESkinCacheDefaultBehavior GetSkinCacheDefaultBehavior();
 
 // Is it actually enabled?
