Fixed static analysis "warning V547: Expression is always true."

diff --git a/Engine/Source/Runtime/Engine/Private/UnrealClient.cpp b/Engine/Source/Runtime/Engine/Private/UnrealClient.cpp
--- a/Engine/Source/Runtime/Engine/Private/UnrealClient.cpp
+++ b/Engine/Source/Runtime/Engine/Private/UnrealClient.cpp
@@ -809,13 +809,9 @@ int32 FStatUnitData::DrawStat(FViewport* InViewport, FCanvas* InCanvas, int32 In
 				FColor Color = (ResolutionFraction < AlertResolutionFraction) ? StatRed : ((ResolutionFraction < FMath::Min(ResolutionFraction * 0.97f, 1.0f)) ? StatOrange : StatGreen);
 				DrawDefaultAvgCell(FString::Printf(TEXT("%3.1f%% (%dx%d)"), ScreenPercentage, RenderResolution.X, RenderResolution.Y), Color);
 			}
-			else if (DynamicResolutionStateInfos.Status == EDynamicResolutionStatus::DebugForceEnabled)
-			{
-				DrawDefaultAvgCell(FString::Printf(TEXT("%3.1f%% (%dx%d)"), ScreenPercentage, RenderResolution.X, RenderResolution.Y), StatMagenda);
-			}
 			else
 			{
-				check(0);
+				DrawDefaultAvgCell(FString::Printf(TEXT("%3.1f%% (%dx%d)"), ScreenPercentage, RenderResolution.X, RenderResolution.Y), StatMagenda);
 			}
 			PushRows(/* RowsCount = */ 1);
 
