Build fix for using not using GStoreDuplicatedVerticesForRecomputeTangents directly but incorrect function wrapper

diff --git a/Engine/Source/Runtime/Engine/Private/GPUSkinCache.cpp b/Engine/Source/Runtime/Engine/Private/GPUSkinCache.cpp
--- a/Engine/Source/Runtime/Engine/Private/GPUSkinCache.cpp
+++ b/Engine/Source/Runtime/Engine/Private/GPUSkinCache.cpp
@@ -1116,7 +1116,7 @@ FGPUSkinCache::FGPUSkinCache(ERHIFeatureLevel::Type InFeatureLevel, bool bInRequ
 {
 	check(World);
 
-	if (GSkinCacheRecomputeTangents == 1 && GPUSkinCacheStoreDuplicatedVertices == 0)
+	if (GSkinCacheRecomputeTangents == 1 && GStoreDuplicatedVerticesForRecomputeTangents == 0)
 	{
 		UE_LOG(LogSkinCache, Warning, TEXT("r.SkinCache.RecomputeTangents is set to 1 to update all skinned objects but duplicated vertices are not are not always stored. Set r.SkinCache.RecomputeTangents to 2 or r.SkinCache.StoreDuplicatedVerticesForRecomputeTangents to 1."));
 	}
@@ -2650,7 +2650,7 @@ void FGPUSkinCache::CVarSinkFunction()
 	if (NewGPUSkinCacheValue != GEnableGPUSkinCache || NewRecomputeTangentsValue != GSkinCacheRecomputeTangents
 		|| NewSceneMaxSizeInMb != GSkinCacheSceneMemoryLimitInMB || NewNumTangentIntermediateBuffers != GNumTangentIntermediateBuffers)
 	{		
-		if (NewRecomputeTangentsValue == 1 && GPUSkinCacheStoreDuplicatedVertices == 0)
+		if (NewRecomputeTangentsValue == 1 && GStoreDuplicatedVerticesForRecomputeTangents == 0)
 		{
 			UE_LOG(LogSkinCache, Warning, TEXT("r.SkinCache.RecomputeTangents is set to 1 to update all skinned objects but duplicated vertices are not are not always stored. Set r.SkinCache.RecomputeTangents to 2 or r.SkinCache.StoreDuplicatedVerticesForRecomputeTangents to 1."));
 		}
