SceneGraph
* Submitting a disabled and WIP version of the NetEntityFactory
* When enabled this allows entities to be replicated without having to use an EntityReplicator
* Only works when Iris is running

#rnx

diff --git a/Engine/Source/Runtime/Engine/Classes/Engine/NetDriver.h b/Engine/Source/Runtime/Engine/Classes/Engine/NetDriver.h
--- a/Engine/Source/Runtime/Engine/Classes/Engine/NetDriver.h
+++ b/Engine/Source/Runtime/Engine/Classes/Engine/NetDriver.h
@@ -2085,6 +2085,9 @@ public:
 	/** Returns identifier used for NetTrace */
 	inline uint32 GetNetTraceId() const { return NetTraceId; }
 
+	/** Is the netdriver allowed to create destruction infos */
+	ENGINE_API bool CanCreateDestructionInfo() const;
+
 	/** Sends a message to a client to destroy an actor to the client.  The actor may already be destroyed locally. */
 	ENGINE_API int64 SendDestructionInfo(UNetConnection* Connection, FActorDestructionInfo* DestructionInfo);
 
diff --git a/Engine/Source/Runtime/Engine/Private/Net/Iris/ReplicationSystem/EngineReplicationBridge.cpp b/Engine/Source/Runtime/Engine/Private/Net/Iris/ReplicationSystem/EngineReplicationBridge.cpp
--- a/Engine/Source/Runtime/Engine/Private/Net/Iris/ReplicationSystem/EngineReplicationBridge.cpp
+++ b/Engine/Source/Runtime/Engine/Private/Net/Iris/ReplicationSystem/EngineReplicationBridge.cpp
@@ -508,7 +508,7 @@ void UEngineReplicationBridge::StopReplicatingActor(AActor* Actor, EEndPlayReaso
 		Flags |= EEndReplicationFlags::SkipPendingEndReplicationValidation;
 	}
 
-	const bool bShouldCreateDestructionInfo = RefHandle.IsStatic() && bShouldDestroyObject && GetReplicationSystem()->IsServer() && !GIsReconstructingBlueprintInstances && !UE::Net::ShouldIgnoreStaticActorDestruction();
+	const bool bShouldCreateDestructionInfo = RefHandle.IsStatic() && bShouldDestroyObject && GetReplicationSystem()->IsServer();
 	if (bShouldCreateDestructionInfo)
 	{
 		const UObjectReplicationBridge::FDestructionParameters EndReplicationParameters
@@ -693,6 +693,16 @@ UEngineReplicationBridge* UEngineReplicationBridge::Create(UNetDriver* NetDriver
 	return Bridge;
 }
 
+bool UEngineReplicationBridge::CanCreateDestructionInfo() const
+{
+	if(ensure(NetDriver))
+	{
+		return NetDriver->CanCreateDestructionInfo();
+	}
+
+	return false;
+}
+
 float UEngineReplicationBridge::GetPollFrequencyOfRootObject(const UObject* ReplicatedObject) const
 {
 	const AActor* ReplicatedActor = CastChecked<AActor>(ReplicatedObject);
diff --git a/Engine/Source/Runtime/Engine/Private/NetDriver.cpp b/Engine/Source/Runtime/Engine/Private/NetDriver.cpp
--- a/Engine/Source/Runtime/Engine/Private/NetDriver.cpp
+++ b/Engine/Source/Runtime/Engine/Private/NetDriver.cpp
@@ -4049,6 +4049,11 @@ FActorDestructionInfo* UNetDriver::CreateDestructionInfo(AActor* ThisActor, FAct
 	return &NewInfo;
 }
 
+bool UNetDriver::CanCreateDestructionInfo() const
+{
+	return !GIsReconstructingBlueprintInstances && !UE::Net::ShouldIgnoreStaticActorDestruction();
+}
+
 void UNetDriver::NotifyActorDestroyed(AActor* ThisActor, bool IsSeamlessTravel)
 {
 #if UE_WITH_IRIS
@@ -4058,9 +4063,9 @@ void UNetDriver::NotifyActorDestroyed(AActor* ThisActor, bool IsSeamlessTravel)
 		if (ReplicationSystem->IsServer())
 		{
 			const bool bActorHasRole = ThisActor->GetRemoteRole() != ROLE_None;
-			const bool bShouldCreateDestructionInfoForInitiallyDormantActor = IsDormInitialStartupActor(ThisActor) && bActorHasRole && !IsSeamlessTravel && !GIsReconstructingBlueprintInstances && !UE::Net::ShouldIgnoreStaticActorDestruction();
+			const bool bShouldCreateDestructionInfoForInitiallyDormantActor = IsDormInitialStartupActor(ThisActor) && bActorHasRole && !IsSeamlessTravel;
 
-			if (bShouldCreateDestructionInfoForInitiallyDormantActor)
+			if (bShouldCreateDestructionInfoForInitiallyDormantActor && CanCreateDestructionInfo())
 			{
 				if (UEngineReplicationBridge* Bridge = ReplicationSystem->GetReplicationBridgeAs<UEngineReplicationBridge>())
 				{
@@ -4098,9 +4103,9 @@ void UNetDriver::NotifyActorDestroyed(AActor* ThisActor, bool IsSeamlessTravel)
 
 		const bool bIsActorStatic = !GetNetGuidCache()->IsDynamicObject( ThisActor ) || IsStaticReplicator;
 		const bool bActorHasRole = ThisActor->GetRemoteRole() != ROLE_None;
-		const bool bShouldCreateDestructionInfo = bIsServer && bIsActorStatic && bActorHasRole && !IsSeamlessTravel && !GIsReconstructingBlueprintInstances && !UE::Net::ShouldIgnoreStaticActorDestruction();
+		const bool bShouldCreateDestructionInfo = bIsServer && bIsActorStatic && bActorHasRole && !IsSeamlessTravel;
 
-		if (bShouldCreateDestructionInfo)
+		if (bShouldCreateDestructionInfo && CanCreateDestructionInfo())
 		{
 			UE_LOG(LogNet, VeryVerbose, TEXT("NotifyActorDestroyed %s - StartupActor"), *ThisActor->GetPathName() );
 			DestructionInfo = CreateDestructionInfo(ThisActor, DestructionInfo);
diff --git a/Engine/Source/Runtime/Engine/Public/Net/Iris/ReplicationSystem/EngineReplicationBridge.h b/Engine/Source/Runtime/Engine/Public/Net/Iris/ReplicationSystem/EngineReplicationBridge.h
--- a/Engine/Source/Runtime/Engine/Public/Net/Iris/ReplicationSystem/EngineReplicationBridge.h
+++ b/Engine/Source/Runtime/Engine/Public/Net/Iris/ReplicationSystem/EngineReplicationBridge.h
@@ -134,7 +134,7 @@ protected:
 	virtual bool IsAllowedToDestroyInstance(const UObject* Instance) const override;
 	virtual void OnProtocolMismatchDetected(FNetRefHandle ObjectHandle) override;
 	virtual void OnProtocolMismatchReported(FNetRefHandle RefHandle, uint32 ConnectionId) override;
-
+	virtual bool CanCreateDestructionInfo() const override;
 	virtual float GetPollFrequencyOfRootObject(const UObject* ReplicatedObject) const override;
 
 	/** Returns true if the class is derived from Actor and its CDO has set bReplicates. */
diff --git a/Engine/Source/Runtime/Experimental/Iris/Core/Private/Iris/ReplicationSystem/ReplicationBridge.cpp b/Engine/Source/Runtime/Experimental/Iris/Core/Private/Iris/ReplicationSystem/ReplicationBridge.cpp
--- a/Engine/Source/Runtime/Experimental/Iris/Core/Private/Iris/ReplicationSystem/ReplicationBridge.cpp
+++ b/Engine/Source/Runtime/Experimental/Iris/Core/Private/Iris/ReplicationSystem/ReplicationBridge.cpp
@@ -987,6 +987,12 @@ UE::Net::FNetRefHandle UReplicationBridge::StoreDestructionInfo(FNetRefHandle Ha
 		return FNetRefHandle::GetInvalid();
 	}
 
+	// Not allowed to create destruction infos at the moment
+	if (!CanCreateDestructionInfo())
+	{
+		return FNetRefHandle::GetInvalid();
+	}
+
 	FReplicationProtocolManager* ProtocolManager = GetReplicationProtocolManager();
 
 	const FNetObjectFactoryId NetFactoryId = Parameters.NetFactoryId;
diff --git a/Engine/Source/Runtime/Experimental/Iris/Core/Public/Iris/ReplicationSystem/ReplicationBridge.h b/Engine/Source/Runtime/Experimental/Iris/Core/Public/Iris/ReplicationSystem/ReplicationBridge.h
--- a/Engine/Source/Runtime/Experimental/Iris/Core/Public/Iris/ReplicationSystem/ReplicationBridge.h
+++ b/Engine/Source/Runtime/Experimental/Iris/Core/Public/Iris/ReplicationSystem/ReplicationBridge.h
@@ -191,6 +191,9 @@ protected:
 	/** Invoked when we start to replicate an object for a specific connection to fill in any initial dependencies */
 	IRISCORE_API virtual void GetInitialDependencies(FNetRefHandle Handle, FNetDependencyInfoArray& OutDependencies) const;
 
+	/** Returns if the bridge is allowed to create new destruction info at this moment. */
+	virtual bool CanCreateDestructionInfo() const { return true; }
+
 protected:
 	// Forward calls to internal operations that we allow replication bridges to access
 
