- Add support for float32 formats to Sub UV generation

diff --git a/Engine/Source/Runtime/Engine/Private/Particles/SubUVAnimation.cpp b/Engine/Source/Runtime/Engine/Private/Particles/SubUVAnimation.cpp
--- a/Engine/Source/Runtime/Engine/Private/Particles/SubUVAnimation.cpp
+++ b/Engine/Source/Runtime/Engine/Private/Particles/SubUVAnimation.cpp
@@ -476,7 +476,7 @@ FIntPoint GNeighbors[] =
 
 bool IsSupportedFormat(ETextureSourceFormat SrcFormat)
 {
-	return SrcFormat == TSF_G8 || SrcFormat == TSF_BGRA8 || SrcFormat == TSF_RGBA16 || SrcFormat == TSF_RGBA16F;
+	return SrcFormat == TSF_G8 || SrcFormat == TSF_R32F || SrcFormat == TSF_BGRA8 || SrcFormat == TSF_RGBA16 || SrcFormat == TSF_RGBA16F || SrcFormat == TSF_RGBA32F;
 }
 
 uint32 GetByteSizePerPixel(ETextureSourceFormat SrcFormat)
@@ -484,9 +484,11 @@ uint32 GetByteSizePerPixel(ETextureSourceFormat SrcFormat)
 	switch (SrcFormat)
 	{
 		case TSF_G8:		return 1;
+		case TSF_R32F:		return sizeof(float);
 		case TSF_BGRA8:		return 4;
 		case TSF_RGBA16:	return 8;
 		case TSF_RGBA16F:	return sizeof(FFloat16) * 4;
+		case TSF_RGBA32F:	return sizeof(float) * 4;
 		default:			return 0;
 	}
 }
@@ -500,6 +502,11 @@ bool ComputeOpacityValue(const uint8* BGRA, int32 x, int32 y, int32 TextureSizeX
 		const uint32 Opacity = BGRA[Offset];
 		return Opacity > AlphaThreshold8;
 	}
+	else if (SrcFormat == TSF_R32F)
+	{
+		const float Opacity = *((const float*)(BGRA + Offset));
+		return Opacity > AlphaThresholdF;
+	}
 	else if (SrcFormat == TSF_BGRA8)
 	{
 		const uint8* BGRA8 = (const uint8*)(BGRA + Offset);
@@ -538,15 +545,15 @@ bool ComputeOpacityValue(const uint8* BGRA, int32 x, int32 y, int32 TextureSizeX
 
 		if (OpacitySourceMode == OSM_Alpha)
 		{
-			Opacity = *(RGBA16 + 3) ;
+			Opacity = *(RGBA16 + 3);
 		}
 		else if (OpacitySourceMode == OSM_RedChannel)
 		{
-			Opacity = *(RGBA16 + 0) ;
+			Opacity = *(RGBA16 + 0);
 		}
 		else if (OpacitySourceMode == OSM_GreenChannel)
 		{
-			Opacity = *(RGBA16 + 1) ;
+			Opacity = *(RGBA16 + 1);
 		}
 		else if (OpacitySourceMode == OSM_BlueChannel)
 		{
@@ -558,7 +565,7 @@ bool ComputeOpacityValue(const uint8* BGRA, int32 x, int32 y, int32 TextureSizeX
 			uint32 G = *(RGBA16 + 1);
 			uint32 B = *(RGBA16 + 2);
 
-			Opacity = ((R + G + B) / 3) ;
+			Opacity = ((R + G + B) / 3);
 		}
 
 		return Opacity > AlphaThreshold16;
@@ -595,7 +602,34 @@ bool ComputeOpacityValue(const uint8* BGRA, int32 x, int32 y, int32 TextureSizeX
 
 		return Opacity > AlphaThresholdF;
 	}
+	else if (SrcFormat == TSF_RGBA32F)
+	{
+		const float* RGBAf = (const float*)( BGRA + Offset);
+
+		float Opacity = 1.0f;
+		if (OpacitySourceMode == OSM_Alpha)
+		{
+			Opacity = RGBAf[3];
+		}
+		else if (OpacitySourceMode == OSM_RedChannel)
+		{
+			Opacity = RGBAf[0];
+		}
+		else if (OpacitySourceMode == OSM_GreenChannel)
+		{
+			Opacity = RGBAf[1];
+		}
+		else if (OpacitySourceMode == OSM_BlueChannel)
+		{
+			Opacity = RGBAf[2];
+		}
+		else
+		{
+			Opacity = (RGBAf[0] + RGBAf[1] + RGBAf[2]) / 3.0f;
+		}
 
+		return Opacity > AlphaThresholdF;
+	}
 	return true;
 }
 
@@ -625,7 +659,7 @@ void FSubUVDerivedData::GetFeedback(UTexture2D* SubUVTexture, int32 SubImages_Ho
 			}
 
 			OutErrors.Add(
-				FText::Format(LOCTEXT("CutoutNotSupported", "Image source format ({0}) is not in a supported format for cutouts ({1}). It will be ignored."),
+				FText::Format(LOCTEXT("CutoutNotSupported", "Image source format ({0}) is not in supported format for cutouts ({1}). It will be ignored."),
 					FText::FromString(SourceFormatEnum->GetNameStringByValue(SourceFormat)),
 					FText::FromString(SupportedFormatsString)
 				)
