[AutoRTFM] Update ISoundGenerator::SynthCommand for Verse transactional safety.

We can't enqueue work on the command queue until the Verse transaction commits safely.

Fixing this required changing `TUniqueFunction<void()>` to `TFunction<void()>` because our task code is currently not compatible with `TUniqueFunction`.

diff --git a/Engine/Plugins/Runtime/Metasound/Source/MetasoundGenerator/Public/MetasoundGenerator.h b/Engine/Plugins/Runtime/Metasound/Source/MetasoundGenerator/Public/MetasoundGenerator.h
--- a/Engine/Plugins/Runtime/Metasound/Source/MetasoundGenerator/Public/MetasoundGenerator.h
+++ b/Engine/Plugins/Runtime/Metasound/Source/MetasoundGenerator/Public/MetasoundGenerator.h
@@ -282,7 +282,7 @@ namespace Metasound
 		// Enqueues a command for this generator to execute when its next buffer is
 		// requested by the mixer.  Enqueued commands are executed before OnGenerateAudio,
 		// and on the same thread.  They can safely access generator state.
-		void OnNextBuffer(TUniqueFunction<void(FMetasoundGenerator&)> Command)
+		void OnNextBuffer(TFunction<void(FMetasoundGenerator&)> Command)
 		{
 			SynthCommand([this, Command = MoveTemp(Command)]() { Command(*this); });
 		}
diff --git a/Engine/Source/Runtime/Engine/Classes/Sound/SoundGenerator.h b/Engine/Source/Runtime/Engine/Classes/Sound/SoundGenerator.h
--- a/Engine/Source/Runtime/Engine/Classes/Sound/SoundGenerator.h
+++ b/Engine/Source/Runtime/Engine/Classes/Sound/SoundGenerator.h
@@ -64,7 +64,7 @@ protected:
 
 	// Protected method to execute lambda in audio render thread
 	// Used for conveying parameter changes or events to the generator thread.
-	ENGINE_API void SynthCommand(TUniqueFunction<void()> Command);
+	ENGINE_API void SynthCommand(TFunction<void()> Command);
 
 private:
 
diff --git a/Engine/Source/Runtime/Engine/Private/SoundGenerator.cpp b/Engine/Source/Runtime/Engine/Private/SoundGenerator.cpp
--- a/Engine/Source/Runtime/Engine/Private/SoundGenerator.cpp
+++ b/Engine/Source/Runtime/Engine/Private/SoundGenerator.cpp
@@ -1,6 +1,7 @@
 // Copyright Epic Games, Inc. All Rights Reserved.
 
 #include "Sound/SoundGenerator.h"
+#include "AutoRTFM/AutoRTFM.h"
 #include "HAL/LowLevelMemTracker.h"
 
 ISoundGenerator::ISoundGenerator()
@@ -31,11 +32,14 @@ int32 ISoundGenerator::GetNextBuffer(float* OutAudio, int32 NumSamples, bool bRe
 	return OnGenerateAudio(OutAudio, NumSamplesToGenerate);
 }
 
-void ISoundGenerator::SynthCommand(TUniqueFunction<void()> Command)
+void ISoundGenerator::SynthCommand(TFunction<void()> Command)
 {
 	LLM_SCOPE(ELLMTag::AudioSynthesis);
 
-	CommandQueue.Enqueue(MoveTemp(Command));
+	UE_AUTORTFM_ONCOMMIT(this, Command = MoveTemp(Command))
+	{
+		CommandQueue.Enqueue(MoveTemp(Command));
+	};
 }
 
 void ISoundGenerator::PumpPendingMessages()
