Fix assumption that compact pose always has the same set of bone indices as the ref pose override

diff --git a/Engine/Source/Runtime/Engine/Public/BoneContainer.h b/Engine/Source/Runtime/Engine/Public/BoneContainer.h
--- a/Engine/Source/Runtime/Engine/Public/BoneContainer.h
+++ b/Engine/Source/Runtime/Engine/Public/BoneContainer.h
@@ -374,31 +374,19 @@ public:
 	{
 		const int32 CompactPoseBoneCount = GetCompactPoseNumBones();
 		OutTransforms.Reset(CompactPoseBoneCount);
-		if (RefPoseOverride.IsValid())
-		{
-			OutTransforms.Append(RefPoseOverride->RefBonePoses);
-		}
-		else
+		OutTransforms.SetNumUninitialized(CompactPoseBoneCount);
+		
+		const TArray<FTransform>& RefPoseTransforms = RefPoseOverride.IsValid() ? RefPoseOverride->RefBonePoses : RefSkeleton->GetRefBonePose();
+		for (int32 CompactBoneIndex = 0; CompactBoneIndex < CompactPoseBoneCount; ++CompactBoneIndex)
 		{
-			OutTransforms.SetNumUninitialized(CompactPoseBoneCount);
-			const TArray<FTransform>& RefPoseTransforms = RefSkeleton->GetRefBonePose();
-			for (int32 CompactBoneIndex = 0; CompactBoneIndex < CompactPoseBoneCount; ++CompactBoneIndex)
-			{
-				OutTransforms[CompactBoneIndex] = RefPoseTransforms[BoneIndicesArray[CompactBoneIndex]];
-			}
+			OutTransforms[CompactBoneIndex] = RefPoseTransforms[BoneIndicesArray[CompactBoneIndex]];
 		}
 	}
 	
 	const FTransform& GetRefPoseTransform(const FCompactPoseBoneIndex& BoneIndex) const
 	{
-		if (RefPoseOverride.IsValid())
-		{
-			return RefPoseOverride->RefBonePoses[BoneIndex.GetInt()];
-		}
-		else
-		{
-			return RefSkeleton->GetRefBonePose()[BoneIndicesArray[BoneIndex.GetInt()]];
-		}
+		const TArray<FTransform>& RefPoseTransforms = RefPoseOverride.IsValid() ? RefPoseOverride->RefBonePoses : RefSkeleton->GetRefBonePose();
+		return RefPoseTransforms[BoneIndicesArray[BoneIndex.GetInt()]];
 	}
 
 	/** Override skeleton ref pose. */
