[AssetManager] Fix calling completion delegate even for empty asset list as specified by the API.

#rnx

diff --git a/Engine/Source/Runtime/Engine/Private/AssetManager.cpp b/Engine/Source/Runtime/Engine/Private/AssetManager.cpp
--- a/Engine/Source/Runtime/Engine/Private/AssetManager.cpp
+++ b/Engine/Source/Runtime/Engine/Private/AssetManager.cpp
@@ -1809,12 +1809,6 @@ TSharedPtr<FStreamableHandle> UAssetManager::ChangeBundleStateForPrimaryAssets(
 	FAssetManagerLoadParams&& LoadParams,
 	UE::FSourceLocation Location)
 {
-	if (AssetsToChange.IsEmpty())
-	{
-		ensureAlwaysMsgf(!AddBundles.IsEmpty() || !RemoveBundles.IsEmpty() || bRemoveAllBundles,
-			TEXT("Attempted to change state on empty primary assets set."));
-		return nullptr;
-	}
 	TArray<TSharedPtr<FStreamableHandle> > AllHandles;
 	int32 MoviePlayerNumAssets = 0;
 	for (const FPrimaryAssetId& PrimaryAssetId : AssetsToChange)
diff --git a/Engine/Source/Runtime/Engine/Private/StreamableManager.cpp b/Engine/Source/Runtime/Engine/Private/StreamableManager.cpp
--- a/Engine/Source/Runtime/Engine/Private/StreamableManager.cpp
+++ b/Engine/Source/Runtime/Engine/Private/StreamableManager.cpp
@@ -2281,11 +2281,7 @@ TSharedPtr<FStreamableHandle> FStreamableManager::CreateCombinedHandle(
 {
 	if (ChildHandles.IsEmpty())
 	{
-		// Complete immediately if nothing to load
-		TSharedPtr<FStreamableHandle> ErrorHandle = CreateErrorHandle(
-			UE::UnifiedError::StreamableManager::PackageLoadCancelled::MakeError(),
-			FStreamableAsyncLoadParams{});
-		Params.OnComplete.ExecuteIfBound(ErrorHandle);
+		FStreamableHandle::ExecuteDelegate(MoveTemp(Params.OnComplete));
 		return nullptr;
 	}
 
