SceneCapture: Fix issue arising from CL 39621162 where default scene capture post-process settings should remain blueprint-derived.

This otherwise caused issues with loading old Composure levels, where the level instances would be incorrectly reset, and miss default properties such as the "ComposureInvertAlpha" PPM.

diff --git a/Engine/Source/Runtime/Engine/Private/Components/SceneCaptureComponent.cpp b/Engine/Source/Runtime/Engine/Private/Components/SceneCaptureComponent.cpp
--- a/Engine/Source/Runtime/Engine/Private/Components/SceneCaptureComponent.cpp
+++ b/Engine/Source/Runtime/Engine/Private/Components/SceneCaptureComponent.cpp
@@ -962,7 +962,8 @@ void USceneCaptureComponent2D::Serialize(FArchive& Ar)
 		bAutoCalculateOrthoPlanes = false;
 	}
 
-	if (Ar.CustomVer(FFortniteMainBranchObjectVersion::GUID) < FFortniteMainBranchObjectVersion::SceneCaptureDefaultSettings)
+	if (Ar.CustomVer(FFortniteMainBranchObjectVersion::GUID) < FFortniteMainBranchObjectVersion::SceneCaptureDefaultSettings
+		&& !GetArchetype()->IsInBlueprint())
 	{
 		// For backwards compatibility, we must preserve the old defaults that do not reflect code overrides in SetupViewFamilyForSceneCapture.
 		PostProcessSettings = FPostProcessSettings::GetDefault();
@@ -1469,7 +1470,8 @@ void USceneCaptureComponentCube::Serialize(FArchive& Ar)
 {
 	Ar.UsingCustomVersion(FFortniteMainBranchObjectVersion::GUID);
 
-	if (Ar.CustomVer(FFortniteMainBranchObjectVersion::GUID) < FFortniteMainBranchObjectVersion::SceneCaptureDefaultSettings)
+	if (Ar.CustomVer(FFortniteMainBranchObjectVersion::GUID) < FFortniteMainBranchObjectVersion::SceneCaptureDefaultSettings
+		&& !GetArchetype()->IsInBlueprint())
 	{
 		// For backwards compatibility, we must preserve the old defaults that do not reflect code overrides in SetupViewFamilyForSceneCapture.
 		PostProcessSettings = FPostProcessSettings::GetDefault();
