Fixed scene component cached bounding box not serialized for migration

diff --git a/Engine/Source/Runtime/Engine/Private/Components/SceneComponent.cpp b/Engine/Source/Runtime/Engine/Private/Components/SceneComponent.cpp
--- a/Engine/Source/Runtime/Engine/Private/Components/SceneComponent.cpp
+++ b/Engine/Source/Runtime/Engine/Private/Components/SceneComponent.cpp
@@ -3731,6 +3731,7 @@ void USceneComponent::Serialize(FArchive& Ar)
 	if (Ar.IsMigratingRemoteObjects())
 	{
 		Ar << ComponentToWorld;
+		Ar << Bounds;
 
 		AttachChildren.Append(RestoreChildren);
 		if (Ar.IsLoading())
