Expand on access to external code declarations from within plugins for material shaders.
-MaterialExpressions.ini is added as a known ini file so that plugins can now add entries
-marks MaterialExpressionExternalCodeBase functions as public so that we can extend the class in a plugin
#tests Create a new material expression in a plugin and verify that it works

diff --git a/Engine/Source/Runtime/Core/Private/Misc/ConfigCacheIni.cpp b/Engine/Source/Runtime/Core/Private/Misc/ConfigCacheIni.cpp
--- a/Engine/Source/Runtime/Core/Private/Misc/ConfigCacheIni.cpp
+++ b/Engine/Source/Runtime/Core/Private/Misc/ConfigCacheIni.cpp
@@ -5697,6 +5697,9 @@ static void LoadRemainingConfigFiles(FConfigContext& Context)
 
 	Context.Load(TEXT("Editor"), GEditorIni);
 
+	// early loading allows for plugins to append entries in material expressions
+	Context.Load(TEXT("MaterialExpressions"));
+
 	// Upgrade editor user settings before loading the editor per project user settings
 	FConfigManifest::MigrateEditorUserSettings();
 	Context.Load(TEXT("EditorPerProjectUserSettings"), GEditorPerProjectIni);
diff --git a/Engine/Source/Runtime/Engine/Public/Materials/MaterialExpressionExternalCodeBase.h b/Engine/Source/Runtime/Engine/Public/Materials/MaterialExpressionExternalCodeBase.h
--- a/Engine/Source/Runtime/Engine/Public/Materials/MaterialExpressionExternalCodeBase.h
+++ b/Engine/Source/Runtime/Engine/Public/Materials/MaterialExpressionExternalCodeBase.h
@@ -21,8 +21,8 @@ public:
 
 	//~ Begin UMaterialExpression Interface
 #if WITH_EDITOR
-	virtual void Build(MIR::FEmitter& Emitter) override;
-	virtual int32 Compile(class FMaterialCompiler* Compiler, int32 OutputIndex) override;
+	ENGINE_API virtual void Build(MIR::FEmitter& Emitter) override;
+	ENGINE_API virtual int32 Compile(class FMaterialCompiler* Compiler, int32 OutputIndex) override;
 #endif
 	//~ End UMaterialExpression Interface
 
