[Engine] Fixed normal direction for PDI DrawTorus

diff --git a/Engine/Source/Runtime/Engine/Private/PrimitiveDrawingUtils.cpp b/Engine/Source/Runtime/Engine/Private/PrimitiveDrawingUtils.cpp
--- a/Engine/Source/Runtime/Engine/Private/PrimitiveDrawingUtils.cpp
+++ b/Engine/Source/Runtime/Engine/Private/PrimitiveDrawingUtils.cpp
@@ -828,7 +828,11 @@ void DrawTorus(FPrimitiveDrawInterface* PDI, const FMatrix& Transform, const FVe
 			MeshVertex.TextureCoordinate[0] = TC;
 			MeshVertex.TextureCoordinate[0].X += TCStepX * InnerIndex;
 
-			FVector TangentX = -UnitDir * FMath::Sin(InnerAngle) + ZAxis * FMath::Cos(InnerAngle);
+			// We invert the direction before deriving tangents, to counteract winding order
+			// @see DrawSphere, etc.
+			Dir *= -1.0;
+
+			FVector TangentX = UnitDir * FMath::Sin(InnerAngle) + ZAxis * FMath::Cos(InnerAngle);
 			FVector TangentY = Dir ^ TangentX;
 			FVector TangentZ = Dir;
 			MeshVertex.SetTangents((FVector3f)TangentX, (FVector3f)TangentY, (FVector3f)TangentZ);
