[Audio Insights] VirtualLoopVirtualize message: use Audio channel instead of AudioMixerChannel (all virtualloop messages use Audio channel)

diff --git a/Engine/Source/Runtime/Engine/Private/AudioVirtualLoop.cpp b/Engine/Source/Runtime/Engine/Private/AudioVirtualLoop.cpp
--- a/Engine/Source/Runtime/Engine/Private/AudioVirtualLoop.cpp
+++ b/Engine/Source/Runtime/Engine/Private/AudioVirtualLoop.cpp
@@ -340,8 +340,8 @@ bool FAudioVirtualLoop::Update(float DeltaTime, bool bForceUpdate)
 #if UE_AUDIO_PROFILERTRACE_ENABLED
 void FAudioVirtualLoop::OnTraceStarted()
 {
-	// Make sure we send mixer source data to Audio Insights if it's opened in the middle of a PIE session
-	const bool bChannelEnabled = UE_TRACE_CHANNELEXPR_IS_ENABLED(AudioMixerChannel);
+	// Make sure we send virtual loop data to Audio Insights if it's opened in the middle of a PIE session
+	const bool bChannelEnabled = UE_TRACE_CHANNELEXPR_IS_ENABLED(AudioChannel);
 	if (bChannelEnabled && ActiveSound)
 	{
 		check(ActiveSound);
@@ -350,7 +350,7 @@ void FAudioVirtualLoop::OnTraceStarted()
 
 		if (const FAudioDevice* AudioDevice = ActiveSound->AudioDevice)
 		{
-			UE_TRACE_LOG(Audio, VirtualLoopVirtualize, AudioMixerChannel)
+			UE_TRACE_LOG(Audio, VirtualLoopVirtualize, AudioChannel)
 				<< VirtualLoopVirtualize.DeviceId(static_cast<uint32>(AudioDevice->DeviceID))
 				<< VirtualLoopVirtualize.Timestamp(FPlatformTime::Cycles64())
 				<< VirtualLoopVirtualize.PlayOrder(ActiveSound->GetPlayOrder())
