[UConsole] Fix console gap.

SBHead is the index of the next slot to be filled in the scrollback buffer, so this needs to be accounted for when figuring out which element of the scrollback buffer should appear first. Prior to this fix, the element one past the end of the scrollback buffer would try to be drawn, creating an empty line. If SBPos was greater than 0 the gap would be filled.

diff --git a/Engine/Source/Runtime/Engine/Classes/Engine/Console.h b/Engine/Source/Runtime/Engine/Classes/Engine/Console.h
--- a/Engine/Source/Runtime/Engine/Classes/Engine/Console.h
+++ b/Engine/Source/Runtime/Engine/Classes/Engine/Console.h
@@ -82,9 +82,10 @@ class UConsole
 	/** Holds the scrollback buffer */
 	TArray<FString> Scrollback;
 
-	/**  Where in the scrollback buffer are we */
+	/** Index to the next slot to be filled in the scrollback buffer. NOTE: The most recent line is at (SBHead-1). */
 	int32 SBHead;
 
+	/** Scrollback position offset (0 = view newest entries, increases to view older entries) */
 	int32 SBPos;
 
 	/** Max number of command history entries */
diff --git a/Engine/Source/Runtime/Engine/Private/UserInterface/Console.cpp b/Engine/Source/Runtime/Engine/Private/UserInterface/Console.cpp
--- a/Engine/Source/Runtime/Engine/Private/UserInterface/Console.cpp
+++ b/Engine/Source/Runtime/Engine/Private/UserInterface/Console.cpp
@@ -644,8 +644,8 @@ void UConsole::OutputTextLine(const FString& Text)
 	// If we are full, delete the first line
 	if (Scrollback.Num() >= ConsoleSettings->MaxScrollbackSize)
 	{
-		Scrollback.RemoveAt(0, 1);
-		SBHead = ConsoleSettings->MaxScrollbackSize - 1;
+		Scrollback.RemoveAt(0, Scrollback.Num() - ConsoleSettings->MaxScrollbackSize + 1);
+		SBHead = ConsoleSettings->MaxScrollbackSize;
 	}
 	else
 	{
@@ -682,7 +682,7 @@ void UConsole::OutputText(const FString& Text)
 		// Remove the line from the string.
 		RemainingText = (RemainingText).Mid(LineLength + 1, MAX_int32);
 		StringLength -= LineLength + 1;
-	};
+	}
 }
 
 
@@ -1346,7 +1346,7 @@ void UConsole::PostRender_Console_Open(UCanvas* Canvas)
 		ConsoleTileWhite.BlendMode = SE_BLEND_AlphaBlend;		// Preserve alpha to allow single-pass composite
 
 		// figure out which element of the scrollback buffer should appear first (at the bottom of the scrollback area)
-		int32 idx = SBHead - SBPos;
+		int32 idx = SBHead - SBPos - 1;
 
 		TArray<FString> SelectedLines;
 
