Adding remote object support to 'obj list' command. Added '-ServerId=x' filter to 'obj list' command
Added FRemoteServerId constructor that takes a string and parses it
Did a cleanup pass om FRemoteServerId uses cases and replaced FRemoteServerId(0) constructors with the default constructor. It does the same thing without hardcoding the invalid id value in random places.

diff --git a/Engine/Source/Runtime/CoreUObject/Private/UObject/RemoteObject.cpp b/Engine/Source/Runtime/CoreUObject/Private/UObject/RemoteObject.cpp
--- a/Engine/Source/Runtime/CoreUObject/Private/UObject/RemoteObject.cpp
+++ b/Engine/Source/Runtime/CoreUObject/Private/UObject/RemoteObject.cpp
@@ -18,6 +18,7 @@
 #include "UObject/UObjectThreadContext.h"
 #include "Misc/CommandLine.h"
 #include "Misc/CoreDelegates.h"
+#include "String/LexFromString.h"
 #include <atomic>
 
 DEFINE_LOG_CATEGORY(LogRemoteObject);
@@ -121,20 +122,28 @@ public:
 };
 FRemoteObjectMaps* ObjectMaps = nullptr;
 
-void InitRemoteObjects()
+void InitServerId()
 {
-	ObjectMaps = new FRemoteObjectMaps();
-
-	uint32 ServerIdNumber = 0;
-	if (!FParse::Value(FCommandLine::Get(), TEXT("MultiServerLocalId="), ServerIdNumber))
+	FString ServerId;
+	if (!FParse::Value(FCommandLine::Get(), TEXT("MultiServerLocalId="), ServerId))
 	{
-		if (!FParse::Value(FCommandLine::Get(), TEXT("LocalPeerId="), ServerIdNumber))
+		if (!FParse::Value(FCommandLine::Get(), TEXT("LocalPeerId="), ServerId))
 		{
-			ServerIdNumber = 1;
+			RemoteServerId == FRemoteServerId(ERemoteServerIdConstants::FirstValid);
 		}
 	}
-	checkf(ServerIdNumber != (uint32)ERemoteServerIdConstants::Invalid, TEXT("Remote Server Id=%u is not valid"), ServerIdNumber);
-	RemoteServerId = FRemoteServerId(ServerIdNumber);
+	if (!ServerId.IsEmpty())
+	{
+		RemoteServerId = FRemoteServerId(ServerId);
+	}
+	checkf(RemoteServerId.IsValid(), TEXT("Remote ServerId is not valid"));
+}
+
+void InitRemoteObjects()
+{
+	ObjectMaps = new FRemoteObjectMaps();
+
+	InitServerId();
 
 	UE::RemoteObject::Transfer::InitRemoteObjectTransfer();
 
@@ -650,6 +659,32 @@ FArchive& operator<<(FArchive& Ar, FRemoteObjectId& Id)
 	return Ar;
 }
 
+FRemoteServerId::FRemoteServerId(const FString& InText)
+{
+	if (InText == TEXT("Asset"))
+	{
+		Id = (uint32)ERemoteServerIdConstants::Asset;
+	}
+	else if (InText == TEXT("Database"))
+	{
+		Id = (uint32)ERemoteServerIdConstants::Database;
+	}
+	else
+	{
+		uint32 ServerIdNumber = (uint32)ERemoteServerIdConstants::Invalid;
+		LexFromString(ServerIdNumber, InText);
+		if (ensureMsgf(ServerIdNumber <= (uint32)ERemoteServerIdConstants::Max, TEXT("Parsed Remote Server Id value %u that is bigger than allowed max %u"), ServerIdNumber, (uint32)ERemoteServerIdConstants::Max))
+		{
+			Id = ServerIdNumber;
+		}
+		else
+		{
+			UE_LOG(LogRemoteObject, Warning, TEXT("Clamping ServerId number %u to the maximum allowed %u"), ServerIdNumber, (uint32)ERemoteServerIdConstants::Max);
+			Id = (uint32)ERemoteServerIdConstants::Max;
+		}		
+	}
+}
+
 FArchive& operator<<(FArchive& Ar, FRemoteServerId& Id)
 {
 	Ar << Id.Id;
diff --git a/Engine/Source/Runtime/CoreUObject/Public/UObject/RemoteObject.h b/Engine/Source/Runtime/CoreUObject/Public/UObject/RemoteObject.h
--- a/Engine/Source/Runtime/CoreUObject/Public/UObject/RemoteObject.h
+++ b/Engine/Source/Runtime/CoreUObject/Public/UObject/RemoteObject.h
@@ -46,8 +46,8 @@ namespace UE::RemoteObject::Handle
 		FRemoteServerId OwningServerId;
 
 		/** Server id of the which server the object should be on for physics
-			Physics ID = 0 (default) means physics ID should not be considered when migrating the object.   */
-		FRemoteServerId PhysicsOverrideServerId = FRemoteServerId(0);
+			Default (invalid) means physics ID should not be considered when migrating the object.   */
+		FRemoteServerId PhysicsOverrideServerId;
 
 		FRemoteObjectStub() = default;
 		COREUOBJECT_API explicit FRemoteObjectStub(UObject* Object);
diff --git a/Engine/Source/Runtime/CoreUObject/Public/UObject/RemoteObjectTypes.h b/Engine/Source/Runtime/CoreUObject/Public/UObject/RemoteObjectTypes.h
--- a/Engine/Source/Runtime/CoreUObject/Public/UObject/RemoteObjectTypes.h
+++ b/Engine/Source/Runtime/CoreUObject/Public/UObject/RemoteObjectTypes.h
@@ -15,6 +15,7 @@ class FArchive;
 enum class ERemoteServerIdConstants : uint32
 {
 	Invalid = 0,
+	FirstValid,
 	Max = (1 << 10) - 1,
 
 	Database = Max,
@@ -40,6 +41,8 @@ struct FRemoteServerId
 	{
 		checkf(Id < (uint32)ERemoteServerIdConstants::FirstReserved, TEXT("Remote server id can not be greater than %u, got: %u"), (uint32)ERemoteServerIdConstants::FirstReserved - 1, Id);
 	}
+	explicit COREUOBJECT_API FRemoteServerId(const FString& InText);
+
 	uint32 GetIdNumber() const
 	{
 		return Id;
diff --git a/Engine/Source/Runtime/Engine/Private/UnrealEngine.cpp b/Engine/Source/Runtime/Engine/Private/UnrealEngine.cpp
--- a/Engine/Source/Runtime/Engine/Private/UnrealEngine.cpp
+++ b/Engine/Source/Runtime/Engine/Private/UnrealEngine.cpp
@@ -9325,6 +9325,17 @@ bool UEngine::HandleObjCommand( const TCHAR* Cmd, FOutputDevice& Ar )
 		ParseObject<UClass>(Cmd, TEXT("CLASS="  ), CheckType, nullptr);
 		ParseObject<UObject>(Cmd, TEXT("OUTER="), CheckOuter, nullptr);
 
+#if UE_WITH_REMOTE_OBJECT_HANDLE
+		FRemoteServerId ServerId;
+		{
+			FString ServerIdText;
+			if (FParse::Value(Cmd, TEXT("SERVERID="), ServerIdText))
+			{
+				ServerId = FRemoteServerId(ServerIdText);
+			}
+		}
+#endif
+
 		ParseObject<UPackage>(Cmd, TEXT("PACKAGE="), InsidePackage, nullptr);
 		if (InsidePackage == nullptr)
 		{
@@ -9356,7 +9367,8 @@ bool UEngine::HandleObjCommand( const TCHAR* Cmd, FOutputDevice& Ar )
 		}
 
 		const bool bAll				= FParse::Param( Cmd, TEXT("ALL") );
-		const bool bShowMemory = !FParse::Param(Cmd, TEXT("SKIPMEMORYSIZE"));
+		// When listing objects with remote object handles enabled skip memory size calculation as the memory counting archive may accidentally attempt to resolve remote objects
+		const bool bShowMemory = !FParse::Param(Cmd, TEXT("SKIPMEMORYSIZE")) && !UE_WITH_REMOTE_OBJECT_HANDLE;
 
 		// if we specified a parameter in the command, but no objects of that parameter were found,
 		// and they didn't specify "all", then don't list all objects
@@ -9368,7 +9380,11 @@ bool UEngine::HandleObjCommand( const TCHAR* Cmd, FOutputDevice& Ar )
 				&&	(InsideObject	||	!FCString::Strifind(Cmd,TEXT("INSIDE=")))
 				&&	(InsideClass || !FCString::Strifind(Cmd, TEXT("INSIDECLASS=")))))
 		{
-			const bool bTrackDetailedObjectInfo = bAll || (CheckType != NULL && CheckType != UObject::StaticClass()) || CheckOuter != nullptr || InsideObject != nullptr || InsidePackage != nullptr || InsideClass != nullptr || !ObjectName.IsEmpty();
+			const bool bTrackDetailedObjectInfo = bAll || (CheckType != NULL && CheckType != UObject::StaticClass()) || CheckOuter != nullptr || InsideObject != nullptr || InsidePackage != nullptr || InsideClass != nullptr || !ObjectName.IsEmpty() 
+#if UE_WITH_REMOTE_OBJECT_HANDLE
+				|| ServerId.IsValid()
+#endif
+				;
 			const bool bOnlyListGCObjects = FParse::Param(Cmd, TEXT("GCONLY"));
 			const bool bOnlyListGCObjectsNoClusters = FParse::Param(Cmd, TEXT("GCNOCLUSTERS"));
 			const bool bOnlyListRootObjects = FParse::Param(Cmd, TEXT("ROOTONLY"));
@@ -9477,6 +9493,17 @@ bool UEngine::HandleObjCommand( const TCHAR* Cmd, FOutputDevice& Ar )
 					}
 				}
 
+#if UE_WITH_REMOTE_OBJECT_HANDLE
+				if (ServerId.IsValid())
+				{
+					FRemoteObjectId ObjectId(*It);
+					if (ObjectId.GetServerId() != ServerId)
+					{
+						continue;
+					}
+				}
+#endif
+
 				SIZE_T NumBytes = 0;
 				SIZE_T MaxBytes = 0;
 				FResourceSizeEx TrueResourceSize = FResourceSizeEx(EResourceSizeMode::Exclusive);
@@ -9559,8 +9586,15 @@ bool UEngine::HandleObjCommand( const TCHAR* Cmd, FOutputDevice& Ar )
 			else
 			{
 				Ar.Logf(
-					TEXT("%140s %10s %10s %10s %15s %15s %15s"),
+					TEXT("%140s")
+#if UE_WITH_REMOTE_OBJECT_HANDLE
+					TEXT("%23s")
+#endif // UE_WITH_REMOTE_OBJECT_HANDLE
+					TEXT("%10s %10s %10s %15s %15s %15s"),
 					TEXT("Object"),
+#if UE_WITH_REMOTE_OBJECT_HANDLE
+					TEXT("RemoteId"),
+#endif // UE_WITH_REMOTE_OBJECT_HANDLE
 					TEXT("NumKB"),
 					TEXT("MaxKB"),
 					TEXT("ResExcKB"),
@@ -9589,8 +9623,15 @@ bool UEngine::HandleObjCommand( const TCHAR* Cmd, FOutputDevice& Ar )
 				else
 				{
 					Ar.Logf(
-						TEXT("%140s %10.2f %10.2f %10.2f %15.2f %15.2f %15.2f"),
+						TEXT("%140s")
+#if UE_WITH_REMOTE_OBJECT_HANDLE
+						TEXT("%23s")
+#endif // UE_WITH_REMOTE_OBJECT_HANDLE
+						TEXT("%10.2f %10.2f %10.2f %15.2f %15.2f %15.2f"),
 						*ObjItem.Object->GetFullName(),
+#if UE_WITH_REMOTE_OBJECT_HANDLE
+						*FRemoteObjectId(ObjItem.Object).ToString(),
+#endif // UE_WITH_REMOTE_OBJECT_HANDLE
 						ObjItem.Num / 1024.0f,
 						ObjItem.Max / 1024.0f,
 						ObjItem.TrueResourceSize.GetTotalMemoryBytes() / 1024.0f,
