Fixing compilation errors when building the editor with remote object handles enabled

diff --git a/Engine/Source/Runtime/Engine/Private/Components/SkeletalMeshComponent.cpp b/Engine/Source/Runtime/Engine/Private/Components/SkeletalMeshComponent.cpp
--- a/Engine/Source/Runtime/Engine/Private/Components/SkeletalMeshComponent.cpp
+++ b/Engine/Source/Runtime/Engine/Private/Components/SkeletalMeshComponent.cpp
@@ -3960,14 +3960,14 @@ void USkeletalMeshComponent::ApplyMigratedPhysicsState(Chaos::FSerializedDataBuf
 		// If we received an Internal physics state, we want to apply it to both, Internal and External so it is available immediatly in the game thread API
 		Chaos::ESerializedDataContext SerializationDataContext = MigratedStateHeader.DataContext == Chaos::ESerializedDataContext::Internal ? Chaos::ESerializedDataContext::Both : MigratedStateHeader.DataContext;
 
-		for (FBodyInstance* BodyInstance : Bodies)
+		for (FBodyInstance* BodyInstanceIt : Bodies)
 		{
 			Chaos::FSerializedDataBufferPtr DataBuffer = MakeUnique<Chaos::FSerializedDataBuffer>();
 			MigratedState.ReadElementDataIntoBuffer(*DataBuffer);
 
-			if (ensure(BodyInstance))
+			if (ensure(BodyInstanceIt))
 			{
-				BodyInstance->ApplySerializedState(MoveTemp(DataBuffer), SerializationDataContext);
+				BodyInstanceIt->ApplySerializedState(MoveTemp(DataBuffer), SerializationDataContext);
 			}
 		}
 
diff --git a/Engine/Source/Runtime/Experimental/Chaos/Public/Chaos/Serialization/SerializedMultiPhysicsState.h b/Engine/Source/Runtime/Experimental/Chaos/Public/Chaos/Serialization/SerializedMultiPhysicsState.h
--- a/Engine/Source/Runtime/Experimental/Chaos/Public/Chaos/Serialization/SerializedMultiPhysicsState.h
+++ b/Engine/Source/Runtime/Experimental/Chaos/Public/Chaos/Serialization/SerializedMultiPhysicsState.h
@@ -18,7 +18,7 @@ namespace Chaos
 			Header.DataContext = DataContext;
 		}
 
-		void ReadElementDataIntoBuffer(FSerializedDataBuffer& TargetBuffer);
+		CHAOS_API void ReadElementDataIntoBuffer(FSerializedDataBuffer& TargetBuffer);
 
 		template<typename ElementType>
 		void AppendDataFromElement(ElementType& InElement)
@@ -33,7 +33,7 @@ namespace Chaos
 			Header.DataTagPerElementIndex.Emplace(DataTag);
 		}
 
-		void Serialize(FArchive& Ar);
+		CHAOS_API void Serialize(FArchive& Ar);
 
 		int32 GetSize() const
 		{
