SceneGraph
* NetEntityFactory now properly supports dynamic entities built from prefabs
* Added LogNetEntity log category to isolate replication related logs for entities

#rnx

diff --git a/Engine/Source/Runtime/Engine/Private/Net/Iris/ReplicationSystem/NetSubObjectFactory.cpp b/Engine/Source/Runtime/Engine/Private/Net/Iris/ReplicationSystem/NetSubObjectFactory.cpp
--- a/Engine/Source/Runtime/Engine/Private/Net/Iris/ReplicationSystem/NetSubObjectFactory.cpp
+++ b/Engine/Source/Runtime/Engine/Private/Net/Iris/ReplicationSystem/NetSubObjectFactory.cpp
@@ -39,7 +39,7 @@ TUniquePtr<UE::Net::FNetObjectCreationHeader> UNetSubObjectFactory::CreateAndFil
 
 	FNetObjectReference ObjectRef = Bridge->GetOrCreateObjectReference(SubObject);
 
-	if (!ObjectRef.GetRefHandle().IsDynamic() || SubObject->IsNameStableForNetworking())
+	if (ObjectRef.GetRefHandle().IsStatic() || SubObject->IsNameStableForNetworking())
 	{
 		// No more information needed since we don't need to spawn the object on the remote.
 		TUniquePtr<FNetStaticSubObjectCreationHeader> Header(new FNetStaticSubObjectCreationHeader);
@@ -121,8 +121,8 @@ UNetSubObjectFactory::FInstantiateResult UNetSubObjectFactory::InstantiateReplic
 
 		if (!SubObject)
 		{
-			UE_LOG(LogIris, Error, TEXT("UNetSubObjectFactory::InstantiateNetObjectFromHeader %s: Failed to find static or stable name object referenced by SubObject: %s, Owner: %s, RootObject: %s"), *Context.Handle.ToString(), *Bridge->DescribeObjectReference(SubObjectHeader->ObjectReference, Context.ResolveContext), *Bridge->PrintObjectFromNetRefHandle(Context.RootObjectOfSubObject), *GetPathNameSafe(Bridge->GetReplicatedObject(Context.RootObjectOfSubObject)));
-			
+			UE_LOG(LogIris, Error, TEXT("UNetSubObjectFactory::InstantiateNetObjectFromHeader %s: Failed to find static or stable name object referenced by SubObject: %s, Owner: %s, RootObject: %s"), 
+			*Context.Handle.ToString(), *Bridge->DescribeObjectReference(SubObjectHeader->ObjectReference, Context.ResolveContext), *Bridge->PrintObjectFromNetRefHandle(Context.RootObjectOfSubObject), *GetPathNameSafe(Bridge->GetReplicatedObject(Context.RootObjectOfSubObject)));
 			return FInstantiateResult();
 		}
 
@@ -248,7 +248,7 @@ namespace UE::Net
 //------------------------------------------------------------------------
 FString FNetStaticSubObjectCreationHeader::ToString() const
 {
-	return FString::Printf(TEXT("FNetStaticSubObjectCreationHeader (ProtocolId:0x%x):\n\t"
+	return FString::Printf(TEXT("\n\tFNetStaticSubObjectCreationHeader (ProtocolId:0x%x):\n\t"
 								"ObjectReference=%s\n\t"),								
 								GetProtocolId(),
 								*ObjectReference.ToString());
@@ -273,7 +273,7 @@ bool FNetStaticSubObjectCreationHeader::Deserialize(const FCreationHeaderContext
 
 FString FNetDynamicSubObjectCreationHeader::ToString() const
 {
-	return FString::Printf(TEXT("FNetDynamicSubObjectCreationHeader (ProtocolId:0x%x):\n\t"
+	return FString::Printf(TEXT("\n\tFNetDynamicSubObjectCreationHeader (ProtocolId:0x%x):\n\t"
 								"ObjectClassReference=%s\n\t"
 								"OuterReference=%s\n\t"
 								"bUsePersistenLevel=%u\n\t"
