PR #13027: Rename TiimerHandle.cpp to TimerHandle.cpp

diff --git a/Engine/Source/Runtime/Engine/Private/Engine/TiimerHandle.cpp b/Engine/Source/Runtime/Engine/Private/Engine/TiimerHandle.cpp
deleted file mode 100644
diff --git a/Engine/Source/Runtime/Engine/Private/Engine/TimerHandle.cpp b/Engine/Source/Runtime/Engine/Private/Engine/TimerHandle.cpp
new file mode 100644
--- /dev/null
+++ b/Engine/Source/Runtime/Engine/Private/Engine/TimerHandle.cpp
@@ -0,0 +1,6 @@
+// Copyright Epic Games, Inc. All Rights Reserved.
+
+#include "Engine/TimerHandle.h"
+#include "UObject/ObjectMacros.h"
+
+#include UE_INLINE_GENERATED_CPP_BY_NAME(TimerHandle)
