[Backout] - CL40843190
#rnx
Original CL Desc
-----------------------------------------------------------------
World Partition - HLOD: Disable distance culling on primitive components that are included in HLOD as this can potentially lead to a gap of visibility between the HLOD representation of the component & the component itself.

diff --git a/Engine/Source/Runtime/Engine/Classes/Components/PrimitiveComponent.h b/Engine/Source/Runtime/Engine/Classes/Components/PrimitiveComponent.h
--- a/Engine/Source/Runtime/Engine/Classes/Components/PrimitiveComponent.h
+++ b/Engine/Source/Runtime/Engine/Classes/Components/PrimitiveComponent.h
@@ -2947,11 +2947,6 @@ public:
 	 */
 	ENGINE_API bool ShouldComponentAddToScene() const;
 	
-	/**
-	 * Returns whether the component can ever be distanced culled.
-	 */
-	ENGINE_API bool CanBeDistanceCulled() const;
-
 	/**
 	 * Changes the value of CullDistance.
 	 * @param NewCullDistance - The value to assign to CullDistance.
diff --git a/Engine/Source/Runtime/Engine/Private/Components/PrimitiveComponent.cpp b/Engine/Source/Runtime/Engine/Private/Components/PrimitiveComponent.cpp
--- a/Engine/Source/Runtime/Engine/Private/Components/PrimitiveComponent.cpp
+++ b/Engine/Source/Runtime/Engine/Private/Components/PrimitiveComponent.cpp
@@ -638,7 +638,8 @@ void UPrimitiveComponent::CreateRenderState_Concurrent(FRegisterComponentContext
 	// Make sure cached cull distance is up-to-date if its zero and we have an LD cull distance
 	if( CachedMaxDrawDistance == 0.f && LDMaxDrawDistance > 0.f )
 	{
-		CachedMaxDrawDistance = CanBeDistanceCulled() ? LDMaxDrawDistance : 0;
+		bool bNeverCull = bNeverDistanceCull || GetLODParentPrimitive();
+		CachedMaxDrawDistance = bNeverCull ? 0.f : LDMaxDrawDistance;
 	}
 
 	Super::CreateRenderState_Concurrent(Context);
@@ -1723,7 +1724,8 @@ void UPrimitiveComponent::PostLoad()
 			CachedMaxDrawDistance = FMath::Min( LDMaxDrawDistance, CachedMaxDrawDistance );
 		}
 
-		CachedMaxDrawDistance = CanBeDistanceCulled() ? CachedMaxDrawDistance : 0.0f;
+		bool bNeverCull = bNeverDistanceCull || GetLODParentPrimitive();
+		CachedMaxDrawDistance = bNeverCull ? 0.f : CachedMaxDrawDistance;
 	} 
 
 	if (GetLightmapType() == ELightmapType::ForceSurface && GetStaticLightingType() == LMIT_None)
@@ -2340,32 +2342,10 @@ void UPrimitiveComponent::SetCullDistance(const float NewCullDistance)
 	}
 }
 
-bool UPrimitiveComponent::CanBeDistanceCulled() const
-{
-	if (bNeverDistanceCull)
-	{
-		return false;
-	}
-
-	if (GetLODParentPrimitive())
-	{
-		return false;
-	}
-
-	if (ULevel* Level = GetTypedOuter<ULevel>())
-	{
-		if (Level->IsWorldPartitionRuntimeCell())
-		{
-			return !IsHLODRelevant();
-		}
-	}
-	
-	return true;
-}
-
 void UPrimitiveComponent::SetCachedMaxDrawDistance(const float NewCachedMaxDrawDistance)
 {
-	float NewMaxDrawDistance = CanBeDistanceCulled() ? NewCachedMaxDrawDistance : 0;
+	bool bNeverCull = bNeverDistanceCull || GetLODParentPrimitive();
+	float NewMaxDrawDistance = bNeverCull ? 0.f : NewCachedMaxDrawDistance;
 
 	if( !FMath::IsNearlyEqual(CachedMaxDrawDistance, NewMaxDrawDistance) )
 	{
