Added a new stat for texture steaming analytics
#tests Local build, verification of locally generated report data and charts
#rnx

diff --git a/Engine/Source/Runtime/Engine/Private/Streaming/StreamingManagerTexture.cpp b/Engine/Source/Runtime/Engine/Private/Streaming/StreamingManagerTexture.cpp
--- a/Engine/Source/Runtime/Engine/Private/Streaming/StreamingManagerTexture.cpp
+++ b/Engine/Source/Runtime/Engine/Private/Streaming/StreamingManagerTexture.cpp
@@ -1825,7 +1825,8 @@ void FRenderAssetStreamingManager::UpdateResourceStreaming( float DeltaTime, boo
 	LogViewLocationChange();
 	STAT(DisplayedStats.Apply();)
 
-	CSV_CUSTOM_STAT(TextureStreaming, StreamingPool, ((float)(DisplayedStats.RequiredPool + (GPoolSizeVRAMPercentage > 0 ? 0 : DisplayedStats.NonStreamingMips))) / (1024.0f * 1024.0f), ECsvCustomStatOp::Set);
+	float StreamingPoolStat = ((float)(DisplayedStats.RequiredPool + (GPoolSizeVRAMPercentage > 0 ? 0 : DisplayedStats.NonStreamingMips)));
+	CSV_CUSTOM_STAT(TextureStreaming, StreamingPool, StreamingPoolStat / (1024.0f * 1024.0f), ECsvCustomStatOp::Set);
 	CSV_CUSTOM_STAT(TextureStreaming, SafetyPool, ((float)DisplayedStats.SafetyPool) / (1024.0f * 1024.0f), ECsvCustomStatOp::Set);
 	CSV_CUSTOM_STAT(TextureStreaming, TemporaryPool, ((float)DisplayedStats.TemporaryPool) / (1024.0f * 1024.0f), ECsvCustomStatOp::Set);
 	CSV_CUSTOM_STAT(TextureStreaming, CachedMips, ((float)DisplayedStats.CachedMips) / (1024.0f * 1024.0f), ECsvCustomStatOp::Set);
@@ -1833,7 +1834,10 @@ void FRenderAssetStreamingManager::UpdateResourceStreaming( float DeltaTime, boo
 	CSV_CUSTOM_STAT(TextureStreaming, ResidentMeshMem, ((float)DisplayedStats.ResidentMeshMem) / (1024.0f * 1024.0f), ECsvCustomStatOp::Set);
 	CSV_CUSTOM_STAT(TextureStreaming, StreamedMeshMem, ((float)DisplayedStats.StreamedMeshMem) / (1024.0f * 1024.0f), ECsvCustomStatOp::Set);
 	CSV_CUSTOM_STAT(TextureStreaming, NonStreamingMips, ((float)DisplayedStats.NonStreamingMips) / (1024.0f * 1024.0f), ECsvCustomStatOp::Set);
-	CSV_CUSTOM_STAT(TextureStreaming, PendingStreamInData, ((float)DisplayedStats.PendingRequests) / (1024.0f * 1024.0f), ECsvCustomStatOp::Set)
+	CSV_CUSTOM_STAT(TextureStreaming, PendingStreamInData, ((float)DisplayedStats.PendingRequests) / (1024.0f * 1024.0f), ECsvCustomStatOp::Set);
+
+	float DesiredDataLoadedPct = StreamingPoolStat == 0.0f ? 0.0f : 100.0f * (float)DisplayedStats.WantedMips / StreamingPoolStat;
+	CSV_CUSTOM_STAT(TextureStreaming, DesiredDataLoadedPercent, DesiredDataLoadedPct, ECsvCustomStatOp::Set);
 
 	RenderAssetInstanceAsyncWork->EnsureCompletion();
 
diff --git a/Engine/Source/Runtime/Engine/Private/Streaming/TextureStreamingHelpers.cpp b/Engine/Source/Runtime/Engine/Private/Streaming/TextureStreamingHelpers.cpp
--- a/Engine/Source/Runtime/Engine/Private/Streaming/TextureStreamingHelpers.cpp
+++ b/Engine/Source/Runtime/Engine/Private/Streaming/TextureStreamingHelpers.cpp
@@ -31,6 +31,7 @@ DECLARE_MEMORY_STAT_POOL(TEXT("IO Bandwidth"), STAT_Streaming14_MipIOBandwidth,
 #if WITH_EDITOR
 DECLARE_MEMORY_STAT_POOL(TEXT("Editor Required Pool"), STAT_Streaming15_EditorRequiredPool, STATGROUP_Streaming, FPlatformMemory::MCR_StreamingPool);
 #endif
+DECLARE_MEMORY_STAT_POOL(TEXT("Desired Data Loaded Percent"), STAT_Streaming16_DesiredDataLoadedPct, STATGROUP_Streaming, FPlatformMemory::MCR_UsedStreamingPool);
 
 DECLARE_CYCLE_STAT(TEXT("Setup Async Task"), STAT_Streaming01_SetupAsyncTask, STATGROUP_Streaming);
 DECLARE_CYCLE_STAT(TEXT("Update Streaming Data"), STAT_Streaming02_UpdateStreamingData, STATGROUP_Streaming);
@@ -429,6 +430,8 @@ void FRenderAssetStreamingStats::Apply()
 #if WITH_EDITOR
 	SET_MEMORY_STAT(STAT_Streaming15_EditorRequiredPool, RequiredPool_Editor);
 #endif
+
+	SET_MEMORY_STAT(STAT_Streaming16_DesiredDataLoadedPct, (bHasBudgetLimit && StreamingPool != 0) ? 100.0f * float(WantedMips) / float(StreamingPool) : 0.0f);
 	SET_CYCLE_COUNTER(STAT_Streaming01_SetupAsyncTask, SetupAsyncTaskCycles);
 	SET_CYCLE_COUNTER(STAT_Streaming02_UpdateStreamingData, UpdateStreamingDataCycles);
 	SET_CYCLE_COUNTER(STAT_Streaming03_StreamRenderAssets, StreamRenderAssetsCycles);
