[WorldPartition] Optimize streaming performance ratio calculations by only calculating blocking ratios when the cell is setup to block on slow streaming.

#tests ReaplayRuns
#rnx

diff --git a/Engine/Source/Runtime/Engine/Private/WorldPartition/WorldPartitionRuntimeCellData.cpp b/Engine/Source/Runtime/Engine/Private/WorldPartition/WorldPartitionRuntimeCellData.cpp
--- a/Engine/Source/Runtime/Engine/Private/WorldPartition/WorldPartitionRuntimeCellData.cpp
+++ b/Engine/Source/Runtime/Engine/Private/WorldPartition/WorldPartitionRuntimeCellData.cpp
@@ -52,6 +52,11 @@ void UWorldPartitionRuntimeCellData::ResetStreamingSourceInfo(const FWorldPartit
 
 DECLARE_CYCLE_STAT_WITH_FLAGS(TEXT("Append Streaming Source Info"), STAT_WorldPartitionAppendStreamingSourceInfo, STATGROUP_WorldPartition, EStatFlags::Verbose);
 void UWorldPartitionRuntimeCellData::AppendStreamingSourceInfo(const FWorldPartitionStreamingSource& Source, const FSphericalSector& SourceShape, const FWorldPartitionStreamingContext& Context) const
+{
+	AppendStreamingSourceInfo(Source, SourceShape, Context, false);
+}
+
+void UWorldPartitionRuntimeCellData::AppendStreamingSourceInfo(const FWorldPartitionStreamingSource& Source, const FSphericalSector& SourceShape, const FWorldPartitionStreamingContext& Context, bool bBlockOnSlowLoading) const
 {
 	SCOPE_CYCLE_COUNTER(STAT_WorldPartitionAppendStreamingSourceInfo);
 
@@ -71,7 +76,7 @@ void UWorldPartitionRuntimeCellData::AppendStreamingSourceInfo(const FWorldParti
 
 	CachedMinSourcePriority = FMath::Min((uint8)Source.Priority, CachedMinSourcePriority);
 
-	if (Source.bBlockOnSlowLoading)
+	if (bBlockOnSlowLoading && Source.bBlockOnSlowLoading)
 	{
 		bCachedWasRequestedByBlockingSource = true;
 
diff --git a/Engine/Source/Runtime/Engine/Private/WorldPartition/WorldPartitionRuntimeCellDataSpatialHash.cpp b/Engine/Source/Runtime/Engine/Private/WorldPartition/WorldPartitionRuntimeCellDataSpatialHash.cpp
--- a/Engine/Source/Runtime/Engine/Private/WorldPartition/WorldPartitionRuntimeCellDataSpatialHash.cpp
+++ b/Engine/Source/Runtime/Engine/Private/WorldPartition/WorldPartitionRuntimeCellDataSpatialHash.cpp
@@ -87,12 +87,17 @@ float UWorldPartitionRuntimeCellDataSpatialHash::ComputeSourceToCellAngleFactor(
 
 void UWorldPartitionRuntimeCellDataSpatialHash::AppendStreamingSourceInfo(const FWorldPartitionStreamingSource& Source, const FSphericalSector& SourceShape, const FWorldPartitionStreamingContext& Context) const
 {
-	Super::AppendStreamingSourceInfo(Source, SourceShape, Context);
+	Super::AppendStreamingSourceInfo(Source, SourceShape, Context, false);
+}
+
+void UWorldPartitionRuntimeCellDataSpatialHash::AppendStreamingSourceInfo(const FWorldPartitionStreamingSource& Source, const FSphericalSector& SourceShape, const FWorldPartitionStreamingContext& Context, bool bBlockOnSlowLoading) const
+{
+	Super::AppendStreamingSourceInfo(Source, SourceShape, Context, bBlockOnSlowLoading);
 
 	const double SquareDistance2D = FVector::DistSquared2D(SourceShape.GetCenter(), Position);
 
 	// Update cached values based on the 2D distance
-	if (Source.bBlockOnSlowLoading)
+	if (bBlockOnSlowLoading && Source.bBlockOnSlowLoading)
 	{
 		CachedMinSquareDistanceToBlockingSource2D = FMath::Min(SquareDistance2D, CachedMinSquareDistanceToBlockingSource2D);
 		CachedMinBlockOnSlowStreamingRatio2D = FMath::Min(CachedMinBlockOnSlowStreamingRatio2D, FMath::Sqrt(CachedMinSquareDistanceToBlockingSource2D) / SourceShape.GetRadius());
diff --git a/Engine/Source/Runtime/Engine/Public/WorldPartition/WorldPartitionRuntimeCell.h b/Engine/Source/Runtime/Engine/Public/WorldPartition/WorldPartitionRuntimeCell.h
--- a/Engine/Source/Runtime/Engine/Public/WorldPartition/WorldPartitionRuntimeCell.h
+++ b/Engine/Source/Runtime/Engine/Public/WorldPartition/WorldPartitionRuntimeCell.h
@@ -429,7 +429,7 @@ protected:
 
 public:
 	//~Begin UWorldPartitionRuntimeCellData Proxy
-	inline void AppendStreamingSourceInfo(const FWorldPartitionStreamingSource& Source, const FSphericalSector& SourceShape, const FWorldPartitionStreamingContext& Context) const { RuntimeCellData->AppendStreamingSourceInfo(Source, SourceShape, Context); }
+	inline void AppendStreamingSourceInfo(const FWorldPartitionStreamingSource& Source, const FSphericalSector& SourceShape, const FWorldPartitionStreamingContext& Context) const { RuntimeCellData->AppendStreamingSourceInfo(Source, SourceShape, Context, bBlockOnSlowLoading); }
 	inline void MergeStreamingSourceInfo() const { RuntimeCellData->MergeStreamingSourceInfo(); }
 	ENGINE_API int32 SortCompare(const UWorldPartitionRuntimeCell* Other) const;
 	
diff --git a/Engine/Source/Runtime/Engine/Public/WorldPartition/WorldPartitionRuntimeCellData.h b/Engine/Source/Runtime/Engine/Public/WorldPartition/WorldPartitionRuntimeCellData.h
--- a/Engine/Source/Runtime/Engine/Public/WorldPartition/WorldPartitionRuntimeCellData.h
+++ b/Engine/Source/Runtime/Engine/Public/WorldPartition/WorldPartitionRuntimeCellData.h
@@ -27,7 +27,9 @@ class UWorldPartitionRuntimeCellData : public UObject
 #endif
 
 	ENGINE_API virtual void ResetStreamingSourceInfo(const FWorldPartitionStreamingContext& Context) const;
+	UE_DEPRECATED(5.6, "Use version that takes a boolean to flag if the cell blocks on slow loading (bBlockOnSlowLoading).")
 	ENGINE_API virtual void AppendStreamingSourceInfo(const FWorldPartitionStreamingSource& Source, const FSphericalSector& SourceShape, const FWorldPartitionStreamingContext& Context) const;
+	ENGINE_API virtual void AppendStreamingSourceInfo(const FWorldPartitionStreamingSource& Source, const FSphericalSector& SourceShape, const FWorldPartitionStreamingContext& Context, bool bBlockOnSlowLoading) const;
 	ENGINE_API virtual void MergeStreamingSourceInfo() const {}
 	ENGINE_API virtual int32 SortCompare(const UWorldPartitionRuntimeCellData* InOther) const;
 
diff --git a/Engine/Source/Runtime/Engine/Public/WorldPartition/WorldPartitionRuntimeCellDataSpatialHash.h b/Engine/Source/Runtime/Engine/Public/WorldPartition/WorldPartitionRuntimeCellDataSpatialHash.h
--- a/Engine/Source/Runtime/Engine/Public/WorldPartition/WorldPartitionRuntimeCellDataSpatialHash.h
+++ b/Engine/Source/Runtime/Engine/Public/WorldPartition/WorldPartitionRuntimeCellDataSpatialHash.h
@@ -15,7 +15,9 @@ class UWorldPartitionRuntimeCellDataSpatialHash : public UWorldPartitionRuntimeC
 
 	//~Begin UWorldPartitionRuntimeCellData
 	ENGINE_API virtual void ResetStreamingSourceInfo(const FWorldPartitionStreamingContext& Context) const override;
+	UE_DEPRECATED(5.6, "Use version that takes a boolean to flag if the cell blocks on slow loading (bBlockOnSlowLoading).")
 	ENGINE_API virtual void AppendStreamingSourceInfo(const FWorldPartitionStreamingSource& Source, const FSphericalSector& SourceShape, const FWorldPartitionStreamingContext& Context) const override;
+	ENGINE_API virtual void AppendStreamingSourceInfo(const FWorldPartitionStreamingSource& Source, const FSphericalSector& SourceShape, const FWorldPartitionStreamingContext& Context, bool bBlockOnSlowLoading) const override;
 	ENGINE_API virtual void MergeStreamingSourceInfo() const override;
 	ENGINE_API virtual int32 SortCompare(const UWorldPartitionRuntimeCellData* InOther) const override;
 	ENGINE_API virtual FBox GetCellBounds() const override;
