Fixed broken logic in TAttributeContainer::Remove

diff --git a/Engine/Source/Runtime/Engine/Classes/Animation/AttributesContainer.h b/Engine/Source/Runtime/Engine/Classes/Animation/AttributesContainer.h
--- a/Engine/Source/Runtime/Engine/Classes/Animation/AttributesContainer.h
+++ b/Engine/Source/Runtime/Engine/Classes/Animation/AttributesContainer.h
@@ -616,8 +616,7 @@ namespace UE
 						// If removing the last one, remove all data for the attribute type
 						if (TypedArray.Num() == 1)
 						{
-							AttributeIdentifiers.RemoveAtSwap(TypeIndex);
-
+							UniqueTypes.RemoveAtSwap(TypeIndex);
 							UniqueTypedBoneIndices.RemoveAtSwap(TypeIndex);
 							AttributeIdentifiers.RemoveAtSwap(TypeIndex);
 							Values.RemoveAtSwap(TypeIndex);
