[AssetManager] Fix crashing when LoadAssetList returned nullptr

#rnx

diff --git a/Engine/Source/Runtime/Engine/Private/AssetManager.cpp b/Engine/Source/Runtime/Engine/Private/AssetManager.cpp
--- a/Engine/Source/Runtime/Engine/Private/AssetManager.cpp
+++ b/Engine/Source/Runtime/Engine/Private/AssetManager.cpp
@@ -1961,7 +1961,7 @@ TSharedPtr<FStreamableHandle> UAssetManager::ChangeBundleStateForPrimaryAssets(
 				NewHandle = LoadAssetList(MoveTemp(PathsToLoad), MoveTemp(AssetListLoadParams), FString(DebugName), CopyTemp(Location));
 				// save handle so possible duplicate requests will not override this one
 				// set only if loading is not completed yet because completion callback might already reset the PendingState
-				if (!NewHandle->HasLoadCompleted())
+				if (NewHandle && !NewHandle->HasLoadCompleted())
 				{
 					NameData->PendingState.Handle = NewHandle;
 				}
