[MLDeformer] A little safety improvement related to external morph targets.

diff --git a/Engine/Source/Runtime/Engine/Private/Components/SkinnedMeshComponent.cpp b/Engine/Source/Runtime/Engine/Private/Components/SkinnedMeshComponent.cpp
--- a/Engine/Source/Runtime/Engine/Private/Components/SkinnedMeshComponent.cpp
+++ b/Engine/Source/Runtime/Engine/Private/Components/SkinnedMeshComponent.cpp
@@ -906,8 +906,26 @@ void USkinnedMeshComponent::ClearExternalMorphSets(int32 LOD)
 
 void USkinnedMeshComponent::ResizeExternalMorphTargetSets()
 {
-	ExternalMorphSets.Reset();
-	ExternalMorphSets.AddDefaulted(GetNumLODs());
+	const int32 NewNumLODs = GetNumLODs();
+	if (ExternalMorphSets.Num() == NewNumLODs)
+	{
+		return;
+	}
+
+	if (NewNumLODs > ExternalMorphSets.Num())
+	{
+		const int32 NumLODsToAdd = NewNumLODs - ExternalMorphSets.Num();
+		ExternalMorphSets.AddDefaulted(NumLODsToAdd);
+	}
+	else
+	{
+		check(NewNumLODs < ExternalMorphSets.Num());
+		const int32 NumLODsToRemove = ExternalMorphSets.Num() - NewNumLODs;
+		for (int32 Index = 0; Index < NumLODsToRemove; ++Index)
+		{
+			ExternalMorphSets.Pop(EAllowShrinking::No);
+		}
+	}
 }
 
 void USkinnedMeshComponent::RefreshExternalMorphTargetWeights(bool bZeroOldWeights)
@@ -1324,7 +1342,7 @@ void USkinnedMeshComponent::SendRenderDynamicData_Concurrent()
 			{
 				ActiveMorphTargets.Empty();
 			}
-
+			
 			MeshObject->Update(
 				UseLOD,
 				this,
