Fix to ensure the max delta time is applied even when timestep smoothing is disabled. Only affects control rig physics as all other systems (RBAN etc) apply timestep smoothing.

diff --git a/Engine/Source/Runtime/Engine/Private/PhysicsEngine/ImmediatePhysics/ImmediatePhysicsChaos/ImmediatePhysicsSimulation_Chaos.cpp b/Engine/Source/Runtime/Engine/Private/PhysicsEngine/ImmediatePhysics/ImmediatePhysicsChaos/ImmediatePhysicsSimulation_Chaos.cpp
--- a/Engine/Source/Runtime/Engine/Private/PhysicsEngine/ImmediatePhysics/ImmediatePhysicsChaos/ImmediatePhysicsSimulation_Chaos.cpp
+++ b/Engine/Source/Runtime/Engine/Private/PhysicsEngine/ImmediatePhysics/ImmediatePhysicsChaos/ImmediatePhysicsSimulation_Chaos.cpp
@@ -893,7 +893,7 @@ namespace ImmediatePhysics_Chaos
 		{
 			// Adjust things so that if we go back to smoothing the step time, we do so gracefully
 			Implementation->NumRollingAverageStepTimes = Implementation->MaxNumRollingAverageStepTimes;
-			Implementation->RollingAverageStepTime = DeltaTime;
+			Implementation->RollingAverageStepTime = FMath::Min(DeltaTime, MaxStepTime);
 		}
 		return Implementation->RollingAverageStepTime;
 	}
