Fixed GPU crash due to hair skeletal mesh LOD being accessed.

diff --git a/Engine/Plugins/Runtime/HairStrands/Source/HairStrandsCore/Private/GroomManager.cpp b/Engine/Plugins/Runtime/HairStrands/Source/HairStrandsCore/Private/GroomManager.cpp
--- a/Engine/Plugins/Runtime/HairStrands/Source/HairStrandsCore/Private/GroomManager.cpp
+++ b/Engine/Plugins/Runtime/HairStrands/Source/HairStrandsCore/Private/GroomManager.cpp
@@ -2926,6 +2926,12 @@ static void RunHairLODSelection(
 	TArray<FUniqueGroomBinding> UniqueGroomBindingAssets;
 	UniqueGroomAssets.Reserve(Instances.Num());
 
+	if (!Instances.IsEmpty())
+	{
+		// Sync the skeletal mesh update tasks involved with LOD selection so that GetLOD returns a valid result.
+		FSkeletalMeshUpdater::WaitForStage(GraphBuilder, ESkeletalMeshUpdateStage::Filter);
+	}
+
 	// Find unique groom assets
 	TArray<FInstanceData> InstanceDatas;
 	InstanceDatas.Reserve(Instances.Num());
diff --git a/Engine/Source/Runtime/Engine/Internal/SkeletalMeshUpdater.h b/Engine/Source/Runtime/Engine/Internal/SkeletalMeshUpdater.h
--- a/Engine/Source/Runtime/Engine/Internal/SkeletalMeshUpdater.h
+++ b/Engine/Source/Runtime/Engine/Internal/SkeletalMeshUpdater.h
@@ -104,8 +104,16 @@ private:
 
 enum class ESkeletalMeshUpdateStage : uint8
 {
+	// Filtering of dynamic datas to mesh objects.
+	Filter,
+
+	// Processing mesh deformer mesh object allocations.
 	MeshDeformer,
+
+	// Processing inline mesh object allocations.
 	Inline,
+
+	// Processing skin cache mesh object allocations.
 	SkinCache
 };
 
@@ -263,6 +271,7 @@ public:
 	struct FSubmitTasks
 	{
 		// These fire in order; e.g. syncing SkinCache syncs everything.
+		UE::Tasks::FTask Filter;
 		UE::Tasks::FTask MeshDeformer;
 		UE::Tasks::FTask Inline;
 		UE::Tasks::FTask SkinCache;
diff --git a/Engine/Source/Runtime/Engine/Private/SkeletalMeshUpdater.cpp b/Engine/Source/Runtime/Engine/Private/SkeletalMeshUpdater.cpp
--- a/Engine/Source/Runtime/Engine/Private/SkeletalMeshUpdater.cpp
+++ b/Engine/Source/Runtime/Engine/Private/SkeletalMeshUpdater.cpp
@@ -302,6 +302,8 @@ struct FSkeletalMeshUpdater::FTaskData
 
 	void ProcessForeground()
 	{
+		Tasks.Filter.Trigger();
+
 		{
 			TRACE_CPUPROFILER_EVENT_SCOPE(MeshDeformer);
 			for (auto& UpdatePacket : Packets)
@@ -346,6 +348,7 @@ struct FSkeletalMeshUpdater::FTaskData
 
 	struct
 	{
+		UE::Tasks::FTaskEvent Filter{ UE_SOURCE_LOCATION };
 		UE::Tasks::FTaskEvent MeshDeformer{ UE_SOURCE_LOCATION };
 		UE::Tasks::FTaskEvent Inline{ UE_SOURCE_LOCATION };
 
@@ -405,7 +408,8 @@ FSkeletalMeshUpdater::FSubmitTasks FSkeletalMeshUpdater::Submit(FRDGBuilder& Gra
 
 	return GraphBuilder.Blackboard.Create<FSubmitTasks>(FSubmitTasks
 	{
-		  .MeshDeformer = TaskData.Tasks.MeshDeformer
+		  .Filter       = TaskData.Tasks.Filter
+		, .MeshDeformer = TaskData.Tasks.MeshDeformer
 		, .Inline       = TaskData.Tasks.Inline
 		, .SkinCache    = MoveTemp(ForegroundTask)
 	});
@@ -418,6 +422,10 @@ void FSkeletalMeshUpdater::WaitForStage(FRDGBuilder& GraphBuilder, ESkeletalMesh
 		switch (Stage)
 		{
 		default: checkNoEntry();
+		case ESkeletalMeshUpdateStage::Filter:
+			SubmitTasks->Filter.Wait();
+			SubmitTasks->Filter = {};
+			break;
 		case ESkeletalMeshUpdateStage::MeshDeformer:
 			SubmitTasks->MeshDeformer.Wait();
 			SubmitTasks->MeshDeformer = {};
