Fix a crash that can occur in InstancedSkinnedMesh when a skinned asset doesn't have a skeleton. This seemed to be able to happen when reimporting the mesh while in use in the level.

diff --git a/Engine/Source/Runtime/Engine/Private/Components/InstancedSkinnedMeshComponent.cpp b/Engine/Source/Runtime/Engine/Private/Components/InstancedSkinnedMeshComponent.cpp
--- a/Engine/Source/Runtime/Engine/Private/Components/InstancedSkinnedMeshComponent.cpp
+++ b/Engine/Source/Runtime/Engine/Private/Components/InstancedSkinnedMeshComponent.cpp
@@ -267,7 +267,7 @@ FInstancedSkinnedMeshData::FInstancedSkinnedMeshData(
 )
 :	LODIndex(InLODIndex)
 {
-	int32 BoneCount = InSkinnedAsset->GetSkeleton()->GetReferenceSkeleton().GetRawBoneNum();
+	int32 BoneCount = InSkinnedAsset->GetSkeleton() ? InSkinnedAsset->GetSkeleton()->GetReferenceSkeleton().GetRawBoneNum() : 0;
 	ReferenceToLocal.SetNumUninitialized(BoneCount);
 	for (int32 TransformIndex = 0; TransformIndex < BoneCount; ++TransformIndex)
 	{
