Unity build fix

diff --git a/Engine/Source/Runtime/Engine/Private/PhysicsEngine/ImmediatePhysics/ImmediatePhysicsChaos/ImmediatePhysicsSimulation_Chaos.cpp b/Engine/Source/Runtime/Engine/Private/PhysicsEngine/ImmediatePhysics/ImmediatePhysicsChaos/ImmediatePhysicsSimulation_Chaos.cpp
--- a/Engine/Source/Runtime/Engine/Private/PhysicsEngine/ImmediatePhysics/ImmediatePhysicsChaos/ImmediatePhysicsSimulation_Chaos.cpp
+++ b/Engine/Source/Runtime/Engine/Private/PhysicsEngine/ImmediatePhysics/ImmediatePhysicsChaos/ImmediatePhysicsSimulation_Chaos.cpp
@@ -1268,7 +1268,7 @@ namespace ImmediatePhysics_Chaos
 		TFunction<void(const FCollisionData&)> Visitor, Chaos::ECollisionVisitorFlags VisitorFlags) const
 	{
 		using namespace Chaos;
-		const Private::FCollisionConstraintAllocator& ConstraintAllocator = Implementation->Collisions.GetConstraintAllocator();
+		const Chaos::Private::FCollisionConstraintAllocator& ConstraintAllocator = Implementation->Collisions.GetConstraintAllocator();
 		ConstraintAllocator.VisitConstCollisions(
 			[&](const FPBDCollisionConstraint& Collision)
 			{
@@ -1277,7 +1277,6 @@ namespace ImmediatePhysics_Chaos
 			}, VisitorFlags);
 	}
 
-
 	void FSimulation::DebugDrawConstraints()
 	{
 #if CHAOS_DEBUG_DRAW
