New material translator changes:
- Added initial support for computing analytical partial derivatives from any given FValue.
- Added Emitter APIs to emit partial derivatives that appropriately resolves to using hardware derivatives (ddx/ddy) or analytical derivatives depending on whether the execution stage support them.
- Added EmitConstantScalar API for emitting a scalar value of a specified kind.
- Added FStageSwitch instruction, which simply bypasses to a different argument value depending in each execution stage.
- Added FHardwarePartialDerivative instruction that resolves to HLSL's ddx()/ddy().
- Added "Compute" stage which calculates the same material attributes as those flagged for the pixel stage execution, except it uses analytical derivatives (since ddx/ddy isn't available).
- Fixed bug regarding calculation of used texcoords usage count in the output shader.
- Fixed emission of and/or/xor operators in HLSL translator.
- Some IR dump improvements.

#rnx

diff --git a/Engine/Source/Runtime/Engine/Private/Materials/MaterialExpressionsIR.cpp b/Engine/Source/Runtime/Engine/Private/Materials/MaterialExpressionsIR.cpp
--- a/Engine/Source/Runtime/Engine/Private/Materials/MaterialExpressionsIR.cpp
+++ b/Engine/Source/Runtime/Engine/Private/Materials/MaterialExpressionsIR.cpp
@@ -902,7 +902,7 @@ static void BuildTextureSampleExpression(MIR::FValue* Texture, UMaterialExpressi
 			MipValue = Emitter.EmitOperator(MIR::BO_Add, MipValue, Emitter.EmitExternalInput(MIR::EExternalInput::ViewMaterialTextureMipBias));
 		}
 	}
-	else
+	else if (Expr->MipValueMode == TMVM_Derivative)
 	{
 		check(Expr->MipValueMode = TMVM_Derivative);
 		TexCoordsDdx = Emitter.EmitConstruct(TexCoords->Type, Emitter.Input(&Expr->CoordinatesDX));
@@ -1229,4 +1229,14 @@ void UMaterialExpressionPixelNormalWS::Build(MIR::FEmitter& Emitter)
 	Emitter.Output(0, Output);
 }
 
+void UMaterialExpressionDDX::Build(MIR::FEmitter& Emitter)
+{
+	Emitter.Output(0, Emitter.EmitPartialDerivative(Emitter.Input(&Value), MIR::EDerivativeAxis::X));
+}
+
+void UMaterialExpressionDDY::Build(MIR::FEmitter& Emitter)
+{
+	Emitter.Output(0, Emitter.EmitPartialDerivative(Emitter.Input(&Value), MIR::EDerivativeAxis::Y));
+}
+
 #endif // WITH_EDITOR
diff --git a/Engine/Source/Runtime/Engine/Private/Materials/MaterialIR.cpp b/Engine/Source/Runtime/Engine/Private/Materials/MaterialIR.cpp
--- a/Engine/Source/Runtime/Engine/Private/Materials/MaterialIR.cpp
+++ b/Engine/Source/Runtime/Engine/Private/Materials/MaterialIR.cpp
@@ -9,18 +9,37 @@
 namespace MIR
 {
 
-EStage GetMaterialOutputStage(EMaterialProperty Property)
+const TCHAR* LexToString(EStage Stage)
 {
-	return Property == MP_WorldPositionOffset ? MIR::Stage_Vertex : MIR::Stage_Pixel;
+	switch (Stage)
+	{
+		case Stage_Vertex:	return TEXT("Vertex");
+		case Stage_Pixel: 	return TEXT("Pixel");
+		case Stage_Compute: return TEXT("Compute");
+		case NumStages: UE_MIR_UNREACHABLE();
+	}
+	return nullptr;
+}
+
+bool MaterialOutputEvaluatesInStage(EMaterialProperty Property, EStage Stage)
+{
+	if (Property == MP_WorldPositionOffset)
+	{
+		return Stage == Stage_Vertex;
+	}
+	else
+	{
+		return Stage != Stage_Vertex;
+	}
 }
 
 const TCHAR* LexToString(EValueKind Kind)
 {
 	switch (Kind)
 	{
+		case VK_Poison: return TEXT("Poison");
 		case VK_Constant: return TEXT("Constant");
 		case VK_ExternalInput: return TEXT("ExternalInput");
-		case VK_InlineHLSL: return TEXT("ExternalCode");
 		case VK_TextureObject: return TEXT("TextureObject");
 		case VK_UniformParameter: return TEXT("UniformParameter");
 		case VK_Dimensional: return TEXT("Dimensional");
@@ -30,8 +49,15 @@ const TCHAR* LexToString(EValueKind Kind)
 		case VK_Subscript: return TEXT("Subscript");
 		case VK_Cast: return TEXT("Cast");
 		case VK_TextureSample: return TEXT("TextureSample");
-		default: UE_MIR_UNREACHABLE();
-	};
+		case VK_InlineHLSL: return TEXT("ExternalCode");
+		case VK_StageSwitch: return TEXT("StageSwitch");
+		case VK_HardwarePartialDerivative: return TEXT("HardwarePartialDerivative");
+		/* invalid entries */
+		case VK_InstructionBegin:
+		case VK_InstructionEnd:
+			UE_MIR_UNREACHABLE();
+	}
+	return nullptr;
 }
 
 bool FValue::IsAnalyzed(EStage Stage) const
@@ -68,20 +94,27 @@ uint32 FValue::GetSizeInBytes() const
 {
 	switch (Kind)
 	{
+		case VK_Poison: return sizeof(FPoison);
 		case VK_Constant: return sizeof(FConstant);
 		case VK_ExternalInput: return sizeof(FExternalInput);
-		case VK_InlineHLSL: return sizeof(FInlineHLSL);
 		case VK_TextureObject: return sizeof(FTextureObject);
 		case VK_UniformParameter: return sizeof(FUniformParameter);
+		case VK_Dimensional: return sizeof(FDimensional) + sizeof(FValue*) * static_cast<const FDimensional*>(this)->GetComponents().Num();
 		case VK_SetMaterialOutput: return sizeof(FSetMaterialOutput);
 		case VK_Operator: return sizeof(FOperator);
-		case VK_Dimensional: return sizeof(FDimensional) + sizeof(FValue*) * static_cast<const FDimensional*>(this)->GetComponents().Num();
 		case VK_Branch: return sizeof(FBranch);
 		case VK_Subscript: return sizeof(FSubscript);
 		case VK_Cast: return sizeof(FCast);
 		case VK_TextureSample: return sizeof(FTextureSample);
-		default: UE_MIR_UNREACHABLE();
+		case VK_InlineHLSL: return sizeof(FInlineHLSL);
+		case VK_StageSwitch: return sizeof(FStageSwitch);
+		case VK_HardwarePartialDerivative: return sizeof(FHardwarePartialDerivative);
+		/* invalid entries */
+		case VK_InstructionBegin:
+		case VK_InstructionEnd:
+			UE_MIR_UNREACHABLE();
 	}
+	return 0;
 }
 
 bool FValue::IsA(EValueKind InKind) const
@@ -163,10 +196,31 @@ TConstArrayView<FValue*> FValue::GetUses() const
 			return {}; // @todo-laura.hermanns - Add input arguments as soon as external code expressions support parameters
 		}
 
+		case VK_StageSwitch:
+		{
+			auto This = static_cast<const FStageSwitch*>(this);
+			return { This->Args, FStageSwitch::NumUses * NumStages };
+		}
+
+		case VK_HardwarePartialDerivative:
+		{
+			auto This = static_cast<const FHardwarePartialDerivative*>(this);
+			return { &This->Arg, FHardwarePartialDerivative::NumUses };
+		}
+
 		default: UE_MIR_UNREACHABLE();
 	}
 }
 
+TConstArrayView<FValue*> FValue::GetUsesForStage(MIR::EStage Stage) const
+{
+	if (const FStageSwitch* This = As<FStageSwitch>())
+	{
+		return { &This->Args[int32(Stage)], FStageSwitch::NumUses };
+	}
+	return GetUses();
+}
+
 bool FValue::IsScalar() const
 {
 	return Type->AsScalar() != nullptr;
@@ -390,7 +444,7 @@ const FType* GetExternalInputType(EExternalInput Id)
 	}
 }
 
-MIR::EExternalInput TexCoordIndexToExternalInput(int TexCoordIndex)
+MIR::EExternalInput TexCoordIndexToExternalInput(int32 TexCoordIndex)
 {
 	check(TexCoordIndex < TexCoordMaxNum);
 	return (MIR::EExternalInput)((int)MIR::EExternalInput::TexCoord0 + TexCoordIndex);
@@ -401,6 +455,16 @@ bool IsExternalInputTexCoord(EExternalInput Id)
 	return Id >= EExternalInput::TexCoord0 && Id <= EExternalInput::TexCoord7;
 }
 
+int32 ExternalInputToTexCoordIndex(EExternalInput Id)
+{
+	int32 IntId = int32(Id) - 1;
+	if (IntId < 0 || IntId > TexCoordMaxNum * TexCoordGroups)
+	{
+		return UINT32_MAX;
+	}
+	return IntId % TexCoordMaxNum;
+}
+
 bool IsExternalInputTexCoordDdx(EExternalInput Id)
 {
 	return Id >= EExternalInput::TexCoord0_Ddx && Id <= EExternalInput::TexCoord7_Ddx;
diff --git a/Engine/Source/Runtime/Engine/Private/Materials/MaterialIRDebug.cpp b/Engine/Source/Runtime/Engine/Private/Materials/MaterialIRDebug.cpp
--- a/Engine/Source/Runtime/Engine/Private/Materials/MaterialIRDebug.cpp
+++ b/Engine/Source/Runtime/Engine/Private/Materials/MaterialIRDebug.cpp
@@ -75,7 +75,7 @@ struct FDebugDumpIRUseGraphState
 		}
 
 		int32 UseIndex = -1;
-		for (const FValue* Use : Value->GetUses())
+		for (const FValue* Use : Value->GetUsesForStage(Stage))
 		{
 			++UseIndex;
 
@@ -188,6 +188,15 @@ void DebugDumpIRUseGraph(const FMaterialIRModule& Module)
 
 /* Module IR to textual representation dumping */
 
+/**
+ * Returns whether given instruction kind has a dynamic number of arguments, such as the
+ * Operator instruction which can have one, two or three arguments.
+ */
+static bool InstrHasVariableArgCount(MIR::EValueKind Kind)
+{
+	return Kind == MIR::VK_Operator;
+}
+
 /// Helper struct to wrap the state used during IR to text dumping.
 struct FDebugDumpIRState
 {
@@ -235,16 +244,21 @@ struct FDebugDumpIRState
 			Out.Append(TEXT(" ("));
 
 			bool bAddComma = false;
-			TConstArrayView<FValue*> Uses = Instr->GetUses();
+			TConstArrayView<FValue*> Uses = Instr->GetUsesForStage(CurrentStage);
 			for (int32 UseIndex = 0; UseIndex < Uses.Num(); ++UseIndex)
 			{
+				const FValue* Use = Uses[UseIndex];
+				if (!Use && InstrHasVariableArgCount(Instr->Kind))
+				{
+					continue;
+				}
+
 				if (bAddComma)
 				{
 					Out.Append(TEXT(", "));
 				}
 				bAddComma = true;
 				
-				const FValue* Use = Uses[UseIndex];
 				if (!Use)
 				{
 					Out.Append(TEXT("null"));
@@ -319,6 +333,10 @@ struct FDebugDumpIRState
 				Out.Appendf(TEXT(" \"%s\""), *InlineHLSL->ExternalCodeDeclaration->Definition);
 			}
 		}
+		else if (const FHardwarePartialDerivative* Derivative = Instr->As<FHardwarePartialDerivative>())
+		{
+			Out.Append(Derivative->Axis == MIR::EDerivativeAxis::X ? TEXT(" \"ddx\"") : TEXT(" \"ddy\""));
+		}
 	}
 
 	/// Appends a reference to the specified value. This will look like "%x" if Value is an
@@ -453,7 +471,7 @@ FString DebugDumpIR(FStringView MaterialName, const FMaterialIRModule& Module)
 	for (int32 StageIndex = 0; StageIndex < MIR::NumStages; ++StageIndex)
 	{
 		State.CurrentStage = (MIR::EStage)StageIndex;
-		State.Out.Appendf(TEXT("\n; Stage %d\n"), StageIndex);
+		State.Out.Appendf(TEXT("\n; Stage %d \"%s\"\n"), StageIndex, MIR::LexToString(State.CurrentStage));
 		State.AppendBlock(Module.GetRootBlock((MIR::EStage)StageIndex), 0);
 	}
 
diff --git a/Engine/Source/Runtime/Engine/Private/Materials/MaterialIREmitter.cpp b/Engine/Source/Runtime/Engine/Private/Materials/MaterialIREmitter.cpp
--- a/Engine/Source/Runtime/Engine/Private/Materials/MaterialIREmitter.cpp
+++ b/Engine/Source/Runtime/Engine/Private/Materials/MaterialIREmitter.cpp
@@ -481,6 +481,24 @@ FValue* FEmitter::EmitConstantFloatScalar(TFloat InX)
 	return EmitPrototype(this, Scalar);
 }
 
+FValue* FEmitter::EmitConstantIntScalar(TInteger InX)
+{
+	FConstant Scalar = MakePrototype<FConstant>(FPrimitiveType::GetScalar(ScalarKind_Int));
+	Scalar.Integer = InX;
+	return EmitPrototype(this, Scalar);
+}
+
+FValue* FEmitter::EmitConstantScalar(EScalarKind Kind, TFloat FromFloat)
+{
+	switch (Kind)
+	{
+		case ScalarKind_Bool: return EmitConstantBoolScalar(FromFloat != 0.0f);
+		case ScalarKind_Int: return EmitConstantIntScalar(TInteger(FromFloat));
+		case ScalarKind_Float: return EmitConstantFloatScalar(FromFloat);
+		default: UE_MIR_UNREACHABLE();
+	}
+}
+
 FValue* FEmitter::EmitConstantFloat2(const FVector2f& InValue)
 {
 	FValue* X = EmitConstantFloatScalar(InValue.X);
@@ -505,13 +523,6 @@ FValue* FEmitter::EmitConstantFloat4(const FVector4f& InValue)
 	return EmitVector4(X, Y, Z, W);
 }
 
-FValue* FEmitter::EmitConstantIntScalar(TInteger InX)
-{
-	FConstant Scalar = MakePrototype<FConstant>(FPrimitiveType::GetScalar(ScalarKind_Int));
-	Scalar.Integer = InX;
-	return EmitPrototype(this, Scalar);
-}
-
 FValue* FEmitter::EmitConstantInt2(const FIntVector2& InValue)
 {
 	FValue* X = EmitConstantIntScalar(InValue.X);
@@ -820,17 +831,43 @@ FValue* FEmitter::EmitCustomPrimitiveData(uint32 PrimitiveDataIndex)
 	return nullptr;
 }
 
-FSetMaterialOutput* FEmitter::EmitSetMaterialOutput(EMaterialProperty InProperty, FValue* Arg)
+FValue* FEmitter::EmitStageSwitch(const FType* Type, TConstArrayView<FValue*> ValuePerStage)
 {
-	EStage PropertyStage = GetMaterialOutputStage(InProperty);
+	check(ValuePerStage.Num() <= NumStages);
+	FStageSwitch Prototype = MakePrototype<FStageSwitch>(Type);
+	for (int i = 0; i < ValuePerStage.Num(); ++i)
+	{
+		Prototype.Args[i] = ValuePerStage[i];
+	}
+	return EmitPrototype(this, Prototype);
+}
 
+FSetMaterialOutput* FEmitter::EmitSetMaterialOutput(EMaterialProperty InProperty, FValue* Arg)
+{
 	FSetMaterialOutput Proto = MakePrototype<FSetMaterialOutput>(nullptr);
-	Proto.Block[PropertyStage]	= Module->RootBlock[PropertyStage];
 	Proto.Property  			= InProperty;
 	Proto.Arg 					= Arg;
 
+	// Initialize the instruction block to the root of each stage it is evaluated in.
+	for (int i = 0; i < NumStages; ++i)
+	{
+		if (MaterialOutputEvaluatesInStage(InProperty, (EStage)i))
+		{
+			Proto.Block[i]	= Module->RootBlock[i];
+		}
+	}
+
+	// Make the instruction
 	FSetMaterialOutput* Instr = static_cast<FSetMaterialOutput*>(EmitPrototype(this, Proto));
-	Module->Outputs[PropertyStage].Add(Instr);
+
+	// Add the instruction to list of outputs of the stages it is evaluated in.
+	for (int i = 0; i < NumStages; ++i)
+	{
+		if (MaterialOutputEvaluatesInStage(InProperty, (EStage)i))
+		{
+			Module->Outputs[i].Add(Instr);
+		}
+	}
 
 	return Instr;
 }
@@ -1205,8 +1242,8 @@ static const FOperatorSignature* GetOperatorSignature(EOperator Op)
 		const FOperatorSignature UnaryArithmetic = { { OPTF_Arithmetic }, ORT_FirstArgumentType };
 		const FOperatorSignature BinaryArithmetic = { { OPTF_CastToCommonArithmetic, OPTF_CastToCommonArithmetic }, ORT_FirstArgumentType };
 		const FOperatorSignature TernaryArithmetic = { { OPTF_CastToCommonArithmetic, OPTF_CastToCommonArithmetic, OPTF_CastToCommonArithmetic }, ORT_FirstArgumentType };
-		const FOperatorSignature BinaryComparison = { { OPTF_Any, OPTF_AsFirstArgument }, ORT_BooleanWithFirstArgumentDimensions };
-		const FOperatorSignature BinaryLogical = { { OPTF_Boolean, OPTF_AsFirstArgument }, ORT_FirstArgumentType };
+		const FOperatorSignature BinaryArithmeticComparison = { { OPTF_CastToCommonArithmetic, OPTF_CastToCommonArithmetic }, ORT_BooleanWithFirstArgumentDimensions };
+		const FOperatorSignature BinaryLogical = { { OPTF_Boolean | OPTF_CastToCommon, OPTF_Boolean | OPTF_CastToCommon }, ORT_FirstArgumentType };
 
 		static FOperatorSignature S[MIR::OperatorCount];
 		S[UO_Negate] 				= UnaryArithmetic;
@@ -1231,12 +1268,12 @@ static const FOperatorSignature* GetOperatorSignature(EOperator Op)
 		S[UO_Sqrt]					= UnaryArithmetic;
 		S[UO_Tan] 					= UnaryArithmetic;
 		S[UO_Truncate] 				= UnaryArithmetic;
-		S[BO_GreaterThan] 			= BinaryComparison;
-		S[BO_GreaterThanOrEquals]	= BinaryComparison;
-		S[BO_LessThan]				= BinaryComparison;
-		S[BO_LessThanOrEquals]		= BinaryComparison;
-		S[BO_Equals]				= BinaryComparison;
-		S[BO_NotEquals]				= BinaryComparison;
+		S[BO_GreaterThan] 			= BinaryArithmeticComparison;
+		S[BO_GreaterThanOrEquals]	= BinaryArithmeticComparison;
+		S[BO_LessThan]				= BinaryArithmeticComparison;
+		S[BO_LessThanOrEquals]		= BinaryArithmeticComparison;
+		S[BO_Equals]				= { { OPTF_CastToCommon, OPTF_CastToCommon }, ORT_BooleanWithFirstArgumentDimensions };
+		S[BO_NotEquals]				= { { OPTF_CastToCommon, OPTF_CastToCommon }, ORT_BooleanWithFirstArgumentDimensions };
 		S[BO_Add] 					= BinaryArithmetic;
 		S[BO_Subtract] 				= BinaryArithmetic;
 		S[BO_Multiply] 				= BinaryArithmetic;
@@ -2000,6 +2037,329 @@ const FType* FEmitter::TryGetCommonType(const FType* A, const FType* B)
 	return FPrimitiveType::Get(ScalarKind, NumRows, 1);
 }
 
+FValue* FEmitter::EmitPartialDerivative(FValue* Value, EDerivativeAxis Axis)
+{
+	// Any operation on poison arguments is a poison.
+	if (Value->IsPoison())
+	{
+		return FPoison::Get();
+	}
+
+	// Differentiation is only valid on primitive types.
+	const MIR::FPrimitiveType* ValuePrimitiveType = Value->Type->AsPrimitive();
+	if (!ValuePrimitiveType || !ScalarKindIsAnyFloat(ValuePrimitiveType->ScalarKind))
+	{
+		Errorf(TEXT("Trying to differentiate a value of type `%s` is invalid. Expected a float type."), Value->Type->GetSpelling().GetData());
+		return FPoison::Get();
+	}
+
+	// Make the hardware derivative instruction.
+	FHardwarePartialDerivative HwDerivativeProto = MakePrototype<FHardwarePartialDerivative>(Value->Type);
+	HwDerivativeProto.Arg = Value;
+	HwDerivativeProto.Axis = Axis;
+	FValue* HwDerivative = EmitPrototype(this, HwDerivativeProto);
+
+	// Compute the analytical derivative for stages that don't support hardware derivatives.
+	FValue* AnalyticalDerivative = EmitAnalyticalPartialDerivative(Value, Axis);
+
+	// Emit the stage switch instruction so that hardware derivatives are used on stages
+	// that support it and analytical deriatives in the other stages.
+	FValue* StageValues[NumStages] = {};
+	for (int i = 0; i < NumStages; ++i)
+	{
+		StageValues[i] = (i == Stage_Pixel) ? HwDerivative : AnalyticalDerivative;
+	}
+	return EmitStageSwitch(Value->Type, StageValues);
+}
+
+static FValue* DifferentiateExternalInput(FEmitter& Emitter, FExternalInput* ExternalInput, EDerivativeAxis Axis)
+{
+	// Texture coordinate external inputs have their own matching DDX/DDY inputs.
+	int32 TexCoordIndex = ExternalInputToTexCoordIndex(ExternalInput->Id);
+	if (TexCoordIndex < TexCoordMaxNum)
+	{
+		EExternalInput PartialDerivativeExternalInput = (Axis == EDerivativeAxis::X) 
+			? EExternalInput((int32)EExternalInput::TexCoord0_Ddx + TexCoordIndex)
+			: EExternalInput((int32)EExternalInput::TexCoord0_Ddy + TexCoordIndex);
+
+		return Emitter.EmitExternalInput(PartialDerivativeExternalInput);
+	}
+
+	// All other inputs are assumed constant.
+	const MIR::FPrimitiveType* PrimType = static_cast<const MIR::FPrimitiveType*>(ExternalInput->Type);
+	return Emitter.EmitConstruct(PrimType, Emitter.EmitConstantScalarZero(PrimType->ScalarKind));
+}
+
+static FValue* DifferentiateOperator(FEmitter& Emitter, FOperator* Op, EDerivativeAxis Axis)
+{
+	const MIR::FPrimitiveType* PrimType = static_cast<const MIR::FPrimitiveType*>(Op->Type);
+
+	// Considering an operator acting on f(x), g(x) and h(x) arguments (e.g. "f(x) + g(x)"),
+	// calculate base terms and 
+	FValue* F = Op->AArg;
+	FValue* G = Op->BArg;
+	FValue* H = Op->CArg;
+	FValue* dF = F && !F->Type->IsBoolean() ? Emitter.EmitAnalyticalPartialDerivative(F, Axis) : nullptr; // Note: select's first argument is a boolean, avoid making the derivative then
+	FValue* dG = G ? Emitter.EmitAnalyticalPartialDerivative(G, Axis) : nullptr;
+	FValue* dH = H ? Emitter.EmitAnalyticalPartialDerivative(H, Axis) : nullptr;
+
+	// Convenience local functions as multiplications and division operations are common in derivatives.
+	auto Zero = [&Emitter, PrimType] () { return Emitter.EmitConstantScalarZero(PrimType->ScalarKind); };
+	auto One = [&Emitter, PrimType] () { return Emitter.EmitConstantScalarOne(PrimType->ScalarKind); };
+	auto Constant = [&Emitter, PrimType] (TFloat FromFloat) { return Emitter.EmitConstantScalar(PrimType->ScalarKind, FromFloat); };
+	auto Neg = [&Emitter] (FValue* A) { return Emitter.EmitOperator(UO_Negate, A); };
+	auto Add = [&Emitter] (FValue* A, FValue* B) { return Emitter.EmitOperator(BO_Add, A, B); };
+	auto Sub = [&Emitter] (FValue* A, FValue* B) { return Emitter.EmitOperator(BO_Subtract, A, B); };
+	auto Mul = [&Emitter] (FValue* A, FValue* B) { return Emitter.EmitOperator(BO_Multiply, A, B); };
+	auto Div = [&Emitter] (FValue* A, FValue* B) { return Emitter.EmitOperator(BO_Divide, A, B); };
+	auto Sqrt = [&Emitter] (FValue* A) { return Emitter.EmitOperator(UO_Sqrt, A); };
+	
+	// Some constants
+	constexpr TFloat Ln2 = 0.69314718055994530941723212145818;
+	constexpr TFloat Ln10 = 2.3025850929940456840179914546844;
+	
+	switch (Op->Op)
+	{
+		// d/dx -f(x) = -f'(x)
+		case UO_Negate:
+			return Neg(dF);
+
+		// d/dx |f(x)| = f(x) f'(x) / |f(x)|
+		case UO_Abs:
+			return Div(Mul(F, dF), Op);
+
+		// d/dx arccos(f(x)) = -1 / sqrt(1 - f(x)^2) * f'(x)
+		case UO_ACos:
+			return Neg(Div(dF, Sqrt(Sub(One(), Mul(F, F)))));
+
+		// d/dx arcsin(f(x)) = 1 / sqrt(1 - f(x)^2) * f'(x)
+		case UO_ASin:
+			return Div(dF, Sqrt(Sub(One(), Mul(F, F))));
+
+		// d/dx arctan(f(x)) = 1 / (1 + f(x)^2) * f'(x)
+		case UO_ATan:
+			return Div(dF, Add(One(), Mul(F, F)));
+
+		// d/dx cos(f(x)) = -sin(f(x)) * f'(x)
+		case UO_Cos:
+			return Neg(Mul(Emitter.EmitOperator(UO_Sin, F), dF));
+
+		// d/dx e^f(x) = e^f(x) * f'(x)
+		case UO_Exponential:
+			return Mul(Op, dF);
+
+		// d/dx 2^f(x) = ln(2) * 2^f(x) * f'(x)
+		case UO_Exponential2:
+			return Mul(Mul(Constant(Ln2), Op), dF);
+
+		// d/dx frac(f(x)) = f'(x), since frac(x) = x - floor(x)
+		case UO_Frac:
+			return dF;
+
+		// d/dx |f(x)| (length in vector case) = f(x) f'(x) / |f(x)|
+		case UO_Length:
+			return Div(Mul(F, dF), Op);
+
+		// d/dx log(f(x)) = 1 / f(x) * f'(x)
+		case UO_Logarithm:
+			return Div(dF, F);
+
+		// d/dx log2(f(x)) = 1 / (f(x) * ln(2)) * f'(x)
+		case UO_Logarithm2:
+			return Div(dF, Mul(F, Constant(Ln2)));
+
+		// d/dx log10(f(x)) = 1 / (f(x) * ln(10)) * f'(x)
+		case UO_Logarithm10:
+			return Div(dF, Mul(F, Constant(Ln10)));
+
+		// d/dx saturate(f(x)) = f'(x) if f(x) is inside (0-1) range, 0 otherwise
+		case UO_Saturate:
+			return Emitter.EmitOperator(TO_Select,
+						Emitter.EmitOperator(BO_And,
+							Emitter.EmitOperator(BO_LessThan, Zero(), F), // 0 < f(x)
+							Emitter.EmitOperator(BO_LessThan, F, One())), // f(x) < 1
+						dF, Zero());
+
+		// d/dx sin(f(x)) = cos(f(x)) * f'(x)
+		case UO_Sin:
+			return Mul(Emitter.EmitOperator(UO_Cos, F), dF);
+
+		// d/dx sqrt(f(x)) = 1 / (2 * sqrt(f(x))) * f'(x)
+		case UO_Sqrt:
+			return Div(dF, Mul(Emitter.EmitConstantScalar(PrimType->ScalarKind, 2), Sqrt(F)));
+
+		// d/dx tan(f(x)) = 1 / cos^2(f(x)) * f'(x)
+		case UO_Tan:
+		{
+			FValue* CosVal = Emitter.EmitOperator(UO_Cos, F);
+			return Div(dF, Mul(CosVal, CosVal));
+		}
+
+		// These functions are piecewise constant, that is mostly constant with some
+		// discontinuities. We assume they're always constant, as they're not differentiable
+		// at the discontinuities.
+		case UO_Ceil:
+		case UO_Floor:
+		case UO_Round:
+		case UO_Truncate:
+			return Zero();
+
+		// d/dx (f(x) + g(x)) = f'(x) + g'(x)
+		case BO_Add:
+			return Add(dF, dG);
+
+		// d/dx (f(x) - g(x)) = f'(x) - g'(x)
+		case BO_Subtract:
+			return Sub(dF, dG);
+
+		// d/dx (f(x) * g(x)) = f'(x) * g(x) + f(x) * g'(x)
+		case BO_Multiply:
+			return Add(Mul(dF, G), Mul(F, dG));
+
+		// d/dx (f(x) / g(x)) = (f'(x) * g(x) - f(x) * g'(x)) / g(x)^2
+		case BO_Divide:
+			return Div(Sub(Mul(dF, G), Mul(F, dG)), Mul(G, G));
+
+		// fmod(f(x), g(x)) = f(x) - g(x) * floor(f(x) / g(x)).
+		// Thus:
+		//     d/dx fmod(f(x), g(x)) = f'(x) - g(x) * floor(f(x) / g(x))
+		// since `floor` is piecewise constant.
+		case BO_Fmod:
+			return Sub(dF, Mul(dG, Emitter.EmitOperator(UO_Floor, Div(F, G))));
+
+		// d/dx max(f(x), g(x)) = f'(x) if f(x) > g(x), else g'(x)
+		case BO_Max:
+			return Emitter.EmitOperator(TO_Select,
+						Emitter.EmitOperator(BO_GreaterThan, F, G),
+						dF,
+						dG);
+
+		// d/dx min(f(x), g(x)) = f'(x) if f(x) < g(x), else g'(x)
+		case BO_Min:
+			return Emitter.EmitOperator(TO_Select,
+						Emitter.EmitOperator(BO_LessThan, F, G),
+						dF,
+						dG);
+
+		// The multiplication rule applies for the dot product too.
+		// d/dx (f(x) • g(x)) = f'(x) • g(x) + f(x) • g'(x)
+		case BO_Dot:
+			return Add(Emitter.EmitOperator(BO_Dot, dF, G), 
+					   Emitter.EmitOperator(BO_Dot, F, dG));
+
+		
+		// The multiplication rule applies for the cross product too.
+		// d/dx (f(x) × g(x)) = f'(x) × g(x) + f(x) × g'(x)
+		case BO_Cross:
+			return Add(Emitter.EmitOperator(BO_Cross, dF, G), 
+					   Emitter.EmitOperator(BO_Cross, F, dG));
+
+		// clamp(x, min, max) (F=x, min=G, max=H)
+		// The derivative is defined when x is between min and max (f'(x)). At and outside
+		// bounds, the clamp result is constant and thus the derivative is zero.
+		case TO_Clamp:
+			return Emitter.EmitOperator(TO_Select,
+				Emitter.EmitOperator(BO_And,
+					Emitter.EmitOperator(BO_LessThan, G, F),
+					Emitter.EmitOperator(BO_LessThan, F, H)),
+				dF, Zero());
+
+
+		// lerp(a, b, t) = a + t * (b - a)
+		// d/dx lerp(f(x), g(x), h(x)) = f'(x) + d/dx (h(x) * ((g(x) - f(x)))
+		// d/dx (h(x) * ((g(x) - f(x))) = h'(x) * ((g(x) - f(x))) + h(x) * (g'(x) - f'(x))
+		case TO_Lerp:
+			return Add(dF, 
+					   Add(Mul(dH, Sub(G, F)), 
+						   Mul(H, Sub(dG, dF))));
+
+		// d/dx select(cond, g(x), h(x)) ≈ select(cond, g'(x), h'(x))
+		case TO_Select:
+			return Emitter.EmitOperator(TO_Select, dG, dH);
+
+
+		// smoothstep(f(x), g(x), h(x)) = 3 z^2 - 2 z^3  with z = saturate((h - f) / (g - f))
+		case TO_Smoothstep:
+		{
+			FValue* Z  = Emitter.EmitOperator(UO_Saturate, Div(Sub(H, F), Sub(G, F)));
+			FValue* dZ = Emitter.EmitAnalyticalPartialDerivative(Z, Axis);
+			// d/dx 3 z(x)^2 - 2 z(x)^3 = 6 * z(x) * z'(x) - 6 * z(x)^2 * z'(x) = 6 * (z(x) - z(x)^2) * z'(x)
+			return Mul(dZ, Mul(Constant(6), Sub(Z, Mul(Z, Z))));
+		}
+
+		default:
+			UE_MIR_UNREACHABLE();
+	}
+}
+
+FValue* FEmitter::EmitAnalyticalPartialDerivative(FValue* Value, EDerivativeAxis Axis)
+{
+	// Any operation on poison arguments is a poison.
+	if (Value->IsPoison())
+	{
+		return FPoison::Get();
+	}
+
+	// Differentiation is only valid on primitive types.
+	const MIR::FPrimitiveType* ValuePrimitiveType = Value->Type->AsPrimitive();
+	if (!ValuePrimitiveType || !ScalarKindIsAnyFloat(ValuePrimitiveType->ScalarKind))
+	{
+		Errorf(TEXT("Trying to differentiate a value of type `%s` is invalid. Expected a float type."), Value->Type->GetSpelling().GetData());
+		return FPoison::Get();
+	}
+
+	switch (Value->Kind)
+	{
+		case MIR::VK_ExternalInput:
+			return DifferentiateExternalInput(*this, static_cast<FExternalInput*>(Value), Axis);
+
+		case MIR::VK_Dimensional:
+		{
+			TConstArrayView<FValue*> ValueComponents = Value->As<FDimensional>()->GetComponents();
+			FDimensional* Derivative = NewDimensionalValue(this, ValuePrimitiveType);
+			TArrayView<FValue*> DerivativeComponents = Derivative->GetMutableComponents();
+			for (int32 i = 0; i < ValueComponents.Num(); ++i)
+			{
+				DerivativeComponents[i] = EmitAnalyticalPartialDerivative(ValueComponents[i], Axis);
+			}
+			return Derivative;
+		}
+
+		case MIR::VK_Operator:
+			return DifferentiateOperator(*this, static_cast<FOperator*>(Value), Axis);
+		
+		case MIR::VK_Branch:
+		{
+			FBranch* Branch = Value->As<FBranch>();
+			return EmitBranch(Branch->ConditionArg,
+				EmitAnalyticalPartialDerivative(Branch->TrueArg, Axis),
+				EmitAnalyticalPartialDerivative(Branch->FalseArg, Axis));
+		}
+
+		case MIR::VK_Subscript:
+		{
+			FSubscript* Subscript = Value->As<FSubscript>();
+			return EmitSubscript(EmitAnalyticalPartialDerivative(Subscript->Arg, Axis), Subscript->Index);
+		}
+
+		case MIR::VK_Cast:
+		{
+			FCast* Cast = Value->As<FCast>();
+			return EmitConstruct(Cast->Type, EmitAnalyticalPartialDerivative(Cast->Arg, Axis));
+		}
+
+		// These values are uniform (constant), thus their value is always zero.
+		case MIR::VK_Constant:
+		case MIR::VK_UniformParameter:
+		case MIR::VK_TextureSample:
+		case MIR::VK_InlineHLSL:
+			return EmitConstruct(ValuePrimitiveType, EmitConstantScalarZero(ValuePrimitiveType->ScalarKind));
+
+		default:
+			UE_MIR_UNREACHABLE();
+	}
+}
+
 const FType* FEmitter::GetCommonType(const FType* A, const FType* B)
 {
 	if (const FType* CommonType = TryGetCommonType(A, B))
diff --git a/Engine/Source/Runtime/Engine/Private/Materials/MaterialIRModuleBuilder.cpp b/Engine/Source/Runtime/Engine/Private/Materials/MaterialIRModuleBuilder.cpp
--- a/Engine/Source/Runtime/Engine/Private/Materials/MaterialIRModuleBuilder.cpp
+++ b/Engine/Source/Runtime/Engine/Private/Materials/MaterialIRModuleBuilder.cpp
@@ -361,7 +361,7 @@ struct FMaterialIRModuleBuilderImpl
 				}
 
 				// Before analyzing this value, make sure all used values are analyzed first.
-				for (MIR::FValue* Use : Value->GetUses())
+				for (MIR::FValue* Use : Value->GetUsesForStage(CurrentStage))
 				{
 					if (Use && !Use->IsAnalyzed(CurrentStage))
 					{
@@ -379,7 +379,7 @@ struct FMaterialIRModuleBuilderImpl
 				ValueStack.Pop();
 
 				// Go through each use instruction and increment its counter of users (this instruction).
-				for (MIR::FValue* Use : Value->GetUses())
+				for (MIR::FValue* Use : Value->GetUsesForStage(CurrentStage))
 				{
 					// If this used value is an instruction, update its counter of users (in current stage).
 					if (MIR::FInstruction* UseInstr = MIR::AsInstruction(Use))
@@ -447,7 +447,7 @@ struct FMaterialIRModuleBuilderImpl
 				Instr->Next[StageIndex] = Instr->Block[StageIndex] ->Instructions;
 				Instr->Block[StageIndex]->Instructions = Instr;
 
-				TConstArrayView<MIR::FValue*> Uses = Instr->GetUses();
+				TConstArrayView<MIR::FValue*> Uses = Instr->GetUsesForStage((MIR::EStage)StageIndex);
 				for (int32 UseIndex = 0; UseIndex < Uses.Num(); ++UseIndex)
 				{
 					MIR::FValue* Use = Uses[UseIndex];
@@ -489,21 +489,6 @@ struct FMaterialIRModuleBuilderImpl
 
 	void Step_Finalize()
 	{
-		/* Produce the module statistics */
-		for (int TexCoordIndex = 0; TexCoordIndex < MIR::TexCoordMaxNum; ++TexCoordIndex)
-		{
-			MIR::EExternalInput TexCoordInput = MIR::TexCoordIndexToExternalInput(TexCoordIndex);
-			if (Module->Statistics.ExternalInputUsedMask[MIR::Stage_Vertex][(int)TexCoordInput])
-			{
-				Module->Statistics.NumVertexTexCoords = TexCoordIndex + 1;
-			}
-			if (Module->Statistics.ExternalInputUsedMask[MIR::Stage_Pixel][(int)TexCoordInput])
-			{
-				Module->Statistics.NumPixelTexCoords = TexCoordIndex + 1;
-			}
-		}
-
-		/* Configure the compilation output */
 		FMaterialCompilationOutput& CompilationOutput = Module->CompilationOutput;
 		CompilationOutput.NumUsedUVScalars = Module->Statistics.NumPixelTexCoords * 2;
 	}
diff --git a/Engine/Source/Runtime/Engine/Private/Materials/MaterialIRToHLSLTranslator.cpp b/Engine/Source/Runtime/Engine/Private/Materials/MaterialIRToHLSLTranslator.cpp
--- a/Engine/Source/Runtime/Engine/Private/Materials/MaterialIRToHLSLTranslator.cpp
+++ b/Engine/Source/Runtime/Engine/Private/Materials/MaterialIRToHLSLTranslator.cpp
@@ -214,9 +214,6 @@ static bool IsOperatorInfix(MIR::EOperator Op)
 		case MIR::BO_Fmod:
 		case MIR::BO_Max:
 		case MIR::BO_Min:
-		case MIR::BO_Or:
-		case MIR::BO_And:
-		case MIR::BO_Xor:
 			return true;
 		default: return false;
 	}
@@ -230,15 +227,15 @@ struct FTranslator : FMaterialIRToHLSLTranslation
 	FHLSLPrinter Printer;
 	FString PixelAttributesHLSL;
 	FString WorldPositionOffsetHLSL;
-	FString EvaluateNormalMaterialAttributeHLSL;
-	FString EvaluateOtherMaterialAttributesHLSL;
+	FString EvaluateNormalMaterialAttributeHLSL[MIR::NumStages];
+	FString EvaluateOtherMaterialAttributesHLSL[MIR::NumStages];
 
 	void GeneratePixelAttributesHLSL()
 	{
 		for (int32 PropertyIndex = 0; PropertyIndex < MP_MAX; ++PropertyIndex)
 		{
 			EMaterialProperty Property = (EMaterialProperty)PropertyIndex;
-			if (!MIR::Internal::IsMaterialPropertyEnabled(Property) || MIR::GetMaterialOutputStage(Property) != MIR::Stage_Pixel)
+			if (!MIR::Internal::IsMaterialPropertyEnabled(Property) || !MIR::MaterialOutputEvaluatesInStage(Property, MIR::Stage_Pixel))
 			{
 				continue;
 			}
@@ -263,18 +260,18 @@ struct FTranslator : FMaterialIRToHLSLTranslation
 		WorldPositionOffsetHLSL = MoveTemp(Printer.Buffer);
 	}
 
-	void GenerateFragmentStageHLSL()
+	void GenerateOtherStageHLSL(MIR::EStage Stage)
 	{
-		BeginStage(MIR::Stage_Pixel);
+		BeginStage(Stage);
 
 		LowerBlock(Module->GetRootBlock(CurrentStage));
 
-			Printer << TEXT("PixelMaterialInputs.FrontMaterial = GetInitialisedSubstrateData()") << EndOfStatement;
-			Printer << TEXT("PixelMaterialInputs.Subsurface = 0") << EndOfStatement;
+		Printer << TEXT("PixelMaterialInputs.FrontMaterial = GetInitialisedSubstrateData()") << EndOfStatement;
+		Printer << TEXT("PixelMaterialInputs.Subsurface = 0") << EndOfStatement;
 
-		EvaluateOtherMaterialAttributesHLSL = MoveTemp(Printer.Buffer);
+		EvaluateOtherMaterialAttributesHLSL[CurrentStage] = MoveTemp(Printer.Buffer);
 	}
-
+	
 	void BeginStage(MIR::EStage Stage)
 	{
 		Printer = {};
@@ -323,7 +320,7 @@ struct FTranslator : FMaterialIRToHLSLTranslation
 			const MIR::FSetMaterialOutput* SetMaterialOutput = MIR::As<MIR::FSetMaterialOutput>(Instr);
 			if (SetMaterialOutput && SetMaterialOutput->Property == MP_Normal) 
 			{
-				EvaluateNormalMaterialAttributeHLSL = MoveTemp(Printer.Buffer);
+				EvaluateNormalMaterialAttributeHLSL[CurrentStage] = MoveTemp(Printer.Buffer);
 			}
         }
 
@@ -371,9 +368,9 @@ struct FTranslator : FMaterialIRToHLSLTranslation
 			case MIR::VK_Subscript: LowerSubscript(static_cast<MIR::FSubscript*>(Instr)); break;
 			case MIR::VK_TextureSample: LowerTextureSample(static_cast<MIR::FTextureSample*>(Instr)); break;
 			case MIR::VK_InlineHLSL: LowerInlineHLSL(static_cast<const MIR::FInlineHLSL*>(Instr)); break;
-
-			default:
-				UE_MIR_UNREACHABLE();
+			case MIR::VK_StageSwitch: LowerStageSwitch(static_cast<const MIR::FStageSwitch*>(Instr)); break;
+			case MIR::VK_HardwarePartialDerivative: LowerHardwarePartialDerivative(static_cast<const MIR::FHardwarePartialDerivative*>(Instr)); break;
+			default: UE_MIR_UNREACHABLE();
 		}
 	
 		return NoOp;
@@ -563,42 +560,38 @@ struct FTranslator : FMaterialIRToHLSLTranslation
 	{
 		if (IsOperatorInfix(Operator->Op))
 		{
-			Printer << "(";
-
-			LowerValue(Operator->AArg);
-
 			const TCHAR* OpString;
 			switch (Operator->Op)
 			{
-				case MIR::BO_Add: OpString = TEXT(" + "); break;
-				case MIR::BO_Divide: OpString = TEXT(" / "); break;
+				case MIR::BO_Add: OpString = TEXT("+"); break;
+				case MIR::BO_Divide: OpString = TEXT("/"); break;
 				case MIR::BO_Equals: OpString = TEXT("=="); break;
 				case MIR::BO_GreaterThan: OpString = TEXT(">"); break;
 				case MIR::BO_GreaterThanOrEquals: OpString = TEXT(">="); break;
 				case MIR::BO_LessThan: OpString = TEXT("<"); break;
 				case MIR::BO_LessThanOrEquals: OpString = TEXT("<="); break;
-				case MIR::BO_Multiply: OpString = TEXT(" * "); break;
+				case MIR::BO_Multiply: OpString = TEXT("*"); break;
 				case MIR::BO_NotEquals: OpString = TEXT("!="); break;
-				case MIR::BO_Subtract: OpString = TEXT(" - "); break;
+				case MIR::BO_Subtract: OpString = TEXT("-"); break;
 				default: UE_MIR_UNREACHABLE();
 			}
-			Printer << OpString;
 
-			LowerValue(Operator->BArg);
-
-			Printer << ")";
+			Printer << "(" << LowerValue(Operator->AArg) << " " << OpString << " " << LowerValue(Operator->BArg) << ")";
 		}
 		else
 		{
 			const TCHAR* OpString;
 			switch (Operator->Op)
 			{
-				case MIR::TO_Clamp: OpString = TEXT("clamp"); break;
+				case MIR::BO_And: OpString = TEXT("and"); break;
 				case MIR::BO_Cross: OpString = TEXT("cross"); break;
 				case MIR::BO_Dot: OpString = TEXT("dot"); break;
 				case MIR::BO_Fmod: OpString = TEXT("fmod"); break;
 				case MIR::BO_Max: OpString = TEXT("max"); break;
 				case MIR::BO_Min: OpString = TEXT("min"); break;
+				case MIR::BO_Or: OpString = TEXT("or"); break;
+				case MIR::BO_Xor: OpString = TEXT("xor"); break;
+				case MIR::TO_Clamp: OpString = TEXT("clamp"); break;
 				case MIR::TO_Lerp: OpString = TEXT("lerp"); break;
 				case MIR::TO_Select: OpString = TEXT("select"); break;
 				case MIR::TO_Smoothstep: OpString = TEXT("smoothstep"); break;
@@ -641,6 +634,7 @@ struct FTranslator : FMaterialIRToHLSLTranslation
 				check(MIR::IsTernaryOperator(Operator->Op));
 				Printer << ", " << LowerValue(Operator->CArg);
 			}
+
 			Printer << ")";
 		}
 	}
@@ -680,7 +674,7 @@ struct FTranslator : FMaterialIRToHLSLTranslation
 		}
 	}
 
-	void LowerTextureSample(MIR::FTextureSample* TextureSample)
+	void LowerTextureSample(const MIR::FTextureSample* TextureSample)
 	{
 		bool bSamplerNeedsBrackets = LowerSamplerType(TextureSample->SamplerType);
 		if (bSamplerNeedsBrackets)
@@ -811,6 +805,16 @@ struct FTranslator : FMaterialIRToHLSLTranslation
 		return NoOp;
 	}
 	
+	void LowerStageSwitch(const MIR::FStageSwitch* StageSwitch)
+	{
+		LowerValue(StageSwitch->Args[CurrentStage]);
+	}
+
+	void LowerHardwarePartialDerivative(const MIR::FHardwarePartialDerivative* StageSwitch)
+	{
+		Printer << (StageSwitch->Axis == MIR::EDerivativeAxis::X ? "ddx(" : "ddy(") << LowerValue(StageSwitch->Arg) << ")";
+	}
+
 	/* Finalization */
 
 	void GenerateTemplateStringParameters(TMap<FString, FString>& Params)
@@ -831,12 +835,12 @@ struct FTranslator : FMaterialIRToHLSLTranslation
 
 		// "Normal" is treated in a special way because the rest of the attributes may lead back to reading it.
 		// Therefore, in the way MaterialTemplate.ush is structured, it needs to be evaluated before other attributes.
-		Params.Add(TEXT("calc_pixel_material_inputs_analytic_derivatives_normal"), EvaluateNormalMaterialAttributeHLSL);
-		Params.Add(TEXT("calc_pixel_material_inputs_normal"), EvaluateNormalMaterialAttributeHLSL);
+		Params.Add(TEXT("calc_pixel_material_inputs_analytic_derivatives_normal"), EvaluateNormalMaterialAttributeHLSL[MIR::Stage_Compute]);
+		Params.Add(TEXT("calc_pixel_material_inputs_normal"), EvaluateNormalMaterialAttributeHLSL[MIR::Stage_Pixel]);
 
 		// Then the other other attributes.
-		Params.Add(TEXT("calc_pixel_material_inputs_analytic_derivatives_other_inputs"), EvaluateOtherMaterialAttributesHLSL);
-		Params.Add(TEXT("calc_pixel_material_inputs_other_inputs"), EvaluateOtherMaterialAttributesHLSL);
+		Params.Add(TEXT("calc_pixel_material_inputs_analytic_derivatives_other_inputs"), EvaluateOtherMaterialAttributesHLSL[MIR::Stage_Compute]);
+		Params.Add(TEXT("calc_pixel_material_inputs_other_inputs"), EvaluateOtherMaterialAttributesHLSL[MIR::Stage_Pixel]);
 		
 		// MaterialAttributes
 		TArray<FGuid> OrderedVisibleAttributes = FMaterialAttributeDefinitionMap::GetOrderedVisibleAttributeList();
@@ -1058,7 +1062,8 @@ void FMaterialIRToHLSLTranslation::Run(TMap<FString, FString>& OutParameters, FS
 	FTranslator Translator{ *this };
 	Translator.GeneratePixelAttributesHLSL();
 	Translator.GenerateVertexStageHLSL();
-	Translator.GenerateFragmentStageHLSL();
+	Translator.GenerateOtherStageHLSL(MIR::Stage_Pixel);
+	Translator.GenerateOtherStageHLSL(MIR::Stage_Compute);
 	Translator.GenerateTemplateStringParameters(OutParameters);
 	Translator.GetShaderCompilerEnvironment(OutEnvironment);
 }
diff --git a/Engine/Source/Runtime/Engine/Private/Materials/MaterialIRValueAnalyzer.cpp b/Engine/Source/Runtime/Engine/Private/Materials/MaterialIRValueAnalyzer.cpp
--- a/Engine/Source/Runtime/Engine/Private/Materials/MaterialIRValueAnalyzer.cpp
+++ b/Engine/Source/Runtime/Engine/Private/Materials/MaterialIRValueAnalyzer.cpp
@@ -19,8 +19,16 @@
 
 static void AnalyzeExternalInput(FMaterialIRValueAnalyzer& Analyzer, MIR::FExternalInput* ExternalInput)
 {
-	 Analyzer.Module->GetStatistics().ExternalInputUsedMask[MIR::Stage_Vertex][(int)ExternalInput->Id] = true;
-	 Analyzer.Module->GetStatistics().ExternalInputUsedMask[MIR::Stage_Pixel][(int)ExternalInput->Id] = true;
+	FMaterialIRModule::FStatistics& Statistics = Analyzer.Module->GetStatistics();
+	Statistics.ExternalInputUsedMask[MIR::Stage_Vertex][(int)ExternalInput->Id] = true;
+	Statistics.ExternalInputUsedMask[MIR::Stage_Pixel][(int)ExternalInput->Id] = true;
+
+	int32 TexCoordIndex = MIR::ExternalInputToTexCoordIndex(ExternalInput->Id);
+	if (TexCoordIndex < MIR::TexCoordMaxNum)
+	{
+		Statistics.NumVertexTexCoords = FMath::Max(Statistics.NumVertexTexCoords, TexCoordIndex + 1);
+		Statistics.NumPixelTexCoords = FMath::Max(Statistics.NumPixelTexCoords, TexCoordIndex + 1);
+	}
 }
 
 static EMaterialShaderFrequency MapToMaterialShaderFrequencyOrAny(MIR::EStage Stage)
diff --git a/Engine/Source/Runtime/Engine/Public/Materials/MaterialExpressionDDX.h b/Engine/Source/Runtime/Engine/Public/Materials/MaterialExpressionDDX.h
--- a/Engine/Source/Runtime/Engine/Public/Materials/MaterialExpressionDDX.h
+++ b/Engine/Source/Runtime/Engine/Public/Materials/MaterialExpressionDDX.h
@@ -21,6 +21,7 @@ class UMaterialExpressionDDX : public UMaterialExpression
 
 	//~ Begin UMaterialExpression Interface
 #if WITH_EDITOR
+	virtual void Build(MIR::FEmitter& Emitter) override;
 	virtual int32 Compile(class FMaterialCompiler* Compiler, int32 OutputIndex) override;
 	virtual void GetCaption(TArray<FString>& OutCaptions) const override;
 #endif
diff --git a/Engine/Source/Runtime/Engine/Public/Materials/MaterialExpressionDDY.h b/Engine/Source/Runtime/Engine/Public/Materials/MaterialExpressionDDY.h
--- a/Engine/Source/Runtime/Engine/Public/Materials/MaterialExpressionDDY.h
+++ b/Engine/Source/Runtime/Engine/Public/Materials/MaterialExpressionDDY.h
@@ -21,6 +21,7 @@ class UMaterialExpressionDDY : public UMaterialExpression
 
 	//~ Begin UMaterialExpression Interface
 #if WITH_EDITOR
+	virtual void Build(MIR::FEmitter& Emitter) override;
 	virtual int32 Compile(class FMaterialCompiler* Compiler, int32 OutputIndex) override;
 	virtual void GetCaption(TArray<FString>& OutCaptions) const override;
 #endif
diff --git a/Engine/Source/Runtime/Engine/Public/Materials/MaterialIR.h b/Engine/Source/Runtime/Engine/Public/Materials/MaterialIR.h
--- a/Engine/Source/Runtime/Engine/Public/Materials/MaterialIR.h
+++ b/Engine/Source/Runtime/Engine/Public/Materials/MaterialIR.h
@@ -19,13 +19,19 @@ enum EStage
 {
 	Stage_Vertex,
 	Stage_Pixel,
+	Stage_Compute,
 	NumStages
 };
 
 /**
- * @brief Returns the stage in which specified material property shall be evaluated.
+ * Returns the string representation of given stage.
  */
-EStage GetMaterialOutputStage(EMaterialProperty Property);
+const TCHAR* LexToString(EStage Stage);
+
+/**
+ * @brief Returns whether specified property is evaluated in specified stage.
+ */
+bool MaterialOutputEvaluatesInStage(EMaterialProperty Property, EStage Stage);
 
 /**
  * @brief A collection of bit flags for a specific Value instance.
@@ -40,14 +46,17 @@ enum class EValueFlags : uint8
 	/** Value has been analyzed for the pixel stage. */
 	AnalyzedForStagePixel = 1 << 1,
 
+	/** Value has been analyzed for the compute stage. */
+	AnalyzedForStageCompute = 1 << 2,
+
 	/** Mask that includes all "analyzed in stage" flags. */
-	AnalyzedInAnyStageMask = AnalyzedForStageVertex | AnalyzedForStagePixel,
+	AnalyzedInAnyStageMask = AnalyzedForStageVertex | AnalyzedForStagePixel | AnalyzedForStageCompute,
 
 	/** Whether to skip substution of abstract tokens like "<PREV>". Only true for custom nodes. */
-	SubstituteTagsInInlineHLSL = 1 << 2,
+	SubstituteTagsInInlineHLSL = 1 << 3,
 
 	/** Whether 'FInlineHLSL::Code' field is used. Otherwise, 'ExternalCodeDeclaration' is used. */
-	HasDynamicHLSLCode = 1 << 3, 
+	HasDynamicHLSLCode = 1 << 4, 
 };
 
 ENUM_CLASS_FLAGS(EValueFlags);
@@ -99,6 +108,8 @@ enum EValueKind
 	VK_Cast,
 	VK_TextureSample,
 	VK_InlineHLSL,
+	VK_StageSwitch,
+	VK_HardwarePartialDerivative,
 
 	VK_InstructionEnd,
 };
@@ -172,13 +183,20 @@ struct FValue
 	bool IsA(EValueKind InKind) const;
 
 	/**
-	 * Gets the immutable array of this value uses.
+	 * Gets the immutable array of all this value uses.
 	 * 
 	 * An use is another value referenced by this one (e.g. the operands of a binary expression).
 	 * 
 	 * @return The immutable array of uses.
 	 */
 	TConstArrayView<FValue*> GetUses() const;
+	
+	/**
+	 * Gets the immutable array of this value uses filtered for a specific stage stage.
+	 * 
+	 * @return The immutable array of uses.
+	 */
+	TConstArrayView<FValue*> GetUsesForStage(MIR::EStage Stage) const;
 
 	/** @return Whether this value is poison. */
 	bool IsPoison() const;
@@ -357,6 +375,9 @@ struct FConstant : TValue<VK_Constant>
 	}
 };
 
+/**
+ * @brief Enumeration of all supported material external inputs.
+ */
 enum class EExternalInput
 {
 	None,
@@ -393,8 +414,11 @@ enum class EExternalInput
 	Count,
 };
 
+/** The number of texcoord external input groups (uv, ddx, ddy). */
+static constexpr int32 TexCoordGroups = 3;
+
 /** Maximum number of supported texture coordinates. */
-static constexpr int TexCoordMaxNum = 8;
+static constexpr int32 TexCoordMaxNum = 8;
 
 /** @return The string representation of given given external input. */
 const TCHAR* LexToString(EExternalInput Input);
@@ -403,11 +427,14 @@ const TCHAR* LexToString(EExternalInput Input);
 const FType* GetExternalInputType(EExternalInput Id);
 
 /** Converts the given texture coordinate index to its corresponding external input mapping. */
-MIR::EExternalInput TexCoordIndexToExternalInput(int TexCoordIndex);
+MIR::EExternalInput TexCoordIndexToExternalInput(int32 TexCoordIndex);
 
 /** @return Whether the given external is a texture coordinate. */
 bool IsExternalInputTexCoord(EExternalInput Id);
 
+/** @return The index of the specified texture coordinates external input.  */
+int32 ExternalInputToTexCoordIndex(EExternalInput Id);
+
 /** @return Whether the given external is a texture coordinate ddx. */
 bool IsExternalInputTexCoordDdx(EExternalInput Id);
 
@@ -415,7 +442,7 @@ bool IsExternalInputTexCoordDdx(EExternalInput Id);
 bool IsExternalInputTexCoordDdy(EExternalInput Id);
 
 /**
- * It represents an external input (like texture coordinates) value.
+ * @brief It represents an external input (like texture coordinates) value.
  */
 struct FExternalInput : TValue<VK_ExternalInput>
 {
@@ -423,7 +450,7 @@ struct FExternalInput : TValue<VK_ExternalInput>
 };
 
 /**
- * It represents a texture object value. 
+ * @brief It represents a texture object value. 
  */
 struct FTextureObject : TValue<VK_TextureObject>
 {
@@ -439,7 +466,7 @@ struct FTextureObject : TValue<VK_TextureObject>
 };
 
 /**
- * It represents a material uniform parameter.
+ * @brief It represents a material uniform parameter.
  */
 struct FUniformParameter: TValue<VK_UniformParameter>
 {
@@ -454,11 +481,10 @@ struct FUniformParameter: TValue<VK_UniformParameter>
 	uint32 Analysis_UniformParameterIndex;
 };
 
-
-/* Instructions */
+/*------------------------------------ Instructions ------------------------------------*/
 
 /**
- * A block of instructions.
+ * @brief A block of instructions.
  * 
  * A block is a sequence of instructions in order of execution. Blocks are organized in a
  * tree-like structure, used to model blocks nested inside other blocks.
@@ -548,7 +574,7 @@ struct TDimensional : FDimensional
 };
 
 /**
- * This instruction sets material attribute (i.e., "BaseColor") to its value.
+ * @brief Instruction that sets material attribute (i.e., "BaseColor") to its value.
  */
 struct FSetMaterialOutput : TInstruction<VK_SetMaterialOutput, 1>
 {
@@ -775,5 +801,40 @@ struct FTextureSample : TInstruction<VK_TextureSample, 5>
 	EMaterialSamplerType SamplerType;
 };
 
+/**
+ * @brief Utility value for selecting a different value based on the execution stage.
+ */
+struct FStageSwitch : TInstruction<VK_StageSwitch, 1>
+{
+	FValue* Args[NumStages];
+};
+
+/**
+ * @brief Specifies the axis for computing screen-space derivatives.
+ *
+ * This enum is used to indicate whether a partial derivative should be taken 
+ * along the X or Y screen-space direction, corresponding to HLSL's ddx and ddy
+ * functions.
+ */
+enum class EDerivativeAxis
+{
+	X, // Corresponding to d/dx
+	Y, // Corresponding to d/dy
+};
+
+/**
+ * @brief Instruction that maps to hardware ddx()/ddy().
+ * 
+ * @remark This is only available in stages that support hardware derivatives (e.g. pixel shader)
+ */
+struct FHardwarePartialDerivative : TInstruction<VK_HardwarePartialDerivative, 1>
+{
+	/** The value argument of ddx()/ddy() */
+	FValue* Arg;
+
+	/** The direction of partial derivative */
+	EDerivativeAxis Axis;
+};
+
 } // namespace MIR
 #endif
diff --git a/Engine/Source/Runtime/Engine/Public/Materials/MaterialIREmitter.h b/Engine/Source/Runtime/Engine/Public/Materials/MaterialIREmitter.h
--- a/Engine/Source/Runtime/Engine/Public/Materials/MaterialIREmitter.h
+++ b/Engine/Source/Runtime/Engine/Public/Materials/MaterialIREmitter.h
@@ -140,11 +140,12 @@ public:
 	FValue* EmitConstantTrue();
 	FValue* EmitConstantFalse();
 	FValue* EmitConstantBoolScalar(bool InX);
+	FValue* EmitConstantIntScalar(TInteger InX);
 	FValue* EmitConstantFloatScalar(TFloat InX);
+	FValue* EmitConstantScalar(EScalarKind Kind, TFloat FromFloat);
 	FValue* EmitConstantFloat2(const FVector2f& InValue);
 	FValue* EmitConstantFloat3(const FVector3f& InValue);
 	FValue* EmitConstantFloat4(const FVector4f& InValue);
-	FValue* EmitConstantIntScalar(TInteger InX);
 	FValue* EmitConstantInt2(const FIntVector2& InValue);
 	FValue* EmitConstantInt3(const FIntVector3& InValue);
 	FValue* EmitConstantInt4(const FIntVector4& InValue);
@@ -161,6 +162,16 @@ public:
 	FValue* EmitParameter(FName Name, FMaterialParameterMetadata& Metadata, EMaterialSamplerType SamplerType = SAMPLERTYPE_Color);
 	FValue* EmitCustomPrimitiveData(uint32 PrimitiveDataIndex);
 
+	/**
+	 * Emits a StageSwitch value.
+	 * 
+	 * @param Type the type of the emitted value.
+	 * @param ValuePerStage the array of values for each MIR::EStage entry.
+	 * 
+	 * @remark ValuePerStage must have no more than MIR::NumStages elements.
+	 */
+	FValue* EmitStageSwitch(const FType* Type, TConstArrayView<FValue*> ValuePerStage);
+
 	/* Instructions */
 
 	FSetMaterialOutput* EmitSetMaterialOutput(EMaterialProperty InProperty, FValue* Arg);
@@ -187,11 +198,59 @@ public:
 	FValue* EmitInlineHLSL(const FType* Type, FString Code, EValueFlags ValueFlags = EValueFlags::None, EGraphProperties UsedGraphProperties = EGraphProperties::None);
 	FValue* EmitInlineHLSL(const FMaterialExternalCodeDeclaration* InExternalCodeDeclaration, EGraphProperties UsedGraphProperties = EGraphProperties::None);
 
+	/**
+	 * @brief Emits the partial derivative of given value.
+	 * 
+	 * This function differentiates the specified value with respect to the given axis.
+	 * If differentiation is not possible, it returns a poison value to indicate an error.
+	 * 
+	 * This API can be used regardless of the stages value will execute in. Hardware partial
+	 * derivatives will be used by default where available. If not available, the analytical
+	 * derivative will be computed instead.
+	 * 
+	 * @param Value The value to differentiate.
+	 * @param Axis The axis along which to compute the partial derivative (X for ddx, Y for ddy).
+	 * 
+	 * @return A new FValue* representing the partial derivative, or a poison value if invalid.
+	 */
+	FValue* EmitPartialDerivative(FValue* Value, EDerivativeAxis Axis);
+
+	/**
+	 * @brief Computes and emits the analytical partial derivative of a give value.
+	 * 
+	 * @param Value The value to differentiate.
+	 * @param Axis The axis along which to compute the partial derivative (X for ddx, Y for ddy).
+	 * 
+	 * @return A new FValue* representing the analytical derivative, or a poison value if invalid.
+	 */
+	FValue* EmitAnalyticalPartialDerivative(FValue* Value, EDerivativeAxis Axis);
+
 	/* Types */
 
+	/**
+	 * Tries to find a common type between A and B.
+	 * 
+	 * The trivial case is when A and B are the same type, in which case that type is
+	 * the common type. When A and B are primitive, the common type is defined as the type
+	 * with the highest dimensions between the two and with a scalar kind that follows the
+	 * progression:
+	 * 
+	 * 	bool -> int -> float
+	 * 
+	 * As an example, the common type between a float3 and a bool4 will be a float4.
+	 * 
+	 * Note that this may cause a loss of data when casting from int to float.
+	 */
 	const FType* TryGetCommonType(const FType* A, const FType* B);
+
+	/**
+	 * Gets the common type between A and B and reports an error if no such type exists.
+	 * 
+	 * @remarks See `TryGetCommonType()`.
+	 */
 	const FType* GetCommonType(const FType* A, const FType* B);
 
+
 	/* Error reporting */
 
 	bool CurrentExpressionHasErrors() const { return bCurrentExpressionHasErrors; }
diff --git a/Engine/Source/Runtime/Engine/Public/Materials/MaterialIRModule.h b/Engine/Source/Runtime/Engine/Public/Materials/MaterialIRModule.h
--- a/Engine/Source/Runtime/Engine/Public/Materials/MaterialIRModule.h
+++ b/Engine/Source/Runtime/Engine/Public/Materials/MaterialIRModule.h
@@ -36,10 +36,10 @@ public:
 		TBitArray<> ExternalInputUsedMask[MIR::NumStages]; 
 
 		// Number of vertex texture coordinates used.
-		int NumVertexTexCoords;
+		int32 NumVertexTexCoords;
 
 		// Number of pixel texture coordinates used.
-		int NumPixelTexCoords;
+		int32 NumPixelTexCoords;
 	};
 
 public:
