Fixed bug in the new material translator concerning texture coordinate index handling.

#rnx

diff --git a/Engine/Source/Runtime/Engine/Private/Materials/MaterialIR.cpp b/Engine/Source/Runtime/Engine/Private/Materials/MaterialIR.cpp
--- a/Engine/Source/Runtime/Engine/Private/Materials/MaterialIR.cpp
+++ b/Engine/Source/Runtime/Engine/Private/Materials/MaterialIR.cpp
@@ -450,11 +450,6 @@ MIR::EExternalInput TexCoordIndexToExternalInput(int32 TexCoordIndex)
 	return (MIR::EExternalInput)((int)MIR::EExternalInput::TexCoord0 + TexCoordIndex);
 }
 
-bool IsExternalInputTexCoord(EExternalInput Id)
-{
-	return Id >= EExternalInput::TexCoord0 && Id <= EExternalInput::TexCoord7;
-}
-
 int32 ExternalInputToTexCoordIndex(EExternalInput Id)
 {
 	int32 IntId = int32(Id) - 1;
@@ -465,6 +460,11 @@ int32 ExternalInputToTexCoordIndex(EExternalInput Id)
 	return IntId % TexCoordMaxNum;
 }
 
+bool IsExternalInputTexCoord(EExternalInput Id)
+{
+	return Id >= EExternalInput::TexCoord0 && Id <= EExternalInput::TexCoord7;
+}
+
 bool IsExternalInputTexCoordDdx(EExternalInput Id)
 {
 	return Id >= EExternalInput::TexCoord0_Ddx && Id <= EExternalInput::TexCoord7_Ddx;
@@ -475,6 +475,11 @@ bool IsExternalInputTexCoordDdy(EExternalInput Id)
 	return Id >= EExternalInput::TexCoord0_Ddy && Id <= EExternalInput::TexCoord7_Ddy;
 }
 
+bool IsExternalInputTexCoordOrPartialDerivative(EExternalInput Id)
+{
+	return Id >= EExternalInput::TexCoord0 && Id <= EExternalInput::TexCoord7_Ddy;
+}
+
 FBlock* FBlock::FindCommonParentWith(MIR::FBlock* Other)
 {
 	FBlock* A = this;
diff --git a/Engine/Source/Runtime/Engine/Private/Materials/MaterialIREmitter.cpp b/Engine/Source/Runtime/Engine/Private/Materials/MaterialIREmitter.cpp
--- a/Engine/Source/Runtime/Engine/Private/Materials/MaterialIREmitter.cpp
+++ b/Engine/Source/Runtime/Engine/Private/Materials/MaterialIREmitter.cpp
@@ -2075,9 +2075,10 @@ FValue* FEmitter::EmitPartialDerivative(FValue* Value, EDerivativeAxis Axis)
 static FValue* DifferentiateExternalInput(FEmitter& Emitter, FExternalInput* ExternalInput, EDerivativeAxis Axis)
 {
 	// Texture coordinate external inputs have their own matching DDX/DDY inputs.
-	int32 TexCoordIndex = ExternalInputToTexCoordIndex(ExternalInput->Id);
-	if (TexCoordIndex < TexCoordMaxNum)
+	if (IsExternalInputTexCoordDdx(ExternalInput->Id))
 	{
+		int32 TexCoordIndex = ExternalInputToTexCoordIndex(ExternalInput->Id);
+
 		EExternalInput PartialDerivativeExternalInput = (Axis == EDerivativeAxis::X) 
 			? EExternalInput((int32)EExternalInput::TexCoord0_Ddx + TexCoordIndex)
 			: EExternalInput((int32)EExternalInput::TexCoord0_Ddy + TexCoordIndex);
diff --git a/Engine/Source/Runtime/Engine/Private/Materials/MaterialIRValueAnalyzer.cpp b/Engine/Source/Runtime/Engine/Private/Materials/MaterialIRValueAnalyzer.cpp
--- a/Engine/Source/Runtime/Engine/Private/Materials/MaterialIRValueAnalyzer.cpp
+++ b/Engine/Source/Runtime/Engine/Private/Materials/MaterialIRValueAnalyzer.cpp
@@ -23,9 +23,9 @@ static void AnalyzeExternalInput(FMaterialIRValueAnalyzer& Analyzer, MIR::FExter
 	Statistics.ExternalInputUsedMask[MIR::Stage_Vertex][(int)ExternalInput->Id] = true;
 	Statistics.ExternalInputUsedMask[MIR::Stage_Pixel][(int)ExternalInput->Id] = true;
 
-	int32 TexCoordIndex = MIR::ExternalInputToTexCoordIndex(ExternalInput->Id);
-	if (TexCoordIndex < MIR::TexCoordMaxNum)
+	if (MIR::IsExternalInputTexCoordOrPartialDerivative(ExternalInput->Id))
 	{
+		int32 TexCoordIndex = MIR::ExternalInputToTexCoordIndex(ExternalInput->Id);
 		Statistics.NumVertexTexCoords = FMath::Max(Statistics.NumVertexTexCoords, TexCoordIndex + 1);
 		Statistics.NumPixelTexCoords = FMath::Max(Statistics.NumPixelTexCoords, TexCoordIndex + 1);
 	}
@@ -35,8 +35,9 @@ static EMaterialShaderFrequency MapToMaterialShaderFrequencyOrAny(MIR::EStage St
 {
 	switch (Stage)
 	{
-	case MIR::EStage::Stage_Vertex: return EMaterialShaderFrequency::Vertex;
-	case MIR::EStage::Stage_Pixel: return EMaterialShaderFrequency::Pixel;
+		case MIR::EStage::Stage_Vertex: return EMaterialShaderFrequency::Vertex;
+		case MIR::EStage::Stage_Pixel: return EMaterialShaderFrequency::Pixel;
+		case MIR::EStage::Stage_Compute: return EMaterialShaderFrequency::Compute;
 	}
 	return EMaterialShaderFrequency::Any;
 }
diff --git a/Engine/Source/Runtime/Engine/Public/Materials/MaterialIR.h b/Engine/Source/Runtime/Engine/Public/Materials/MaterialIR.h
--- a/Engine/Source/Runtime/Engine/Public/Materials/MaterialIR.h
+++ b/Engine/Source/Runtime/Engine/Public/Materials/MaterialIR.h
@@ -429,18 +429,21 @@ const FType* GetExternalInputType(EExternalInput Id);
 /** Converts the given texture coordinate index to its corresponding external input mapping. */
 MIR::EExternalInput TexCoordIndexToExternalInput(int32 TexCoordIndex);
 
-/** @return Whether the given external is a texture coordinate. */
-bool IsExternalInputTexCoord(EExternalInput Id);
-
 /** @return The index of the specified texture coordinates external input.  */
 int32 ExternalInputToTexCoordIndex(EExternalInput Id);
 
+/** @return Whether the given external is a texture coordinate. */
+bool IsExternalInputTexCoord(EExternalInput Id);
+
 /** @return Whether the given external is a texture coordinate ddx. */
 bool IsExternalInputTexCoordDdx(EExternalInput Id);
 
 /** @return Whether the given external is a texture coordinate ddy. */
 bool IsExternalInputTexCoordDdy(EExternalInput Id);
 
+/** @return Whether the given external is a texture coordinate, a ddx or ddy. */
+bool IsExternalInputTexCoordOrPartialDerivative(EExternalInput Id);
+
 /**
  * @brief It represents an external input (like texture coordinates) value.
  */
