RemoteObjects:  Actors can now be destroyed while loaned out, and when the loan gets returned the owning server will handle the DestroyActor properly.

#rnx

diff --git a/Engine/Source/Runtime/Engine/Private/Actor.cpp b/Engine/Source/Runtime/Engine/Private/Actor.cpp
--- a/Engine/Source/Runtime/Engine/Private/Actor.cpp
+++ b/Engine/Source/Runtime/Engine/Private/Actor.cpp
@@ -1041,26 +1041,37 @@ void AActor::Serialize(FArchive& Ar)
 #if UE_WITH_REMOTE_OBJECT_HANDLE
 	if (Ar.IsMigratingRemoteObjects())
 	{
-		bool bLocalActorInitialized = IsActorInitialized();
-
-		Ar << CreationTime;
-		Ar << bLocalActorInitialized;
+		uint8 MigrationFlags = 0;
 
 		// TODO: This a workaround for not currently supporting native delegates in object migration
 		// we are treating LifeSpan as a property of Actor rather than solving TimerManager's delegates generically.
 		// Calling SetLifeSpan here is also not ideal because it's virtual which means it can do "anything" while in Serialize.
+		// We could fix it with FUObjectAnnotationSparse
 		float CurrentLifespan = 0.0f; 
 		if (Ar.IsSaving())
 		{
 			CurrentLifespan = GetLifeSpan();
 			SetLifeSpan(0.0f);
+
+			MigrationFlags |= (static_cast<uint8>(ActorHasBegunPlay) & 3) << 0;
+			MigrationFlags |= (bHasFinishedSpawning & 1) << 2;
+			MigrationFlags |= (bActorInitialized & 1) << 3;
+			MigrationFlags |= (bTickFunctionsRegistered & 1) << 4;
 		}
+
+		Ar << CreationTime;
 		Ar << CurrentLifespan;
+		Ar << MigrationFlags;
+
 		if (Ar.IsLoading())
 		{
 			// When migrating actors are in a state after BeginPlay has been called but we need to call SetLifeSpan to initialize local timer
 			SetLifeSpan(CurrentLifespan);
-			bActorInitialized = bLocalActorInitialized;
+
+			ActorHasBegunPlay = static_cast<EActorBeginPlayState>(MigrationFlags & 3);
+			bHasFinishedSpawning = (MigrationFlags >> 2) & 0x1;
+			bActorInitialized = (MigrationFlags >> 3) & 0x1;
+			bTickFunctionsRegistered = (MigrationFlags >> 4) & 0x1;
 		}
 	}
 #endif // UE_WITH_REMOTE_OBJECT_HANDLE
@@ -1235,6 +1246,21 @@ static USceneComponent* FixupNativeActorComponents(AActor* Actor);
 
 void AActor::PostMigrate(const FUObjectMigrationContext& MigrationContext)
 {
+	// We want to be able to store "annotations" (metadata) about this Actor while we loan it out, so that
+	// when we have the loan returned to us, we can check what has changed...
+	struct FActorLoanedState
+	{
+		bool bActorWasDestroyed = false;
+		bool bTickFunctionsRegistered = false;
+		bool operator==(const FActorLoanedState&) const = default;
+		bool IsDefault() const
+		{
+			static const FActorLoanedState DefaultLoanedState;
+			return *this == DefaultLoanedState;
+		}
+	};
+	static FUObjectAnnotationSparse<FActorLoanedState,/*bAutoRemove=*/true> RemoteObjects_ActorLoanedStates;
+
 	Super::PostMigrate(MigrationContext);
 
 	// We are using this as a proxy for having ownership until we sort out those mechanics
@@ -1252,6 +1278,11 @@ void AActor::PostMigrate(const FUObjectMigrationContext& MigrationContext)
 		ensureMsgf(MigrationContext.RemoteServerId != LocalServerId, TEXT("%hs: Sending Object %s from %s to %s but destination is the same as the Local Server Id."), __func__, *GetName(), *LocalServerId.ToString(), *MigrationContext.RemoteServerId.ToString());
 
 		const bool bKeepOwnership = (MigrationSendType == EObjectMigrationSendType::Loan);
+		if (bKeepOwnership)
+		{
+			RemoteObjects_ActorLoanedStates.AddAnnotation(this, { .bActorWasDestroyed = bActorIsBeingDestroyed, .bTickFunctionsRegistered = bTickFunctionsRegistered });
+		}
+
 		if (!bKeepOwnership || !bHasOwnership)
 		{
 			// Remove the actor from the world on this server as it is about to be migrated to another server.
@@ -1293,7 +1324,10 @@ void AActor::PostMigrate(const FUObjectMigrationContext& MigrationContext)
 	else if (MigrationContext.MigrationSide == EObjectMigrationSide::Receive)
 	{
 		EObjectMigrationRecvType MigrationRecvType = MigrationContext.GetObjectMigrationRecvType(this);
-		UE_LOG(LogActor, Verbose, TEXT("%hs: Recv %s (LocalSrvId: %s. DstSrvId: %s. HasOwnership: %d. MigrationType: %s)"), __func__, *GetName(), *LocalServerId.ToString(), *MigrationContext.RemoteServerId.ToString(), bHasOwnership, ToString(MigrationRecvType));
+		UE_LOG(LogActor, Verbose, TEXT("%hs: Recv %s (LocalSrvId: %s. FromSrvId: %s. HasOwnership: %d. MigrationType: %s)"), __func__, *GetName(), *LocalServerId.ToString(), *MigrationContext.RemoteServerId.ToString(), bHasOwnership, ToString(MigrationRecvType));
+
+		// Build up the list of Components we have again (even if we previously owned this object, it could have been mutated)
+		ResetOwnedComponents();
 
 		// We can get rid of this check further down the road...
 		// Set the actor's world transform if it has a native rootcomponent.
@@ -1301,45 +1335,62 @@ void AActor::PostMigrate(const FUObjectMigrationContext& MigrationContext)
 		ensureMsgf(SceneRootComponent == RootComponent, TEXT("%hs: We were expecting RootComponent (%s) to be equal to FixupNativeActorComponents (%s). Transform fix-up code would be required to handle this case."), __func__, *GetNameSafe(RootComponent), *GetNameSafe(SceneRootComponent));
 
 		// Add the actor that was migrated from another server to this server's world
-		constexpr bool bOnlyAddIfUnique = true;
-		GetLevel()->TryAddActorToList(this, bOnlyAddIfUnique);
-
-		// If we don't have Ownership, these components were never registered
-		if (!bHasOwnership)
+		if (!bActorIsBeingDestroyed)
 		{
-			// This will build up the list of OwnedComponents (even if they are not currently registered)
-			ResetOwnedComponents();
+			constexpr bool bOnlyAddIfUnique = true;
+			GetLevel()->TryAddActorToList(this, bOnlyAddIfUnique);
 		}
 
-		bHasFinishedSpawning = true;
-
 		// If we're being passed ownership, register us with the global systems
 		if (MigrationRecvType == EObjectMigrationRecvType::AssignedOwnership)
 		{
-			BuildReplicatedComponentsInfo();
+			ensureMsgf(!bActorIsBeingDestroyed, TEXT("We do not want to be assigned ownership of Actors that have already been destroyed. We do not yet know what to do with it."));
 
 			#if UE_WITH_IRIS
 			BeginReplication();
 			#endif // UE_WITH_IRIS
 
-			// Each component will be in charge of registering itself in its own PostMigrate
+			// Each component will be in charge of registering itself in its own PostMigrate, but we must register
+			// the Actor (which we'll do be saying it's currently the opposite of the deserialized value, then flipping it)
 			constexpr bool bRegisterComponents = false;
-			RegisterAllActorTickFunctions(true, bRegisterComponents);
-
-			// We are going to assume that the Physics queries on the previous owning server had taken multi-server physics into account
-			// thus we do not actually need to check for Initial Overlaps anymore
-			// ForEachSceneComponent: SceneComponent->ClearSkipUpdateOverlaps();
-			// UpdateInitialOverlaps(/*bFromLevelStreaming*/ false);
+			bTickFunctionsRegistered = !bTickFunctionsRegistered;
+			RegisterAllActorTickFunctions(!bTickFunctionsRegistered, bRegisterComponents);
 
 			World->AddNetworkActor(this);
 		}
+		else if (MigrationRecvType == EObjectMigrationRecvType::ReturnedLoan)
+		{
+			FActorLoanedState LoanedState = RemoteObjects_ActorLoanedStates.GetAndRemoveAnnotation(this);
+			if (bActorIsBeingDestroyed)
+			{
+				// If we weren't destroyed when we loaned it out, we should destroy it now that we've received it back...
+				if (!LoanedState.bActorWasDestroyed)
+				{
+					ensureMsgf(IsPendingKillPending(), TEXT("%hs: We thought %s was destroyed while on loan, but our received state does not match a destroyed Actor"), __func__, *GetName());
+					UE_LOG(LogActor, Warning, TEXT("%hs: Destroying %s back on owning server (it was destroyed while loaned-out)"), __func__, *GetName());
+
+					// Destroy it again on the owning server, as it needs to be deregistered from our systems
+					bActorIsBeingDestroyed = false;
+					ClearGarbage();
+					World->DestroyActor(this);
+				}
+			}
+			else if (bTickFunctionsRegistered != LoanedState.bTickFunctionsRegistered)
+			{
+				// Remote server changed the tick registration
+				constexpr bool bRegisterComponents = false;
+				bTickFunctionsRegistered = LoanedState.bTickFunctionsRegistered;
+				RegisterAllActorTickFunctions(!bTickFunctionsRegistered, bRegisterComponents);
+			}
+
+		}
 		else if (MigrationRecvType == EObjectMigrationRecvType::Borrowed)
 		{
 			// Make sure that all tick functions have been deregistered if we don't have ownership.
+			bool bTickFunctionsRegisteredOnOwningServer = bTickFunctionsRegistered;
 			RegisterAllActorTickFunctions(false, false);
+			bTickFunctionsRegistered = bTickFunctionsRegisteredOnOwningServer;
 		}
-
-		ActorHasBegunPlay = EActorBeginPlayState::HasBegunPlay;
 	}
 }
 #endif // UE_WITH_REMOTE_OBJECT_HANDLE
diff --git a/Engine/Source/Runtime/Engine/Private/Components/PrimitiveComponent.cpp b/Engine/Source/Runtime/Engine/Private/Components/PrimitiveComponent.cpp
--- a/Engine/Source/Runtime/Engine/Private/Components/PrimitiveComponent.cpp
+++ b/Engine/Source/Runtime/Engine/Private/Components/PrimitiveComponent.cpp
@@ -1268,7 +1268,7 @@ void UPrimitiveComponent::Serialize(FArchive& Ar)
 		}
 		
 		Ar << MoveIgnoreMask;
-		
+
 		MigratePhysicsState(Ar);
 	}
 #endif
@@ -1317,7 +1317,8 @@ void UPrimitiveComponent::MigratePhysicsState(FArchive& Ar)
 			return;
 		}
 
-		if (ObjMigrationRecvType == EObjectMigrationRecvType::ReturnedLoan)
+		const bool bFinishedDestroy = HasAnyFlags(EObjectFlags::RF_MirroredGarbage);
+		if (ObjMigrationRecvType == EObjectMigrationRecvType::ReturnedLoan || bFinishedDestroy)
 		{
 			return;
 		}
