Minor bugfix in new material translator that fixes implementation of DifferentiateExternalInput().

#rnx

diff --git a/Engine/Source/Runtime/Engine/Private/Materials/MaterialIREmitter.cpp b/Engine/Source/Runtime/Engine/Private/Materials/MaterialIREmitter.cpp
--- a/Engine/Source/Runtime/Engine/Private/Materials/MaterialIREmitter.cpp
+++ b/Engine/Source/Runtime/Engine/Private/Materials/MaterialIREmitter.cpp
@@ -2075,7 +2075,7 @@ FValue* FEmitter::EmitPartialDerivative(FValue* Value, EDerivativeAxis Axis)
 static FValue* DifferentiateExternalInput(FEmitter& Emitter, FExternalInput* ExternalInput, EDerivativeAxis Axis)
 {
 	// Texture coordinate external inputs have their own matching DDX/DDY inputs.
-	if (IsExternalInputTexCoordDdx(ExternalInput->Id))
+	if (IsExternalInputTexCoord(ExternalInput->Id))
 	{
 		int32 TexCoordIndex = ExternalInputToTexCoordIndex(ExternalInput->Id);
 
