[AssetManager] Fix download params are not passed in for the StreamableManager request

#rnx

diff --git a/Engine/Source/Runtime/Engine/Private/AssetManager.cpp b/Engine/Source/Runtime/Engine/Private/AssetManager.cpp
--- a/Engine/Source/Runtime/Engine/Private/AssetManager.cpp
+++ b/Engine/Source/Runtime/Engine/Private/AssetManager.cpp
@@ -1951,7 +1951,7 @@ TSharedPtr<FStreamableHandle> UAssetManager::ChangeBundleStateForPrimaryAssets(
 
 				FAssetManagerLoadParams AssetListLoadParams;
 				AssetListLoadParams.Priority = LoadParams.Priority;
-				if (AssetListLoadParams.DownloadParams)
+				if (LoadParams.DownloadParams)
 				{
 					AssetListLoadParams.DownloadParams = CopyTemp(*LoadParams.DownloadParams);
 				}
