Integrating TangentBezierSpline with SplineComponent

diff --git a/Engine/Source/Editor/ComponentVisualizers/Private/SplineComponentVisualizer.cpp b/Engine/Source/Editor/ComponentVisualizers/Private/SplineComponentVisualizer.cpp
--- a/Engine/Source/Editor/ComponentVisualizers/Private/SplineComponentVisualizer.cpp
+++ b/Engine/Source/Editor/ComponentVisualizers/Private/SplineComponentVisualizer.cpp
@@ -31,9 +31,15 @@
 #include "CanvasItem.h"
 #include "CanvasTypes.h"
 #include "Math/UnrealMathUtility.h"
+#include "Widgets/Input/SButton.h"
+#include "Widgets/Input/SComboButton.h"
+#include "Widgets/Input/SEditableTextBox.h"
+#include "Widgets/Input/SNumericEntryBox.h"
+#include "Widgets/Views/SListView.h"
 
 IMPLEMENT_HIT_PROXY(HSplineVisProxy, HComponentVisProxy);
 IMPLEMENT_HIT_PROXY(HSplineKeyProxy, HSplineVisProxy);
+IMPLEMENT_HIT_PROXY(HSplineAttributeKeyProxy, HSplineVisProxy);
 IMPLEMENT_HIT_PROXY(HSplineSegmentProxy, HSplineVisProxy);
 IMPLEMENT_HIT_PROXY(HSplineTangentHandleProxy, HSplineVisProxy);
 
@@ -133,11 +139,17 @@ void USplineComponentVisualizerSelectionState::GetVerifiedSelectedTangentHandle(
 void USplineComponentVisualizerSelectionState::Reset()
 {
 	SplinePropertyPath = FComponentPropertyPath();
-	SelectedKeys.Reset();
-	LastKeyIndexSelected = INDEX_NONE;
+	ClearSelectedKeys();
 	CachedRotation = FQuat();
 	ClearSelectedSegmentIndex();
 	ClearSelectedTangentHandle();
+	ClearSelectedAttribute();
+}
+
+void USplineComponentVisualizerSelectionState::ClearSelectedKeys()
+{
+	SelectedKeys.Reset();
+	LastKeyIndexSelected = INDEX_NONE;
 }
 
 void USplineComponentVisualizerSelectionState::ClearSelectedSegmentIndex()
@@ -209,6 +221,8 @@ public:
 		UI_COMMAND(DiscontinuousSpline, "Allow Discontinuous Splines", "Whether the visualization allows Arrive and Leave tangents to be set separately.", EUserInterfaceActionType::ToggleButton, FInputChord());
 		UI_COMMAND(ToggleClosedLoop, "Closed Loop", "Toggle the Closed Loop setting of the spline", EUserInterfaceActionType::ToggleButton, FInputChord());
 		UI_COMMAND(ResetToDefault, "Reset to Default", "Reset this spline to its archetype default.", EUserInterfaceActionType::Button, FInputChord());
+		UI_COMMAND(AddAttributeKey, "Add Attribute Here", "Add a new attribute value at the cursor location.", EUserInterfaceActionType::Button, FInputChord());
+		UI_COMMAND(DeleteAttributeKey, "Delete Attribute", "Delete the currently selected attribute.", EUserInterfaceActionType::Button, FInputChord());
 	}
 
 public:
@@ -322,6 +336,12 @@ public:
 
 	/** Reset this spline to its default */
 	TSharedPtr<FUICommandInfo> ResetToDefault;
+
+	/** Add attribute key */
+	TSharedPtr<FUICommandInfo> AddAttributeKey;
+
+	/** Delete attribute key */
+	TSharedPtr<FUICommandInfo> DeleteAttributeKey;
 };
 
 TWeakPtr<SWindow> FSplineComponentVisualizer::WeakExistingWindow;
@@ -547,6 +567,16 @@ void FSplineComponentVisualizer::OnRegister()
 		FExecuteAction::CreateSP(this, &FSplineComponentVisualizer::OnResetToDefault),
 		FCanExecuteAction::CreateSP(this, &FSplineComponentVisualizer::CanResetToDefault));
 
+	SplineComponentVisualizerActions->MapAction(
+		Commands.AddAttributeKey,
+		FExecuteAction::CreateSP(this, &FSplineComponentVisualizer::OnAddAttributeKey),
+		FCanExecuteAction::CreateSP(this, &FSplineComponentVisualizer::CanAddAttributeKey));
+	
+	SplineComponentVisualizerActions->MapAction(
+		Commands.DeleteAttributeKey,
+		FExecuteAction::CreateSP(this, &FSplineComponentVisualizer::OnDeleteAttributeKey),
+		FCanExecuteAction::CreateSP(this, &FSplineComponentVisualizer::CanDeleteAttributeKey));
+	
 	bool bAlign = false;
 	bool bUseLineTrace = false;
 	bool bUseBounds = false;
@@ -832,6 +862,21 @@ void FSplineComponentVisualizer::DrawVisualization(const UActorComponent* Compon
 			OldKeyRightVector = NewKeyRightVector;
 			OldKeyScale = NewKeyScale;
 		}
+
+		// Render the selected attribute channel's values.
+		if (FName SelectedAttributeName = SelectionState->GetSelectedAttributeName(); SelectedAttributeName != NAME_None)
+		{
+			for (int Idx = 0; Idx < SplineComp->GetNumberOfPropertyValues(SelectedAttributeName); ++Idx)
+			{
+				PDI->SetHitProxy(new HSplineAttributeKeyProxy(Component, Idx));
+				
+				float Parameter = SplineComp->GetFloatPropertyInputKeyAtIndex(Idx, SelectedAttributeName);
+				FVector PositionOnCurve = SplineComp->GetLocationAtSplineInputKey(Parameter, ESplineCoordinateSpace::World);
+				PDI->DrawPoint(PositionOnCurve, FLinearColor::Red, GrabHandleSize, SDPG_Foreground);
+			}
+					
+			PDI->SetHitProxy(NULL);
+		}
 	}
 }
 
@@ -844,35 +889,35 @@ void FSplineComponentVisualizer::DrawVisualizationHUD(const UActorComponent* Com
 
 		if (SplineComp == EditedSplineComp)
 		{
+			const FIntRect CanvasRect = Canvas->GetViewRect();
+
+			auto DrawText = [&](const FText& SnapHelpText, float YOffset = 50.f)
+			{
+				int32 XL;
+				int32 YL;
+				StringSize(GEngine->GetLargeFont(), XL, YL, *SnapHelpText.ToString());
+				const float DrawPositionX = FMath::FloorToFloat(static_cast<float>(CanvasRect.Min.X) + static_cast<float>(CanvasRect.Width() - XL) * 0.5f);
+				const float DrawPositionY = static_cast<float>(CanvasRect.Min.Y) + YOffset;
+				Canvas->DrawShadowedString(DrawPositionX, DrawPositionY, *SnapHelpText.ToString(), GEngine->GetLargeFont(), FLinearColor::Yellow);
+			};
+			
 			if (bIsSnappingToActor)
 			{
-				const FIntRect CanvasRect = Canvas->GetViewRect();
-
 				static const FText SnapToActorHelp = LOCTEXT("SplinePointSnapToActorMessage", "Snap to Actor: Use Ctrl-LMB to select actor to use as target.");
 				static const FText AlignToActorHelp = LOCTEXT("SplinePointAlignToActorMessage", "Snap Align to Actor: Use Ctrl-LMB to select actor to use as target.");
 				static const FText AlignPerpToActorHelp = LOCTEXT("SplinePointAlignPerpToActorMessage", "Snap Align Perpendicular to Actor: Use Ctrl-LMB to select actor to use as target.");
 
-				auto DisplaySnapToActorHelpText = [&](const FText& SnapHelpText)
-				{
-					int32 XL;
-					int32 YL;
-					StringSize(GEngine->GetLargeFont(), XL, YL, *SnapHelpText.ToString());
-					const float DrawPositionX = FMath::FloorToFloat(static_cast<float>(CanvasRect.Min.X) + static_cast<float>(CanvasRect.Width() - XL) * 0.5f);
-					const float DrawPositionY = static_cast<float>(CanvasRect.Min.Y) + 50.0f;
-					Canvas->DrawShadowedString(DrawPositionX, DrawPositionY, *SnapHelpText.ToString(), GEngine->GetLargeFont(), FLinearColor::Yellow);
-				};
-
 				if (SnapToActorMode == ESplineComponentSnapMode::Snap)
 				{
-					DisplaySnapToActorHelpText(SnapToActorHelp);
+					DrawText(SnapToActorHelp);
 				}
 				else if (SnapToActorMode == ESplineComponentSnapMode::AlignToTangent)
 				{
-					DisplaySnapToActorHelpText(AlignToActorHelp);
+					DrawText(AlignToActorHelp);
 				}
 				else
 				{
-					DisplaySnapToActorHelpText(AlignPerpToActorHelp);
+					DrawText(AlignPerpToActorHelp);
 				}
 			}
 		}
@@ -888,6 +933,12 @@ void FSplineComponentVisualizer::ChangeSelectionState(int32 Index, bool bIsCtrlH
 	check(SelectionState);
 	SelectionState->Modify();
 
+	if (Index != INDEX_NONE)
+	{
+		// disallow selecting keys and attribute keys at the same time
+		SelectionState->SetSelectedAttribute(INDEX_NONE, SelectionState->GetSelectedAttributeName());
+	}
+	
 	TSet<int32>& SelectedKeys = SelectionState->ModifySelectedKeys();
 	if (Index == INDEX_NONE)
 	{
@@ -1113,6 +1164,26 @@ bool FSplineComponentVisualizer::VisProxyHandleClick(FEditorViewportClient* InVi
 				SelectionState->SetSelectedTangentHandleType(KeyProxy->bArriveTangent ? ESelectedTangentHandle::Arrive : ESelectedTangentHandle::Leave);
 				SelectionState->SetCachedRotation(SplineComp->GetQuaternionAtSplinePoint(SelectionState->GetSelectedTangentHandle(), ESplineCoordinateSpace::World));
 
+				bVisProxyClickHandled = true;
+			}
+		}
+		else if (VisProxy->IsA(HSplineAttributeKeyProxy::StaticGetType()))
+		{
+			// Spline attribute key clicked
+			const FScopedTransaction Transaction(LOCTEXT("SelectAttributeValue", "Select Attribute Value"));
+
+			SelectionState->Modify();
+
+			ResetTempModes();
+
+			if (const USplineComponent* SplineComp = UpdateSelectedSplineComponent(VisProxy))
+			{
+				HSplineAttributeKeyProxy* AttributeKeyProxy = static_cast<HSplineAttributeKeyProxy*>(VisProxy);
+				
+				SelectionState->SetSelectedAttribute(AttributeKeyProxy->KeyIndex, SelectionState->GetSelectedAttributeName());
+				const float InputKey = SplineComp->GetFloatPropertyInputKeyAtIndex(AttributeKeyProxy->KeyIndex, SelectionState->GetSelectedAttributeName());
+				SelectionState->SetCachedRotation(SplineComp->GetQuaternionAtSplineInputKey(InputKey, ESplineCoordinateSpace::World));
+				
 				bVisProxyClickHandled = true;
 			}
 		}
@@ -1158,6 +1229,7 @@ bool FSplineComponentVisualizer::GetWidgetLocation(const FEditorViewportClient*
 		
 		const TSet<int32>& SelectedKeys = SelectionState->GetSelectedKeys();
 		int32 LastKeyIndexSelected = SelectionState->GetLastKeyIndexSelected();
+		int32 SelectedAttributeIndex = SelectionState->GetSelectedAttributeIndex();
 
 		int32 SelectedTangentHandle = SelectionState->GetSelectedTangentHandle();
 		ESelectedTangentHandle SelectedTangentHandleType = SelectionState->GetSelectedTangentHandleType();
@@ -1183,6 +1255,16 @@ bool FSplineComponentVisualizer::GetWidgetLocation(const FEditorViewportClient*
 
 			return true;
 		}
+		else if (SelectedAttributeIndex != INDEX_NONE)
+		{
+			// If we have a selected attribute value, use that.
+			if (SelectedAttributeIndex >= 0 && SelectedAttributeIndex < SplineComp->GetNumberOfPropertyValues(SelectionState->GetSelectedAttributeName()))
+			{
+				const float Param = SplineComp->GetFloatPropertyInputKeyAtIndex(SelectedAttributeIndex, SelectionState->GetSelectedAttributeName());
+				OutLocation = SplineComp->GetLocationAtSplineInputKey(Param, ESplineCoordinateSpace::World);
+				return true;
+			}
+		}
 		else if (LastKeyIndexSelected != INDEX_NONE)
 		{
 			// Otherwise use the last key index set
@@ -1304,7 +1386,11 @@ bool FSplineComponentVisualizer::HandleInputDelta(FEditorViewportClient* Viewpor
 			}
 		}
 
-		if (SelectionState->GetSelectedTangentHandle() != INDEX_NONE)
+		if (SelectionState->GetSelectedAttributeIndex() != INDEX_NONE)
+		{
+			bInputHandled = TransformSelectedAttribute(EPropertyChangeType::Interactive, DeltaTranslate);
+		}
+		else if (SelectionState->GetSelectedTangentHandle() != INDEX_NONE)
 		{
 			// Transform the tangent using an EPropertyChangeType::Interactive change. Later on, at the end of mouse tracking, a non-interactive change will be notified via void TrackingStopped :
 			bInputHandled = TransformSelectedTangent(EPropertyChangeType::Interactive, DeltaTranslate);
@@ -1496,6 +1582,45 @@ bool FSplineComponentVisualizer::TransformSelectedKeys(EPropertyChangeType::Type
 	return false;
 }
 
+bool FSplineComponentVisualizer::TransformSelectedAttribute(EPropertyChangeType::Type InPropertyChangeType, const FVector& InDeltaTranslate)
+{
+	USplineComponent* SplineComp = GetEditedSplineComponent();
+	if (SplineComp != nullptr)
+	{
+		check(SelectionState);
+		const int32 SelectedAttributeIndex = SelectionState->GetSelectedAttributeIndex();
+		const FName SelectedAttributeName = SelectionState->GetSelectedAttributeName();
+		check(SelectedAttributeIndex != INDEX_NONE);
+		check(SelectedAttributeName != NAME_None);
+
+		SplineComp->Modify();
+		
+		const float CurrentParam = SplineComp->GetFloatPropertyInputKeyAtIndex(SelectedAttributeIndex, SelectedAttributeName);
+		const FVector CurrentPosition = SplineComp->GetLocationAtSplineInputKey(CurrentParam, ESplineCoordinateSpace::World);
+		const FVector NewPosition = CurrentPosition + InDeltaTranslate;
+
+		const float NewParam = SplineComp->FindInputKeyClosestToWorldLocation(NewPosition);
+		const int32 NewIndex = SplineComp->SetFloatPropertyInputKeyAtIndex(SelectedAttributeIndex, NewParam, SelectedAttributeName);
+
+		if (NewIndex != SelectedAttributeIndex)
+		{
+			// Index was invalidated because we reordered the attribute storage, update selection and invalidate proxies
+			SelectionState->SetSelectedAttribute(NewIndex, SelectionState->GetSelectedAttributeName());
+			GEditor->RedrawAllViewports();
+		}
+		
+		SelectionState->SetCachedRotation(SplineComp->GetQuaternionAtSplineInputKey(NewParam, ESplineCoordinateSpace::World));
+
+		SplineComp->bSplineHasBeenEdited = true;
+		NotifyPropertyModified(SplineComp, SplineCurvesProperty, InPropertyChangeType);
+		GEditor->RedrawLevelEditingViewports(true);
+
+		return true;
+	}
+
+	return false;
+}
+
 bool FSplineComponentVisualizer::HandleInputKey(FEditorViewportClient* ViewportClient, FViewport* Viewport, FKey Key, EInputEvent Event)
 {
 	bool bHandled = false;
@@ -1515,15 +1640,26 @@ bool FSplineComponentVisualizer::HandleInputKey(FEditorViewportClient* ViewportC
 			check(SelectionState);
 
 			// Recache widget rotation
-			int32 Index = SelectionState->GetSelectedTangentHandle();
-			if (Index == INDEX_NONE)
+			
+			if (int32 AttributeIndex = SelectionState->GetSelectedAttributeIndex(); AttributeIndex != INDEX_NONE)
 			{
-				// If not set, fall back to last key index selected
-				Index = SelectionState->GetLastKeyIndexSelected();
+				const float Param = SplineComp->GetFloatPropertyInputKeyAtIndex(AttributeIndex, SelectionState->GetSelectedAttributeName());
+				
+				SelectionState->Modify();
+				SelectionState->SetCachedRotation(SplineComp->GetQuaternionAtSplineInputKey(Param, ESplineCoordinateSpace::World));
 			}
+			else
+			{
+				int32 Index = SelectionState->GetSelectedTangentHandle();
+				if (Index == INDEX_NONE)
+				{
+					// If not set, fall back to last key index selected
+					Index = SelectionState->GetLastKeyIndexSelected();
+				}
 
-			SelectionState->Modify();
-			SelectionState->SetCachedRotation(SplineComp->GetQuaternionAtSplinePoint(Index, ESplineCoordinateSpace::World));
+				SelectionState->Modify();
+				SelectionState->SetCachedRotation(SplineComp->GetQuaternionAtSplinePoint(Index, ESplineCoordinateSpace::World));
+			}
 		}
 
 		// Reset duplication on LMB release
@@ -3209,6 +3345,76 @@ bool FSplineComponentVisualizer::CanResetToDefault() const
     }
 }
 
+void FSplineComponentVisualizer::OnAddAttributeKey()
+{
+	const FScopedTransaction Transaction(LOCTEXT("AddSplinePoint", "Add Spline Point"));
+
+	ResetTempModes();
+
+	USplineComponent* SplineComp = GetEditedSplineComponent();
+	check(SplineComp != nullptr);
+	check(SelectionState);
+
+	SelectionState->Modify();
+
+	const float NearestParam = SplineComp->FindInputKeyClosestToWorldLocation(SelectionState->GetSelectedSplinePosition());
+	// const float 
+	int32 NewAttributeIndex = SplineComp->SetFloatPropertyAtSplineInputKey(NearestParam, 0.f, SelectionState->GetSelectedAttributeName());
+
+	// Update selection
+	SelectionState->SetSelectedSegmentIndex(INDEX_NONE);
+	SelectionState->SetSelectedAttribute(NewAttributeIndex, SelectionState->GetSelectedAttributeName());
+	SelectionState->SetSelectedSplinePosition(FVector::ZeroVector);
+	SelectionState->SetCachedRotation(SplineComp->GetQuaternionAtSplineInputKey(NearestParam, ESplineCoordinateSpace::World));
+}
+
+bool FSplineComponentVisualizer::CanAddAttributeKey() const
+{
+	check (SelectionState);
+	const USplineComponent* SplineComp = GetEditedSplineComponent();
+	return SplineComp && SplineComp->SupportsAttributes() && SelectionState->GetSelectedAttributeIndex() == INDEX_NONE && SelectionState->GetSelectedAttributeName() != NAME_None;
+}
+
+void FSplineComponentVisualizer::OnDeleteAttributeKey()
+{
+	const FScopedTransaction Transaction(LOCTEXT("RemoveSplinePoint", "Remove Spline Point"));
+
+	ResetTempModes();
+
+	USplineComponent* SplineComp = GetEditedSplineComponent();
+	check(SplineComp != nullptr);
+	check(SelectionState);
+
+	SelectionState->Modify();
+
+	SplineComp->RemovePropertyAtIndex(SelectionState->GetSelectedAttributeIndex(), SelectionState->GetSelectedAttributeName());
+
+	// Update selection
+	if (const int32 NumAttributeValues = SplineComp->GetNumberOfPropertyValues(SelectionState->GetSelectedAttributeName()); NumAttributeValues < 1)
+	{
+		// Clear attribute selection & select a control point
+		ChangeSelectionState(0, false);
+		SelectionState->SetCachedRotation(SplineComp->GetQuaternionAtSplinePoint(SelectionState->GetLastKeyIndexSelected(), ESplineCoordinateSpace::World));
+	}
+	else
+	{
+		// Update attribute selection
+		SelectionState->SetSelectedAttribute(FMath::Clamp(SelectionState->GetSelectedAttributeIndex(), 0, NumAttributeValues - 1), SelectionState->GetSelectedAttributeName());
+		const float NewParam = SplineComp->GetFloatPropertyInputKeyAtIndex(SelectionState->GetSelectedAttributeIndex(), SelectionState->GetSelectedAttributeName());
+		SelectionState->SetCachedRotation(SplineComp->GetQuaternionAtSplineInputKey(NewParam, ESplineCoordinateSpace::World));
+	}
+
+	// Necessary because the attribute proxies store invalid indices, we must redraw them.
+	GEditor->RedrawAllViewports();
+}
+
+bool FSplineComponentVisualizer::CanDeleteAttributeKey() const
+{
+	check (SelectionState);
+	const USplineComponent* SplineComp = GetEditedSplineComponent();
+	return SplineComp && SplineComp->SupportsAttributes() && SelectionState->GetSelectedAttributeIndex() != INDEX_NONE && SelectionState->GetSelectedAttributeName() != NAME_None;
+}
+
 bool FSplineComponentVisualizer::HandleSelectFirstLastSplinePoint(USplineComponent* InSplineComponent, bool bFirstPoint)
 {
 	const FScopedTransaction Transaction(LOCTEXT("SelectFirstSplinePoint", "Select First Spline Point"));
@@ -3484,6 +3690,8 @@ void FSplineComponentVisualizer::GenerateContextMenuSections(FMenuBuilder& InMen
 		InMenuBuilder.AddMenuEntry(FSplineComponentVisualizerCommands::Get().DiscontinuousSpline);
 	}
 	InMenuBuilder.EndSection();
+	
+	GenerateAttributeMenuSection(InMenuBuilder);
 }
 
 void FSplineComponentVisualizer::GenerateSelectSplinePointsSubMenu(FMenuBuilder& MenuBuilder) const
@@ -3554,6 +3762,322 @@ void FSplineComponentVisualizer::GenerateLockAxisSubMenu(FMenuBuilder& MenuBuild
 	MenuBuilder.AddMenuEntry(FSplineComponentVisualizerCommands::Get().SetLockedAxisZ);
 }
 
+void FSplineComponentVisualizer::GenerateAttributeMenuSection(FMenuBuilder& InMenuBuilder) const
+{
+	InMenuBuilder.BeginSection("Attributes", LOCTEXT("Attributes", "Attributes"));
+
+	InMenuBuilder.AddWidget(GenerateChannelWidget().ToSharedRef(), FText::GetEmpty(), true);
+	
+	if (SelectionState)
+	{
+		if (SelectionState->GetSelectedAttributeIndex() == INDEX_NONE)
+		{
+			InMenuBuilder.AddMenuEntry(FSplineComponentVisualizerCommands::Get().AddAttributeKey);
+		}
+		else
+		{
+			InMenuBuilder.AddWidget(GenerateAttributeEditorWidget().ToSharedRef(), FText::GetEmpty(), true);
+			InMenuBuilder.AddMenuEntry(FSplineComponentVisualizerCommands::Get().DeleteAttributeKey);
+		}
+	}
+
+	InMenuBuilder.EndSection();
+}
+
+TSharedPtr<SWidget> FSplineComponentVisualizer::GenerateChannelWidget() const
+{
+	return SNew(SVerticalBox)
+		
+		+ SVerticalBox::Slot()
+		.Padding(10.f, 5.f)
+		.AutoHeight()
+		[
+			GenerateChannelCreatorWidget().ToSharedRef()
+		]
+    
+		+ SVerticalBox::Slot()
+		.Padding(10.f, 1.f)
+		.AutoHeight()
+		[
+			GenerateChannelSelectorWidget().ToSharedRef()
+		];
+}
+
+TSharedPtr<SWidget> FSplineComponentVisualizer::GenerateAttributeEditorWidget() const
+{
+	// Precompute these values so they can be used as arguments for the widget itself (not just used by the widget lambdas).
+	float LowerBound = 0.f;
+	float UpperBound = 0.f;
+	if (const USplineComponent* SplineComp = GetEditedSplineComponent())
+	{
+		UpperBound = static_cast<float>(SplineComp->GetNumberOfSplineSegments());
+	}
+	
+	auto GetParameter = [this]() -> TOptional<float>
+	{
+		if (USplineComponent* LocalSplineComp = GetEditedSplineComponent())
+		{
+			const int32 AttributeIndex = SelectionState ? SelectionState->GetSelectedAttributeIndex() : INDEX_NONE;
+			const FName AttributeName = SelectionState ? SelectionState->GetSelectedAttributeName() : NAME_None;
+	
+			if (AttributeIndex != INDEX_NONE && AttributeName != NAME_None)
+			{
+				return LocalSplineComp->GetFloatPropertyInputKeyAtIndex(AttributeIndex, AttributeName);
+			}
+		}
+		
+		return TOptional<float>();
+	};
+	
+	auto SetParameter = [this, LowerBound, UpperBound](float NewParameter)
+	{
+		if (USplineComponent* LocalSplineComp = GetEditedSplineComponent())
+		{
+			if (SelectionState)
+			{
+				const int32 AttributeIndex = SelectionState->GetSelectedAttributeIndex();
+				const FName AttributeName = SelectionState->GetSelectedAttributeName();
+	
+				if (AttributeIndex != INDEX_NONE && AttributeName != NAME_None)
+				{
+					const float NewParam = FMath::Clamp(NewParameter, LowerBound, UpperBound);
+					const int32 NewIndex = LocalSplineComp->SetFloatPropertyInputKeyAtIndex(AttributeIndex, NewParam, AttributeName);
+
+					// Necessary if SetFloatPropertyInputKeyAtIndex caused attribute storage reordering.
+					SelectionState->SetSelectedAttribute(NewIndex, SelectionState->GetSelectedAttributeName());
+					
+					// Without this, the viewport will not reflect changes until the context menu is closed.
+					SelectionState->SetCachedRotation(LocalSplineComp->GetQuaternionAtSplineInputKey(NewParam, ESplineCoordinateSpace::World));
+					GEditor->RedrawAllViewports();
+				}
+			}
+		}
+	};
+	
+	auto GetValue = [this]() -> TOptional<float>
+	{
+		if (USplineComponent* LocalSplineComp = GetEditedSplineComponent())
+		{
+			const int32 AttributeIndex = SelectionState ? SelectionState->GetSelectedAttributeIndex() : INDEX_NONE;
+			const FName AttributeName = SelectionState ? SelectionState->GetSelectedAttributeName() : NAME_None;
+	
+			if (AttributeIndex != INDEX_NONE && AttributeName != NAME_None)
+			{
+				return LocalSplineComp->GetFloatPropertyAtIndex(AttributeIndex, AttributeName);
+			}
+		}
+
+		return TOptional<float>();
+	};
+
+	auto SetValue = [this](float NewValue)
+	{
+		if (USplineComponent* LocalSplineComp = GetEditedSplineComponent())
+		{
+			const int32 AttributeIndex = SelectionState ? SelectionState->GetSelectedAttributeIndex() : INDEX_NONE;
+			const FName AttributeName = SelectionState ? SelectionState->GetSelectedAttributeName() : NAME_None;
+	
+			if (AttributeIndex != INDEX_NONE && AttributeName != NAME_None)
+			{
+				LocalSplineComp->SetFloatPropertyAtIndex(AttributeIndex, NewValue, AttributeName);
+				GEditor->RedrawAllViewports();	// Force redraw while context menu is preventing it
+			}
+		}
+	};
+
+	TSharedPtr<SVerticalBox> AttributeEditor = SNew(SVerticalBox)
+
+	// Parameter editor
+	+ SVerticalBox::Slot()
+	.AutoHeight()
+	.Padding(10.f, 1.f)
+	[
+		SNew(SHorizontalBox)
+
+		+ SHorizontalBox::Slot()
+		.FillWidth(1.0f)
+		.Padding(10.f, 1.f)
+		.VAlign(VAlign_Center)
+		[
+			SNew(STextBlock)
+			.Text(FText::FromString(TEXT("Parameter: ")))
+			.Justification(ETextJustify::Left)
+		]
+
+		+ SHorizontalBox::Slot()
+		.AutoWidth()
+		.Padding(10.f, 1.f)
+		[
+			SNew(SNumericEntryBox<float>)
+			.MinDesiredValueWidth(100.0f)
+			.Value_Lambda(GetParameter)
+			.OnValueChanged_Lambda(SetParameter)
+			.AllowSpin(true)
+			.MinValue(LowerBound)
+			.MinSliderValue(LowerBound)
+			.MaxValue(UpperBound)
+			.MaxSliderValue(UpperBound)
+		]
+	]
+
+	// Value editor
+	+ SVerticalBox::Slot()
+	.AutoHeight()
+	.Padding(10.f, 1.f)
+	[
+		SNew(SHorizontalBox)
+
+		+ SHorizontalBox::Slot()
+		.FillWidth(1.0f)
+		.Padding(10.f, 1.f)
+		.VAlign(VAlign_Center)
+		[
+			SNew(STextBlock)
+			.Text(FText::FromString(TEXT("Value: ")))
+			.Justification(ETextJustify::Left)
+		]
+
+		+ SHorizontalBox::Slot()
+		.AutoWidth()
+		.Padding(10.f, 1.f)
+		[
+			SNew(SNumericEntryBox<float>)
+			.MinDesiredValueWidth(100.0f)
+			.Value_Lambda(GetValue)
+			.OnValueChanged_Lambda(SetValue)
+			.AllowSpin(true)
+		]
+	];
+
+	return AttributeEditor;
+}
+
+TSharedPtr<SWidget> FSplineComponentVisualizer::GenerateChannelCreatorWidget() const
+{
+	TSharedPtr<SEditableTextBox> ChannelNameBox = SNew(SEditableTextBox)
+		.HintText(FText::FromString(TEXT("Attribute Name")));
+
+	auto OnCreateChannelPressed = [this](FName Name) -> FReply
+	{
+		if (USplineComponent* SplineComp = GetEditedSplineComponent())
+		{
+			SplineComp->CreateFloatPropertyChannel(Name);
+
+			if (SelectionState)
+			{
+				SelectionState->SetSelectedAttribute(INDEX_NONE, Name);
+				GEditor->RedrawAllViewports();
+			}
+		}
+
+		// Necessary for existing selector widget to show new channel option.
+		UpdateSharedAttributeNames();
+
+		return FReply::Handled();
+	};
+
+	return SNew(SHorizontalBox)
+    
+		+ SHorizontalBox::Slot()
+		.Padding(10.f, 1.f)
+		.FillWidth(1)
+		[
+			ChannelNameBox.ToSharedRef()
+		]
+        
+		+ SHorizontalBox::Slot()
+		.Padding(10.f, 1.f)
+		.AutoWidth()
+		[
+			SNew(SButton)
+			.Text(FText::FromString(TEXT("Create")))
+			.OnClicked_Lambda([OnCreateChannelPressed, WeakChannelNameBox = ChannelNameBox.ToWeakPtr()]()
+			{
+				if (TSharedPtr<SEditableTextBox> StrongChannelNameBox = WeakChannelNameBox.Pin())
+				{
+					FName ChannelName(*StrongChannelNameBox->GetText().ToString());
+					StrongChannelNameBox->SetText(FText::GetEmpty());
+					return OnCreateChannelPressed(ChannelName);
+				}
+				
+				return FReply::Unhandled();
+			})
+		];
+}
+
+TSharedPtr<SWidget> FSplineComponentVisualizer::GenerateChannelSelectorWidget() const
+{
+	TSharedPtr<SComboButton> ComboButton = SNew(SComboButton)
+		.ButtonContent()
+		[
+			SNew(STextBlock)
+			.Text_Lambda([this]()
+			{
+				return FText::FromString(SelectionState->GetSelectedAttributeName().ToString());
+			})
+			.Justification(ETextJustify::Center)
+		];
+
+	TSharedPtr<SHorizontalBox> ChannelSelector = SNew(SHorizontalBox)
+	
+		+ SHorizontalBox::Slot()
+		.Padding(10.f, 1.f)
+		.AutoWidth()
+		.VAlign(VAlign_Center)
+		[
+			SNew(STextBlock)
+			.Text(FText::FromString(TEXT("Selected Channel: ")))
+		]
+	
+		+ SHorizontalBox::Slot()
+		.Padding(10.f, 1.f)
+		.FillWidth(1)
+		.VAlign(VAlign_Center)
+		[
+			ComboButton.ToSharedRef()
+		];
+
+	auto OnAttributeSelected = [this, WeakCombo = ComboButton.ToWeakPtr()](const TSharedPtr<FName>& SelectedItem, ESelectInfo::Type)
+	{
+		SelectionState->SetSelectedAttribute(INDEX_NONE, SelectedItem ? *SelectedItem : NAME_None);
+
+		if (TSharedPtr<SComboButton> PinnedComboButton = WeakCombo.Pin())
+		{
+			PinnedComboButton->SetIsOpen(false);
+		}
+
+		// Necessary to cause attribute handles to be rendered before closing the context menu
+		GEditor->RedrawAllViewports();
+	};
+
+	UpdateSharedAttributeNames();
+	
+	// We must defer adding menu content til ComboButton is assigned because of the weak ptr capture above.
+	// Otherwise, we capture it before it is valid.
+	ComboButton->SetMenuContent(SNew(SVerticalBox)
+		+ SVerticalBox::Slot()
+		.AutoHeight()
+		[
+			SNew(SListView<TSharedPtr<FName>>)
+			.SelectionMode(ESelectionMode::SingleToggle)
+			.ListItemsSource(&SharedAttributeNames)
+			.OnGenerateRow_Lambda([](TSharedPtr<FName> Item, const TSharedRef<STableViewBase>& OwnerTable)
+			{
+				return SNew(STableRow<TSharedPtr<FName>>, OwnerTable)
+				.Padding(1.f)
+				[
+					SNew(STextBlock)
+					.Text(FText::FromName(*Item))
+					.Justification(ETextJustify::Center)
+				];
+			})
+			.OnSelectionChanged_Lambda(OnAttributeSelected)
+		]
+	);
+
+	return ChannelSelector;
+}
+
 void FSplineComponentVisualizer::CreateSplineGeneratorPanel()
 {
 	SAssignNew(SplineGeneratorPanel, SSplineGeneratorPanel, SharedThis(this));
diff --git a/Engine/Source/Editor/ComponentVisualizers/Public/SplineComponentVisualizer.h b/Engine/Source/Editor/ComponentVisualizers/Public/SplineComponentVisualizer.h
--- a/Engine/Source/Editor/ComponentVisualizers/Public/SplineComponentVisualizer.h
+++ b/Engine/Source/Editor/ComponentVisualizers/Public/SplineComponentVisualizer.h
@@ -85,6 +85,19 @@ public:
 	int32 GetSelectedTangentHandle() const { return SelectedTangentHandle; }
 	void SetSelectedTangentHandle(const int32 InSelectedTangentHandle) { SelectedTangentHandle = InSelectedTangentHandle; }
 
+	int32 GetSelectedAttributeIndex() const { return SelectedAttributeIndex; }
+	FName GetSelectedAttributeName() const { return SelectedAttributeName; }
+	void SetSelectedAttribute(const int32 InSelectedAttributeIndex = INDEX_NONE, const FName& InSelectedAttributeName = NAME_None)
+	{
+		if (InSelectedAttributeIndex != INDEX_NONE)
+		{
+			// if we actually select an attribute point, don't persist other selection.
+			ClearSelectedKeys();	// necessary? maybe...
+		}
+		SelectedAttributeIndex = InSelectedAttributeIndex;
+		SelectedAttributeName = InSelectedAttributeName;
+	}
+
 	ESelectedTangentHandle GetSelectedTangentHandleType() const { return SelectedTangentHandleType; }
 	void SetSelectedTangentHandleType(const ESelectedTangentHandle InSelectedTangentHandle) { SelectedTangentHandleType = InSelectedTangentHandle; }
 
@@ -95,8 +108,10 @@ public:
 	void SetCachedRotation(const FQuat& InCachedRotation) { CachedRotation = InCachedRotation; }
 
 	void Reset();
+	void ClearSelectedKeys();
 	void ClearSelectedSegmentIndex();
 	void ClearSelectedTangentHandle();
+	void ClearSelectedAttribute() { SetSelectedAttribute(); }
 
 	bool IsSplinePointSelected(const int32 InIndex) const;
 
@@ -125,6 +140,13 @@ protected:
 	UPROPERTY()
 	ESelectedTangentHandle SelectedTangentHandleType = ESelectedTangentHandle::None;
 
+	/** Index of attribute handle we have selected */
+	UPROPERTY()
+	int32 SelectedAttributeIndex = INDEX_NONE;
+
+	UPROPERTY()
+	FName SelectedAttributeName = NAME_None;
+	
 	/** Position on spline we have selected */
 	UPROPERTY()
 	FVector SelectedSplinePosition;
@@ -167,6 +189,24 @@ struct HSplineKeyProxy : public HSplineVisProxy
 	}
 };
 
+/** Proxy for spline attribute key (as opposed to a positional control point) */
+struct HSplineAttributeKeyProxy : public HSplineVisProxy
+{
+	DECLARE_HIT_PROXY();
+
+	HSplineAttributeKeyProxy(const UActorComponent* InComponent, int32 InKeyIndex) 
+		: HSplineVisProxy(InComponent)
+		, KeyIndex(InKeyIndex)
+	{}
+
+	int32 KeyIndex;
+	
+	virtual EMouseCursor::Type GetMouseCursor() override
+	{
+		return EMouseCursor::CardinalCross;
+	}
+};
+
 /** Proxy for a spline segment */
 struct HSplineSegmentProxy : public HSplineVisProxy
 {
@@ -292,6 +332,9 @@ protected:
 	/** Transforms selected tangent by given translate, rotate and scale */
 	bool TransformSelectedKeys(EPropertyChangeType::Type InPropertyChangeType, const FVector& InDeltaTranslate, const FRotator& InDeltaRotate = FRotator::ZeroRotator, const FVector& InDeltaScale = FVector::ZeroVector);
 
+	/** Transforms selected attribute by given translation */
+	bool TransformSelectedAttribute(EPropertyChangeType::Type InPropertyChangeType, const FVector& InDeltaTranslate);
+	
 	/** Update the key selection state of the visualizer */
 	virtual void ChangeSelectionState(int32 Index, bool bIsCtrlHeld);
 
@@ -391,6 +434,12 @@ protected:
 	void OnResetToDefault();
 	bool CanResetToDefault() const;
 
+	void OnAddAttributeKey();
+	bool CanAddAttributeKey() const;
+
+	void OnDeleteAttributeKey();
+	bool CanDeleteAttributeKey() const;
+	
 	/** Select first or last spline point */
 	void OnSelectFirstLastSplinePoint(bool bFirstPoint);
 
@@ -414,6 +463,17 @@ protected:
 	/** Generate the submenu containing the lock axis types */
 	void GenerateLockAxisSubMenu(FMenuBuilder& MenuBuilder) const;
 
+	/** Extends the provided MenuBuilder with an Attribute section */
+	void GenerateAttributeMenuSection(FMenuBuilder& InMenuBuilder) const;
+	
+	/** Generate the channel creator/selector widget. */
+	TSharedPtr<SWidget> GenerateChannelWidget() const;
+	TSharedPtr<SWidget> GenerateChannelCreatorWidget() const;
+	TSharedPtr<SWidget> GenerateChannelSelectorWidget() const;
+
+	/** Generate the attribute editor widget. */
+	TSharedPtr<SWidget> GenerateAttributeEditorWidget() const;
+
 	/** Helper function to set edited component we are currently editing */
 	void SetEditedSplineComponent(const USplineComponent* InSplineComponent);
 
@@ -467,6 +527,20 @@ protected:
 
 	FDelegateHandle DeselectedInEditorDelegateHandle;
 
+	void UpdateSharedAttributeNames() const
+	{
+		if (const USplineComponent* SplineComp = GetEditedSplineComponent())
+		{
+			SharedAttributeNames.Empty();
+			SharedAttributeNames.Add(MakeShared<FName>(NAME_None));
+			Algo::Transform(SplineComp->GetFloatPropertyChannels(), SharedAttributeNames, [](const FName& Name)
+			{
+				return MakeShared<FName>(Name);
+			});
+		}
+	}
+	mutable TArray<TSharedPtr<FName>> SharedAttributeNames;
+	
 private:
 	TSharedPtr<SSplineGeneratorPanel> SplineGeneratorPanel;
 	static TWeakPtr<SWindow> WeakExistingWindow;
diff --git a/Engine/Source/Editor/DetailCustomizations/Private/SplineComponentDetails.cpp b/Engine/Source/Editor/DetailCustomizations/Private/SplineComponentDetails.cpp
--- a/Engine/Source/Editor/DetailCustomizations/Private/SplineComponentDetails.cpp
+++ b/Engine/Source/Editor/DetailCustomizations/Private/SplineComponentDetails.cpp
@@ -1155,7 +1155,7 @@ void FSplinePointDetails::OnSetRotation(FRotator::FReal NewValue, ETextCommit::T
 				continue;
 			}
 
-			const FQuat CurrentRotationRelative = SplineComp->GetRotationAtSplinePoint(Index, ESplineCoordinateSpace::Local).Quaternion();
+			const FQuat CurrentRotationRelative = SplineComp->GetQuaternionAtSplinePoint(Index, ESplineCoordinateSpace::Local);
 
 			if (bEditingRotationAbsolute)
 			{
@@ -1184,13 +1184,7 @@ void FSplinePointDetails::OnSetRotation(FRotator::FReal NewValue, ETextCommit::T
 				NewRotationRelative = NewRotationRotator.Quaternion();
 			}
 
-			SplineComp->SetRotationAtSplinePoint(Index, NewRotationRelative.Rotator(), ESplineCoordinateSpace::Local);
-
-			FQuat DeltaRotate(NewRotationRelative * CurrentRotationRelative.Inverse());
-			// Rotate tangent according to delta rotation
-			FVector NewTangent = SplineComp->GetLeaveTangentAtSplinePoint(Index, ESplineCoordinateSpace::World);
-			NewTangent = DeltaRotate.RotateVector(NewTangent); // apply world-space delta rotation to world-space tangent
-			SplineComp->SetTangentAtSplinePoint(Index, NewTangent, ESplineCoordinateSpace::World, false);
+			SplineComp->SetQuaternionAtSplinePoint(Index, NewRotationRelative, ESplineCoordinateSpace::Local);
 		}
 	}
 
diff --git a/Engine/Source/Runtime/Engine/Classes/Components/SplineComponent.h b/Engine/Source/Runtime/Engine/Classes/Components/SplineComponent.h
--- a/Engine/Source/Runtime/Engine/Classes/Components/SplineComponent.h
+++ b/Engine/Source/Runtime/Engine/Classes/Components/SplineComponent.h
@@ -67,7 +67,7 @@ struct FSplineCurves
 		NewCurve.Position = InProxy.GetSplinePointsPosition();
 		NewCurve.Rotation = InProxy.GetSplinePointsRotation();
 		NewCurve.Scale = InProxy.GetSplinePointsScale();
-		NewCurve.ReparamTable = InProxy.GetSplineReparamTable();
+		NewCurve.UpdateSpline(NewCurve.Position.bIsLooped, false);
 		return NewCurve;
 	}
 	
@@ -216,6 +216,9 @@ protected:
 	
 public:
 
+	UFUNCTION(CallInEditor, Category="Spline")
+	void PopulateFromLegacy();
+
 	// Gives the name of the backing spline property to allow users
 	// the ability to correctly trigger property change events.
 	ENGINE_API static FName GetSplinePropertyName();
@@ -371,7 +374,6 @@ public:
 	ENGINE_API const FInterpCurveVector& GetSplinePointsPosition() const;
 	ENGINE_API const FInterpCurveQuat& GetSplinePointsRotation() const;
 	ENGINE_API const FInterpCurveVector& GetSplinePointsScale() const;
-	ENGINE_API const FInterpCurveFloat& GetSplineReparamTable() const;
 
 	virtual USplineMetadata* GetSplinePointsMetadata() { return nullptr; }
 	virtual const USplineMetadata* GetSplinePointsMetadata() const { return nullptr; }
@@ -451,14 +453,6 @@ public:
 	UFUNCTION(BlueprintCallable, Category = Spline)
 	ENGINE_API float GetDistanceAlongSplineAtLocation(const FVector& InLocation, ESplineCoordinateSpace::Type CoordinateSpace) const;
 
-	/** Get a metadata property float value along the spline at spline input key */
-	UFUNCTION(BlueprintCallable, Category = Spline)
-	ENGINE_API float GetFloatPropertyAtSplineInputKey(float InKey, FName PropertyName) const;
-
-	/** Get a metadata property vector value along the spline at spline input key */
-	UFUNCTION(BlueprintCallable, Category = Spline)
-	ENGINE_API FVector GetVectorPropertyAtSplineInputKey(float InKey, FName PropertyName) const;
-
 	/** Specify unselected spline component segment color in the editor */
 	UFUNCTION(BlueprintCallable, Category = Editor)
 	ENGINE_API void SetUnselectedSplineSegmentColor(const FLinearColor& SegmentColor);
@@ -863,6 +857,133 @@ public:
 	UFUNCTION(BlueprintPure = false, Category = Spline)
 	ENGINE_API bool ConvertSplineToPolyline_InTimeRange(ESplineCoordinateSpace::Type CoordinateSpace, const float MaxSquareDistanceFromSpline, float StartTimeAlongSpline, float EndTimeAlongSpline, bool bUseConstantVelocity, TArray<FVector>& OutPoints, TArray<double>& OutDistancesAlongSpline, bool bAllowWrappingIfClosed = true) const;
 
+	/* ---- Metadata API ---- */
+	
+	/**
+	 * @return True if this component currently supports arbitrary attribute channels, false otherwise.
+	 */
+	UFUNCTION(BlueprintCallable, Category = "Spline|Metadata")
+	ENGINE_API bool SupportsAttributes() const;
+
+	/**
+	 * Get the number of attribute values in a specified attribute channel.
+	 * @param PropertyName The name of the attribute channel.
+	 * @return The number of attribute values in the specified channel.
+	 */
+	UFUNCTION(BlueprintCallable, Category = "Spline|Metadata")
+	ENGINE_API int32 GetNumberOfPropertyValues(FName PropertyName) const;
+		
+	/**
+	 * @return An array of all float attribute channel names.
+	 */
+	UFUNCTION(BlueprintCallable, Category = "Spline|Metadata")
+	ENGINE_API TArray<FName> GetFloatPropertyChannels() const;
+
+	/**
+	 * @return An array of all vector attribute channel names.
+	 */
+	UFUNCTION(BlueprintCallable, Category = "Spline|Metadata")
+	ENGINE_API TArray<FName> GetVectorPropertyChannels() const;
+
+	/**
+	 * Create a float attribute channel.
+	 * @param PropertyName The desired name of the new attribute channel.
+	 */
+	UFUNCTION(BlueprintCallable, Category = "Spline|Metadata")
+	ENGINE_API void CreateFloatPropertyChannel(FName PropertyName);
+
+	/**
+	 * Create a vector attribute channel.
+	 * @param PropertyName The desired name of the new attribute channel.
+	 */
+	UFUNCTION(BlueprintCallable, Category = "Spline|Metadata")
+	ENGINE_API void CreateVectorPropertyChannel(FName PropertyName);
+	
+	/**
+	 * Get a float attribute value by parameter.
+	 * @param InKey The parameter to use when evaluating the specified attribute channel.
+	 * @param PropertyName The name of the attribute channel to evaluate.
+	 * @return The attribute value.
+	 */
+	UFUNCTION(BlueprintCallable, Category = "Spline|Metadata")
+	ENGINE_API float GetFloatPropertyAtSplineInputKey(float InKey, FName PropertyName) const;
+
+	/**
+	 * Get a vector attribute value by parameter.
+	 * @param InKey The parameter to use when evaluating the specified attribute channel.
+	 * @param PropertyName The name of the attribute channel to evaluate.
+	 * @return The attribute value.
+	 */
+	UFUNCTION(BlueprintCallable, Category = "Spline|Metadata")
+	ENGINE_API FVector GetVectorPropertyAtSplineInputKey(float InKey, FName PropertyName) const;
+
+	/**
+	 * Add a float attribute value by parameter.
+	 * @param InKey The desired attribute parameter.
+	 * @param Value The desired attribute value.
+	 * @param PropertyName The name of the attribute channel that will contain the new attribute value.
+	 * @return The index of the new attribute.
+	 */
+	UFUNCTION(BlueprintCallable, Category = "Spline|Metadata")
+	ENGINE_API int32 SetFloatPropertyAtSplineInputKey(float InKey, float Value, FName PropertyName) const;
+
+	/**
+	 * Add a vector attribute value by parameter.
+	 * @param InKey The desired attribute parameter.
+	 * @param Value The desired attribute value.
+	 * @param PropertyName The name of the attribute channel that will contain the new attribute value.
+	 * @return The index of the new attribute.
+	 */
+	UFUNCTION(BlueprintCallable, Category = "Spline|Metadata")
+	ENGINE_API int32 SetVectorPropertyAtSplineInputKey(float InKey, const FVector& Value, FName PropertyName) const;
+	
+	// todo: vector variants of below functions
+	
+	/**
+	 * Get a float attribute parameter by index.
+	 * @param Index The index of the specified attribute.
+	 * @param PropertyName The name of the attribute channel containing the specified attribute.
+	 * @return The attribute parameter.
+	 */
+	UFUNCTION(BlueprintCallable, Category = "Spline|Metadata")
+	ENGINE_API float GetFloatPropertyInputKeyAtIndex(int32 Index, FName PropertyName) const;
+
+	/**
+	 * Set a float attribute parameter by index.
+	 * @param Index The index of the specified attribute.
+	 * @param InKey The desired attribute parameter.
+	 * @param PropertyName The name of the attribute channel containing the specified attribute.
+	 * @return The index of the specified attribute after changing the attribute parameter (it can change).
+	 */
+	UFUNCTION(BlueprintCallable, Category = "Spline|Metadata")
+	ENGINE_API int32 SetFloatPropertyInputKeyAtIndex(int32 Index, float InKey, FName PropertyName);
+
+	/**
+	 * Get a float attribute value by index.
+	 * @param Index The index of the specified attribute.
+	 * @param PropertyName The name of the attribute channel containing the specified attribute.
+	 * @return The attribute value.
+	 */
+	UFUNCTION(BlueprintCallable, Category = "Spline|Metadata")
+	ENGINE_API float GetFloatPropertyAtIndex(int32 Index, FName PropertyName) const;
+
+	/**
+	 * Set a float attribute value by index.
+	 * @param Index The index of the specified attribute.
+	 * @param Value The desired attribute value.
+	 * @param PropertyName The name of the attribute channel containing the specified attribute.
+	 */
+	UFUNCTION(BlueprintCallable, Category = "Spline|Metadata")
+	ENGINE_API void SetFloatPropertyAtIndex(int32 Index, float Value, FName PropertyName);
+
+	/**
+	 * Remove an attribute value by index.
+	 * @param Index The index of the specified attribute.
+	 * @param PropertyName The name of the attribute channel containing the specified attribute.
+	 */
+	UFUNCTION(BlueprintCallable, Category = "Spline|Metadata")
+	ENGINE_API void RemovePropertyAtIndex(int32 Index, FName PropertyName);
+
 	// todo: deprecate
 	FInterpCurveVector& GetSplinePointsPosition() { return WarninglessSplineCurves().Position; }
 	FInterpCurveQuat& GetSplinePointsRotation() { return WarninglessSplineCurves().Rotation; }
@@ -874,6 +995,7 @@ protected:
 	ENGINE_API float GetSegmentLength(const int32 Index, const float Param = 1.0f) const;
 	
 private:
+	
 	/**
 	 * This function takes on the roles of the previously existing functions GetPositionPointSafe, GetRotationPointSafe, and GetScalePointSafe.
 	 * Specifically, it makes sure that a provided index is valid by
diff --git a/Engine/Source/Runtime/Engine/Classes/Curves/Spline.h b/Engine/Source/Runtime/Engine/Classes/Curves/Spline.h
--- a/Engine/Source/Runtime/Engine/Classes/Curves/Spline.h
+++ b/Engine/Source/Runtime/Engine/Classes/Curves/Spline.h
@@ -9,6 +9,8 @@
 
 struct FSplinePoint;
 struct FSplineCurves;
+class FLegacySpline;
+class FNewSpline;
 
 /**
  * A general purpose, reflected spline.
@@ -31,11 +33,11 @@ struct FSpline
 	 * Nick would like to change this behavior in the future. See UE-250236. */
 	ENGINE_API void AddPoint(const FSplinePoint& InPoint);
 	/** Adds point by index. Provided parameter is ignored. The new point is inserted before all points with parameter values greater than the new point's parameter. */
-	ENGINE_API void AddPointAtIndex(const FSplinePoint& InPoint, int32 Index);
+	ENGINE_API void InsertPoint(const FSplinePoint& InPoint, int32 Index);
 	ENGINE_API FSplinePoint GetPoint(const int32 Index) const;
 	ENGINE_API void RemovePoint(const int32 Index);
 	
-	ENGINE_API void SetLocation(int32 PointIndex, const FVector& InLocation);
+	ENGINE_API void SetLocation(int32 Index, const FVector& InLocation);
 	ENGINE_API FVector GetLocation(const int32 Index) const;
 	
 	ENGINE_API void SetInTangent(const int32 Index, const FVector& InTangent);
@@ -44,16 +46,18 @@ struct FSpline
 	ENGINE_API void SetOutTangent(const int32 Index, const FVector& OutTangent);
 	ENGINE_API FVector GetOutTangent(const int32 Index) const;
 
-	ENGINE_API void SetRotation(int32 PointIndex, const FQuat& InRotation);
+	ENGINE_API void SetRotation(int32 Index, const FQuat& InRotation);
 	ENGINE_API FQuat GetRotation(const int32 Index) const;
 	
-	ENGINE_API void SetScale(int32 PointIndex, const FVector& InScale);
+	ENGINE_API void SetScale(int32 Index, const FVector& InScale);
 	ENGINE_API FVector GetScale(const int32 Index) const;
 	
-	ENGINE_API void SetSplinePointType(int32 PointIndex, EInterpCurveMode Type);
-	ENGINE_API EInterpCurveMode GetSplinePointType(int32 PointIndex) const;
+	ENGINE_API void SetSplinePointType(int32 Index, EInterpCurveMode Type);
+	ENGINE_API EInterpCurveMode GetSplinePointType(int32 Index) const;
 
-	ENGINE_API float GetSplinePointParameter(int32 PointIndex) const;
+	ENGINE_API float GetParameterAtIndex(int32 Index) const;
+	ENGINE_API float GetParameterAtDistance(float Distance) const;
+	ENGINE_API float GetDistanceAtParameter(float Parameter) const;
 	
 	/* Parameter Interface */
 
@@ -61,35 +65,32 @@ struct FSpline
 	ENGINE_API FVector EvaluateDerivative(float Param) const;
 	ENGINE_API FQuat EvaluateRotation(float Param) const;
 	ENGINE_API FVector EvaluateScale(float Param) const;
-	template<typename T>
-	T EvaluateAttribute(float Param, FName AttributeName) const
-	{
-		return T();
-	}
+
+	/* Attribute Interface */
+	
+	bool SupportsAttributes() const { return IsNew(); }
+	ENGINE_API bool HasAttributeChannel(FName AttributeName) const;
+	ENGINE_API int32 NumAttributeValues(FName AttributeName) const;
+	ENGINE_API float GetAttributeParameter(int32 Index, const FName& Name) const;
+	ENGINE_API int32 SetAttributeParameter(int32 Index, float Parameter, const FName& Name);
+	ENGINE_API void RemoveAttributeValue(int32 Index, FName AttributeName);
+	ENGINE_API TArray<FName> GetFloatPropertyChannels() const;
+	ENGINE_API TArray<FName> GetVectorPropertyChannels() const;
+	
+	template<typename T> ENGINE_API T GetAttributeValue(int32 Index, const FName& Name) const;
+	template<typename T> ENGINE_API void SetAttributeValue(int32 Index, const T& Value, const FName& Name);
+	template<typename T> ENGINE_API bool CreateAttributeChannel(FName AttributeName) const;
+	template<typename T> ENGINE_API int32 AddAttributeValue(float Param, const T& Value, FName AttributeName) const;
+	
+	
+	template<typename T> ENGINE_API T EvaluateAttribute(float Param, FName AttributeName) const;
 
 	ENGINE_API float FindNearest(const FVector& InLocation) const;
 	ENGINE_API float FindNearestOnSegment(const FVector& InLocation, int32 SegmentIndex) const;
 	
 	/* Misc Interface */
 	
-	bool operator==(const FSpline& Other) const
-	{
-		if (CurrentImplementation == Other.CurrentImplementation)
-		{
-			if (LegacyData && Other.LegacyData)
-			{
-				return *LegacyData == *Other.LegacyData;
-			}
-
-			// if (NewData && Other.NewData)
-			// {
-			// 	return *NewData == *Other.NewData;
-			// }
-		}
-
-		return false;
-	}
-	
+	ENGINE_API bool operator==(const FSpline& Other) const;
 	bool operator!=(const FSpline& Other) const
 	{
 		return !(*this == Other);
@@ -100,22 +101,16 @@ struct FSpline
 		Spline.Serialize(Ar);
 		return Ar;
 	}
-	bool Serialize(FArchive& Ar);
+	ENGINE_API bool Serialize(FArchive& Ar);
 
 	uint32 GetVersion() const { return Version; }
 	
-	/**
-	 * These functions will potentially be slow when not using the legacy spline.
-	 * This is because we will be forced to cache the underlying spline into
-	 * a data structure that allows us to return a const ref.
-	 */
 	ENGINE_API const FInterpCurveVector& GetSplinePointsPosition() const;
 	ENGINE_API const FInterpCurveQuat& GetSplinePointsRotation() const;
 	ENGINE_API const FInterpCurveVector& GetSplinePointsScale() const;
-	ENGINE_API const FInterpCurveFloat& GetSplineReparamTable() const;		// I may actually implement reparameterization with this interp curve for both spline backends.
 
 	/** Returns the length of the specified spline segment up to the parametric value given */
-	ENGINE_API float GetSegmentLength(const int32 Index, const float Param, bool bClosedLoop = false, const FVector& Scale3D = FVector(1.0f)) const;
+	ENGINE_API float GetSegmentLength(const int32 Index, const float Param, const FVector& Scale3D = FVector(1.0f)) const;
 
 	/** Returns total length along this spline */
 	ENGINE_API float GetSplineLength() const;
@@ -150,108 +145,9 @@ struct FSpline
 
 private:
 	
-	class FLegacySpline
-	{
-		FInterpCurveVector PositionCurve;
-		FInterpCurveQuat RotationCurve;
-		FInterpCurveVector ScaleCurve;
-		FInterpCurveFloat ReparamTable;
-		
-	public:
-
-		ENGINE_API FLegacySpline();
-		static FLegacySpline FromSplineCurves(const FSplineCurves& InSpline);
-		
-		/* Control Point Index Interface */
-
-		void AddPoint(const FSplinePoint& InPoint);
-		void AddPointAtIndex(const FSplinePoint& InPoint, int32 Index);
-		FSplinePoint GetPoint(const int32 Index) const;
-		void RemovePoint(int32 Index);
-		
-		void SetLocation(int32 PointIndex, const FVector& InLocation);
-		FVector GetLocation(const int32 Index) const;
-		
-		void SetInTangent(const int32 Index, const FVector& InTangent);
-		FVector GetInTangent(const int32 Index) const;
-		
-		void SetOutTangent(const int32 Index, const FVector& OutTangent);
-		FVector GetOutTangent(const int32 Index) const;
-
-		void SetRotation(int32 PointIndex, const FQuat& InRotation);
-		FQuat GetRotation(const int32 Index) const;
-		
-		void SetScale(int32 PointIndex, const FVector& InScale);
-		FVector GetScale(const int32 Index) const;
-		
-		void SetSplinePointType(int32 PointIndex, EInterpCurveMode Type);
-		EInterpCurveMode GetSplinePointType(int32 PointIndex) const;
-
-		float GetSplinePointParameter(int32 PointIndex) const;
-		
-		/* Parameter Interface */
-
-		FVector Evaluate(float Param) const;
-		FVector EvaluateDerivative(float Param) const;
-		FQuat EvaluateRotation(float Param) const;
-		FVector EvaluateScale(float Param) const;
-		float FindNearest(const FVector& InLocation) const;
-		float FindNearestOnSegment(const FVector& InLocation, int32 SegmentIndex) const;
-		
-		/* Misc Interface */
-		
-		bool operator==(const FLegacySpline& Other) const
-		{
-			return PositionCurve == Other.PositionCurve && RotationCurve == Other.RotationCurve && ScaleCurve == Other.ScaleCurve;
-		}
-
-		bool operator!=(const FLegacySpline& Other) const
-		{
-			return !(*this == Other);
-		}
-		
-		friend ENGINE_API FArchive& operator<<(FArchive& Ar, FLegacySpline& Spline)
-		{
-			Spline.Serialize(Ar);
-			return Ar;
-		}
-		bool Serialize(FArchive& Ar);
-		
-		FInterpCurveVector& GetSplinePointsPosition() { return PositionCurve; }
-		const FInterpCurveVector& GetSplinePointsPosition() const { return PositionCurve; }
-		FInterpCurveQuat& GetSplinePointsRotation() { return RotationCurve; }
-		const FInterpCurveQuat& GetSplinePointsRotation() const { return RotationCurve; }
-		FInterpCurveVector& GetSplinePointsScale() { return ScaleCurve; }
-		const FInterpCurveVector& GetSplinePointsScale() const { return ScaleCurve; }
-		FInterpCurveFloat& GetSplineReparamTable() { return ReparamTable; }
-		const FInterpCurveFloat& GetSplineReparamTable() const { return ReparamTable; }
-
-		/** Returns the length of the specified spline segment up to the parametric value given */
-		float GetSegmentLength(const int32 Index, const float Param, bool bClosedLoop = false, const FVector& Scale3D = FVector(1.0f)) const;
-
-		/** Returns total length along this spline */
-		float GetSplineLength() const;
-		
-		/** Returns the total number of control points on this spline. */
-		int32 GetNumControlPoints() const;
-		
-		/** Reset the spline to an empty spline. */
-		void Reset();
-		
-		/** Reset the rotation attribute channel to default values. */
-		void ResetRotation();
-		
-		/** Reset the scale attribute channel to default values. */
-		void ResetScale();
-		
-		/** Reset the spline to the default spline (2 points). */
-		void ResetToDefault();
-		
-		/** Update the spline's internal data according to the passed-in params. */
-		void UpdateSpline(const FUpdateSplineParams& InParams);
-	};
-	
-	static inline const FLegacySpline DummySpline;
+	static inline const FInterpCurveVector PositionCurve;
+	static inline const FInterpCurveQuat RotationCurve;
+	static inline const FInterpCurveVector ScaleCurve;
 	
 	// Used for upgrade logic in spline component.
 	// Not ideal, but allows us to automatically populate the proxy
@@ -263,10 +159,11 @@ private:
 	uint8 CurrentImplementation;
 	
 	uint32 Version;
-	
-	struct {										// Invalid when CurrentImplementation is 0 or 3.
-		TSharedPtr<FLegacySpline> LegacyData;		// Valid when CurrentImplementation is 1.
-		//TSharedPtr<FHermiteSpline> NewData;		// Valid when CurrentImplementation is 2.
+
+	// probably better implemented as a TSharedPtr<ISplineInterface> or something
+	struct {													// Invalid when CurrentImplementation is 0 or 3.
+		TSharedPtr<FLegacySpline> LegacyData;					// Valid when CurrentImplementation is 1.
+		TSharedPtr<FNewSpline> NewData;							// Valid when CurrentImplementation is 2.
 	};
 	
 	bool IsEnabled() const { return CurrentImplementation != 0; }
diff --git a/Engine/Source/Runtime/Engine/Private/Components/SplineComponent.cpp b/Engine/Source/Runtime/Engine/Private/Components/SplineComponent.cpp
--- a/Engine/Source/Runtime/Engine/Private/Components/SplineComponent.cpp
+++ b/Engine/Source/Runtime/Engine/Private/Components/SplineComponent.cpp
@@ -13,6 +13,7 @@
 #include "PrimitiveViewRelevance.h"
 #include "PrimitiveSceneProxy.h"
 #include "PrimitiveDrawingUtils.h"
+#include "Algo/ForEach.h"
 #include "Net/UnrealNetwork.h"
 #include "Net/Core/PushModel/PushModel.h"
 #include "Styling/SlateColor.h"
@@ -47,6 +48,11 @@ USplineMetadata::USplineMetadata(const FObjectInitializer& ObjectInitializer)
 {
 }
 
+void USplineComponent::PopulateFromLegacy()
+{
+	SynchronizeSplines();
+}
+
 FName USplineComponent::GetSplinePropertyName()
 {
 	if (UE::SplineComponent::GUseSplineCurves)
@@ -101,7 +107,7 @@ USplineComponent::USplineComponent(const FObjectInitializer& ObjectInitializer)
 	SplineInfo_DEPRECATED = Spline.GetSplinePointsPosition();
 	SplineRotInfo_DEPRECATED = Spline.GetSplinePointsRotation();
 	SplineScaleInfo_DEPRECATED = Spline.GetSplinePointsScale();
-	SplineReparamTable_DEPRECATED = Spline.GetSplineReparamTable();
+	SplineReparamTable_DEPRECATED = WarninglessSplineCurves().ReparamTable;
 #endif
 }
 
@@ -412,7 +418,7 @@ float USplineComponent::GetSegmentLength(const int32 Index, const float Param) c
 {
 	return UE::SplineComponent::GUseSplineCurves
 		? WarninglessSplineCurves().GetSegmentLength(Index, Param, bClosedLoop, GetComponentTransform().GetScale3D())
-		: Spline.GetSegmentLength(Index, Param, bClosedLoop, GetComponentTransform().GetScale3D());
+		: Spline.GetSegmentLength(Index, Param, GetComponentTransform().GetScale3D());
 }
 
 float USplineComponent::GetSegmentParamFromLength(const int32 Index, const float Length, const float SegmentLength) const
@@ -605,15 +611,21 @@ float USplineComponent::GetDistanceAlongSplineAtSplineInputKey(float InKey) cons
 
 	if ((InKey >= 0) && (InKey < NumSegments))
 	{
-		const int32 PointIndex = FMath::FloorToInt(InKey);
-		const float Fraction = InKey - PointIndex;
-		const int32 ReparamPointIndex = PointIndex * ReparamStepsPerSegment;
+		float Distance;
 		
-		const float Distance = UE::SplineComponent::GUseSplineCurves
-			? WarninglessSplineCurves().ReparamTable.Points[ReparamPointIndex].InVal
-			: Spline.GetSplineReparamTable().Points[ReparamPointIndex].InVal;
+		if (UE::SplineComponent::GUseSplineCurves)
+		{
+			const int32 PointIndex = FMath::FloorToInt(InKey);
+			const float Fraction = InKey - PointIndex;
+			const int32 ReparamPointIndex = PointIndex * ReparamStepsPerSegment;
+			Distance = WarninglessSplineCurves().ReparamTable.Points[ReparamPointIndex].InVal + GetSegmentLength(PointIndex, Fraction);
+		}
+		else
+		{
+			Distance = Spline.GetDistanceAtParameter(InKey);
+		}
 		
-		return Distance + GetSegmentLength(PointIndex, Fraction);
+		return Distance;
 	}
 	else if (InKey >= NumSegments)
 	{
@@ -635,8 +647,59 @@ float USplineComponent::GetDistanceAlongSplineAtLocation(const FVector& InLocati
 	return GetDistanceAlongSplineAtSplineInputKey(Key);
 }
 
+int32 USplineComponent::GetNumberOfPropertyValues(FName PropertyName) const
+{
+	if (UE::SplineComponent::GUseSplineCurves == false && Spline.SupportsAttributes())
+	{
+		return Spline.NumAttributeValues(PropertyName);
+	}
+
+	return 0;
+}
+
+TArray<FName> USplineComponent::GetFloatPropertyChannels() const
+{
+	if (UE::SplineComponent::GUseSplineCurves == false && Spline.SupportsAttributes())
+	{
+		return Spline.GetFloatPropertyChannels();
+	}
+
+	return TArray<FName>();
+}
+
+TArray<FName> USplineComponent::GetVectorPropertyChannels() const
+{
+	if (UE::SplineComponent::GUseSplineCurves == false && Spline.SupportsAttributes())
+	{
+		return Spline.GetVectorPropertyChannels();
+	}
+
+	return TArray<FName>();
+}
+
 template<class T>
-T GetPropertyValueAtSplineInputKey(const USplineMetadata* Metadata, float InKey, FName PropertyName)
+bool CreatePropertyChannel(const USplineMetadata* Metadata, const FSpline& Spline, FName PropertyName)
+{
+	if (Metadata && Metadata->GetClass()->FindPropertyByName(PropertyName))
+	{
+		// can't create, it already exists on legacy metadata
+		return false;
+	}
+	
+	if (UE::SplineComponent::GUseSplineCurves == false && Spline.SupportsAttributes())
+	{
+		if (!Spline.HasAttributeChannel(PropertyName))
+		{
+			Spline.CreateAttributeChannel<T>(PropertyName);
+			return true;
+		}
+	}
+
+	return false;
+}
+
+template<class T>
+T GetPropertyAtSplineInputKey(const USplineMetadata* Metadata, const FSpline& Spline, float InKey, FName PropertyName)
 {
 	if (Metadata)
 	{
@@ -646,18 +709,133 @@ T GetPropertyValueAtSplineInputKey(const USplineMetadata* Metadata, float InKey,
 			return Curve->Eval(InKey, T(0));
 		}
 	}
+	
+	if (UE::SplineComponent::GUseSplineCurves == false && Spline.SupportsAttributes())
+	{
+		if (Spline.HasAttributeChannel(PropertyName))
+		{
+			return Spline.EvaluateAttribute<T>(InKey, PropertyName);
+		}
+	}
 
 	return T(0);
 }
 
+template<class T>
+int32 SetPropertyAtSplineInputKey(const FSpline& Spline, float InKey, const T& InValue, FName PropertyName)
+{
+	if (UE::SplineComponent::GUseSplineCurves == false && Spline.SupportsAttributes())
+	{
+		if (Spline.HasAttributeChannel(PropertyName))
+		{
+			return Spline.AddAttributeValue<T>(InKey, InValue, PropertyName);
+		}
+	}
+
+	return INDEX_NONE;
+}
+
+float GetInputKeyAtIndex(const FSpline& Spline, int32 Index, FName PropertyName)
+{
+	if (UE::SplineComponent::GUseSplineCurves == false && Spline.SupportsAttributes())
+	{
+		return Spline.GetAttributeParameter(Index, PropertyName);
+	}
+
+	return 0.f;
+}
+
+int32 SetInputKeyAtIndex(FSpline& Spline, int32 Index, float InKey, FName PropertyName)
+{
+	if (UE::SplineComponent::GUseSplineCurves == false && Spline.SupportsAttributes())
+	{
+		return Spline.SetAttributeParameter(Index, InKey , PropertyName);
+	}
+
+	return INDEX_NONE;
+}
+
+template<typename T>
+float GetPropertyAtIndex(const FSpline& Spline, int32 Index, FName PropertyName)
+{
+	if (UE::SplineComponent::GUseSplineCurves == false && Spline.SupportsAttributes())
+	{
+		return Spline.GetAttributeValue<T>(Index, PropertyName);
+	}
+
+	return 0.f;
+}
+
+template<typename T>
+void SetPropertyAtIndex(FSpline& Spline, int32 Index, float Value, FName PropertyName)
+{
+	if (UE::SplineComponent::GUseSplineCurves == false && Spline.SupportsAttributes())
+	{
+		Spline.SetAttributeValue<T>(Index, Value, PropertyName);
+	}
+}
+
+void USplineComponent::RemovePropertyAtIndex(int32 Index, FName PropertyName)
+{
+	if (UE::SplineComponent::GUseSplineCurves == false && Spline.SupportsAttributes())
+	{
+		Spline.RemoveAttributeValue(Index, PropertyName);
+	}
+}
+
+bool USplineComponent::SupportsAttributes() const
+{
+	return UE::SplineComponent::GUseSplineCurves == false && Spline.SupportsAttributes();
+}
+
+void USplineComponent::CreateFloatPropertyChannel(FName PropertyName)
+{
+	CreatePropertyChannel<float>(GetSplinePointsMetadata(), Spline, PropertyName);
+}
+
+void USplineComponent::CreateVectorPropertyChannel(FName PropertyName)
+{
+	CreatePropertyChannel<FVector>(GetSplinePointsMetadata(), Spline, PropertyName);
+}
+
 float USplineComponent::GetFloatPropertyAtSplineInputKey(float InKey, FName PropertyName) const
 {
-	return GetPropertyValueAtSplineInputKey<float>(GetSplinePointsMetadata(), InKey, PropertyName);
+	return GetPropertyAtSplineInputKey<float>(GetSplinePointsMetadata(), Spline, InKey, PropertyName);
 }
 
 FVector USplineComponent::GetVectorPropertyAtSplineInputKey(float InKey, FName PropertyName) const
 {
-	return GetPropertyValueAtSplineInputKey<FVector>(GetSplinePointsMetadata(), InKey, PropertyName);
+	return GetPropertyAtSplineInputKey<FVector>(GetSplinePointsMetadata(), Spline, InKey, PropertyName);
+}
+
+int32 USplineComponent::SetFloatPropertyAtSplineInputKey(float InKey, float Value, FName PropertyName) const
+{
+	return SetPropertyAtSplineInputKey<float>(Spline, InKey, Value, PropertyName);
+}
+
+int32 USplineComponent::SetVectorPropertyAtSplineInputKey(float InKey, const FVector& Value, FName PropertyName) const
+{
+	return SetPropertyAtSplineInputKey<FVector>(Spline, InKey, Value, PropertyName);
+}
+
+float USplineComponent::GetFloatPropertyInputKeyAtIndex(int32 Index, FName PropertyName) const
+{
+	return GetInputKeyAtIndex(Spline, Index, PropertyName);
+}
+	
+int32 USplineComponent::SetFloatPropertyInputKeyAtIndex(int32 Index, float InputKey, FName PropertyName)
+{
+	return SetInputKeyAtIndex(Spline, Index, InputKey, PropertyName);
+}
+	
+float USplineComponent::GetFloatPropertyAtIndex(int32 Index, FName PropertyName) const
+{
+	return GetPropertyAtIndex<float>(Spline, Index, PropertyName);
+}
+	
+void USplineComponent::SetFloatPropertyAtIndex(int32 Index, float Value, FName PropertyName)
+{
+	SetPropertyAtIndex<float>(Spline, Index, Value, PropertyName);
 }
 
 void USplineComponent::SetClosedLoop(bool bInClosedLoop, bool bUpdateSpline)
@@ -887,7 +1065,7 @@ void USplineComponent::AddSplinePointAtIndex(const FVector& Position, int32 Inde
 		NewPoint.Rotation = FQuat::Identity.Rotator();
 		NewPoint.Scale = FVector(1.f);
 
-		Spline.AddPointAtIndex(NewPoint, Index);
+		Spline.InsertPoint(NewPoint, Index);
 
 		WarninglessSplineCurves().Position.Points.Insert(FInterpCurvePoint<FVector>(InKey, TransformedPosition, FVector::ZeroVector, FVector::ZeroVector, CIM_CurveAuto), Index);
 		WarninglessSplineCurves().Rotation.Points.Insert(FInterpCurvePoint<FQuat>(InKey, FQuat::Identity, FQuat::Identity, FQuat::Identity, CIM_CurveAuto), Index);
@@ -898,10 +1076,9 @@ void USplineComponent::AddSplinePointAtIndex(const FVector& Position, int32 Inde
 			Metadata->InsertPoint(InKey, Index, bClosedLoop);
 		}
 		
-		NumPoints++;
-
 		// Adjust subsequent points' input keys to make room for the value just added
-		for (int I = Index + 1; I < NumPoints; ++I)
+		int32 WarninglessSplinePointsNum = WarninglessSplineCurves().Position.Points.Num();
+		for (int I = Index + 1; I < WarninglessSplinePointsNum; ++I)
 		{
 			WarninglessSplineCurves().Position.Points[I].InVal += 1.0f;
 			WarninglessSplineCurves().Rotation.Points[I].InVal += 1.0f;
@@ -1180,7 +1357,7 @@ float USplineComponent::GetInputKeyValueAtSplinePoint(int32 PointIndex) const
 	
 	return UE::SplineComponent::GUseSplineCurves
 		? WarninglessSplineCurves().Position.Points[PointIndex].InVal
-		: Spline.GetSplinePointParameter(PointIndex);
+		: Spline.GetParameterAtIndex(PointIndex);
 }
 
 FSplinePoint USplineComponent::GetSplinePointAt(int32 PointIndex, ESplineCoordinateSpace::Type CoordinateSpace) const
@@ -1287,20 +1464,29 @@ void USplineComponent::GetLocationAndTangentAtSplinePoint(int32 PointIndex, FVec
 
 float USplineComponent::GetDistanceAlongSplineAtSplinePoint(int32 PointIndex) const
 {
-	const int32 NumPoints = GetNumberOfSplinePoints();
-	const int32 NumSegments = bClosedLoop ? NumPoints : NumPoints - 1;
-
-	const int32 NumReparamPoints = UE::SplineComponent::GUseSplineCurves
-		? WarninglessSplineCurves().ReparamTable.Points.Num()
-		: Spline.GetSplineReparamTable().Points.Num();
+	if (IsClosedLoop() && PointIndex == GetNumberOfSplinePoints())
+	{
+		// Special case, if we are closed and the index here is 1 past the last valid point, the length is the full spline.
+		return GetSplineLength();
+	}
 	
-	// Ensure that if the reparam table is not prepared yet we don't attempt to access it. This can happen
-	// early in the construction of the spline component object.
-	if ((PointIndex >= 0) && (PointIndex < NumSegments + 1) && ((PointIndex * ReparamStepsPerSegment) < NumReparamPoints))
+	if (UE::SplineComponent::GUseSplineCurves)
+	{
+		const int32 NumPoints = GetNumberOfSplinePoints();
+        const int32 NumSegments = bClosedLoop ? NumPoints : NumPoints - 1;
+        const int32 NumReparamPoints = WarninglessSplineCurves().ReparamTable.Points.Num();
+        
+        // Ensure that if the reparam table is not prepared yet we don't attempt to access it. This can happen
+        // early in the construction of the spline component object.
+        if ((PointIndex >= 0) && (PointIndex < NumSegments + 1) && ((PointIndex * ReparamStepsPerSegment) < NumReparamPoints))
+        {
+        	return WarninglessSplineCurves().ReparamTable.Points[PointIndex * ReparamStepsPerSegment].InVal;
+        }
+	}
+	else
 	{
-		return UE::SplineComponent::GUseSplineCurves
-			? WarninglessSplineCurves().ReparamTable.Points[PointIndex * ReparamStepsPerSegment].InVal
-			: Spline.GetSplineReparamTable().Points[PointIndex * ReparamStepsPerSegment].InVal;
+		const float ParameterAtIndex = Spline.GetParameterAtIndex(PointIndex);
+		return Spline.GetDistanceAtParameter(ParameterAtIndex);
 	}
 
 	return 0.0f;
@@ -1368,7 +1554,7 @@ float USplineComponent::GetInputKeyValueAtDistanceAlongSpline(float Distance) co
 	
 	return UE::SplineComponent::GUseSplineCurves
 		? WarninglessSplineCurves().ReparamTable.Eval(Distance, 0.0f)
-		: Spline.GetSplineReparamTable().Eval(Distance, 0.0f);
+		: Spline.GetParameterAtDistance(Distance);
 }
 
 float USplineComponent::GetTimeAtDistanceAlongSpline(float Distance) const
@@ -1383,14 +1569,14 @@ float USplineComponent::GetTimeAtDistanceAlongSpline(float Distance) const
 	const float TimeMultiplier = Duration / (bClosedLoop ? NumPoints : (NumPoints - 1.0f));
 	return (UE::SplineComponent::GUseSplineCurves
 		? WarninglessSplineCurves().ReparamTable.Eval(Distance, 0.0f)
-		: Spline.GetSplineReparamTable().Eval(Distance, 0.0f)) * TimeMultiplier;
+		: Spline.GetParameterAtDistance(Distance)) * TimeMultiplier;
 }
 
 FVector USplineComponent::GetLocationAtDistanceAlongSpline(float Distance, ESplineCoordinateSpace::Type CoordinateSpace) const
 {
 	const float Param = UE::SplineComponent::GUseSplineCurves
 		? WarninglessSplineCurves().ReparamTable.Eval(Distance, 0.0f)
-		: Spline.GetSplineReparamTable().Eval(Distance, 0.0f);
+		: Spline.GetParameterAtDistance(Distance);
 	return GetLocationAtSplineInputKey(Param, CoordinateSpace);
 }
 
@@ -1398,7 +1584,7 @@ FVector USplineComponent::GetTangentAtDistanceAlongSpline(float Distance, ESplin
 {
 	const float Param = UE::SplineComponent::GUseSplineCurves
 		? WarninglessSplineCurves().ReparamTable.Eval(Distance, 0.0f)
-		: Spline.GetSplineReparamTable().Eval(Distance, 0.0f);
+		: Spline.GetParameterAtDistance(Distance);
 	return GetTangentAtSplineInputKey(Param, CoordinateSpace);
 }
 
@@ -1407,7 +1593,7 @@ FVector USplineComponent::GetDirectionAtDistanceAlongSpline(float Distance, ESpl
 {
 	const float Param = UE::SplineComponent::GUseSplineCurves
 		? WarninglessSplineCurves().ReparamTable.Eval(Distance, 0.0f)
-		: Spline.GetSplineReparamTable().Eval(Distance, 0.0f);
+		: Spline.GetParameterAtDistance(Distance);
 	return GetDirectionAtSplineInputKey(Param, CoordinateSpace);
 }
 
@@ -1415,7 +1601,7 @@ FQuat USplineComponent::GetQuaternionAtDistanceAlongSpline(float Distance, ESpli
 {
 	const float Param = UE::SplineComponent::GUseSplineCurves
 		? WarninglessSplineCurves().ReparamTable.Eval(Distance, 0.0f)
-		: Spline.GetSplineReparamTable().Eval(Distance, 0.0f);
+		: Spline.GetParameterAtDistance(Distance);
 	return GetQuaternionAtSplineInputKey(Param, CoordinateSpace);
 }
 
@@ -1423,7 +1609,7 @@ FRotator USplineComponent::GetRotationAtDistanceAlongSpline(float Distance, ESpl
 {
 	const float Param = UE::SplineComponent::GUseSplineCurves
 		? WarninglessSplineCurves().ReparamTable.Eval(Distance, 0.0f)
-		: Spline.GetSplineReparamTable().Eval(Distance, 0.0f);
+		: Spline.GetParameterAtDistance(Distance);
 	return GetRotationAtSplineInputKey(Param, CoordinateSpace);
 }
 
@@ -1431,7 +1617,7 @@ FVector USplineComponent::GetUpVectorAtDistanceAlongSpline(float Distance, ESpli
 {
 	const float Param = UE::SplineComponent::GUseSplineCurves
 		? WarninglessSplineCurves().ReparamTable.Eval(Distance, 0.0f)
-		: Spline.GetSplineReparamTable().Eval(Distance, 0.0f);
+		: Spline.GetParameterAtDistance(Distance);
 	return GetUpVectorAtSplineInputKey(Param, CoordinateSpace);
 }
 
@@ -1439,7 +1625,7 @@ FVector USplineComponent::GetRightVectorAtDistanceAlongSpline(float Distance, ES
 {
 	const float Param = UE::SplineComponent::GUseSplineCurves
 		? WarninglessSplineCurves().ReparamTable.Eval(Distance, 0.0f)
-		: Spline.GetSplineReparamTable().Eval(Distance, 0.0f);
+		: Spline.GetParameterAtDistance(Distance);
 	return GetRightVectorAtSplineInputKey(Param, CoordinateSpace);
 }
 
@@ -1447,7 +1633,7 @@ float USplineComponent::GetRollAtDistanceAlongSpline(float Distance, ESplineCoor
 {
 	const float Param = UE::SplineComponent::GUseSplineCurves
 		? WarninglessSplineCurves().ReparamTable.Eval(Distance, 0.0f)
-		: Spline.GetSplineReparamTable().Eval(Distance, 0.0f);
+		: Spline.GetParameterAtDistance(Distance);
 	return GetRollAtSplineInputKey(Param, CoordinateSpace);
 }
 
@@ -1455,7 +1641,7 @@ FVector USplineComponent::GetScaleAtDistanceAlongSpline(float Distance) const
 {
 	const float Param = UE::SplineComponent::GUseSplineCurves
 		? WarninglessSplineCurves().ReparamTable.Eval(Distance, 0.0f)
-		: Spline.GetSplineReparamTable().Eval(Distance, 0.0f);
+		: Spline.GetParameterAtDistance(Distance);
 	return GetScaleAtSplineInputKey(Param);
 }
 
@@ -1463,7 +1649,7 @@ FTransform USplineComponent::GetTransformAtDistanceAlongSpline(float Distance, E
 {
 	const float Param = UE::SplineComponent::GUseSplineCurves
 		? WarninglessSplineCurves().ReparamTable.Eval(Distance, 0.0f)
-		: Spline.GetSplineReparamTable().Eval(Distance, 0.0f);
+		: Spline.GetParameterAtDistance(Distance);
 	return GetTransformAtSplineInputKey(Param, CoordinateSpace, bUseScale);
 }
 
@@ -1968,16 +2154,27 @@ bool USplineComponent::ConvertSplineToPolyline_InDistanceRange(ESplineCoordinate
 		return bHasAdded;
 	} // end of the wrap-around case, after this values will be in the normal range
 	
-	int32 StartIndex = UE::SplineComponent::GUseSplineCurves
-		? WarninglessSplineCurves().ReparamTable.GetPointIndexForInputValue(RangeStart)
-		: Spline.GetSplineReparamTable().GetPointIndexForInputValue(RangeStart);
-	
-	int32 EndIndex = UE::SplineComponent::GUseSplineCurves
-		? WarninglessSplineCurves().ReparamTable.GetPointIndexForInputValue(RangeEnd)
-		: Spline.GetSplineReparamTable().GetPointIndexForInputValue(RangeEnd);
-	
-	int32 SegmentStart = StartIndex / ReparamStepsPerSegment;
-	int32 SegmentEnd = FMath::Min(NumSegments, 1 + EndIndex / ReparamStepsPerSegment);
+	int32 SegmentStart;
+	int32 SegmentEnd;
+	if (UE::SplineComponent::GUseSplineCurves)
+	{
+		int32 StartIndex = WarninglessSplineCurves().ReparamTable.GetPointIndexForInputValue(RangeStart);
+		int32 EndIndex = WarninglessSplineCurves().ReparamTable.GetPointIndexForInputValue(RangeEnd);
+		SegmentStart = StartIndex / ReparamStepsPerSegment;
+		SegmentEnd = FMath::Min(NumSegments, 1 + EndIndex / ReparamStepsPerSegment);
+	}
+	else
+	{
+		// with new spline, we are strongly assuming that point index is always same as param
+		// so, we can convert param directly to segment index.
+		// we do this without making the assumptions about the reparameterization table (unlike the other branch)
+
+		const float StartParam = Spline.GetParameterAtDistance(RangeStart);
+		const float EndParam = Spline.GetParameterAtDistance(RangeEnd);
+
+		SegmentStart = FMath::FloorToInt32(StartParam);
+		SegmentEnd = FMath::Min(NumSegments, FMath::CeilToInt32(EndParam));
+	}
 
 	TArray<FVector> NewPoints;
 	TArray<double> NewDistances;
@@ -2205,6 +2402,12 @@ void USplineComponent::Draw(FPrimitiveDrawInterface* PDI, const FSceneView* View
 	const int32 GrabHandleSize = 6;
 	FVector OldKeyPos(0);
 
+	auto GetLineColor = [](int32 KeyIdx, int32 NumSegments, int32 NumSteps) -> FLinearColor
+	{
+		float Normalized = (KeyIdx + (NumSegments * NumSteps)) / static_cast<float>(NumSegments * NumSteps);
+		return FLinearColor::MakeFromHSV8(static_cast<uint8>(Normalized * 255), 255, 255);
+	};
+	
 	const int32 NumPoints = SplineInfo.Points.Num();
 	const int32 NumSegments = SplineInfo.bIsLooped ? NumPoints : NumPoints - 1;
 	for (int32 KeyIdx = 0; KeyIdx < NumSegments + 1; KeyIdx++)
@@ -2250,7 +2453,7 @@ void USplineComponent::Draw(FPrimitiveDrawInterface* PDI, const FSceneView* View
 					const float Key = (KeyIdx - 1) + (StepIdx / static_cast<float>(NumSteps));
 					const FVector NewPos = LocalToWorld.TransformPosition(SplineInfo.Eval(Key, FVector(0)));
 #if WITH_EDITOR
-					PDI->DrawTranslucentLine(OldPos, NewPos, LineColor, DepthPriorityGroup, SegmentLineThickness);
+					PDI->DrawTranslucentLine(OldPos, NewPos, GetLineColor(KeyIdx * NumSteps + StepIdx, NumSegments, NumSteps), DepthPriorityGroup, SegmentLineThickness);
 #else
 					PDI->DrawTranslucentLine(OldPos, NewPos, LineColor, DepthPriorityGroup);
 #endif
@@ -2380,13 +2583,6 @@ const FInterpCurveVector& USplineComponent::GetSplinePointsScale() const
 		: Spline.GetSplinePointsScale();
 }
 
-const FInterpCurveFloat& USplineComponent::GetSplineReparamTable() const
-{
-	return UE::SplineComponent::GUseSplineCurves
-		? WarninglessSplineCurves().ReparamTable
-		: Spline.GetSplineReparamTable();
-}
-
 TArray<ESplinePointType::Type> USplineComponent::GetEnabledSplinePointTypes() const
 {
 	return 
diff --git a/Engine/Source/Runtime/Engine/Private/Curves/Spline.cpp b/Engine/Source/Runtime/Engine/Private/Curves/Spline.cpp
--- a/Engine/Source/Runtime/Engine/Private/Curves/Spline.cpp
+++ b/Engine/Source/Runtime/Engine/Private/Curves/Spline.cpp
@@ -4,6 +4,7 @@
 
 #include "Components/SplineComponent.h"	// for FSplineCurves, ideally removed
 #include "UObject/FortniteMainBranchObjectVersion.h"
+#include "Splines/TangentBezierSpline.h"
 
 #include UE_INLINE_GENERATED_CPP_BY_NAME(Spline)
 
@@ -35,26 +36,345 @@ FAutoConsoleVariableRef CVarSplineImplementation(
 );
 #endif
 	
+bool GApproximateTangents = false;
+FAutoConsoleVariableRef CVarApproximateTangents(
+	TEXT("Spline.ApproximateTangents"),
+	GApproximateTangents,
+	TEXT("True if we should approximate tangents using the central difference formula.")
+);
+
+bool GFallbackFindNearest = false;
+FAutoConsoleVariableRef CVarFallbackFindNearest(
+	TEXT("Spline.FallbackFindNearest"),
+	GFallbackFindNearest,
+	TEXT("True if we should implement FindNearest and FindNearestOnSegment using an intermediate spline representation. Only applies if Spline.Implementation == 2.")
+);
+
+}
+
+using namespace UE::Geometry::Spline;
+
+/**
+ * Converts from EInterpCurveMode to ETangentMode
+ */
+inline ETangentMode ConvertInterpCurveModeToTangentMode(EInterpCurveMode Mode)
+{
+	switch (Mode)
+	{
+	case CIM_Linear:
+		return ETangentMode::Linear;
+	case CIM_CurveAuto:
+		return ETangentMode::Auto;
+	case CIM_Constant:
+		return ETangentMode::Constant;
+	case CIM_CurveUser:
+		return ETangentMode::User;
+	case CIM_CurveBreak:
+		return ETangentMode::Broken;
+	case CIM_CurveAutoClamped:
+		return ETangentMode::AutoClamped;
+	case CIM_Unknown:
+	default:
+		return ETangentMode::Unknown;
+	}
+}
+
+/**
+ * Converts from ETangentMode to EInterpCurveMode
+ */
+inline EInterpCurveMode ConvertTangentModeToInterpCurveMode(ETangentMode Mode)
+{
+	switch (Mode)
+	{
+	case ETangentMode::Linear:
+		return CIM_Linear;
+	case ETangentMode::Auto:
+		return CIM_CurveAuto;
+	case ETangentMode::Constant:
+		return CIM_Constant;
+	case ETangentMode::User:
+		return CIM_CurveUser;
+	case ETangentMode::Broken:
+		return CIM_CurveBreak;
+	case ETangentMode::AutoClamped:
+		return CIM_CurveAutoClamped;
+	case ETangentMode::Unknown:
+	default:
+		return CIM_Unknown;
+	}
+}
+
+/**
+ * FNewSpline Definition
+ *
+ * A spline that provides tangent-based control over curve shape while using 
+ * piecewise Bezier curves internally for evaluation. Supports both manual
+ * tangent control and automatic tangent computation.
+ */
+
+class FNewSpline: public FTangentBezierSpline
+{
+	
+public:
+	
+    FNewSpline();
+    FNewSpline(const FNewSpline& Other);
+	FNewSpline(const FLegacySpline& Other);
+	FNewSpline(const FSplineCurves& Other);
+    FNewSpline& operator=(const FNewSpline& Other);
+
+	bool operator==(const FNewSpline& Other) const;
+	bool operator!=(const FNewSpline& Other) const;
+
+    bool Serialize(FArchive& Ar);
+    friend FArchive& operator<<(FArchive& Ar, FNewSpline& Spline)
+    {
+        Spline.Serialize(Ar);
+        return Ar;
+    }
+
+	float GetSegmentLength(const int32 Index, const float Param, const FVector& Scale3D = FVector(1.0f)) const;
+	float GetSplineLength() const;
+	
+    void SetClosedLoop(bool bInClosedLoop);
+
+	void Reset();
+	void ResetToDefault();
+	
+	void ResetRotation();
+	void ResetScale();
+	
+    void AddPoint(const FSplinePoint& Point);
+	void InsertPoint(const FSplinePoint& Point, int32 Index);
+	FSplinePoint GetPoint(int32 Index) const;
+    void RemovePoint(int32 Index);
+
+	void SetLocation(int32 Index, const FVector& InLocation);
+	FVector GetLocation(const int32 Index) const;
+	
+	void SetInTangent(const int32 Index, const FVector& InTangent);
+	FVector GetInTangent(const int32 Index) const;
+	
+	void SetOutTangent(const int32 Index, const FVector& OutTangent);
+	FVector GetOutTangent(const int32 Index) const;
+
+	void SetRotation(int32 Index, const FQuat& InRotation);
+	FQuat GetRotation(const int32 Index) const;
+	
+	void SetScale(int32 Index, const FVector& InScale);
+	FVector GetScale(const int32 Index) const;
+	
+	void SetSplinePointType(int32 Index, EInterpCurveMode Type);
+	EInterpCurveMode GetSplinePointType(int32 Index) const;
+
+	float GetParameterAtIndex(int32 Index) const;
+	float GetParameterAtDistance(float Distance) const;
+	float GetDistanceAtParameter(float Parameter) const;
+	
+	float FindNearest(const FVector& InLocation) const;
+	float FindNearestOnSegment(const FVector& InLocation, int32 SegmentIndex) const;
+
+	/** Parameter based evaluation API. */
+    FVector Evaluate(float Parameter) const;
+    FVector EvaluateDerivative(float Parameter) const;
+	FQuat EvaluateRotation(float Parameter) const;
+	FVector EvaluateScale(float Parameter) const;
+
+	template <typename AttrType> TArray<FName> GetAttributeChannelNames() const;
+	template <typename AttrType> AttrType EvaluateAttribute(FName Name, float Parameter) const;
+	int32 SetAttributeParameter(const FName& Name, int32 Index, float ParentSpaceParameter);
+	
+	/** Updates the spline using the current configuration. */
+	void UpdateSpline();
+
+	/** Updates the spline configuration, then updates the spline. */
+	void UpdateSpline(const FSpline::FUpdateSplineParams& Params);
+
+	const FInterpCurveVector& GetSplinePointsPosition() const { return PositionCurve; }
+	const FInterpCurveQuat& GetSplinePointsRotation() const { return RotationCurve; }
+	const FInterpCurveVector& GetSplinePointsScale() const { return ScaleCurve; }
+
+private:
+	
+	static inline const FName RotationAttrName = FName("Rotation");
+	static inline const FName ScaleAttrName = FName("Scale");
+	static inline const FName ParamToLengthAttrName = FName("ParamToLength");
+	static inline const FName LengthToParamAttrName = FName("LengthToParam");
+
+    /**
+     * We keep these up to date with Points.
+     * Exists to implement the legacy interface provided by FSpline.
+     */
+    FInterpCurveVector PositionCurve;
+	FInterpCurveQuat RotationCurve;
+	FInterpCurveVector ScaleCurve;
+	
+	int32 ReparamStepsPerSegment = 10;
+
+private:
+
+	/** Clears and repopulates reparameterization attribute channels. */
+	void UpdateReparamTable();
+
+	/** Clears and repopulates PositionCurve, RotationCurve, and ScaleCurve */
+	void RebuildLegacyCurves();
+};
+
+/**
+ * Retrieves the names of all attribute channels of the given template type in the spline
+ * @tparam AttrType - attribute type
+ * @return Array of attribute channel names
+ */
+template <typename AttrType>
+TArray<FName> FNewSpline::GetAttributeChannelNames() const
+{
+	TArray<FName> ChannelNames = InternalSpline.GetAttributeChannelNames<AttrType>();
+	ChannelNames.Remove(RotationAttrName);
+	ChannelNames.Remove(ScaleAttrName);
+	ChannelNames.Remove(ParamToLengthAttrName);
+	ChannelNames.Remove(LengthToParamAttrName);
+	return ChannelNames;
 }
 
+template <typename AttrType>
+AttrType FNewSpline::EvaluateAttribute(FName Name, float Parameter) const
+{
+	const int32 NumEffectivePoints = IsClosedLoop() ? GetNumPoints() + 1 : GetNumPoints();
+	Parameter = GetParameterRange().Interpolate(Parameter / (NumEffectivePoints - 1));
+	return InternalSpline.EvaluateAttribute<AttrType>(Name, Parameter);
+}
+
+int32 FNewSpline::SetAttributeParameter(const FName& Name, int32 Index, float Parameter)
+{
+	const int32 NumEffectivePoints = IsClosedLoop() ? GetNumPoints() + 1 : GetNumPoints();
+   	Parameter = GetParameterRange().Interpolate(Parameter / (NumEffectivePoints - 1));
+	return InternalSpline.SetAttributeParameter(Name, Index, Parameter);
+}
+
+
+/** FLegacySpline Definition */
+
+class FLegacySpline
+{
+	FInterpCurveVector PositionCurve;
+	FInterpCurveQuat RotationCurve;
+	FInterpCurveVector ScaleCurve;
+	FInterpCurveFloat ReparamTable;
+	
+public:
+
+	ENGINE_API FLegacySpline();
+	ENGINE_API FLegacySpline(const FNewSpline& Other);
+	ENGINE_API FLegacySpline(const FSplineCurves& Other);
+	
+	/* Control Point Index Interface */
+
+	void AddPoint(const FSplinePoint& InPoint);
+	void InsertPoint(const FSplinePoint& InPoint, int32 Index);
+	FSplinePoint GetPoint(const int32 Index) const;
+	void RemovePoint(int32 Index);
+	
+	void SetLocation(int32 Index, const FVector& InLocation);
+	FVector GetLocation(const int32 Index) const;
+	
+	void SetInTangent(const int32 Index, const FVector& InTangent);
+	FVector GetInTangent(const int32 Index) const;
+	
+	void SetOutTangent(const int32 Index, const FVector& OutTangent);
+	FVector GetOutTangent(const int32 Index) const;
+
+	void SetRotation(int32 Index, const FQuat& InRotation);
+	FQuat GetRotation(const int32 Index) const;
+	
+	void SetScale(int32 Index, const FVector& InScale);
+	FVector GetScale(const int32 Index) const;
+	
+	void SetSplinePointType(int32 Index, EInterpCurveMode Type);
+	EInterpCurveMode GetSplinePointType(int32 Index) const;
+
+	float GetParameterAtIndex(int32 Index) const;
+	float GetParameterAtDistance(float Distance) const;
+	float GetDistanceAtParameter(float Parameter) const;
+	
+	/* Parameter Interface */
+
+	FVector Evaluate(float Param) const;
+	FVector EvaluateDerivative(float Param) const;
+	FQuat EvaluateRotation(float Param) const;
+	FVector EvaluateScale(float Param) const;
+	float FindNearest(const FVector& InLocation) const;
+	float FindNearestOnSegment(const FVector& InLocation, int32 SegmentIndex) const;
+	
+	/* Misc Interface */
+	
+	bool operator==(const FLegacySpline& Other) const
+	{
+		return PositionCurve == Other.PositionCurve
+			&& RotationCurve == Other.RotationCurve
+			&& ScaleCurve == Other.ScaleCurve;
+	}
+
+	bool operator!=(const FLegacySpline& Other) const
+	{
+		return !(*this == Other);
+	}
+	
+	friend ENGINE_API FArchive& operator<<(FArchive& Ar, FLegacySpline& Spline)
+	{
+		Spline.Serialize(Ar);
+		return Ar;
+	}
+	bool Serialize(FArchive& Ar);
+	
+	FInterpCurveVector& GetSplinePointsPosition() { return PositionCurve; }
+	const FInterpCurveVector& GetSplinePointsPosition() const { return PositionCurve; }
+	FInterpCurveQuat& GetSplinePointsRotation() { return RotationCurve; }
+	const FInterpCurveQuat& GetSplinePointsRotation() const { return RotationCurve; }
+	FInterpCurveVector& GetSplinePointsScale() { return ScaleCurve; }
+	const FInterpCurveVector& GetSplinePointsScale() const { return ScaleCurve; }
+
+	/** Returns the length of the specified spline segment up to the parametric value given */
+	float GetSegmentLength(const int32 Index, const float Param, const FVector& Scale3D = FVector(1.0f)) const;
+
+	/** Returns total length along this spline */
+	float GetSplineLength() const;
+	
+	/** Returns the total number of control points on this spline. */
+	int32 GetNumControlPoints() const;
+	
+	/** Reset the spline to an empty spline. */
+	void Reset();
+	
+	/** Reset the rotation attribute channel to default values. */
+	void ResetRotation();
+	
+	/** Reset the scale attribute channel to default values. */
+	void ResetScale();
+	
+	/** Reset the spline to the default spline (2 points). */
+	void ResetToDefault();
+	
+	/** Update the spline's internal data according to the passed-in params. */
+	void UpdateSpline(const FSpline::FUpdateSplineParams& InParams);
+};
+
+
+
+/** FSpline Implementation */
+
 FSpline::FSpline()
 : CurrentImplementation(UE::Spline::GImplementation)
 , Version(0xffffffff)
 {
 #if WITH_EDITOR
-	PreviousImplementation = UE::Spline::GImplementation;
+	PreviousImplementation = 0;
 #endif
 	
 	switch (CurrentImplementation)
 	{
-	default:	// Intentionally doing nothing
-		break;
-	case 1:
-		LegacyData = MakeShared<FLegacySpline>();
-		break;
-	case 2:
-		// todo: implement for new spline
-		break;
+	default: break;		// Intentionally doing nothing
+	case 1: LegacyData = MakeShared<FLegacySpline>(); break;
+	case 2: NewData = MakeShared<FNewSpline>(); break;
 	}
 
 #if WITH_EDITOR
@@ -70,23 +390,30 @@ FSpline::FSpline(const FSpline& Other)
 
 FSpline& FSpline::operator=(const FSpline& Other)
 {
+	Reset();
+	
 	CurrentImplementation = Other.CurrentImplementation;
 
 	switch (CurrentImplementation)
 	{
-	default:
-		break;
+	default: break;		// Intentionally doing nothing
 	case 1:
 		if (Other.LegacyData)
-			{ LegacyData = MakeShared<FLegacySpline>(*Other.LegacyData); }
+		{
+			FLegacySpline Temp = *Other.LegacyData;
+			LegacyData = MakeShared<FLegacySpline>(Temp);
+		}
 		else
 			{ LegacyData = MakeShared<FLegacySpline>(); }
 		break;
 	case 2:
-		// if (Other.NewData)
-		// 	{ NewData = MakeShared<FHermiteSpline>(*Other.NewData); }
-		// else
-		// 	{ NewData = MakeShared<FHermiteSpline>(); }
+		if (Other.NewData)
+		{
+			FNewSpline Temp = *Other.NewData;
+			NewData = MakeShared<FNewSpline>(Temp);
+		}
+		else
+			{ NewData = MakeShared<FNewSpline>(); }
 		break;
 	}
 	
@@ -115,15 +442,9 @@ FSpline FSpline::FromSplineCurves(const FSplineCurves& InSpline)
 
 	switch (NewSpline.CurrentImplementation)
 	{
-	default:
-		// Intentionally doing nothing
-		break;
-	case 1:
-		NewSpline.LegacyData = MakeShared<FLegacySpline>(FLegacySpline::FromSplineCurves(InSpline));
-		break;
-	case 2:
-		// todo: implement for new spline
-		break;
+	default: break;		// Intentionally doing nothing
+	case 1: NewSpline.LegacyData = MakeShared<FLegacySpline>(InSpline); break;
+	case 2: NewSpline.NewData = MakeShared<FNewSpline>(InSpline); break;
 	}
 	
 	return NewSpline;
@@ -133,31 +454,19 @@ void FSpline::AddPoint(const FSplinePoint& InPoint)
 {
 	switch (CurrentImplementation)
 	{
-	default:
-		// Intentionally doing nothing
-		break;
-	case 1:
-		LegacyData->AddPoint(InPoint);
-		break;
-	case 2:
-		// todo: implement for new spline
-		break;
+	default: break;		// Intentionally doing nothing
+	case 1: LegacyData->AddPoint(InPoint); break;
+	case 2: NewData->AddPoint(InPoint); break;
 	}
 }
 
-void FSpline::AddPointAtIndex(const FSplinePoint& InPoint, int32 Index)
+void FSpline::InsertPoint(const FSplinePoint& InPoint, int32 Index)
 {
 	switch (CurrentImplementation)
 	{
-	default:
-		// Intentionally doing nothing
-		break;
-	case 1:
-		LegacyData->AddPointAtIndex(InPoint, Index);
-		break;
-	case 2:
-		// todo: implement for new spline
-		break;
+	default: break;		// Intentionally doing nothing
+	case 1: LegacyData->InsertPoint(InPoint, Index); break;
+	case 2: NewData->InsertPoint(InPoint, Index); break;
 	}
 }
 
@@ -165,13 +474,9 @@ FSplinePoint FSpline::GetPoint(const int32 Index) const
 {
 	switch (CurrentImplementation)
 	{
-	default:
-		return FSplinePoint();
-	case 1:
-		return LegacyData->GetPoint(Index);
-	case 2:
-		// todo: implement for new spline
-		return FSplinePoint();
+	default: return FSplinePoint();		// Intentionally doing nothing
+	case 1: return LegacyData->GetPoint(Index);
+	case 2: return NewData->GetPoint(Index);
 	}
 }
 
@@ -179,31 +484,19 @@ void FSpline::RemovePoint(const int32 Index)
 {
 	switch (CurrentImplementation)
 	{
-	default:
-		// Intentionally doing nothing
-		break;
-	case 1:
-		LegacyData->RemovePoint(Index);
-		break;
-	case 2:
-		// todo: implement for new spline
-		break;
+	default: break;		// Intentionally doing nothing
+	case 1: LegacyData->RemovePoint(Index); break;
+	case 2: NewData->RemovePoint(Index); break;
 	}
 }
 
-void FSpline::SetLocation(int32 PointIndex, const FVector& InLocation)
+void FSpline::SetLocation(int32 Index, const FVector& InLocation)
 {
 	switch (CurrentImplementation)
 	{
-	default:
-		// Intentionally doing nothing
-		break;
-	case 1:
-		LegacyData->SetLocation(PointIndex, InLocation);
-		break;
-	case 2:
-		// todo: implement for new spline
-		break;
+	default: break;		// Intentionally doing nothing
+	case 1: LegacyData->SetLocation(Index, InLocation); break;
+	case 2: NewData->SetLocation(Index, InLocation); break;
 	}
 }
 
@@ -211,14 +504,9 @@ FVector FSpline::GetLocation(const int32 Index) const
 {
 	switch (CurrentImplementation)
 	{
-	default:
-		// Intentionally doing nothing
-		return FVector();
-	case 1:
-		return LegacyData->GetLocation(Index);
-	case 2:
-		// todo: implement for new spline
-		return FVector();
+	default: return FVector();
+	case 1: return LegacyData->GetLocation(Index);
+	case 2: return NewData->GetLocation(Index);
 	}
 }
 
@@ -226,15 +514,9 @@ void FSpline::SetInTangent(const int32 Index, const FVector& InTangent)
 {
 	switch (CurrentImplementation)
 	{
-	default:
-		// Intentionally doing nothing
-		break;
-	case 1:
-		LegacyData->SetInTangent(Index, InTangent);
-		break;
-	case 2:
-		// todo: implement for new spline
-		break;
+	default: break;
+	case 1: LegacyData->SetInTangent(Index, InTangent); break;
+	case 2: NewData->SetInTangent(Index, InTangent); break;
 	}
 }
 
@@ -242,14 +524,9 @@ FVector FSpline::GetInTangent(const int32 Index) const
 {
 	switch (CurrentImplementation)
 	{
-	default:
-		// Intentionally doing nothing
-		return FVector();
-	case 1:
-		return LegacyData->GetInTangent(Index);
-	case 2:
-		// todo: implement for new spline
-		return FVector();
+	default: return FVector();
+	case 1: return LegacyData->GetInTangent(Index);
+	case 2: return NewData->GetInTangent(Index);
 	}
 }
 
@@ -257,15 +534,9 @@ void FSpline::SetOutTangent(const int32 Index, const FVector& OutTangent)
 {
 	switch (CurrentImplementation)
 	{
-	default:
-		// Intentionally doing nothing
-		break;
-	case 1:
-		LegacyData->SetOutTangent(Index, OutTangent);
-		break;
-	case 2:
-		// todo: implement for new spline
-		break;
+	default: break;
+	case 1: LegacyData->SetOutTangent(Index, OutTangent); break;
+	case 2: NewData->SetOutTangent(Index, OutTangent); break;
 	}
 }
 
@@ -273,30 +544,19 @@ FVector FSpline::GetOutTangent(const int32 Index) const
 {
 	switch (CurrentImplementation)
 	{
-	default:
-		// Intentionally doing nothing
-		return FVector();
-	case 1:
-		return LegacyData->GetOutTangent(Index);
-	case 2:
-		// todo: implement for new spline
-		return FVector();
+	default: return FVector();
+	case 1: return LegacyData->GetOutTangent(Index);
+	case 2: return NewData->GetOutTangent(Index);
 	}
 }
 
-void FSpline::SetRotation(int32 PointIndex, const FQuat& InRotation)
+void FSpline::SetRotation(int32 Index, const FQuat& InRotation)
 {
 	switch (CurrentImplementation)
 	{
-	default:
-		// Intentionally doing nothing
-		break;
-	case 1:
-		LegacyData->SetRotation(PointIndex, InRotation);
-		break;
-	case 2:
-		// todo: implement for new spline
-		break;
+	default: break;
+	case 1: LegacyData->SetRotation(Index, InRotation); break;
+	case 2: NewData->SetRotation(Index, InRotation); break;
 	}
 }
 
@@ -304,30 +564,19 @@ FQuat FSpline::GetRotation(const int32 Index) const
 {
 	switch (CurrentImplementation)
 	{
-	default:
-		// Intentionally doing nothing
-		return FQuat();
-	case 1:
-		return LegacyData->GetRotation(Index);
-	case 2:
-		// todo: implement for new spline
-		return FQuat();
+	default: return FQuat();
+	case 1: return LegacyData->GetRotation(Index);
+	case 2: return NewData->GetRotation(Index);
 	}
 }
 
-void FSpline::SetScale(int32 PointIndex, const FVector& InScale)
+void FSpline::SetScale(int32 Index, const FVector& InScale)
 {
 	switch (CurrentImplementation)
 	{
-	default:
-		// Intentionally doing nothing
-		break;
-	case 1:
-		LegacyData->SetScale(PointIndex, InScale);
-		break;
-	case 2:
-		// todo: implement for new spline
-		break;
+	default: break;
+	case 1: LegacyData->SetScale(Index, InScale); break;
+	case 2: NewData->SetScale(Index, InScale); break;
 	}
 }
 
@@ -335,60 +584,59 @@ FVector FSpline::GetScale(const int32 Index) const
 {
 	switch (CurrentImplementation)
 	{
-	default:
-		// Intentionally doing nothing
-		return FVector(1.0);
-	case 1:
-		return LegacyData->GetScale(Index);
-	case 2:
-		// todo: implement for new spline
-		return FVector(1.0);
+	default: return FVector(1.0);
+	case 1: return LegacyData->GetScale(Index);
+	case 2: return NewData->GetScale(Index);
 	}
 }
 
-void FSpline::SetSplinePointType(int32 PointIndex, EInterpCurveMode Type)
+void FSpline::SetSplinePointType(int32 Index, EInterpCurveMode Type)
 {
 	switch (CurrentImplementation)
 	{
-	default:
-		// Intentionally doing nothing
-		break;
-	case 1:
-		LegacyData->SetSplinePointType(PointIndex, Type);
-		break;
-	case 2:
-		// todo: implement for new spline
-		break;
+	default: break;
+	case 1: LegacyData->SetSplinePointType(Index, Type); break;
+	case 2: NewData->SetSplinePointType(Index, Type); break;
 	}
 }
 
-EInterpCurveMode FSpline::GetSplinePointType(int32 PointIndex) const
+EInterpCurveMode FSpline::GetSplinePointType(int32 Index) const
 {
 	switch (CurrentImplementation)
 	{
-	default:
-		// Intentionally doing nothing
-		return CIM_Unknown;
-	case 1:
-		return LegacyData->GetSplinePointType(PointIndex);
-	case 2:
-		// todo: implement for new spline
-		return CIM_Unknown;
+	default: return CIM_Unknown;
+	case 1: return LegacyData->GetSplinePointType(Index);
+	case 2: return NewData->GetSplinePointType(Index);
 	}
 }
 
-float FSpline::GetSplinePointParameter(int32 PointIndex) const
+float FSpline::GetParameterAtIndex(int32 Index) const
 {
 	switch (CurrentImplementation)
 	{
-	default:
-		// Intentionally doing nothing
-		return 0.f;
-	case 1:
-		return LegacyData->GetSplinePointParameter(PointIndex);
-	case 2:
-		// todo: implement for new spline
-		return 0.f;
+	default: return 0.f;
+	case 1: return LegacyData->GetParameterAtIndex(Index);
+	case 2: return NewData->GetParameterAtIndex(Index);
+	}
+}
+
+float FSpline::GetParameterAtDistance(float Distance) const
+{
+	switch (CurrentImplementation)
+	{
+	default: return 0.f;
+	case 1: return LegacyData->GetParameterAtDistance(Distance);
+	case 2: return NewData->GetParameterAtDistance(Distance);
+	}
+}
+
+float FSpline::GetDistanceAtParameter(float Parameter) const
+{
+	switch (CurrentImplementation)
+	{
+	default: return 0.f;
+	case 1: return LegacyData->GetDistanceAtParameter(Parameter);
+	case 2: return NewData->GetDistanceAtParameter(Parameter);
 	}
 }
 
@@ -396,29 +644,46 @@ FVector FSpline::Evaluate(float Param) const
 {
 	switch (CurrentImplementation)
 	{
-	default:
-		// Intentionally doing nothing
-		return FVector();
-	case 1:
-		return LegacyData->Evaluate(Param);
-	case 2:
-		// todo: implement for new spline
-		return FVector();
+	default: return FVector();
+	case 1: return LegacyData->Evaluate(Param);
+	case 2: return NewData->Evaluate(Param);
 	}
 }
 
 FVector FSpline::EvaluateDerivative(float Param) const
 {
-	switch (CurrentImplementation)
+	if (UE::Spline::GApproximateTangents)
 	{
-	default:
-		// Intentionally doing nothing
-		return FVector();
-	case 1:
-		return LegacyData->EvaluateDerivative(Param);
-	case 2:
-		// todo: implement for new spline
-		return FVector();
+		// Approximate using central difference.
+		
+		// This computes the tangent direction using central difference and
+		// assumes that tangent magnitude is linearly changing between control points.
+		// While the assumption about magnitude is probably wrong, it works well.
+		
+		auto ClampValidParameterRange = [Min = 0.f, Max = GetNumControlPoints() - 1](float Param)
+			{ return FMath::Clamp(Param, Min, Max); };
+
+		constexpr float H = UE_KINDA_SMALL_NUMBER;
+		const float ParamLow = ClampValidParameterRange(Param - H);
+		const float ParamHigh = ClampValidParameterRange(Param + H);
+		const FVector Tangent = ((Evaluate(ParamHigh) - Evaluate(ParamLow)) / (ParamHigh - ParamLow)).GetSafeNormal();
+
+		const int32 Index1 = FMath::Clamp((int32)Param, 0, GetNumControlPoints() - 1);
+		const int32 Index2 = FMath::Clamp((int32)(Param+1), 0, GetNumControlPoints() - 1);
+		const float Mag1 = GetInTangent(Index1).Length();
+		const float Mag2 = GetInTangent(Index2).Length();
+		const float Mag = FMath::Lerp(Mag1, Mag2, FMath::Frac(Param));
+
+		return Tangent * Mag;
+	}
+	else
+	{
+		switch (CurrentImplementation)
+		{
+		default: return FVector();
+		case 1: return LegacyData->EvaluateDerivative(Param);
+		case 2: return NewData->EvaluateDerivative(Param);
+		}
 	}
 }
 
@@ -426,14 +691,9 @@ FQuat FSpline::EvaluateRotation(float Param) const
 {
 	switch (CurrentImplementation)
 	{
-	default:
-		// Intentionally doing nothing
-		return FQuat();
-	case 1:
-		return LegacyData->EvaluateRotation(Param);
-	case 2:
-		// todo: implement for new spline
-		return FQuat();
+	default: return FQuat();
+	case 1: return LegacyData->EvaluateRotation(Param);
+	case 2: return NewData->EvaluateRotation(Param);
 	}
 }
 
@@ -441,57 +701,189 @@ FVector FSpline::EvaluateScale(float Param) const
 {
 	switch (CurrentImplementation)
 	{
-	default:
-		// Intentionally doing nothing
-		return FVector();
-	case 1:
-		return LegacyData->EvaluateScale(Param);
-	case 2:
-		// todo: implement for new spline
-		return FVector();
+	default: return FVector();
+	case 1: return LegacyData->EvaluateScale(Param);
+	case 2: return NewData->EvaluateScale(Param);
 	}	
 }
 
-float FSpline::FindNearest(const FVector& InLocation) const
+bool FSpline::HasAttributeChannel(FName AttributeName) const
 {
-	switch (CurrentImplementation)
+	if (NewData)
 	{
-	default:
-		// Intentionally doing nothing
-		return 0.f;
-	case 1:
-		return LegacyData->FindNearest(InLocation);
-	case 2:
-		// todo: implement for new spline
-		return 0.f;
+		return NewData->HasAttributeChannel(AttributeName);
 	}
+
+	return false;
 }
 
-float FSpline::FindNearestOnSegment(const FVector& InLocation, int32 SegmentIndex) const
+int32 FSpline::NumAttributeValues(FName AttributeName) const
 {
-	switch (CurrentImplementation)
+	if (NewData)
 	{
-	default:
-		// Intentionally doing nothing
-		return 0.f;
-	case 1:
-		return LegacyData->FindNearestOnSegment(InLocation, SegmentIndex);
-	case 2:
-		// todo: implement for new spline
-		return 0.f;
+		return NewData->NumAttributeValues(AttributeName);
 	}
+
+	return 0;
 }
 
-bool FSpline::Serialize(FArchive& Ar)
+float FSpline::GetAttributeParameter(int32 Index, const FName& Name) const
 {
-	Ar.UsingCustomVersion(FFortniteMainBranchObjectVersion::GUID);
-	
-	// Data format:
+	if (NewData)
+	{
+		return NewData->GetAttributeParameter(Name, Index);
+	}
 
-	// Byte 1 - The data format, determine by CurrentImplementation at the time of last save
-	// Remaining N Bytes - Spline data (or empty, if byte 0 is 0). The format is determined by byte 1.
-	
-	if (Ar.IsLoading())
+	return 0.f;
+}
+
+int32 FSpline::SetAttributeParameter(int32 Index, float Parameter, const FName& Name)
+{
+	if (NewData)
+	{
+		return NewData->SetAttributeParameter(Name, Index, Parameter);
+	}
+	return Index;
+}
+
+template<typename T>
+T FSpline::GetAttributeValue(int32 Index, const FName& Name) const
+{
+	if (NewData)
+	{
+		return NewData->GetAttributeValue<T>(Name, Index);
+	}
+	
+	return T();
+}
+template ENGINE_API float FSpline::GetAttributeValue<float>(int32 Index, const FName& Name) const;
+template ENGINE_API FVector FSpline::GetAttributeValue<FVector>(int32 Index, const FName& Name) const;
+
+template<typename T>
+void FSpline::SetAttributeValue(int32 Index, const T& Value, const FName& Name)
+{
+	if (NewData)
+	{
+		NewData->SetAttributeValue<T>(Name, Index, Value);
+	}
+}
+template ENGINE_API void FSpline::SetAttributeValue<float>(int32 Index, const float& Value, const FName& Name);
+template ENGINE_API void FSpline::SetAttributeValue<FVector>(int32 Index, const FVector& Value, const FName& Name);
+
+
+template<typename T>
+bool FSpline::CreateAttributeChannel(FName AttributeName) const
+{
+	if (NewData)
+	{
+		return NewData->CreateAttributeChannel<T>(AttributeName);
+	}
+	return false;
+}
+template ENGINE_API bool FSpline::CreateAttributeChannel<float>(FName AttributeName) const;
+template ENGINE_API bool FSpline::CreateAttributeChannel<FVector>(FName AttributeName) const;
+
+template<typename T>
+int32 FSpline::AddAttributeValue(float Param, const T& Value, FName AttributeName) const
+{
+	if (NewData)
+	{
+		return NewData->AddAttributeValue(AttributeName, Value, Param);
+	}
+	
+	return INDEX_NONE;
+}
+template ENGINE_API int32 FSpline::AddAttributeValue<float>(float Param, const float& Value, FName AttributeName) const;
+template ENGINE_API int32 FSpline::AddAttributeValue<FVector>(float Param, const FVector& Value, FName AttributeName) const;
+
+void FSpline::RemoveAttributeValue(int32 Index, FName AttributeName)
+{
+	if (NewData)
+	{
+		NewData->RemoveAttributeValue(AttributeName, Index);
+	}
+}
+
+TArray<FName> FSpline::GetFloatPropertyChannels() const
+{
+	TArray<FName> ChannelNames;
+	if (NewData)
+	{
+		ChannelNames = NewData->GetAttributeChannelNames<float>();
+	}
+	return ChannelNames;
+}
+
+TArray<FName> FSpline::GetVectorPropertyChannels() const
+{
+	TArray<FName> ChannelNames;
+	if (NewData)
+	{
+		ChannelNames = NewData->GetAttributeChannelNames<FVector>();
+	}
+	return ChannelNames;
+}
+
+template<typename T>
+T FSpline::EvaluateAttribute(float Param, FName AttributeName) const
+{
+	if (NewData)
+	{
+		return NewData->EvaluateAttribute<T>(AttributeName, Param);
+	}
+	return T(0);
+}
+
+template ENGINE_API float FSpline::EvaluateAttribute<float>(float Param, FName AttributeName) const;
+template ENGINE_API FVector FSpline::EvaluateAttribute<FVector>(float Param, FName AttributeName) const;
+float FSpline::FindNearest(const FVector& InLocation) const
+{
+	switch (CurrentImplementation)
+	{
+	default: return 0.f;
+	case 1: return LegacyData->FindNearest(InLocation);
+	case 2: return NewData->FindNearest(InLocation);
+	}
+}
+
+float FSpline::FindNearestOnSegment(const FVector& InLocation, int32 SegmentIndex) const
+{
+	switch (CurrentImplementation)
+	{
+	default: return 0.f;
+	case 1: return LegacyData->FindNearestOnSegment(InLocation, SegmentIndex);
+	case 2: return NewData->FindNearestOnSegment(InLocation, SegmentIndex);
+	}
+}
+
+bool FSpline::operator==(const FSpline& Other) const
+{
+	if (CurrentImplementation == Other.CurrentImplementation)
+	{
+		if (LegacyData && Other.LegacyData)
+		{
+			return *LegacyData == *Other.LegacyData;
+		}
+
+		if (NewData && Other.NewData)
+		{
+			return *NewData == *Other.NewData;
+		}
+	}
+
+	return false;
+}
+
+bool FSpline::Serialize(FArchive& Ar)
+{
+	Ar.UsingCustomVersion(FFortniteMainBranchObjectVersion::GUID);
+	
+	// Data format:
+
+	// Byte 1 - The data format, determine by CurrentImplementation at the time of last save
+	// Remaining N Bytes - Spline data (or empty, if byte 0 is 0). The format is determined by byte 1.
+	
+	if (Ar.IsLoading())
 	{
 		// using intermediate because I am scared of serializing a 2 bit field.
 		int8 PreviousImpl;
@@ -512,20 +904,18 @@ bool FSpline::Serialize(FArchive& Ar)
 			{
 				if (IsLegacy())
 				{
-					// read in the data directly to the legacy variant
 					LegacyData = MakeShared<FLegacySpline>();
 					Ar << *LegacyData;
 				}
 				else if (IsNew())
 				{
-					// read in the data to an intermediate legacy format structure
 					FLegacySpline IntermediateData;
 					Ar << IntermediateData;
-					// convert to new format
+					NewData = MakeShared<FNewSpline>(IntermediateData);
 				}
 				else if (!IsEnabled())
 				{
-					// read in the data but do nothing
+					// Intentionally doing nothing with data
 					FLegacySpline IntermediateData;
 					Ar << IntermediateData;
 				}
@@ -534,19 +924,21 @@ bool FSpline::Serialize(FArchive& Ar)
 			{
 				if (IsLegacy())
 				{
-					// read in the data to a new format structure
-					// convert to legacy format variant
+					FNewSpline IntermediateData;
+					Ar << IntermediateData;
+					LegacyData = MakeShared<FLegacySpline>(IntermediateData);
 				}
 				else if (IsNew())
 				{
-					// read in the data directly to the new variant
+					NewData = MakeShared<FNewSpline>();
+					Ar << *NewData;
 				}
 				else if (!IsEnabled())
 				{
-					// read in the data but do nothing
+					// Intentionally doing nothing with data
+					FNewSpline IntermediateData;
+					Ar << IntermediateData;
 				}
-				
-				// will implement in followup change.
 			}
 		}
 	}
@@ -561,266 +953,878 @@ bool FSpline::Serialize(FArchive& Ar)
 		}
 		else if (IsNew())
 		{
-			// todo: serialize new variant
+			Ar << *NewData;
+		}
+	}
+
+	return true;
+}
+
+const FInterpCurveVector& FSpline::GetSplinePointsPosition() const
+{
+	switch (CurrentImplementation)
+	{
+	default: return PositionCurve;
+	case 1: return LegacyData->GetSplinePointsPosition();
+	case 2: return NewData->GetSplinePointsPosition();
+	}
+}
+
+const FInterpCurveQuat& FSpline::GetSplinePointsRotation() const
+{
+	switch (CurrentImplementation)
+	{
+	default: return RotationCurve;
+	case 1: return LegacyData->GetSplinePointsRotation();
+	case 2: return NewData->GetSplinePointsRotation();
+	}
+}
+
+const FInterpCurveVector& FSpline::GetSplinePointsScale() const
+{
+	switch (CurrentImplementation)
+	{
+	default: return ScaleCurve;
+	case 1: return LegacyData->GetSplinePointsScale();
+	case 2: return NewData->GetSplinePointsScale();
+	}
+}
+
+float FSpline::GetSegmentLength(const int32 Index, const float Param, const FVector& Scale3D) const
+{
+	switch (CurrentImplementation)
+	{
+	default: return 0.f;
+	case 1: return LegacyData->GetSegmentLength(Index, Param, Scale3D);
+	case 2: return NewData->GetSegmentLength(Index, Param, Scale3D);
+	}
+}
+
+float FSpline::GetSplineLength() const
+{
+	switch (CurrentImplementation)
+	{
+	default: return 0.f;	// Intentionally doing nothing
+	case 1: return LegacyData->GetSplineLength();
+	case 2: return NewData->GetSplineLength();
+	}
+}
+
+int32 FSpline::GetNumControlPoints() const
+{
+	switch (CurrentImplementation)
+	{
+	default: return 0;
+	case 1: return LegacyData->GetNumControlPoints();
+	case 2: return NewData->GetNumPoints();
+	}
+}
+
+void FSpline::Reset()
+{
+	switch (CurrentImplementation)
+	{
+	default: break;
+	case 1: LegacyData->Reset(); break;
+	case 2: NewData->Reset(); break;
+	}
+}
+	
+void FSpline::ResetRotation()
+{
+	switch (CurrentImplementation)
+	{
+	default: break;
+	case 1: LegacyData->ResetRotation(); break;
+	case 2: NewData->ResetRotation(); break;
+	}
+}
+	
+void FSpline::ResetScale()
+{
+	switch (CurrentImplementation)
+	{
+	default: break;
+	case 1: LegacyData->ResetScale(); break;
+	case 2: NewData->ResetScale(); break;
+	}
+}
+	
+void FSpline::ResetToDefault()
+{
+	switch (CurrentImplementation)
+	{
+	default: break;
+	case 1: LegacyData->ResetToDefault(); break;
+	case 2: NewData->ResetToDefault(); break;
+	}
+}
+
+void FSpline::UpdateSpline(const FUpdateSplineParams& InParams)
+{
+	switch (CurrentImplementation)
+    {
+    default: break;
+    case 1: LegacyData->UpdateSpline(InParams); break;
+    case 2: NewData->UpdateSpline(InParams); break;
+    }
+}
+
+#if WITH_EDITOR
+void FSpline::OnSplineImplementationChanged()
+{
+	using namespace UE::Spline;
+
+	// This implements a state machine I drew out.
+	// I am pretty confident it covers all cases because it handles all edges of a
+	// directed digraph (of 3 nodes, the 3 possible states, with self-edges) where edges represent state transitions.
+
+	// Do you see the nice triangle?
+	
+	if (GImplementation == CurrentImplementation)			// 0->0, 1->1, 2->2
+	{
+	}
+	else if (GImplementation == 0)							// 1->0, 2->0
+	{
+		LegacyData.Reset();
+		NewData.Reset();
+	}
+	else if (CurrentImplementation == 0)					// 0->1, 0->2
+	{
+		// Here we just need to alloc but not copy
+		// GImplementation guaranteed to be 1 or 2
+		if (GImplementation == 1)
+		{
+			LegacyData = MakeShared<FLegacySpline>();
 		}
+		else if (GImplementation == 2)
+		{
+			NewData = MakeShared<FNewSpline>();
+		}
+	}
+	else if (CurrentImplementation == 1)					// 1->2
+	{
+		NewData = MakeShared<FNewSpline>(*LegacyData);
+		LegacyData.Reset();
+	}
+	else if (CurrentImplementation == 2)					// 2->1
+	{
+		LegacyData = MakeShared<FLegacySpline>(*NewData);
+		NewData.Reset();
+	}
+
+	CurrentImplementation = GImplementation;
+}
+#endif
+
+
+
+/** FNewSpline Implementation */
+
+FNewSpline::FNewSpline()
+{
+    CreateAttributeChannel<FQuat>(RotationAttrName);
+    CreateAttributeChannel<FVector>(ScaleAttrName);
+    CreateAttributeChannel<float>(ParamToLengthAttrName);
+    CreateAttributeChannel<float>(LengthToParamAttrName);
+    
+    ResetToDefault();
+}
+
+FNewSpline::FNewSpline(const FNewSpline& Other)
+    : FNewSpline()
+{
+    *this = Other;
+}
+
+FNewSpline::FNewSpline(const FLegacySpline& Other)
+	: FNewSpline()
+{
+	TArray<FTangentBezierControlPoint> Points;
+	
+    for (int32 Idx = 0; Idx < Other.GetNumControlPoints(); Idx++)
+    {
+	    FTangentBezierControlPoint& Point = Points.AddDefaulted_GetRef();
+
+    	Point.Position = Other.GetLocation(Idx);
+    	Point.TangentIn = Other.GetInTangent(Idx);
+    	Point.TangentOut = Other.GetOutTangent(Idx);
+    	Point.TangentMode = ConvertInterpCurveModeToTangentMode(Other.GetSplinePointType(Idx));
+    }
+
+	// initialize spline
+	SetControlPoints(Points);
+
+	auto NormalizeAttributeParameter = [MaxParameter = static_cast<float>(GetNumPoints() - 1)](int32 Index)
+	{
+		return Index / MaxParameter;
+	};
+	
+	// setup rotation/scale attributes
+	for (int32 Idx = 0; Idx < GetNumPoints(); Idx++)
+	{    		
+		const float AttributeParameter = NormalizeAttributeParameter(Idx);
+		AddAttributeValue(RotationAttrName, Other.GetRotation(Idx), AttributeParameter);
+		AddAttributeValue(ScaleAttrName, Other.GetScale(Idx), AttributeParameter);
+	}
+    UpdateSpline();
+}
+
+FNewSpline::FNewSpline(const FSplineCurves& Other)
+	: FNewSpline()
+{
+	TArray<FTangentBezierControlPoint> Points;
+	
+    for (int32 Idx = 0; Idx < Other.Position.Points.Num(); Idx++)
+    {
+		const auto& Position = Other.Position.Points[Idx];
+    	
+    	
+    	FTangentBezierControlPoint& Point = Points.AddDefaulted_GetRef();
+
+    	Point.Position = Position.OutVal;
+    	Point.TangentIn = Position.ArriveTangent;
+    	Point.TangentOut = Position.LeaveTangent;
+    	Point.TangentMode = ConvertInterpCurveModeToTangentMode(Position.InterpMode);
+    }
+	
+	// initialize spline
+	SetControlPoints(Points);
+	
+	auto NormalizeAttributeParameter = [MaxParameter = static_cast<float>(GetNumPoints() - 1)](int32 Index)
+    {
+    	return Index / MaxParameter;
+    };
+	
+    // setup rotation/scale attributes
+    for (int32 Idx = 0; Idx < GetNumPoints(); Idx++)
+    {
+    	const auto& Rotation = Other.Rotation.Points[Idx];
+        const auto& Scale = Other.Scale.Points[Idx];
+    	
+    	const float AttributeParameter = NormalizeAttributeParameter(Idx);
+    	AddAttributeValue(RotationAttrName, Rotation, AttributeParameter);
+    	AddAttributeValue(ScaleAttrName, Scale, AttributeParameter);
+    }
+
+    UpdateSpline();
+}
+
+/** Copy assignment */
+FNewSpline& FNewSpline::operator=(const FNewSpline& Other)
+{
+    if (this != &Other)
+    {
+        InternalSpline = Other.InternalSpline;
+    	TangentModes = Other.TangentModes;
+    	ReparamStepsPerSegment = Other.ReparamStepsPerSegment;
+        Tension = Other.Tension;
+    	bStationaryEndpoint = Other.bStationaryEndpoint;
+
+    	UpdateSpline();
+    }
+    return *this;
+}
+
+bool FNewSpline::operator==(const FNewSpline& Other) const
+{
+    return InternalSpline == Other.InternalSpline && Tension == Other.Tension;
+    //return Points == Other.Points && Tension == Other.Tension && InternalSpline == Other.InternalSpline;
+}
+
+bool FNewSpline::operator!=(const FNewSpline& Other) const
+{
+    return !(*this == Other);
+}
+
+bool FNewSpline::Serialize(FArchive& Ar)
+{
+    FTangentBezierSpline::Serialize(Ar);
+	
+    Ar << ReparamStepsPerSegment;
+	Ar << bStationaryEndpoint;
+	
+	if (Ar.IsLoading())
+	{
+		UpdateSpline();
+	}
+    return true;
+}
+
+float FNewSpline::GetSegmentLength(const int32 Index, const float Param, const FVector& Scale3D) const
+{
+	const int32 NumPoints = GetNumPoints();
+	const int32 LastPoint = NumPoints - 1;
+
+	check(Index >= 0 && ((IsClosedLoop() && Index < NumPoints) || (!IsClosedLoop() && Index < LastPoint)));
+	check(Param >= 0.0f && Param <= 1.0f);
+
+	// Evaluate the length of a Hermite spline segment.
+	// This calculates the integral of |dP/dt| dt, where P(t) is the spline equation with components (x(t), y(t), z(t)).
+	// This isn't solvable analytically, so we use a numerical method (Legendre-Gauss quadrature) which performs very well
+	// with functions of this type, even with very few samples.  In this case, just 5 samples is sufficient to yield a
+	// reasonable result.
+
+	struct FLegendreGaussCoefficient
+	{
+		float Abscissa;
+		float Weight;
+	};
+
+	static const FLegendreGaussCoefficient LegendreGaussCoefficients[] =
+	{
+		{ 0.0f, 0.5688889f },
+		{ -0.5384693f, 0.47862867f },
+		{ 0.5384693f, 0.47862867f },
+		{ -0.90617985f, 0.23692688f },
+		{ 0.90617985f, 0.23692688f }
+	};
+
+	const auto& P0 = GetPosition(Index);
+	const auto& T0 = GetTangentOut(Index);
+	const auto& P1 = GetPosition(Index == LastPoint ? 0 : Index + 1);
+	const auto& T1 = GetTangentOut(Index == LastPoint ? 0 : Index + 1);
+
+	// Special cases for linear
+	if (TangentModes[Index] == ETangentMode::Linear)
+	{
+		return ((P1 - P0) * Scale3D).Size() * Param;
+	}
+
+	// Cache the coefficients to be fed into the function to calculate the spline derivative at each sample point as they are constant.
+	const FVector Coeff1 = ((P0 - P1) * 2.0f + T0 + T1) * 3.0f;
+	const FVector Coeff2 = (P1 - P0) * 6.0f - T0 * 4.0f - T1 * 2.0f;
+	const FVector Coeff3 = T0;
+
+	const float HalfParam = Param * 0.5f;
+
+	float Length = 0.0f;
+	for (const auto& LegendreGaussCoefficient : LegendreGaussCoefficients)
+	{
+		// Calculate derivative at each Legendre-Gauss sample, and perform a weighted sum
+		const float Alpha = HalfParam * (1.0f + LegendreGaussCoefficient.Abscissa);
+		const FVector Derivative = ((Coeff1 * Alpha + Coeff2) * Alpha + Coeff3) * Scale3D;
+		Length += Derivative.Size() * LegendreGaussCoefficient.Weight;
+	}
+	Length *= HalfParam;
+
+	return Length;
+}
+
+float FNewSpline::GetSplineLength() const
+{
+    return EvaluateAttribute<float>(ParamToLengthAttrName, 1.f);
+}
+
+void FNewSpline::SetClosedLoop(bool bInClosedLoop)
+{
+    FTangentBezierSpline::SetClosedLoop(bInClosedLoop);
+    UpdateSpline();
+}
+
+void FNewSpline::Reset()
+{
+    //Points.Empty();
+	//TODO: Not sure what to do here as we don't have a reset to no segments implementation for our Bezier splines
+    UpdateSpline();
+}
+
+void FNewSpline::ResetToDefault()
+{
+    SetControlPoints( { FTangentBezierControlPoint(FVector(0.0, 0.0, 0.0)),
+    					FTangentBezierControlPoint(FVector(100.0, 0.0, 0.0)) } );
+    UpdateSpline();
+}
+
+void FNewSpline::ResetRotation()
+{
+    ClearAttributeChannel(RotationAttrName);
+
+    auto NormalizeAttributeParameter = [MaxParameter = static_cast<float>(GetNumPoints() - 1)](int32 Index)
+    {
+    	return Index / MaxParameter;
+    };
+    
+    for (int32 Idx = 0; Idx < GetNumPoints(); Idx++)
+    {    		
+    	const float AttributeParameter = NormalizeAttributeParameter(Idx);
+    	//Points[Idx].Rotation = FQuat::Identity;
+    	AddAttributeValue(RotationAttrName, FQuat::Identity, AttributeParameter);
+    }
+}
+
+void FNewSpline::ResetScale()
+{
+    ClearAttributeChannel(ScaleAttrName);
+
+    auto NormalizeAttributeParameter = [MaxParameter = static_cast<float>(GetNumPoints() - 1)](int32 Index)
+    {
+    	return Index / MaxParameter;
+    };
+    
+    for (int32 Idx = 0; Idx < GetNumPoints(); Idx++)
+    {    		
+        const float AttributeParameter = NormalizeAttributeParameter(Idx);
+        //Points[Idx].Scale = FVector::OneVector;
+        AddAttributeValue(ScaleAttrName, FVector::OneVector, AttributeParameter);
+    }
+}
+
+void FNewSpline::AddPoint(const FSplinePoint& Point)
+{
+    // auto UpperBound = [](const TArray<FTangentBezierControlPoint>& SplinePoints, float Value)
+    // {
+    // 	const int32 InsertIndex = static_cast<int32>(Value);
+    // 	int32 Count = SplinePoints.Num();
+    // 	int32 First = 0;
+    //
+    // 	while (Count > 0)
+    // 	{
+    // 		const int32 Middle = Count / 2;
+    // 		const int32 CurrentIndex = First + Middle;
+    //
+    // 		if (CurrentIndex <= InsertIndex)
+    // 		{
+    // 			First = CurrentIndex + 1;
+    // 			Count -= Middle + 1;
+    // 		}
+    // 		else
+    // 		{
+    // 			Count = Middle;
+    // 		}
+    // 	}
+    //
+    // 	return First;
+    // };
+    //
+    // const int32 Index = UpperBound(Points, Point.InputKey);
+    
+    FTangentBezierControlPoint NewPoint;
+    NewPoint.Position = Point.Position;
+    NewPoint.TangentIn = Point.ArriveTangent;
+    NewPoint.TangentOut = Point.LeaveTangent;
+    //NewPoint.Rotation = Point.Rotation.Quaternion();
+    // NewPoint.Scale = Point.Scale;
+    //
+    // Points.Insert(NewPoint, Index);
+	FTangentBezierSpline::AddPoint(NewPoint);
+
+	//TODO: Make sure this works as expected
+	AddAttributeValue(RotationAttrName, Point.Rotation.Quaternion(), Point.InputKey);
+	AddAttributeValue(ScaleAttrName, Point.Scale, Point.InputKey);
+	
+    RebuildLegacyCurves();
+}
+
+void FNewSpline::InsertPoint(const FSplinePoint& Point, int32 Index)
+{
+    Index = FMath::Clamp(Index, 0, GetNumPoints() - 1);
+	float SqDistance;
+	float SplineParam = FTangentBezierSpline::FindNearestOnSegment(Point.Position, Index - 1, SqDistance);
+	
+	EInterpCurveMode InterpCurveMode = ConvertSplinePointTypeToInterpCurveMode(Point.Type);
+    int32 InsertedIndex = FTangentBezierSpline::InsertPoint(SplineParam, Point.Position, ConvertInterpCurveModeToTangentMode(InterpCurveMode));
+
+	// If insertion was successful, update tangents if not auto-computed
+	if (InsertedIndex >= 0 && InterpCurveMode != CIM_CurveAuto && InterpCurveMode != CIM_CurveAutoClamped)
+	{
+		SetTangentIn(InsertedIndex, Point.ArriveTangent);
+		SetTangentOut(InsertedIndex, Point.LeaveTangent);
 	}
 
-	return true;
+	// Update rotation and scale attributes
+	// Use the normalized parameter for the new point
+	float NormalizedParam = static_cast<float>(InsertedIndex) / 
+						   FMath::Max(1.0f, static_cast<float>(GetNumPoints() - 1));
+    
+	AddAttributeValue(RotationAttrName, Point.Rotation.Quaternion(), NormalizedParam);
+	AddAttributeValue(ScaleAttrName, Point.Scale, NormalizedParam);
+	
+    RebuildLegacyCurves();
 }
 
-const FInterpCurveVector& FSpline::GetSplinePointsPosition() const
+FSplinePoint FNewSpline::GetPoint(int32 Index) const
 {
-	switch (CurrentImplementation)
-	{
-	default:
-		// Intentionally doing nothing
-		return DummySpline.GetSplinePointsPosition();
-	case 1:
-		return LegacyData->GetSplinePointsPosition();
-	case 2:
-		// todo: implement for new spline
-		return DummySpline.GetSplinePointsPosition();
-	}
+    FSplinePoint Point;
+    
+    if (Index < 0 || Index >= GetNumPoints())
+    {
+    	return Point;
+    }
+
+    Point.Position = GetPosition(Index);
+	Point.ArriveTangent = GetTangentIn(Index);
+    Point.LeaveTangent = GetTangentOut(Index);
+    Point.Rotation = GetAttributeValue<FQuat>(RotationAttrName, Index).Rotator();
+    Point.Scale = GetAttributeValue<FVector>(ScaleAttrName, Index);
+	EInterpCurveMode InterpCurveMode = ConvertTangentModeToInterpCurveMode(TangentModes[Index]);
+    Point.Type = ConvertInterpCurveModeToSplinePointType(InterpCurveMode);
+
+    return Point;
 }
 
-const FInterpCurveQuat& FSpline::GetSplinePointsRotation() const
+/** Removes a point from the spline */
+void FNewSpline::RemovePoint(int32 Index)
 {
-	switch (CurrentImplementation)
+	if (Index < 0 || Index >= GetNumPoints())
 	{
-	default:
-		// Intentionally doing nothing
-		return DummySpline.GetSplinePointsRotation();
-	case 1:
-		return LegacyData->GetSplinePointsRotation();
-	case 2:
-		// todo: implement for new spline
-		return DummySpline.GetSplinePointsRotation();
+		return;
 	}
+
+	FTangentBezierSpline::RemovePoint(Index);
+    RebuildLegacyCurves();
 }
 
-const FInterpCurveVector& FSpline::GetSplinePointsScale() const
+void FNewSpline::SetLocation(int32 Index, const FVector& InLocation)
 {
-	switch (CurrentImplementation)
-	{
-	default:
-		// Intentionally doing nothing
-		return DummySpline.GetSplinePointsScale();
-	case 1:
-		return LegacyData->GetSplinePointsScale();
-	case 2:
-		// todo: implement for new spline
-		return DummySpline.GetSplinePointsScale();
-	}
+    if (Index < 0 || Index >= GetNumPoints())
+    {
+    	return;
+    }
+	//SetValueAt(Index * 4, InLocation);
+	ModifyPoint(Index, {InLocation, GetTangentIn(Index), GetTangentOut(Index), GetTangentMode(Index)});
+    RebuildLegacyCurves();
 }
 
-const FInterpCurveFloat& FSpline::GetSplineReparamTable() const
+FVector FNewSpline::GetLocation(const int32 Index) const
 {
-	switch (CurrentImplementation)
-	{
-	default:
-		// Intentionally doing nothing
-		return DummySpline.GetSplineReparamTable();
-	case 1:
-		return LegacyData->GetSplineReparamTable();
-	case 2:
-		// todo: implement for new spline
-		return DummySpline.GetSplineReparamTable();
-	}
+	if (Index < 0 || Index >= GetNumPoints())
+    {
+    	return FVector();
+    }
+
+    return GetPosition(Index);
 }
 
-float FSpline::GetSegmentLength(const int32 Index, const float Param, bool bClosedLoop, const FVector& Scale3D) const
+void FNewSpline::SetInTangent(const int32 Index, const FVector& InTangent)
 {
-	switch (CurrentImplementation)
-	{
-	default:
-		// Intentionally doing nothing
-		return 0.f;
-	case 1:
-		return LegacyData->GetSegmentLength(Index, Param, bClosedLoop, Scale3D);
-	case 2:
-		// todo: implement for new spline
-		return 0.f;
-	}
+	if (Index < 0 || Index >= GetNumPoints())
+    {
+    	return;
+    }
+
+	SetPointTangentMode(Index, ETangentMode::User);
+    this->SetTangentIn(Index, InTangent);
+    RebuildLegacyCurves();
 }
 
-float FSpline::GetSplineLength() const
+FVector FNewSpline::GetInTangent(const int32 Index) const
 {
-	switch (CurrentImplementation)
-	{
-	default:
-		// Intentionally doing nothing
-		return 0.f;
-	case 1:
-		return LegacyData->GetSplineLength();
-	case 2:
-		// todo: implement for new spline
-		return 0.f;
-	}
+	if (Index < 0 || Index >= GetNumPoints())
+    {
+    	return FVector();
+    }
+
+    return GetTangentIn(Index);
 }
 
-int32 FSpline::GetNumControlPoints() const
+void FNewSpline::SetOutTangent(const int32 Index, const FVector& OutTangent)
 {
-	switch (CurrentImplementation)
-	{
-	default:
-		// Intentionally doing nothing
-		return 0.f;
-	case 1:
-		return LegacyData->GetNumControlPoints();
-	case 2:
-		// todo: implement for new spline
-		return 0.f;
-	}
+	if (Index < 0 || Index >= GetNumPoints())
+    {
+    	return;
+    }
+	SetPointTangentMode(Index, ETangentMode::User);
+	this->SetTangentOut(Index, OutTangent);
+    
+    RebuildLegacyCurves();
 }
 
-void FSpline::Reset()
+FVector FNewSpline::GetOutTangent(const int32 Index) const
 {
-	switch (CurrentImplementation)
-	{
-	default:
-		// Intentionally doing nothing
-		break;
-	case 1:
-		return LegacyData->Reset();
-	case 2:
-		// todo: implement for new spline
-		break;
-	}
+	if (Index < 0 || Index >= GetNumPoints())
+    {
+    	return FVector();
+    }
+
+    return GetTangentOut(Index);
 }
-	
-void FSpline::ResetRotation()
+
+void FNewSpline::SetRotation(int32 Index, const FQuat& InRotation)
 {
-	switch (CurrentImplementation)
+	if (Index < 0 || Index >= GetNumPoints())
+    {
+    	return;
+    }
+
+	// if rotation value not added to attribute channel, add it
+	if (Index >= NumAttributeValues(RotationAttrName))
 	{
-	default:
-		// Intentionally doing nothing
-		break;
-	case 1:
-		return LegacyData->ResetRotation();
-	case 2:
-		// todo: implement for new spline
-		break;
+		AddAttributeValue(RotationAttrName, InRotation, GetParameterAt(Index));
+	}
+	else
+	{
+		SetAttributeValue(RotationAttrName, Index, InRotation);
 	}
+	
+    RebuildLegacyCurves();
 }
+
+FQuat FNewSpline::GetRotation(const int32 Index) const
+{
+	if (Index < 0 || Index >= GetNumPoints())
+    {
+    	return FQuat::Identity;
+    }
 	
-void FSpline::ResetScale()
+    return GetAttributeValue<FQuat>(RotationAttrName, Index);
+}
+
+void FNewSpline::SetScale(int32 Index, const FVector& InScale)
 {
-	switch (CurrentImplementation)
+	if (Index < 0 || Index >= GetNumPoints())
+    {
+    	return;
+    }
+
+	// if scale value not added to attribute channel, add it
+	if (Index >= NumAttributeValues(ScaleAttrName))
 	{
-	default:
-		// Intentionally doing nothing
-		break;
-	case 1:
-		return LegacyData->ResetRotation();
-	case 2:
-		// todo: implement for new spline
-		break;
+		AddAttributeValue(ScaleAttrName, InScale, GetParameterAt(Index));
+	}
+	else
+	{
+		SetAttributeValue(ScaleAttrName, Index, InScale);
 	}
+    
+    RebuildLegacyCurves();
+}
+
+FVector FNewSpline::GetScale(const int32 Index) const
+{
+    if (Index < 0 || Index >= GetNumPoints())
+    {
+    	return FVector::OneVector;
+    }
+
+    return GetAttributeValue<FVector>(ScaleAttrName, Index);
 }
+
+void FNewSpline::SetSplinePointType(int32 Index, EInterpCurveMode Type)
+{
+    if (Index < 0 || Index >= GetNumPoints())
+    {
+    	return;
+    }
 	
-void FSpline::ResetToDefault()
+	SetPointTangentMode(Index, ConvertInterpCurveModeToTangentMode(Type));
+    
+    RebuildLegacyCurves();
+}
+
+EInterpCurveMode FNewSpline::GetSplinePointType(int32 Index) const
 {
-	switch (CurrentImplementation)
-	{
-	default:
-		// Intentionally doing nothing
-		break;
-	case 1:
-		return LegacyData->ResetToDefault();
-	case 2:
-		// todo: implement for new spline
-		break;
-	}
+    if (Index < 0 || Index >= GetNumPoints())
+    {
+    	return CIM_Unknown;
+    }
+	ETangentMode TangentMode = GetTangentMode(Index);
+    return ConvertTangentModeToInterpCurveMode(TangentMode);
 }
 
-void FSpline::UpdateSpline(const FUpdateSplineParams& InParams)
+float FNewSpline::GetParameterAtIndex(int32 Index) const
 {
-	switch (CurrentImplementation)
+    if (Index < 0 || Index >= GetNumPoints())
     {
-    default:
-    	// Intentionally doing nothing
-    	break;
-    case 1:
-    	return LegacyData->UpdateSpline(InParams);
-    case 2:
-    	// todo: implement for new spline
-    	break;
+    	return 0.f;
     }
+    
+    return static_cast<float>(Index);
 }
 
-#if WITH_EDITOR
-void FSpline::OnSplineImplementationChanged()
+float FNewSpline::GetParameterAtDistance(float Distance) const
 {
-	using namespace UE::Spline;
+	const float NormalizedDistance = FMath::Clamp(Distance / FMath::Max(GetSplineLength(), UE_SMALL_NUMBER), 0.f, 1.f);
+    return EvaluateAttribute<float>(LengthToParamAttrName, NormalizedDistance);
+}
 
-	// This implements a state machine I drew out.
-	// I am pretty confident it covers all cases because it handles all edges of a
-	// directed digraph (of 3 nodes, the 3 possible states, with self-edges) where edges represent state transitions.
+float FNewSpline::GetDistanceAtParameter(float Parameter) const
+{
+    const int32 NumEffectivePoints = IsClosedLoop() ? GetNumPoints() + 1 : GetNumPoints();
+    Parameter = Parameter / (NumEffectivePoints - 1);
+    return EvaluateAttribute<float>(ParamToLengthAttrName, Parameter);
+}
 
-	// Do you see the nice triangle?
+float FNewSpline::FindNearest(const FVector& InLocation) const
+{
+	float Dummy;
+
+    if (UE::Spline::GFallbackFindNearest)
+    {
+    	return PositionCurve.FindNearest(InLocation, Dummy);
+    }
+    else
+    {
+	    return FTangentBezierSpline::FindNearest(InLocation, Dummy);
+    }
+}
+
+float FNewSpline::FindNearestOnSegment(const FVector& InLocation, int32 SegmentIndex) const
+{
+	float Dummy;
 	
-	if (GImplementation == CurrentImplementation)			// 0->0, 1->1, 2->2
-	{
-	}
-	else if (GImplementation == 0)							// 1->0, 2->0
-	{
-		LegacyData.Reset();
-		// NewData.Reset();
-	}
-	else if (CurrentImplementation == 0)					// 0->1, 0->2
-	{
-		// Here we just need to alloc but not copy
-		// GImplementation guaranteed to be 1 or 2
-		if (GImplementation == 1)
-		{
-			LegacyData = MakeShared<FLegacySpline>();
-		}
-		else if (GImplementation == 2)
-		{
-			// NewData = MakeShared<FHermiteSpline>();
-		}
-	}
-	else if (CurrentImplementation == 1)					// 1->2
-	{
-		TSharedPtr<FLegacySpline> CurrentData = LegacyData;
-		LegacyData.Reset();
-		// NewData = MakeShared<FHermiteSpline>(*CurrentData);
-	}
-	else if (CurrentImplementation == 2)					// 2->1
-	{
-		// Cannot implement properly yet
-		// TSharedPtr<FLegacySpline> CurrentData = NewData;
-		// NewData.Reset();
-		// LegacyData = MakeShared<FLegacySpline>(*CurrentData);
-		
-		// This is temporary to make sure LegacyData is always valid when we
-		// expect it to be, but this should be removed when above code is active.
-		LegacyData = MakeShared<FLegacySpline>();
-	}
+    if (UE::Spline::GFallbackFindNearest)
+    {
+        if (!PositionCurve.Points.IsValidIndex(SegmentIndex))
+        {
+            return 0.f;
+        }
+    
+        return PositionCurve.FindNearestOnSegment(InLocation, SegmentIndex, Dummy);
+    }
+    else
+    {
+	   return FTangentBezierSpline::FindNearestOnSegment(InLocation, SegmentIndex, Dummy);
+    }
+}
 
-	CurrentImplementation = GImplementation;
+FVector FNewSpline::Evaluate(float Parameter) const
+{
+    const int32 NumEffectivePoints = IsClosedLoop() ? GetNumPoints() + 1 : GetNumPoints();
+    Parameter = GetParameterRange().Interpolate(Parameter / (NumEffectivePoints - 1));
+	
+    return FVector(GetPosition(Parameter));
 }
-#endif
+
+FVector FNewSpline::EvaluateDerivative(float Parameter) const
+{
+    const int32 NumEffectivePoints = IsClosedLoop() ? GetNumPoints() + 1 : GetNumPoints();
+	Parameter = GetParameterRange().Interpolate(Parameter / (NumEffectivePoints - 1));
+    return FVector(GetTangent(Parameter));
+}
+
+FQuat FNewSpline::EvaluateRotation(float Parameter) const
+{
+    const int32 NumEffectivePoints = IsClosedLoop() ? GetNumPoints() + 1 : GetNumPoints();
+	Parameter = GetParameterRange().Interpolate(Parameter / (NumEffectivePoints - 1));
+    return EvaluateAttribute<FQuat>(RotationAttrName, Parameter);
+}
+
+FVector FNewSpline::EvaluateScale(float Parameter) const
+{
+    const int32 NumEffectivePoints = IsClosedLoop() ? GetNumPoints() + 1 : GetNumPoints();
+	Parameter = GetParameterRange().Interpolate(Parameter / (NumEffectivePoints - 1));
+    return EvaluateAttribute<FVector>(ScaleAttrName, Parameter);
+}
+
+void FNewSpline::UpdateSpline()
+{
+    //UpdateTangents();			// updates the tangents on our version of the points
+    UpdateReparamTable();		// updates both reparam tables
+    RebuildLegacyCurves();		// updates legacy curves based on points, never evaluates internal spline.
+}
+
+void FNewSpline::UpdateSpline(const FSpline::FUpdateSplineParams& Params)
+{
+	// todo: handled unhandled params
+    
+    SetClosedLoop(Params.bClosedLoop);
+    bStationaryEndpoint = Params.bStationaryEndpoints;
+    ReparamStepsPerSegment = Params.ReparamStepsPerSegment;
+
+    UpdateSpline();
+}
+
+void FNewSpline::UpdateReparamTable()
+{
+    const int32 NumPoints = GetNumPoints();
+  	const int32 NumSegments = NumPoints - 1;
+    const int32 TotalSteps = NumSegments * ReparamStepsPerSegment + 1;
+
+    // Start by clearing it
+    ClearAttributeChannel(ParamToLengthAttrName);
+    ClearAttributeChannel(LengthToParamAttrName);
+    
+    // Here we get Param -> Length mapping.
+    float AccumulatedLength = 0.0f;
+    for (int32 SegmentIndex = 0; SegmentIndex < NumSegments; ++SegmentIndex)
+    {
+    	for (int32 Step = 0; Step < ReparamStepsPerSegment; ++Step)
+    	{
+    		// These are with respect to current segment (so called 'segment space').
+    		const float SegmentParam = static_cast<float>(Step) / ReparamStepsPerSegment;
+    		const float OffsetLength = GetSegmentLength(SegmentIndex, SegmentParam);
+
+    		// Map the above values to 'spline space' (as opposed to segment space).
+    		const float SplineParam = static_cast<float>(SegmentIndex * ReparamStepsPerSegment + Step) / TotalSteps;
+    		const float CurrentLength = OffsetLength + AccumulatedLength;
+
+    		AddAttributeValue(ParamToLengthAttrName, CurrentLength, SplineParam);
+    	}
+
+    	AccumulatedLength += GetSegmentLength(SegmentIndex, 1.0f);
+    }
+    AddAttributeValue(ParamToLengthAttrName, AccumulatedLength, 1.f);
+
+    // Here we get Normalized Length -> Param mapping
+    for (int32 Step = 0; Step <= NumSegments * ReparamStepsPerSegment + 1; ++Step)
+    {
+    	const float SplineParam = FMath::Clamp(static_cast<float>(Step) / TotalSteps, 0.0f, 1.0f);
+    	const float CurrentLength = EvaluateAttribute<float>(ParamToLengthAttrName, SplineParam);
+    
+    	AddAttributeValue(LengthToParamAttrName, SplineParam * NumSegments, CurrentLength / AccumulatedLength);	// must be normalized param
+    }
+}
+
+void FNewSpline::RebuildLegacyCurves()
+{
+    PositionCurve.Reset();
+    RotationCurve.Reset();
+    ScaleCurve.Reset();
+
+    PositionCurve.bIsLooped = IsClosedLoop();
+    PositionCurve.LoopKeyOffset = 1.f;
+    RotationCurve.bIsLooped = IsClosedLoop();
+    ScaleCurve.bIsLooped = IsClosedLoop();
+
+	// if this fails, some other function which mutates InternalSpline is failing to keep Points appropriately sized
+	//check(Points.Num() == GetNumPoints());
+	
+    for (int32 i = 0; i < GetNumPoints(); ++i)
+    {
+    	const float AttributeParameter = static_cast<float>(i);
+    	
+    	// Additionally populate interp curves here so we don't double loop
+    	PositionCurve.Points.Emplace(AttributeParameter, GetPosition(i), GetTangentIn(i), GetTangentOut(i), ConvertTangentModeToInterpCurveMode(TangentModes[i]));
+    	RotationCurve.AddPoint(AttributeParameter, GetAttributeValue<FQuat>(RotationAttrName, i));
+    	ScaleCurve.AddPoint(AttributeParameter, GetAttributeValue<FVector>(ScaleAttrName, i));
+    }
+}
+
 
 
 
+/** FLegacySpline Implementation */
 
-FSpline::FLegacySpline::FLegacySpline()
+FLegacySpline::FLegacySpline()
 {
 	ResetToDefault();
 }
 
-FSpline::FLegacySpline FSpline::FLegacySpline::FromSplineCurves(const FSplineCurves& InSpline)
+FLegacySpline::FLegacySpline(const FNewSpline& Other)
+	: FLegacySpline()
 {
-	FLegacySpline NewSpline;
-	NewSpline.PositionCurve = InSpline.Position;
-	NewSpline.RotationCurve = InSpline.Rotation;
-	NewSpline.ScaleCurve = InSpline.Scale;
-	NewSpline.ReparamTable = InSpline.ReparamTable;
-	return NewSpline;
+	PositionCurve = Other.GetSplinePointsPosition();
+	RotationCurve = Other.GetSplinePointsRotation();
+	ScaleCurve = Other.GetSplinePointsScale();
+
+	FSpline::FUpdateSplineParams InParams;
+	InParams.bClosedLoop = PositionCurve.bIsLooped;
+	UpdateSpline(InParams);
+
+	PositionCurve.AutoSetTangents(0.f, false);
+}
+
+FLegacySpline::FLegacySpline(const FSplineCurves& InSpline)
+	: FLegacySpline()
+{
+	PositionCurve = InSpline.Position;
+	RotationCurve = InSpline.Rotation;
+	ScaleCurve = InSpline.Scale;
+	ReparamTable = InSpline.ReparamTable;
 }
 
-void FSpline::FLegacySpline::AddPoint(const FSplinePoint& InPoint)
+void FLegacySpline::AddPoint(const FSplinePoint& InPoint)
 {
 	auto UpperBound = [this](float Value)
 	{
@@ -871,9 +1875,9 @@ void FSpline::FLegacySpline::AddPoint(const FSplinePoint& InPoint)
 	), Index);
 }
 
-void FSpline::FLegacySpline::AddPointAtIndex(const FSplinePoint& InPoint, int32 Index)
+void FLegacySpline::InsertPoint(const FSplinePoint& InPoint, int32 Index)
 {
-	const float InKey = (Index == 0) ? 0.0f : GetSplinePointParameter(Index - 1) + 1.0f;
+	const float InKey = (Index == 0) ? 0.0f : GetParameterAtIndex(Index - 1) + 1.0f;
 	
 	PositionCurve.Points.Insert(FInterpCurvePoint<FVector>(
 		InKey,
@@ -908,7 +1912,7 @@ void FSpline::FLegacySpline::AddPointAtIndex(const FSplinePoint& InPoint, int32
 	}
 }
 
-FSplinePoint FSpline::FLegacySpline::GetPoint(const int32 Index) const
+FSplinePoint FLegacySpline::GetPoint(const int32 Index) const
 {
 	FSplinePoint Point;
 
@@ -928,7 +1932,7 @@ FSplinePoint FSpline::FLegacySpline::GetPoint(const int32 Index) const
 	return Point;
 }
 
-void FSpline::FLegacySpline::RemovePoint(int32 Index)
+void FLegacySpline::RemovePoint(int32 Index)
 {
 	if (!PositionCurve.Points.IsValidIndex(Index))
 	{
@@ -948,22 +1952,22 @@ void FSpline::FLegacySpline::RemovePoint(int32 Index)
 	}
 }
 
-void FSpline::FLegacySpline::SetLocation(int32 PointIndex, const FVector& InLocation)
+void FLegacySpline::SetLocation(int32 Index, const FVector& InLocation)
 {
-	if (!PositionCurve.Points.IsValidIndex(PointIndex))
+	if (!PositionCurve.Points.IsValidIndex(Index))
 	{
 		return;
 	}
 	
-	PositionCurve.Points[PointIndex].OutVal = InLocation;
+	PositionCurve.Points[Index].OutVal = InLocation;
 }
 
-FVector FSpline::FLegacySpline::GetLocation(const int32 PointIndex) const
+FVector FLegacySpline::GetLocation(const int32 Index) const
 {
-	return PositionCurve.Points.IsValidIndex(PointIndex) ? PositionCurve.Points[PointIndex].OutVal : FVector();
+	return PositionCurve.Points.IsValidIndex(Index) ? PositionCurve.Points[Index].OutVal : FVector();
 }
 
-void FSpline::FLegacySpline::SetInTangent(const int32 Index, const FVector& InTangent)
+void FLegacySpline::SetInTangent(const int32 Index, const FVector& InTangent)
 {
 	if (!PositionCurve.Points.IsValidIndex(Index))
 	{
@@ -974,12 +1978,12 @@ void FSpline::FLegacySpline::SetInTangent(const int32 Index, const FVector& InTa
 	PositionCurve.Points[Index].InterpMode = CIM_CurveUser;
 }
 
-FVector FSpline::FLegacySpline::GetInTangent(const int32 Index) const
+FVector FLegacySpline::GetInTangent(const int32 Index) const
 {
 	return PositionCurve.Points.IsValidIndex(Index) ? PositionCurve.Points[Index].ArriveTangent : FVector();
 }
 
-void FSpline::FLegacySpline::SetOutTangent(const int32 Index, const FVector& OutTangent)
+void FLegacySpline::SetOutTangent(const int32 Index, const FVector& OutTangent)
 {
 	if (!PositionCurve.Points.IsValidIndex(Index))
 	{
@@ -990,88 +1994,109 @@ void FSpline::FLegacySpline::SetOutTangent(const int32 Index, const FVector& Out
 	PositionCurve.Points[Index].InterpMode = CIM_CurveUser;
 }
 
-FVector FSpline::FLegacySpline::GetOutTangent(const int32 Index) const
+FVector FLegacySpline::GetOutTangent(const int32 Index) const
 {
 	return PositionCurve.Points.IsValidIndex(Index) ? PositionCurve.Points[Index].LeaveTangent : FVector();
 }
 
-void FSpline::FLegacySpline::SetRotation(int32 PointIndex, const FQuat& InRotation)
+void FLegacySpline::SetRotation(int32 Index, const FQuat& InRotation)
 {
-	if (!RotationCurve.Points.IsValidIndex(PointIndex))
+	if (!RotationCurve.Points.IsValidIndex(Index))
 	{
 		return;
 	}
 	
-	RotationCurve.Points[PointIndex].OutVal = InRotation;
+	RotationCurve.Points[Index].OutVal = InRotation;
 }
 
-FQuat FSpline::FLegacySpline::GetRotation(const int32 PointIndex) const
+FQuat FLegacySpline::GetRotation(const int32 Index) const
 {
-	return RotationCurve.Points.IsValidIndex(PointIndex) ? RotationCurve.Points[PointIndex].OutVal : FQuat();
+	return RotationCurve.Points.IsValidIndex(Index) ? RotationCurve.Points[Index].OutVal : FQuat();
 }
 
-void FSpline::FLegacySpline::SetScale(int32 PointIndex, const FVector& InScale)
+void FLegacySpline::SetScale(int32 Index, const FVector& InScale)
 {
-	if (!ScaleCurve.Points.IsValidIndex(PointIndex))
+	if (!ScaleCurve.Points.IsValidIndex(Index))
 	{
 		return;
 	}
 	
-	ScaleCurve.Points[PointIndex].OutVal = InScale;
+	ScaleCurve.Points[Index].OutVal = InScale;
 }
 
-FVector FSpline::FLegacySpline::GetScale(const int32 PointIndex) const
+FVector FLegacySpline::GetScale(const int32 Index) const
 {
-	return ScaleCurve.Points.IsValidIndex(PointIndex) ? ScaleCurve.Points[PointIndex].OutVal : FVector();
+	return ScaleCurve.Points.IsValidIndex(Index) ? ScaleCurve.Points[Index].OutVal : FVector();
 }
 
-void FSpline::FLegacySpline::SetSplinePointType(int32 PointIndex, EInterpCurveMode Type)
+void FLegacySpline::SetSplinePointType(int32 Index, EInterpCurveMode Type)
 {
-	if (!PositionCurve.Points.IsValidIndex(PointIndex))
+	if (!PositionCurve.Points.IsValidIndex(Index))
 	{
 		return;
 	}
 	
-	PositionCurve.Points[PointIndex].InterpMode = Type;
+	PositionCurve.Points[Index].InterpMode = Type;
 }
 
-EInterpCurveMode FSpline::FLegacySpline::GetSplinePointType(const int32 PointIndex) const
+EInterpCurveMode FLegacySpline::GetSplinePointType(const int32 Index) const
 {
-	return PositionCurve.Points.IsValidIndex(PointIndex) ? PositionCurve.Points[PointIndex].InterpMode.GetValue() : CIM_Unknown;
+	return PositionCurve.Points.IsValidIndex(Index) ? PositionCurve.Points[Index].InterpMode.GetValue() : CIM_Unknown;
 }
 
-float FSpline::FLegacySpline::GetSplinePointParameter(int32 PointIndex) const
+float FLegacySpline::GetParameterAtIndex(int32 Index) const
 {
-	return PositionCurve.Points.IsValidIndex(PointIndex) ? PositionCurve.Points[PointIndex].InVal : 0.f;
+	return PositionCurve.Points.IsValidIndex(Index) ? PositionCurve.Points[Index].InVal : 0.f;
+}
+
+float FLegacySpline::GetParameterAtDistance(float Distance) const
+{
+	return ReparamTable.Eval(Distance);
+}
+
+float FLegacySpline::GetDistanceAtParameter(float Parameter) const
+{
+	// this might be duplicating what an interp curve can already do...
+	
+	const float ParameterMax = PositionCurve.Points.Last().InVal;
+	const float Key = (Parameter / ParameterMax) * (ReparamTable.Points.Num() - 1);
+	const int32 LowerKey = FMath::FloorToInt32(Key);
+	ensureAlways(LowerKey >= 0 && LowerKey < ReparamTable.Points.Num());
+	const int32 UpperKey = FMath::CeilToInt32(Key);
+	ensureAlways(UpperKey >= 0 && UpperKey < ReparamTable.Points.Num());
+	const float Alpha = FMath::Frac(Key);
+	const float Distance = FMath::Lerp(ReparamTable.Points[LowerKey].InVal, ReparamTable.Points[UpperKey].InVal, Alpha);
+
+	return Distance;
 }
 
-FVector FSpline::FLegacySpline::Evaluate(float Param) const
+FVector FLegacySpline::Evaluate(float Param) const
 {
 	return PositionCurve.Eval(Param);
 }
 
-FVector FSpline::FLegacySpline::EvaluateDerivative(float Param) const
+FVector FLegacySpline::EvaluateDerivative(float Param) const
 {
 	return PositionCurve.EvalDerivative(Param);
 }
 
-FQuat FSpline::FLegacySpline::EvaluateRotation(float Param) const
+FQuat FLegacySpline::EvaluateRotation(float Param) const
 {
 	return RotationCurve.Eval(Param);
 }
 
-FVector FSpline::FLegacySpline::EvaluateScale(float Param) const
+FVector FLegacySpline::EvaluateScale(float Param) const
 {
 	return ScaleCurve.Eval(Param);
 }
 
-float FSpline::FLegacySpline::FindNearest(const FVector& InLocation) const
+float FLegacySpline::FindNearest(const FVector& InLocation) const
 {
 	float Dummy;
 	return PositionCurve.FindNearest(InLocation, Dummy);
 }
 
-float FSpline::FLegacySpline::FindNearestOnSegment(const FVector& InLocation, int32 SegmentIndex) const
+float FLegacySpline::FindNearestOnSegment(const FVector& InLocation, int32 SegmentIndex) const
 {
 	if (!PositionCurve.Points.IsValidIndex(SegmentIndex))
 	{
@@ -1082,7 +2107,7 @@ float FSpline::FLegacySpline::FindNearestOnSegment(const FVector& InLocation, in
 	return PositionCurve.FindNearestOnSegment(InLocation, SegmentIndex, Dummy);
 }
 
-bool FSpline::FLegacySpline::Serialize(FArchive& Ar)
+bool FLegacySpline::Serialize(FArchive& Ar)
 {
 	Ar.UsingCustomVersion(FFortniteMainBranchObjectVersion::GUID);
 	Ar << PositionCurve;
@@ -1093,12 +2118,12 @@ bool FSpline::FLegacySpline::Serialize(FArchive& Ar)
 	return true;
 }
 
-float FSpline::FLegacySpline::GetSegmentLength(const int32 Index, const float Param, bool bClosedLoop, const FVector& Scale3D) const
+float FLegacySpline::GetSegmentLength(const int32 Index, const float Param, const FVector& Scale3D) const
 {
 	const int32 NumPoints = PositionCurve.Points.Num();
 	const int32 LastPoint = NumPoints - 1;
 
-	check(Index >= 0 && ((bClosedLoop && Index < NumPoints) || (!bClosedLoop && Index < LastPoint)));
+	check(Index >= 0 && ((PositionCurve.bIsLooped && Index < NumPoints) || (!PositionCurve.bIsLooped && Index < LastPoint)));
 	check(Param >= 0.0f && Param <= 1.0f);
 
 	// Evaluate the length of a Hermite spline segment.
@@ -1161,7 +2186,7 @@ float FSpline::FLegacySpline::GetSegmentLength(const int32 Index, const float Pa
 	return Length;
 }
 
-float FSpline::FLegacySpline::GetSplineLength() const
+float FLegacySpline::GetSplineLength() const
 {
 	const int32 NumPoints = ReparamTable.Points.Num();
 
@@ -1174,19 +2199,19 @@ float FSpline::FLegacySpline::GetSplineLength() const
 	return 0.0f;
 }
 
-int32 FSpline::FLegacySpline::GetNumControlPoints() const
+int32 FLegacySpline::GetNumControlPoints() const
 {
 	return PositionCurve.Points.Num();
 }
 
-void FSpline::FLegacySpline::Reset()
+void FLegacySpline::Reset()
 {
 	PositionCurve.Points.Reset();
 	RotationCurve.Points.Reset();
 	ScaleCurve.Points.Reset();
 }
 
-void FSpline::FLegacySpline::ResetRotation()
+void FLegacySpline::ResetRotation()
 {
 	RotationCurve.Points.Reset(PositionCurve.Points.Num());
 
@@ -1196,7 +2221,7 @@ void FSpline::FLegacySpline::ResetRotation()
 	}
 }
 
-void FSpline::FLegacySpline::ResetScale()
+void FLegacySpline::ResetScale()
 {
 	ScaleCurve.Points.Reset(PositionCurve.Points.Num());
 
@@ -1206,7 +2231,7 @@ void FSpline::FLegacySpline::ResetScale()
 	}
 }
 
-void FSpline::FLegacySpline::ResetToDefault()
+void FLegacySpline::ResetToDefault()
 {
 	PositionCurve.Points.Reset(10);
 	RotationCurve.Points.Reset(10);
@@ -1221,7 +2246,7 @@ void FSpline::FLegacySpline::ResetToDefault()
 	ScaleCurve.Points.Emplace(1.0f, FVector(1.0f), FVector::ZeroVector, FVector::ZeroVector, CIM_CurveAuto);
 }
 
-void FSpline::FLegacySpline::UpdateSpline(const FUpdateSplineParams& InParams)
+void FLegacySpline::UpdateSpline(const FSpline::FUpdateSplineParams& InParams)
 {
 	const int32 NumPoints = PositionCurve.Points.Num();
 	check(RotationCurve.Points.Num() == NumPoints && ScaleCurve.Points.Num() == NumPoints);
@@ -1256,7 +2281,7 @@ void FSpline::FLegacySpline::UpdateSpline(const FUpdateSplineParams& InParams)
 	ScaleCurve.AutoSetTangents(0.0f, InParams.bStationaryEndpoints);
 
 	// Now initialize the spline reparam table
-	const int32 NumSegments = InParams.bClosedLoop ? NumPoints : FMath::Max(0, NumPoints - 1);
+	const int32 NumSegments = PositionCurve.bIsLooped ? NumPoints : FMath::Max(0, NumPoints - 1);
 
 	// Start by clearing it
 	ReparamTable.Points.Reset(NumSegments * InParams.ReparamStepsPerSegment + 1);
@@ -1266,11 +2291,11 @@ void FSpline::FLegacySpline::UpdateSpline(const FUpdateSplineParams& InParams)
 		for (int32 Step = 0; Step < InParams.ReparamStepsPerSegment; ++Step)
 		{
 			const float Param = static_cast<float>(Step) / InParams.ReparamStepsPerSegment;
-			const float SegmentLength = (Step == 0) ? 0.0f : GetSegmentLength(SegmentIndex, Param, InParams.bClosedLoop, InParams.Scale3D);
+			const float SegmentLength = (Step == 0) ? 0.0f : GetSegmentLength(SegmentIndex, Param, InParams.Scale3D);
 
 			ReparamTable.Points.Emplace(SegmentLength + AccumulatedLength, SegmentIndex + Param, 0.0f, 0.0f, CIM_Linear);
 		}
-		AccumulatedLength += GetSegmentLength(SegmentIndex, 1.0f, InParams.bClosedLoop, InParams.Scale3D);
+		AccumulatedLength += GetSegmentLength(SegmentIndex, 1.0f, InParams.Scale3D);
 	}
 	ReparamTable.Points.Emplace(AccumulatedLength, static_cast<float>(NumSegments), 0.0f, 0.0f, CIM_Linear);
 }
\ No newline at end of file
diff --git a/Engine/Source/Runtime/GeometryCore/Public/Splines/AttributeChannel.h b/Engine/Source/Runtime/GeometryCore/Public/Splines/AttributeChannel.h
--- a/Engine/Source/Runtime/GeometryCore/Public/Splines/AttributeChannel.h
+++ b/Engine/Source/Runtime/GeometryCore/Public/Splines/AttributeChannel.h
@@ -94,9 +94,26 @@ public:
 	 * Sets the parameter value at the given index
 	 * @param Index - Index of the parameter to set
 	 * @param Param - New parameter value
-	 * @return The new parameter value
+	 * @return New index of the attribute originally at the given index after any reordering
 	 */
-	virtual float SetParameter(int32 Index, float Param) = 0;
+	virtual int32 SetParameter(int32 Index, float Param) = 0;
+
+	/**
+	 * Removes the value at the given index
+	 * @param Index - Index of the value to remove
+	 * @return true if the value was removed
+	 */
+	virtual bool RemoveValue(int32 Index) = 0;
+
+	/**
+     * Compare this channel to another channel of the same type
+     * @param Other - Channel to compare against
+     * @return true if the channels have identical contents
+     */
+    virtual bool IsEqual(const IAttributeChannelBase& Other) const = 0;
+
+	virtual float MapToParentSpace(float Parameter) const = 0;
+	virtual float MapFromParentSpace(float Parameter) const = 0;
 };
 
 /**
@@ -109,14 +126,14 @@ class TAttributeChannel : public IAttributeChannelBase
 	static FName TypeName;
 
 	/** Storage for channel values */
-    TAttributeStorage<ValueType> Storage;
+	TAttributeStorage<ValueType> Storage;
 
 	/** Evaluator for interpolation */
-    TUniquePtr<TBaseEvaluator<ValueType>> Evaluator;
+	TUniquePtr<TBaseEvaluator<ValueType>> Evaluator;
 
-    // Optional overrides:
-    TFunction<TArray<const ValueType*>(TArrayView<const TStorageEntry<ValueType>>, float)> FindWindowOverride;
-    TFunction<ValueType(TArrayView<const ValueType* const>, float)> InterpolateOverride;
+	// Optional overrides:
+	TFunction<TArray<const ValueType*>(TArrayView<const TStorageEntry<ValueType>>, float)> FindWindowOverride;
+	TFunction<ValueType(TArrayView<const ValueType* const>, float)> InterpolateOverride;
 
 public:
 	/** Default constructor */
@@ -125,52 +142,121 @@ public:
 	{}
 
 	/** Constructor with custom evaluator */
-    explicit TAttributeChannel(TUniquePtr<TBaseEvaluator<ValueType>> InEvaluator)
-        : Evaluator(MoveTemp(InEvaluator))
-    {
-    }
+	explicit TAttributeChannel(TUniquePtr<TBaseEvaluator<ValueType>> InEvaluator)
+		: Evaluator(MoveTemp(InEvaluator))
+	{
+	}
+
+	/**
+	 * Compare this channel to another channel of the same type
+	 * @param Other - Channel to compare against
+	 * @return true if the channels have identical contents
+	 */
+	virtual bool IsEqual(const IAttributeChannelBase& Other) const override
+	{
+		// First check if the other channel is of the same type
+		if (GetTypeName() != Other.GetTypeName())
+		{
+			return false;
+		}
+
+		// Safe to downcast since we checked type name
+		const TAttributeChannel<ValueType>* OtherTyped = 
+			static_cast<const TAttributeChannel<ValueType>*>(&Other);
+
+		// Compare number of entries
+		const TArrayView<const TStorageEntry<ValueType>>& ThisEntries = Storage.GetEntries();
+		const TArrayView<const TStorageEntry<ValueType>>& OtherEntries = OtherTyped->Storage.GetEntries();
+        
+		if (ThisEntries.Num() != OtherEntries.Num())
+		{
+			return false;
+		}
+
+		// Compare each entry's parameter and value
+		for (int32 i = 0; i < ThisEntries.Num(); ++i)
+		{
+			// Use the TStorageEntry equality operator that we defined
+			if (!(ThisEntries[i] == OtherEntries[i]))
+			{
+				return false;
+			}
+		}
+
+		// Compare evaluator settings
+		if (Evaluator && OtherTyped->Evaluator)
+		{
+			// Compare evaluators using their IsEqual method
+			if (!Evaluator->IsEqual(*OtherTyped->Evaluator))
+			{
+				return false;
+			}
+		}
+		else if (Evaluator != OtherTyped->Evaluator)
+		{
+			// One has evaluator and the other doesn't
+			return false;
+		}
 
+		return true;
+	}
+    
+	/**
+	 * Equality operator for TAttributeChannel
+	 */
+	bool operator==(const TAttributeChannel<ValueType>& Other) const
+	{
+		return IsEqual(Other);
+	}
+    
+	/**
+	 * Inequality operator for TAttributeChannel
+	 */
+	bool operator!=(const TAttributeChannel<ValueType>& Other) const
+	{
+		return !IsEqual(Other);
+	}
 	/** 
 	 * Implementation of the clone method
 	 * Makes a deep copy of this channel 
 	 */
 	virtual TUniquePtr<IAttributeChannelBase> Clone() const override
 	{
-    	return MakeUnique<TAttributeChannel>(*this);
-    }
+		return MakeUnique<TAttributeChannel>(*this);
+	}
 	
 	/** Copy constructor with deep evaluator clone */
 	TAttributeChannel(const TAttributeChannel& Other)
 		: Storage(Other.Storage)
 		, FindWindowOverride(Other.FindWindowOverride)
 		, InterpolateOverride(Other.InterpolateOverride)
-    {
-    	if (Other.Evaluator)
-    	{
-    		Evaluator = Other.Evaluator->Clone();
-    	}
-    }
+	{
+		if (Other.Evaluator)
+		{
+			Evaluator = Other.Evaluator->Clone();
+		}
+	}
 
 	/** Copy assignment operator with deep evaluator clone */
 	TAttributeChannel& operator=(const TAttributeChannel& Other)
-    {
-    	if (this != &Other)
-    	{
-    		Storage = Other.Storage;
-    		FindWindowOverride = Other.FindWindowOverride;
-    		InterpolateOverride = Other.InterpolateOverride;
-
-    		if (Other.Evaluator)
-    		{
-    			Evaluator = Other.Evaluator->Clone();
-    		}
-    		else
-    		{
-    			Evaluator.Reset();
-    		}
-    	}
-    	return *this;
-    }
+	{
+		if (this != &Other)
+		{
+			Storage = Other.Storage;
+			FindWindowOverride = Other.FindWindowOverride;
+			InterpolateOverride = Other.InterpolateOverride;
+
+			if (Other.Evaluator)
+			{
+				Evaluator = Other.Evaluator->Clone();
+			}
+			else
+			{
+				Evaluator.Reset();
+			}
+		}
+		return *this;
+	}
 
 	/**
 	 * Gets the registered type name for this channel
@@ -203,11 +289,11 @@ public:
 		TypeName = InTypeName;
 	}
 
-    friend FArchive& operator<<(FArchive& Ar, TAttributeChannel& Channel)
-    {
+	friend FArchive& operator<<(FArchive& Ar, TAttributeChannel& Channel)
+	{
 		Channel.Serialize(Ar);
 		return Ar;
-    }
+	}
 
 	/**
 	 * Clears all values from this channel
@@ -217,100 +303,110 @@ public:
 		Storage.Clear();
 	}
 	
-    /**
-     * Adds a new value to the channel at the given parameter
-     * @param Value 
-     * @param Parameter
-     * @return Index of the added value
-     */
-    int32 AddValue(const ValueType& Value, float Parameter)
-    {
-        return Storage.AddValue(Value, Parameter);
-    }
+	/**
+	 * Adds a new value to the channel at the given parameter
+	 * @param Value 
+	 * @param Parameter
+	 * @return Index of the added value
+	 */
+	int32 AddValue(const ValueType& Value, float Parameter)
+	{
+		return Storage.AddValue(Value, Parameter);
+	}
 
 	/**
 	 * Removes a value from the channel at the given index
 	 * @param Index - Index of the value to remove
 	 * @return True if the value was removed
 	 */
-	bool RemoveValue(int32 Index)
+	virtual bool RemoveValue(int32 Index) override
 	{
 		return Storage.RemoveValue(Index);
 	}
 	
-    /**
-     * Evaluates the channel at the given parameter in channel space
-     * @param Parameter - Parameter value in channel's local space
-     */
-    ValueType Evaluate(float Parameter) const
-    {
+	/**
+	 * Evaluates the channel at the given parameter in channel space
+	 * @param Parameter - Parameter value in channel's local space
+	 */
+	ValueType Evaluate(float Parameter) const
+	{
 		int32 NumValues = Storage.GetEntries().Num();
-        if (NumValues == 0) return ValueType();
-        if (NumValues == 1) return Storage.GetEntries()[0].Value;
+		if (NumValues == 0) return ValueType();
+		if (NumValues == 1) return Storage.GetEntries()[0].Value;
 
-        // Get window using override or evaluator
-        TArray<const ValueType*> Window;
-        if (FindWindowOverride)
-        {
-            Window = FindWindowOverride(Storage.GetEntries(), Parameter);
-        }
-        else
-        {
-            Window = Evaluator->FindWindow(Storage.GetEntries(), Parameter);
-        }
-        if (Window.Num() == 0) return ValueType();
+		Parameter = Evaluator->HandleOutOfBounds(Parameter);
+		// Get window using override or evaluator
+		TArray<const ValueType*> Window;
+		if (FindWindowOverride)
+		{
+			Window = FindWindowOverride(Storage.GetEntries(), Parameter);
+		}
+		else
+		{
+			Window = Evaluator->FindWindow(Storage.GetEntries(), Parameter);
+		}
+		if (Window.Num() == 0) return ValueType();
         
-        if (InterpolateOverride)
-        {
-            return InterpolateOverride(Window, Parameter);
-        }
-        return Evaluator->InterpolateWindow(Window, Parameter);
-    }
+		if (InterpolateOverride)
+		{
+			return InterpolateOverride(Window, Parameter);
+		}
+		return Evaluator->InterpolateWindow(Window, Parameter);
+	}
 
 	// Evaluate nth derivative of attribute
 	template<int32 Order>
 	ValueType EvaluateDerivative(float Parameter) const
-    {
-    	if (Storage.GetEntries().Num() == 0) return ValueType();
+	{
+		if (Storage.GetEntries().Num() == 0) return ValueType();
         
-    	// Use policy for derivative calculation
-    	auto Window = Evaluator->FindWindow(Storage.GetEntries(), Parameter);
+		// Use policy for derivative calculation
+		auto Window = Evaluator->FindWindow(Storage.GetEntries(), Parameter);
 
-    	// Get parameter range from evaluator
-    	float Start, End;
-    	Evaluator->GetParentSpaceRange(Start, End);
-    	const float ParamScale = End - Start;
+		// Get parameter range from evaluator
+		float Start, End;
+		Evaluator->GetParentSpaceRange(Start, End);
+		const float ParamScale = End - Start;
         
-    	if (FMath::IsNearlyZero(ParamScale))
-    	{
-    		return ValueType();
-    	}
+		if (FMath::IsNearlyZero(ParamScale))
+		{
+			return ValueType();
+		}
         
-    	return TSplineInterpolationPolicy<ValueType>::CalculateDerivative<Order>(
+		return TSplineInterpolationPolicy<ValueType>::CalculateDerivative<Order>(
 			Window, Parameter, 1.0f / ParamScale);
-    }
+	}
 
-    TBaseEvaluator<ValueType>& GetEvaluator()
-    {
-        return *Evaluator;
-    }
+	TBaseEvaluator<ValueType>& GetEvaluator()
+	{
+		return *Evaluator;
+	}
     
-    const TBaseEvaluator<ValueType>& GetEvaluator() const
-    {
-        return *Evaluator;
-    }
+	const TBaseEvaluator<ValueType>& GetEvaluator() const
+	{
+		return *Evaluator;
+	}
     
-    // Access to underlying storage
+	// Access to underlying storage
 	TAttributeStorage<ValueType>& GetStorage() 
-    {
-        return Storage;
-    }
-    
-    const TAttributeStorage<ValueType>& GetStorage() const
-    {
-        return Storage;
-    }
+	{
+		return Storage;
+	}
     
+	const TAttributeStorage<ValueType>& GetStorage() const
+	{
+		return Storage;
+	}
+
+	virtual float MapToParentSpace(float Parameter) const override
+	{
+		return Evaluator->MapToParentSpace(Parameter);
+	}
+
+	virtual float MapFromParentSpace(float Parameter) const override
+	{
+		return Evaluator->MapFromParentSpace(Parameter);
+	}
 
     // Set overrides
     void SetFindWindowOverride(TFunction<TArray<const ValueType*>(TArrayView<const TStorageEntry<ValueType>>, float)> Override)
@@ -431,13 +527,30 @@ public:
 		return 0.0f;
 	}
 
-	virtual float SetParameter(int32 Index, float Param) override
+	virtual int32 SetParameter(int32 Index, float Parameter) override
 	{
-		if (Storage.GetEntries().IsValidIndex(Index))
+		if (!Storage.GetEntries().IsValidIndex(Index))
+		{
+			return INDEX_NONE;
+		}
+
+		auto& Entries = Storage.GetEntries();
+
+		Entries[Index].Parameter = Parameter;
+
+		while (Index > 0 && Entries[Index].Parameter < Entries[Index - 1].Parameter)
 		{
-			Storage.GetEntries()[Index].Parameter = Param;
+			Entries.Swap(Index, Index - 1);
+			Index--;
 		}
-		return Param;
+
+		while (Index < Entries.Num() - 1 && Entries[Index].Parameter > Entries[Index + 1].Parameter)
+		{
+			Entries.Swap(Index, Index + 1);
+			Index++;
+		}
+
+		return Index;
 	}
 };
 	
diff --git a/Engine/Source/Runtime/GeometryCore/Public/Splines/BSpline.h b/Engine/Source/Runtime/GeometryCore/Public/Splines/BSpline.h
--- a/Engine/Source/Runtime/GeometryCore/Public/Splines/BSpline.h
+++ b/Engine/Source/Runtime/GeometryCore/Public/Splines/BSpline.h
@@ -59,7 +59,6 @@ public:
 		{
 			return Window;
 		}
-
 		const float MappedParam = Base::MapToParentSpace(Parameter);
 		int32 Span = FindKnotSpan(MappedParam);
 		// For closed loop, we need to ensure the span wraps properly
@@ -119,51 +118,7 @@ public:
 		Ar << bClampEnds;
         
 	}
-	/**
-	 * Override SetParentSpaceRange to update knot vector accordingly
-	 * @param InStart - New start parameter
-	 * @param InEnd - New end parameter
-	 */
-	virtual void SetParentSpaceRange(float InStart, float InEnd) override
-	{
-		float OldStart = this->ParentStart;
-	    float OldEnd = this->ParentEnd;
-	    
-	    // Only proceed if there's an actual change
-	    if (FMath::IsNearlyEqual(OldStart, InStart) && FMath::IsNearlyEqual(OldEnd, InEnd))
-	    {
-	        return;
-	    }
-	    
-	    // Call base implementation to update the parameter range
-	    Base::SetParentSpaceRange(InStart, InEnd);
-	    
-		// Scale the knot vector regardless of mode
-		if (KnotVector.Num() >= 2)
-		{
-			// Get the current effective knot range
-			float KnotStart, KnotEnd;
-			GetKnotRange(KnotStart, KnotEnd);
-        
-			const float OldRange = KnotEnd - KnotStart;
-			const float NewRange = InEnd - InStart;
-        
-			if (!FMath::IsNearlyZero(OldRange) && !FMath::IsNearlyZero(NewRange))
-			{
-				const float Scale = NewRange / OldRange;
-				const float Offset = InStart - KnotStart * Scale;
-            
-				// Scale the knot vector in place
-				for (int32 i = 0; i < KnotVector.Num(); ++i)
-				{
-					KnotVector[i] = KnotVector[i] * Scale + Offset;
-				}
-            
-				UE_LOG(LogSpline, Verbose, TEXT("BSpline knot vector scaled from range [%f,%f] to [%f,%f]"), 
-					KnotStart, KnotEnd, InStart, InEnd);
-			}
-		}
-	}
+	
 
 	/**
 	 * Inserts a knot at the specified parameter value
@@ -231,11 +186,10 @@ public:
 			int32 Multiplicity = GetKnotMultiplicity(Parameter);
 			if (bClampEnds)
 			{
-				float KnotStart, KnotEnd;
-				GetKnotRange(KnotStart, KnotEnd);
+				FInterval1f KnotRange = GetKnotRange();
 				// If the parameter is at an endpoint and the multiplicity is already clamped,
 				// do not remove.
-				if ((FMath::IsNearlyEqual(Parameter, KnotStart) || FMath::IsNearlyEqual(Parameter, KnotEnd)) &&
+				if ((FMath::IsNearlyEqual(Parameter, KnotRange.Min) || FMath::IsNearlyEqual(Parameter, KnotRange.Max)) &&
 					Multiplicity <= (GetDegree() + 1))
 				{
 					UE_LOG(LogSpline, Warning, TEXT("Cannot remove clamped endpoint knot: %f"), Parameter);
@@ -359,41 +313,29 @@ public:
 	}
 
 	// Get the valid parameter range in knot space
-	void GetKnotRange(float& OutStart, float& OutEnd) const
+	FInterval1f GetKnotRange() const
 	{
 		if (KnotVector.Num() < 2)
 		{
-			OutStart = 0.0f;  
-			OutEnd = 1.0f;
-			return;
+			return FInterval1f(0.0f, 1.0f);
 		}
 		if (bClampEnds)
 		{
 			// For clamped B-splines, the valid range is [Degree, NumKnots - Degree - 1]
-			OutStart = KnotVector[Degree];
-			OutEnd = KnotVector[KnotVector.Num() - Degree - 1];
+			return FInterval1f(KnotVector[Degree],KnotVector[KnotVector.Num() - Degree - 1]);
 		}
 		else
 		{
 			// For unclamped B-splines, the valid range is [FirstKnot, LastKnot]
-			OutStart = KnotVector[0];
-			OutEnd = KnotVector.Last();
+			return FInterval1f(KnotVector[0], KnotVector.Last());
 		}
 	}
 
-	// Convert from normalized [0,1] to knot space
-	float ToKnotSpace(float NormalizedParam) const
-	{
-		float Start, End;
-		GetKnotRange(Start, End);
-		return FMath::Lerp(Start, End, NormalizedParam);
-	}
-
 	/** 
 	 * Evaluate nth derivative at parameter
 	 * @param Order - The derivative order (0 = position, 1 = first derivative, etc.)
 	 * @param Entries - Spline control points
-	 * @param Parameter - Parameter in normalized space [0,1]
+	 * @param Parameter - Parameter in parent space
 	 * @return nth derivative vector
 	 */
 	template<int32 Order>
@@ -415,6 +357,24 @@ public:
 			Window, KnotVector, Span, MappedParam, Degree);
 	}
 	
+	virtual float GetParentSpaceMapping(float ChildSpaceParam) const override
+    {
+    	if (KnotVector.Num() < 2) return ChildSpaceParam;
+
+		float localT = this->GetKnotRange().GetT(ChildSpaceParam);
+		return this->GetParentSpaceRange().Interpolate(localT);
+    	
+    }
+
+    virtual float GetChildSpaceMapping(float ParentParam) const override
+    {
+    	if (KnotVector.Num() < 2) return ParentParam;
+
+		float localT = this->GetParentSpaceRange().GetT(ParentParam);
+		return this->GetKnotRange().Interpolate(localT);
+    	
+    }
+	
 protected:
 	int32 GetNumValuesFromKnots() const
 	{
@@ -436,29 +396,7 @@ protected:
 				   ? (NumValues + (2 * Degree) + 1) // closed loop: n + 2p + 1
 				   : (NumValues + Degree + 1);
 	}
-	virtual float GetParentSpaceMapping(float NormalizedParam) const override
-	{
-		if (KnotVector.Num() < 2) return NormalizedParam;
-
-		return ToKnotSpace(NormalizedParam);
-	}
-
-	virtual float GetNormalizedSpaceMapping(float KnotParam) const override
-	{
-		if (KnotVector.Num() < 2) return KnotParam;
-
-		float Start, End;
-		GetKnotRange(Start, End);
-
-		// Normalize the parameter
-		const float Range = End - Start;
-				
-		if (FMath::IsNearlyZero(Range))
-		{
-			return 0.0f;
-		}
-		return (KnotParam - Start) / Range;
-	}
+	
 	
 private:
 	/** Find knot span using binary search */
@@ -933,24 +871,23 @@ public:
 		Base::GetEvaluator().UpdateKnotVector(Base::Values.GetStorage().GetEntries(), Mode);
 
 		// Update the parameter range to match the new knot range
-    	float KnotStart, KnotEnd;
-    	this->GetEvaluator().GetKnotRange(KnotStart, KnotEnd);
+    	
+    	FInterval1f KnotRange = this->GetEvaluator().GetKnotRange();
     
     	// Only update if the knot range is valid
-    	if (KnotEnd > KnotStart)
+    	if (!KnotRange.IsEmpty())
     	{
-    		float CurrentStart, CurrentEnd;
-    		this->GetEvaluator().GetParentSpaceRange(CurrentStart, CurrentEnd);
-        
+    		FInterval1f CurrentRange = this->GetEvaluator().GetParentSpaceRange();
+
     		// Check if the knot range is significantly different from parameter range
     		const float Tolerance = 0.01f; // Small tolerance to avoid frequent small updates
-    		if (FMath::Abs(KnotStart - CurrentStart) > Tolerance || 
-    			FMath::Abs(KnotEnd - CurrentEnd) > Tolerance)
+    		if (FMath::Abs(KnotRange.Min - CurrentRange.Min) > Tolerance || 
+    			FMath::Abs(KnotRange.Max - CurrentRange.Max) > Tolerance)
     		{
     			// Update parameter range to match knot range
-    			this->GetEvaluator().SetParentSpaceRange(KnotStart, KnotEnd);
+    			this->GetEvaluator().SetParentSpaceRange(KnotRange);
     			UE_LOG(LogSpline, Verbose, TEXT("Updated parameter range to [%f,%f] to match knot range"),
-    				KnotStart, KnotEnd);
+    				KnotRange.Min, KnotRange.Max);
     		}
     	}
 	}
@@ -1004,6 +941,11 @@ public:
 		}
 	}
 
+	virtual FInterval1f GetParameterRange() const override
+	{
+		return Base::GetEvaluator().GetKnotRange();
+	}
+
 	/**
 	 * Sets whether the spline should interpolate its endpoints
 	 * @param bInClampEnds - True to interpolate endpoints, false for unclamped curve
diff --git a/Engine/Source/Runtime/GeometryCore/Public/Splines/BaseEvaluator.h b/Engine/Source/Runtime/GeometryCore/Public/Splines/BaseEvaluator.h
--- a/Engine/Source/Runtime/GeometryCore/Public/Splines/BaseEvaluator.h
+++ b/Engine/Source/Runtime/GeometryCore/Public/Splines/BaseEvaluator.h
@@ -3,6 +3,7 @@
 #pragma once
 
 #include "CoreMinimal.h"
+#include "BoxTypes.h"
 #include "InterpolationPolicies.h"
 #include "SplineMath.h"
 DEFINE_LOG_CATEGORY_STATIC(LogSpline, Log, All);
@@ -35,18 +36,16 @@ protected:
 	/** whether the spline is closed */
 	bool bClosedLoop = false;
 	
-	/** Start point in parent space */
-	float ParentStart = 0.0f;
+	/** Parameter range in parent space */
+	FInterval1f ParentSpaceRange;
     
-	/** End point in parent space */
-	float ParentEnd = 1.0f;
-
 	/** How to handle parameters outside the mapped range */
 	EOutOfBoundsHandlingMode PreInfinityMode = EOutOfBoundsHandlingMode::Zero;
 	EOutOfBoundsHandlingMode PostInfinityMode = EOutOfBoundsHandlingMode::Zero;
 
 public:
-	TBaseEvaluator()
+	TBaseEvaluator():
+		ParentSpaceRange(0.0f, 1.0f)
 	{
 	}
 	virtual ~TBaseEvaluator() = default;
@@ -60,9 +59,11 @@ public:
 		Ar << bClosedLoop;
 
 		// Serialize parameter range
-		Ar << ParentStart;
-		Ar << ParentEnd;
-
+		float Min = ParentSpaceRange.Min;
+		float Max = ParentSpaceRange.Max;
+		Ar << Min;
+		Ar << Max;
+		
 		// Serialize infinity modes
 		uint8 PreMode = static_cast<uint8>(PreInfinityMode);
 		uint8 PostMode = static_cast<uint8>(PostInfinityMode);
@@ -73,6 +74,7 @@ public:
 		{
 			PreInfinityMode = static_cast<EOutOfBoundsHandlingMode>(PreMode);
 			PostInfinityMode = static_cast<EOutOfBoundsHandlingMode>(PostMode);
+			ParentSpaceRange = FInterval1f(Min, Max);
 		}
 	}
 	friend FArchive& operator<<(FArchive& Ar, TBaseEvaluator& Evaluator)
@@ -80,6 +82,58 @@ public:
 		Evaluator.Serialize(Ar);
 		return Ar;	
 	}
+
+	/**
+	 * Checks if this evaluator is equal to another evaluator
+	 * @param Other - The other evaluator to compare with
+	 * @return True if evaluators are considered equal
+	 */
+	virtual bool IsEqual(const TBaseEvaluator& Other) const
+	{
+		// Compare essential properties
+		if (bClosedLoop != Other.bClosedLoop)
+		{
+			return false;
+		}
+
+		// Compare parameter range
+		if (!FMath::IsNearlyEqual(ParentSpaceRange.Min, Other.ParentSpaceRange.Min) ||
+			!FMath::IsNearlyEqual(ParentSpaceRange.Max, Other.ParentSpaceRange.Max))
+		{
+			return false;
+		}
+
+		// Compare infinity modes
+		if (PreInfinityMode != Other.PreInfinityMode ||
+			PostInfinityMode != Other.PostInfinityMode)
+		{
+			return false;
+		}
+
+		// Compare polynomial degree (if derived classes override this)
+		if (GetDegree() != Other.GetDegree())
+		{
+			return false;
+		}
+
+		return true;
+	}
+
+	/**
+     * Equality operator for TBaseEvaluator
+     */
+    bool operator==(const TBaseEvaluator& Other) const
+    {
+        return IsEqual(Other);
+    }
+    
+    /**
+     * Inequality operator for TBaseEvaluator
+     */
+    bool operator!=(const TBaseEvaluator& Other) const
+    {
+        return !IsEqual(Other);
+    }
 	
 	/** Get the polynomial degree of the evaluator */
 	virtual int32 GetDegree() const{ return 1; }
@@ -169,26 +223,24 @@ public:
 		return MakeUnique<TBaseEvaluator>(*this);
 	}
 	
-	/** Convert from normalized [0,1] to parent space with bounds handling */
-	virtual float MapToParentSpace(float NormalizedParam) const
+	/** Convert from child space to parent space with bounds handling */
+	virtual float MapToParentSpace(float ChildParam) const
 	{
-		const float BoundedParam = HandleOutOfBounds(NormalizedParam);
-		const float Result = GetParentSpaceMapping(BoundedParam);
+		const float Result = GetParentSpaceMapping(ChildParam);
 
-		UE_LOG(LogSpline, VeryVerbose, TEXT("MapToParentSpace: %f -> Bounded: %f -> Result: %f"), 
-			NormalizedParam, BoundedParam, Result);
+		UE_LOG(LogSpline, VeryVerbose, TEXT("MapToParentSpace: %f -> Result: %f"), 
+			ChildParam, Result);
 		return Result;
 	}
 
-	/** Convert from parent space back to normalized [0,1] with bounds handling */
+	/** Convert from parent space back to channel space */
 	virtual float MapFromParentSpace(float ParentParam) const
 	{
-		const float LocalParam = GetNormalizedSpaceMapping(ParentParam);
-		const float Result = HandleOutOfBounds(LocalParam);
-    
-		UE_LOG(LogSpline, VeryVerbose, TEXT("MapFromParentSpace: %f -> Local: %f -> Result: %f"), 
-			ParentParam, LocalParam, Result);
-		return Result;
+		const float ChildSpaceParam = GetChildSpaceMapping(ParentParam);
+		   
+		UE_LOG(LogSpline, VeryVerbose, TEXT("MapFromParentSpace: %f -> ChildSpaceParam: %f"), 
+			ParentParam, ChildSpaceParam);
+		return ChildSpaceParam;
 	}
 
 	// Virtual function to update parent reference if this is an attribute evaluator
@@ -213,20 +265,11 @@ public:
 	bool IsClosedLoop() const { return bClosedLoop; }
 
 	/** Set the mapping range in parent space */
-	virtual void SetParentSpaceRange(float InStart, float InEnd)
-	{
-		ParentStart = InStart;
-		ParentEnd = InEnd;
-	}
-
-	// Auto-expands parameter range for parameters outside current bounds
-	void ExpandRangeToInclude(float Parameter)
+	virtual void SetParentSpaceRange(FInterval1f InRange)
 	{
-		// Expand range if parameter lies outside current bounds
-		ParentStart = FMath::Min(ParentStart, Parameter);
-		ParentEnd = FMath::Max(ParentEnd, Parameter);
+		ParentSpaceRange = InRange;
 	}
-
+	
 	/** Set Pre-Infinity handling mode when out of bounds*/ 
 	void SetPreInfinityMode(EOutOfBoundsHandlingMode InMode)
 	{
@@ -252,82 +295,65 @@ public:
 	}
 
 	/** Get the current parent space range */
-	void GetParentSpaceRange(float& OutStart, float& OutEnd) const
+	FInterval1f GetParentSpaceRange() const
 	{
-		OutStart = ParentStart;
-		OutEnd = ParentEnd;
+		return ParentSpaceRange;
 	}
 
 	virtual FArchive& operator<<(FArchive& Ar)
 	{
-		Ar << ParentStart;
-		Ar << ParentEnd;
-        
-		uint8 PreMode = static_cast<uint8>(PreInfinityMode);
-		uint8 PostMode = static_cast<uint8>(PostInfinityMode);
-        
-		Ar << PreMode;
-		Ar << PostMode;
-		Ar << bClosedLoop;
-        
-		if (Ar.IsLoading())
-		{
-			PreInfinityMode = static_cast<EOutOfBoundsHandlingMode>(PreMode);
-			PostInfinityMode = static_cast<EOutOfBoundsHandlingMode>(PostMode);
-		}
-        
+		Serialize(Ar);
 		return Ar;
 	}
-	
-protected:
-	/** Override to implement specific parent space mapping */
-	virtual float GetParentSpaceMapping(float NormalizedParam) const
-	{
-		return NormalizedParam;
-	}
-
-	/** Override to implement specific normalized space mapping */
-	virtual float GetNormalizedSpaceMapping(float ParentParam) const
-	{
-		return ParentParam;
-	}
-private:
 	/** Apply out-of-bounds handling to parameter */
 	float HandleOutOfBounds(float Parameter) const
 	{
-		if (Parameter < ParentStart)
+		if (Parameter < ParentSpaceRange.Min)
 		{
 			switch (PreInfinityMode)
 			{
 			case EOutOfBoundsHandlingMode::Zero:
 				return 0.0f;
 			case EOutOfBoundsHandlingMode::Constant:
-				return ParentStart;
+				return ParentSpaceRange.Min;
 			case EOutOfBoundsHandlingMode::Cycle:
 				{
-					const float Range = ParentEnd - ParentStart;
-					return FMath::Fmod(Parameter - ParentStart, Range) + ParentStart;
+					const float Range = ParentSpaceRange.Length();
+					return FMath::Fmod(Parameter - ParentSpaceRange.Min, Range) + ParentSpaceRange.Min;
 				}
 			}
 		}
-		else if (Parameter > ParentEnd)
+		else if (Parameter > ParentSpaceRange.Max)
 		{
 			switch (PostInfinityMode)
 			{
 			case EOutOfBoundsHandlingMode::Zero:
 				return 0.0f;
 			case EOutOfBoundsHandlingMode::Constant:
-				return ParentEnd;
+				return ParentSpaceRange.Max;
 			case EOutOfBoundsHandlingMode::Cycle:
 				{
-					const float Range = ParentEnd - ParentStart;
-					return FMath::Fmod(Parameter - ParentStart, Range) + ParentStart;
+					const float Range = ParentSpaceRange.Length();
+					return FMath::Fmod(Parameter - ParentSpaceRange.Min, Range) + ParentSpaceRange.Min;
 				}
 			}
 		}
 		
 		return Parameter;
 	}
+protected:
+	/** Override to implement specific parent space mapping */
+	virtual float GetParentSpaceMapping(float ChildSpaceParam) const
+	{
+		return ChildSpaceParam;
+	}
+
+	/** Override to implement specific child space mapping */
+	virtual float GetChildSpaceMapping(float ParentParam) const
+	{
+		return ParentParam;
+	}
+	
 };
 } // end namespace UE::Geometry::Spline
 } // end namespace UE::Geometry
diff --git a/Engine/Source/Runtime/GeometryCore/Public/Splines/ChannelTypeRegistry.h b/Engine/Source/Runtime/GeometryCore/Public/Splines/ChannelTypeRegistry.h
--- a/Engine/Source/Runtime/GeometryCore/Public/Splines/ChannelTypeRegistry.h
+++ b/Engine/Source/Runtime/GeometryCore/Public/Splines/ChannelTypeRegistry.h
@@ -105,7 +105,6 @@ public:
     	// String types
     	RegisterType<FName>();
     	RegisterType<FString>();
-    	RegisterType<FText>();
     }
 };
 
diff --git a/Engine/Source/Runtime/GeometryCore/Public/Splines/CustomAttributeExample.h b/Engine/Source/Runtime/GeometryCore/Public/Splines/CustomAttributeExample.h
--- a/Engine/Source/Runtime/GeometryCore/Public/Splines/CustomAttributeExample.h
+++ b/Engine/Source/Runtime/GeometryCore/Public/Splines/CustomAttributeExample.h
@@ -42,6 +42,11 @@ struct FFoo
         Foo.Serialize(Ar);
         return Ar;
     }
+
+	bool operator==(const FFoo& Other) const
+    {
+    	return Value == Other.Value;
+    }
     
 };
 
diff --git a/Engine/Source/Runtime/GeometryCore/Public/Splines/ParameterizedTypes.h b/Engine/Source/Runtime/GeometryCore/Public/Splines/ParameterizedTypes.h
--- a/Engine/Source/Runtime/GeometryCore/Public/Splines/ParameterizedTypes.h
+++ b/Engine/Source/Runtime/GeometryCore/Public/Splines/ParameterizedTypes.h
@@ -61,6 +61,22 @@ struct TChooseVectorType
 		"Invalid vector type combination");
 };
 
+// Trait to detect if a type T has an Equals(const T&) method.
+template<typename T, typename = void>
+struct THasEqualsMethod : std::false_type {};
+
+template<typename T>
+struct THasEqualsMethod<T, std::void_t<decltype(std::declval<T>().Equals(std::declval<T>()))>>
+    : std::true_type {};
+
+// Trait to detect if type T has a static EqualTo(const T&, const T&) method.
+template<typename T, typename = void>
+struct THasEqualToMethod : std::false_type {};
+
+template<typename T>
+struct THasEqualToMethod<T, std::void_t<decltype(T::EqualTo(std::declval<T>(), std::declval<T>()))>>
+	: std::true_type {};
+	
 // Forward declarations
 template<typename ValueType>
 class TAttributeChannel;
@@ -77,6 +93,43 @@ struct TStorageEntry
 		Ar << Entry.Value;
 		return Ar;
 	}
+
+	/**
+     * Equality comparison operator for storage entries
+     * @param Other - Entry to compare against
+     * @return true if parameter and value match
+     */
+    bool operator==(const TStorageEntry<ValueType>& Other) const
+    {
+        // Compare parameters using safe floating-point comparison
+	    if (!FMath::IsNearlyEqual(Parameter, Other.Parameter))
+	    {
+	        return false;
+	    }
+		
+		if constexpr (std::is_floating_point_v<ValueType>)
+	    {
+	        return FMath::IsNearlyEqual(Value, Other.Value);
+	    }
+	    // Compare values - try to use Equals() method if available,
+	    else if constexpr (THasEqualsMethod<ValueType>::value)
+	    {
+	        return Value.Equals(Other.Value);
+	    }	    
+		// otherwise use ==
+	    else
+	    {
+	        return Value == Other.Value;
+	    }
+    }
+    
+    /**
+     * Inequality comparison operator for storage entries
+     */
+    bool operator!=(const TStorageEntry<ValueType>& Other) const
+    {
+        return !(*this == Other);
+    }
 };
 
 } // end namespace UE::Geometry::Spline
diff --git a/Engine/Source/Runtime/GeometryCore/Public/Splines/PolyBezierSpline.h b/Engine/Source/Runtime/GeometryCore/Public/Splines/PolyBezierSpline.h
--- a/Engine/Source/Runtime/GeometryCore/Public/Splines/PolyBezierSpline.h
+++ b/Engine/Source/Runtime/GeometryCore/Public/Splines/PolyBezierSpline.h
@@ -53,6 +53,9 @@ public:
                 
         }
     	
+    	// Update parameter range
+    	this->SetParentSpaceRange(this->GetKnotRange());
+    	
     	UE_LOG(LogSpline, Verbose, TEXT("Set Custom knot vector - Points: %d, Knots: %d, Mode: %d"),
 					   NumValues,  this->KnotVector.Num(), static_cast<int32>(Mode));
     	for (int32 i = 0; i < this->KnotVector.Num(); ++i)
@@ -179,7 +182,7 @@ public:
      * Evaluate nth derivative at parameter
      * @tparam Order - The derivative order (0 = position, 1 = first derivative, etc.)
      * @param Entries - Spline control points
-     * @param Parameter - Parameter in normalized space [0,1]
+     * @param Parameter - Parameter in spline space
      * @return nth derivative vector
      */
     template<int32 Order>
@@ -274,41 +277,27 @@ public:
     	// Compute the derivative using B-spline methods
     	return Math::TBSplineDerivativeCalculator<ValueType, Order>::Compute(
 			Window, KnotVector, Span, LocalT, 3);
-    }
-	/**
-	 * Override SetParentSpaceRange to update segment structure accordingly
-	 * @param InStart - New start parameter
-	 * @param InEnd - New end parameter
-	 */
-	virtual void SetParentSpaceRange(float InStart, float InEnd) override
-    {
-    	float OldStart = this->ParentStart;
-    	float OldEnd = this->ParentEnd;
-    
-    	// Only proceed if there's an actual change
-    	if (FMath::IsNearlyEqual(OldStart, InStart) && FMath::IsNearlyEqual(OldEnd, InEnd))
-    	{
-    		return;
-    	}
-    
-    	// Call base implementation to update the parameter range
-    	Base::SetParentSpaceRange(InStart, InEnd);
-    
     }
 	
-    virtual float GetParentSpaceMapping(float NormalizedParam) const override
-    {
-    	int32 NumValues = this->GetNumValuesFromKnots();
-        int32 NumSegments = NumValues / 4;
-        return NormalizedParam * static_cast<float>(NumSegments);
-    }
-
-    virtual float GetNormalizedSpaceMapping(float ParentParam) const override
-    {
-    	int32 NumValues = this->GetNumValuesFromKnots();
-    	int32 NumSegments = NumValues / 4;
-        return NumSegments > 0 ? ParentParam / static_cast<float>(NumSegments) : 0.0f;
-    }
+    // virtual float GetParentSpaceMapping(float ChildParam) const override
+    // {
+    // 	float localT = this->GetKnotRange.GetT(ChildParam);
+    // 	return this->GetParentSpaceRange().Interpolate(localT);
+    // 	// int32 NumValues = this->GetNumValuesFromKnots();
+    //  //    int32 NumSegments = NumValues / 4;
+    //  //    return this->GetParentSpaceRange().GetT(ChildParam) * static_cast<float>(NumSegments);
+    // }
+    //
+    // virtual float GetChildSpaceMapping(float ParentParam) const override
+    // {
+    // 	float localT = this->GetParentSpaceRange().GetT(ParentParam);
+    // 	return this->GetKnotRange().Interpolate(localT);
+    // 	
+    // 	// int32 NumValues = this->GetNumValuesFromKnots();
+    // 	// int32 NumSegments = NumValues / 4;
+    //  //    float NormalizedParam = NumSegments > 0 ? ParentParam / static_cast<float>(NumSegments) : 0.0f;
+    // 	
+    // }
 
 	/**
 	 * Helper to insert local knots for a segment boundary.
@@ -342,6 +331,9 @@ public:
 			this->InsertKnot(StartParam, 1);
 			this->InsertKnot(EndParam, 1);
 		}
+    	
+    	// Update parameter range
+    	this->SetParentSpaceRange(this->GetKnotRange());
 	}
 private:
 	
@@ -573,9 +565,9 @@ public:
            float AnglePerSegment = AngleSpan / static_cast<float>(NumSegments);
            
            // For large angles per segment, increase segment count for better approximation
-           if (AnglePerSegment > PI/4)
+           if (AnglePerSegment > UE_PI/4)
            {
-               NumSegments = FMath::CeilToInt(AngleSpan / (PI/4));
+               NumSegments = FMath::CeilToInt(AngleSpan / (UE_PI/4));
                AnglePerSegment = AngleSpan / static_cast<float>(NumSegments);
            }
            
@@ -648,7 +640,7 @@ public:
     
 		// Create a full circle arc (0 to 2π)
 		TPolyBezierSpline<ScalarType, Dimensions> Result = 
-			CreateCircleArc(Center, Radius, 0.0f, 2.0f * PI, NumSegments);
+			CreateCircleArc(Center, Radius, 0.0f, 2.0f * UE_PI, NumSegments);
     
 		// Set as closed loop
 		Result.SetClosedLoop(true);
@@ -667,7 +659,7 @@ public:
 	    if (NumSegments < 4) NumSegments = 4;
 	    
 	    // Calculate angle per segment
-	    const float AnglePerSegment = 2.0f * PI / static_cast<float>(NumSegments);
+	    const float AnglePerSegment = 2.0f * UE_PI / static_cast<float>(NumSegments);
 	    
 	    // Start with first point at angle 0
 	    float CurrentAngle = 0.0f;
@@ -736,105 +728,115 @@ public:
 	 */
 	bool SetControlPoints(const TArray<ValueType>& Points, EParameterizationPolicy ParameterizationPolicy)
 	{
-	    // For a cubic Bezier spline, we need either:
-	    // - 4N points for closed loop (4 points per segment)
-	    // - 4N-1 points for open spline (last point is shared between segments)
 	    const int32 NumPoints = Points.Num();
-	    	    
-	    // For open splines with 4N-1 points, the last control point of one segment 
-	    // and the first control point of the next segment must be identical
-	    const bool bIsClosedLoop = this->IsClosedLoop();
-
-		// Need at least 2 points to add a valid segment
-		if (NumPoints < 4)
-		{
-			UE_LOG(LogSpline, Warning, TEXT("SetControlPoints requires at least 4 points to add a valid segment. Got %d points."), NumPoints);
-			return false;
-		}
-		
-	    // Check if points array is valid for the spline type
-	    if (bIsClosedLoop)
-	    {
-	        // For closed loops, we need a multiple of 4 points (complete segments)
-	        if (NumPoints % 4 != 0)
-	        {
-	            UE_LOG(LogSpline, Warning, TEXT("Invalid point count for closed loop PolyBezier. Expected multiple of 4, got %d."), NumPoints);
-	            return false;
-	        }
-	    }
-	    else
+    
+	    // Need at least 4 points for a valid cubic bezier segment
+	    if (NumPoints < 4)
 	    {
-	        // For open splines, we need 4N-1 points: 4 for first segment, 3 for each additional segment
-	        // Equivalently: 1 + 3*NumSegments
-	        // So NumSegments = (NumPoints - 1) / 3
-	        // This means NumPoints - 1 must be divisible by 3
-	        if ((NumPoints - 1) % 3 != 0 && NumPoints > 1)
-	        {
-	            UE_LOG(LogSpline, Warning, TEXT("Invalid point count for open PolyBezier. Expected 4N-1 (1 + 3*NumSegments), got %d."), NumPoints);
-	            return false;
-	        }
+	        UE_LOG(LogSpline, Warning, TEXT("SetControlPoints requires at least 4 points to add a valid segment. Got %d points."), NumPoints);
+	        return false;
 	    }
 	    
 	    // Clear existing points
 	    this->Clear();
 	    
-	    // Calculate number of segments
-	    int32 NumSegments;
-	    if (bIsClosedLoop)
-	    {
-	        NumSegments = NumPoints / 4;
-	    }
-	    else
+	    const bool bIsClosedLoop = this->IsClosedLoop();
+	    
+	    // For cubic Bezier curves, points come in groups of 4:
+	    // P0 (start), P1 (control), P2 (control), P3 (end)
+	    // Each complete segment requires 4 points
+	    
+	    // Calculate how many complete segments we have
+	    int32 NumSegments = (NumPoints) / 4;
+	    
+	    // Make sure we have at least one complete segment
+	    if (NumSegments < 1)
 	    {
-	        // For open splines: NumPoints = 1 + 3*NumSegments, so NumSegments = (NumPoints - 1) / 3
-	        NumSegments = NumPoints > 0 ? (NumPoints - 1) / 3 + 1 : 0;
+	        UE_LOG(LogSpline, Warning, TEXT("Not enough points to create a complete segment. Need at least 4 points."));
+	        return false;
 	    }
 	    
-	    // Add all segments
+	    // Add each segment
 	    for (int32 SegmentIndex = 0; SegmentIndex < NumSegments; ++SegmentIndex)
 	    {
 	        const int32 BaseIndex = SegmentIndex * 4;
 	        
-	        // For each segment, we need 4 control points
-	        ValueType P0, P1, P2, P3;
-	        
-	        if (SegmentIndex == 0 || bIsClosedLoop)
+	        // Ensure we have enough points for this segment
+	        if (BaseIndex + 3 >= NumPoints)
 	        {
-	            // For first segment or any segment in closed loop, get all 4 points
-	            P0 = Points[BaseIndex % NumPoints];
-	            P1 = Points[(BaseIndex + 1) % NumPoints];
-	            P2 = Points[(BaseIndex + 2) % NumPoints];
-	            P3 = Points[(BaseIndex + 3) % NumPoints];
-	        }
-	        else
-	        {
-	            // For subsequent segments in open spline, P0 is shared with the last P3
-	            // and we only need 3 more points
-	            P0 = this->GetValueAt(this->GetNumKeys() - 1); // Last point of previous segment
-	            
-	            // Get next 3 control points (P1, P2, P3) from the input array
-	            // For segment i (i > 0), the points start at 1 + 3*(i-1)
-	            const int32 InputStartIndex = 1 + 3 * (SegmentIndex - 1);
-	            P1 = Points[InputStartIndex];
-	            P2 = Points[InputStartIndex + 1];
-	            P3 = Points[InputStartIndex + 2];
+	            break;
 	        }
 	        
+	        // Extract the 4 control points for this segment
+	        ValueType P0 = Points[BaseIndex];
+	        ValueType P1 = Points[BaseIndex + 1];
+	        ValueType P2 = Points[BaseIndex + 2];
+	        ValueType P3 = Points[BaseIndex + 3];
+	        
 	        // Add the segment
 	        if (SegmentIndex == 0)
 	        {
-	            // For first segment, add all 4 points
-	            this->AddBezierSegmentInternal(P0, P1, P2, P3, 0, ParameterizationPolicy);
+	            // For the first segment, use AddBezierSegmentInternal
+	            this->AddBezierSegmentInternal(P0, P1, P2, P3, 0.0f, ParameterizationPolicy);
 	        }
 	        else
 	        {
-	            // For subsequent segments, use AppendBezierSegment which only needs P1, P2, P3
-	            this->AppendBezierSegment(P1, P2, P3, ParameterizationPolicy);
+	            // For subsequent segments, check if P0 matches the last point of the previous segment
+	            ValueType LastPoint = this->GetValueAt(this->GetNumKeys() - 1);
+	            
+	            if ((P0 - LastPoint).SizeSquared() < UE_KINDA_SMALL_NUMBER)
+	            {
+	                // Points are the same, use AppendBezierSegment which only needs P1, P2, P3
+	                this->AppendBezierSegment(P1, P2, P3, ParameterizationPolicy);
+	            }
+	            else
+	            {
+	                // Points don't match, use AddBezierSegmentInternal with appropriately computed parameter
+	                float StartParam = this->GetParameterAt(this->GetNumKeys() - 1);
+	                this->AddBezierSegmentInternal(P0, P1, P2, P3, StartParam, ParameterizationPolicy);
+	            }
 	        }
 	    }
 	    
-	    // Set closed loop state and reparameterize
-	    this->SetClosedLoop(bIsClosedLoop);
+	    // Update closed loop state if needed
+	    if (bIsClosedLoop)
+	    {
+	        // If we're supposed to be a closed loop, make sure the last point connects to the first
+	        ValueType FirstPoint = this->GetValueAt(0);
+	        ValueType LastPoint = this->GetValueAt(this->GetNumKeys() - 1);
+	        
+	        // If not already a closed loop, add a segment to close it
+	        if ((LastPoint - FirstPoint).SizeSquared() > UE_KINDA_SMALL_NUMBER && NumSegments > 0)
+	        {
+	            // Add a segment that connects back to the start
+	            // We need to compute appropriate control points
+	            int32 LastBaseIndex = (NumSegments - 1) * 4;
+	            
+	            // We'll try to maintain tangent continuity if possible
+	            ValueType OutTangent = (Points[LastBaseIndex + 2] - Points[LastBaseIndex + 1]);
+	            ValueType InTangent = (Points[1] - Points[0]);
+	            
+	            // Scale tangents to match segment length
+	            float SegmentLength = static_cast<float>((FirstPoint - LastPoint).Size());
+	            OutTangent = OutTangent.GetSafeNormal() * (SegmentLength * 0.33f);
+	            InTangent = InTangent.GetSafeNormal() * (SegmentLength * 0.33f);
+	            
+	            // Create control points
+	            ValueType P1 = LastPoint + OutTangent;
+	            ValueType P2 = FirstPoint - InTangent;
+	            
+	            // Add the closing segment
+	            this->AppendBezierSegment(P1, P2, FirstPoint, ParameterizationPolicy);
+	        }
+	        
+	        this->SetClosedLoop(true);
+	    }
+	    else
+	    {
+	        this->SetClosedLoop(false);
+	    }
+	    
+	    // Reparameterize for consistent parameter distribution
 	    this->Reparameterize(ParameterizationPolicy);
 	    
 	    return true;
@@ -983,16 +985,10 @@ public:
     {
 		// Get P0 from the last point of the previous segment
 		ValueType P0 = Base::GetValueAt(Base::GetNumKeys() - 1);
-    
-		// Get current parameter range end
-		float CurrentStart, CurrentEnd;
-		this->GetEvaluator().GetParentSpaceRange(CurrentStart, CurrentEnd);
-    
-		// Use a parameter just past the current end to trigger append behavior
-		float AppendParam = CurrentEnd + 0.1f;
-    
+		float StartParam = Base::GetParameterAt(Base::GetNumKeys() - 1);
+		
 		// Add segment with append behavior
-		return AddBezierSegmentInternal(P0, P1, P2, P3, AppendParam, ParameterizationPolicy);
+		return AddBezierSegmentInternal(P0, P1, P2, P3, StartParam, ParameterizationPolicy);
     }
 
     /**
@@ -1013,23 +1009,22 @@ public:
 	{
 		// Get P3 from the first point of the existing spline
 		ValueType P3 = Base::GetValueAt(0);
-    
 		// Get current parameter range start
-		float CurrentStart, CurrentEnd;
-		this->GetEvaluator().GetParentSpaceRange(CurrentStart, CurrentEnd);
+		
+		FInterval1f ParameterRange = this->GetParameterRange();
     
 		// Use a parameter just before the current start to trigger prepend behavior
-		float PrependParam = CurrentStart - 0.1f;
+		float PrependParam = ParameterRange.Min - 0.1f;
     
 		// Add segment with prepend behavior
-		return AddBezierSegmentInternal(P0, P1, P2, P3, PrependParam, ParameterizationPolicy);
+		return AddBezierSegmentInternal(P0, P1, P2, P3, PrependParam, ParameterizationPolicy, ParameterRange.Min);
 	}
 
     /**
      * Inserts a new point at the specified parameter along the spline.
      * Uses the specified parameterization policy to preserve attribute positions.
      * 
-     * @param Parameter Parameter value [0,1] where to insert the point
+     * @param Parameter Parameter value where to insert the point (Spline Space)
      * @param Position Position for the new point
      * @param ParameterizationPolicy How to distribute parameters
      * @return The index of the newly created point
@@ -1047,12 +1042,35 @@ public:
             Base::AddValue(Position, 0.0f);
             return 0;
         }
+        FInterval1f ParameterRange = this->GetParameterRange();
+		
+		// Handle out-of-range parameters
+		if (Parameter <= ParameterRange.Min)
+		{
+			// Prepend
+			PrependPoint(Position, ParameterizationPolicy);
+			Reparameterize(ParameterizationPolicy);
+			return 0;
+		}
+    
+		if (Parameter >= ParameterRange.Max)
+		{
+			// Append - create a point at the end
+			float StartParam = Base::GetParameterAt(Base::GetNumKeys() - 1);
+			ValueType P0 = Base::GetValueAt(Base::GetNumKeys() - 1);
         
-        // Ensure parameter is in valid range
-        Parameter = FMath::Clamp(Parameter, 0.0f, 1.0f);
+			// Calculate reasonable auto-tangents
+			ValueType Dir = Position - P0;
+			ValueType P1 = P0 + Dir * 0.33f;
+			ValueType P2 = P0 + Dir * 0.66f;
+        
+			// Add the segment
+			AddBezierSegmentInternal(P0, P1, P2, Position, StartParam, ParameterizationPolicy);
+			Reparameterize(ParameterizationPolicy);
+			return GetNumSegments();
+		}
             	
-        float MappedParam = this->GetEvaluator().MapToParentSpace(Parameter);
-    	const float ClampedParam = FMath::Clamp(MappedParam, 0.0f, static_cast<float>(NumSegments));
+    	const float ClampedParam = FMath::Clamp(Parameter, 0.0f, static_cast<float>(NumSegments));
     	const int32 SegmentIndex = FMath::Min(FMath::FloorToInt(ClampedParam), NumSegments - 1);
     	const float LocalT = ClampedParam - static_cast<float>(SegmentIndex); // Parameter [0,1] within segment
     	
@@ -1063,10 +1081,6 @@ public:
         ValueType P2 = Base::GetValueAt(BaseIndex + 2);
         ValueType P3 = Base::GetValueAt(BaseIndex + 3);
 
-    	float Param0 = Base::GetParameterAt(BaseIndex);
-    	float Param3 = Base::GetParameterAt(BaseIndex + 3);
-    	float SplitParam = Param0 + (Param3 - Param0) * LocalT;
-    	
         // Use De Casteljau algorithm to split the curve
         const float t = LocalT;
         const float mt = 1.0f - t;
@@ -1085,25 +1099,55 @@ public:
         // Store indices for return value calculation
         const int32 InsertedPointIndex = BaseIndex + 3;
         
-        
+        float OriginalStartParam = Base::GetParameterAt(BaseIndex);
+        float OriginalEndParam = Base::GetParameterAt(BaseIndex+3);
+		float TargetSplitParam = OriginalStartParam + LocalT * (OriginalEndParam - OriginalStartParam);
+    
 		// Remove the original segment
 		RemoveSegmentInternal(SegmentIndex);
 		
         // Add the two new segments using the policy
-        AddBezierSegmentInternal(Q0, Q1, Q2, Q3, Param0, ParameterizationPolicy);
-        AddBezierSegmentInternal(R0, R1, R2, R3, SplitParam, ParameterizationPolicy);
-        
-		
-    	// Insert the knot directly into the knot vector
-    	this->GetEvaluator().InsertKnot(SplitParam);
-
+        AddBezierSegmentInternal(Q0, Q1, Q2, Q3, OriginalStartParam, ParameterizationPolicy, TargetSplitParam);
+		AddBezierSegmentInternal(R0, R1, R2, R3, TargetSplitParam, ParameterizationPolicy, OriginalEndParam);
+		Reparameterize(ParameterizationPolicy);
+				
         return InsertedPointIndex;
     }
 
+    /**
+     * Prepends a point to the spline at the beginning.
+     * The new point will be connected to the first segment.
+     * @param Position The position of the new point
+     * @param ParameterizationPolicy How to distribute parameters
+     * @return Index of the new point
+     */
+    int32 PrependPoint(
+		const ValueType& Position,
+		EParameterizationPolicy ParameterizationPolicy = EParameterizationPolicy::Centripetal)
+	{
+		// Get current parameter range start
+		
+		FInterval1f ParameterRange = this->GetParameterRange();
+    
+		// Use a parameter just before the current start
+		float PrependParam = ParameterRange.Min - 0.1f;
+    
+		// Create a control point at the beginning
+		ValueType P3 = Base::GetValueAt(0); // First point of the spline
+    
+		// Calculate reasonable auto-tangents
+		ValueType Dir = P3 - Position;
+		ValueType P1 = Position + Dir * 0.33f;
+		ValueType P2 = Position + Dir * 0.66f;
+    
+		// Add the segment by reusing existing code
+		return AddBezierSegmentInternal(Position, P1, P2, P3, PrependParam, ParameterizationPolicy, ParameterRange.Min);
+	}
+	
     /**
      * Inserts a Bezier segment at the specified parameter
      * 
-     * @param Parameter Parameter value [0,1] where to insert the segment
+     * @param Parameter Parameter value where to insert the segment (Spline Space)
      * @param P1 First control point
      * @param P2 Second control point
      * @param P3 End point
@@ -1117,14 +1161,15 @@ public:
         const ValueType& P3, 
         EParameterizationPolicy ParameterizationPolicy = EParameterizationPolicy::Centripetal)
     {
+		FInterval1f ParameterRange = this->GetParameterRange();
         // Handle edge cases
-        if (Parameter >= 1.0f)
+        if (Parameter >= ParameterRange.Max)
         {
             // Add new segment at the end
             return AppendBezierSegment(P1, P2, P3, ParameterizationPolicy);
         }
         
-        if (Parameter <= 0.0f)
+        if (Parameter <= ParameterRange.Min)
         {
 			// we treat P1, P2, P3 as the starting three control points, and connect to existing spline at the start
             return PrependBezierSegment(P1, P2, P3, ParameterizationPolicy);
@@ -1221,6 +1266,11 @@ public:
     
         return bSuccess;
     }
+	
+	virtual FInterval1f GetParameterRange() const override
+	{
+		return Base::GetEvaluator().GetKnotRange();
+	}
 
     /** Returns the number of Bezier segments in the spline */
     virtual int32 GetNumSegments() const override
@@ -1233,11 +1283,10 @@ public:
      * Finds the closest point on the spline to the given point
      * @param Point - Target point to find nearest point for
      * @param OutSquaredDistance - Output squared distance to the nearest point
-     * @return Parameter value of the nearest point [0,1]
+     * @return Parameter value of the nearest point in parent space
      */
 	virtual float FindNearest(const ValueType& Point, float& OutSquaredDistance) const override
     {
-    	// Can use same structure as TBSpline since we're cubic
     	float BestDistSq = TNumericLimits<float>::Max();
     	float BestParam = 0.0f;
 
@@ -1245,9 +1294,8 @@ public:
     	for (int32 i = 0; i < NumSegments; ++i)
     	{
     		float SegmentDistSq;
-    		const float LocalT = FindNearestOnSegment(Point, i, SegmentDistSq);
-    		// Convert local [0,1] parameter to normalized spline parameter
-    		const float SplineParam = (static_cast<float>(i) + LocalT) / static_cast<float>(NumSegments);
+    		const float SplineParam = FindNearestOnSegment(Point, i, SegmentDistSq);
+    		
     		if (SegmentDistSq < BestDistSq)
     		{
     			BestDistSq = SegmentDistSq;
@@ -1265,7 +1313,7 @@ public:
      * @param Point - Target point to find nearest point for
      * @param SegmentIndex - Index of the segment to search
      * @param OutSquaredDistance - Output squared distance to the nearest point
-     * @return Parameter value within the segment [0,1]
+     * @return Parameter value within spline space
      */
 	virtual float FindNearestOnSegment(const ValueType& Point, int32 SegmentIndex, float& OutSquaredDistance) const override
     {
@@ -1292,7 +1340,11 @@ public:
     	Coeffs[2] = 3 * (P2 - 2*P1 + P0);          // quadratic coefficient  
     	Coeffs[3] = P3 - 3*P2 + 3*P1 - P0;         // cubic coefficient
 
-    	return Math::FindNearestPoint_Cubic(Coeffs, 0.0f, 1.0f, OutSquaredDistance);
+    	float LocalT = Math::FindNearestPoint_Cubic(Coeffs, 0.0f, 1.0f, OutSquaredDistance);
+		// Convert local [0,1] parameter to normalized spline parameter
+		const float SplineParamNormalized = (static_cast<float>(SegmentIndex) + LocalT) / static_cast<float>(GetNumSegments());
+		const float SplineParam = this->GetParameterRange().Interpolate(SplineParamNormalized);
+		return SplineParam;
     }
 
     /**
@@ -1345,87 +1397,47 @@ private:
     	ValueType P1(0.33f, 0, 0);
     	ValueType P2(0.66f, 0, 0);
     	ValueType P3(1, 0, 0);
-    	AddBezierSegmentInternal(P0, P1, P2, P3, 0, EParameterizationPolicy::Uniform);
+    	AddBezierSegmentInternal(P0, P1, P2, P3, 0, EParameterizationPolicy::Centripetal);
 	}
 	
 	/**
-	 * Computes an appropriate parameter for a new segment based on insertion position and parameterization policy
-	 * @param P0 Start point of the segment
-	 * @param P1 First control point
-	 * @param P2 Second control point
-	 * @param P3 End point of the segment
-	 * @param RequestedParam The requested parameter (used to determine prepend vs append)
-	 * @param ParameterizationPolicy The parameterization policy to use
-	 * @return A properly computed parameter value for the segment
+	 * Computes an appropriate parameter for a new segment based on insertion position.
+	 * For all parameterization modes, segments occupy 1.0 unit of parameter space.
+	 * 
+	 * @param RequestedParam - The requested parameter value (used to determine prepend/append/insert)
+	 * @return The parameter value to use for the new segment
 	 */
-	float ComputeNewSegmentParameter(
-	    const ValueType& P0,
-	    const ValueType& P1, 
-	    const ValueType& P2, 
-	    const ValueType& P3,
-	    float RequestedParam,
-	    EParameterizationPolicy ParameterizationPolicy)
-	{
-	    // Get current parameter range
-	    float CurrentStart, CurrentEnd;
-	    this->GetEvaluator().GetParentSpaceRange(CurrentStart, CurrentEnd);
-	    
-	    // Handle special case for empty splines
-	    bool bIsEmpty = (this->GetNumSegments() == 0);
-	    
-	    // If the spline is empty, just use the requested parameter (or 0 if near zero)
-	    if (bIsEmpty)
-	    {
-	        return FMath::IsNearlyZero(RequestedParam) ? 0.0f : RequestedParam;
-	    }
-	    
-	    // Check if parameter is within current range
-	    if (RequestedParam >= CurrentStart && RequestedParam <= CurrentEnd)
-	    {
-	        // For insertion within the curve, use the requested parameter directly
-	        return RequestedParam;
-	    }
-	    
-	    // Handle out-of-bounds cases
-	    float NewParam;
-	    float SegmentLength;
-	    
-	    // Calculate segment length based on parameterization policy
-	    switch (ParameterizationPolicy)
-	    {
-	    case EParameterizationPolicy::ChordLength:
-	        SegmentLength = static_cast<float>((P3 - P0).Size());
-	        break;
-	    case EParameterizationPolicy::Centripetal:
-	        SegmentLength = FMath::Sqrt(static_cast<float>((P3 - P0).Size()));
-	        break;
-	    case EParameterizationPolicy::Uniform:
-	    default:
-	        SegmentLength = 1.0f;
-	        break;
-	    }
-	    
-	    // Determine prepend or append
-	    if (RequestedParam < CurrentStart)
-	    {
-	        // Prepend case - parameter before current start
-	        NewParam = CurrentStart - SegmentLength;
-	    }
-	    else
-	    {
-	        // Append case - parameter after current end
-	        NewParam = CurrentEnd + SegmentLength;
-	        
-	        // For uniform parameterization, use integer values for append
-	        if (ParameterizationPolicy == EParameterizationPolicy::Uniform)
-	        {
-	            const int32 NumSegments = GetNumSegments();
-	            NewParam = static_cast<float>(NumSegments + 1);
-	        }
-	    }
-	    
-	    return NewParam;
-	}
+	float ComputeNewSegmentParameter(float RequestedParam)
+    {
+    	// Get current parameter range
+    	FInterval1f CurrentParameterRange = this->GetEvaluator().GetParentSpaceRange();
+    
+    	// Handle special case for empty splines
+    	bool bIsEmpty = (this->GetNumSegments() == 0);
+    	if (bIsEmpty)
+    	{
+    		return FMath::IsNearlyZero(RequestedParam) ? 0.0f : RequestedParam;
+    	}
+    
+    	// Check if parameter is within current range
+    	if (CurrentParameterRange.Contains(RequestedParam))
+    	{
+    		// For insertion within the curve, use the requested parameter directly
+    		return RequestedParam;
+    	}
+    
+    	// Determine prepend or append
+    	if (RequestedParam < CurrentParameterRange.Min)
+    	{
+    		// Prepend case - place 1.0 unit before the start
+    		return CurrentParameterRange.Min - 1.0f;
+    	}
+    	else
+    	{
+    		// Append case - place at NumSegments (end of the current parameter range)
+    		return static_cast<float>(GetNumSegments());
+    	}
+    }
 
 	/**
 	 * Computes the end parameter for a segment given the start parameter and policy
@@ -1477,6 +1489,7 @@ private:
      * @param P2 Second control point
      * @param P3 End point
      * @param StartParam Start parameter for the segment
+     * @param EndParam Optional End parameter for the segment, or -1 for automatic calculation
      * @param ParameterizationPolicy How to distribute parameters
      * @return Index of the new segment
      */
@@ -1486,13 +1499,17 @@ private:
 		const ValueType& P2, 
 		const ValueType& P3,
 		float StartParam,
-		EParameterizationPolicy ParameterizationPolicy = EParameterizationPolicy::Centripetal)
+		EParameterizationPolicy ParameterizationPolicy = EParameterizationPolicy::Centripetal,
+		float EndParam = -1.0f) // Default to -1 for automatic calculation
     {
-    	// Compute the appropriate start parameter based on requested value and policy
-    	StartParam = ComputeNewSegmentParameter(P0, P1, P2, P3, StartParam, ParameterizationPolicy);
+    	// Compute the appropriate start parameter based on requested value
+    	StartParam = ComputeNewSegmentParameter(StartParam);
     
-    	// Calculate endpoint parameter based on policy
-    	float EndParam = ComputeEndParameter(StartParam, P0, P3, ParameterizationPolicy);
+    	// If EndParam is provided, use it; otherwise, calculate it
+    	if (EndParam < 0.0f || EndParam < StartParam)
+    	{
+    		EndParam = StartParam + 1.0f;
+    	}
 
     	// Calculate intermediate control point parameters
     	float IntermediateParams[2] = {
@@ -1507,15 +1524,6 @@ private:
     	Base::AddValue(P3, EndParam);
     
     	this->GetEvaluator().InsertLocalKnots(StartParam, EndParam);
-
-    	// Update parameter range
-    	float CurrentStart, CurrentEnd;
-    	this->GetEvaluator().GetParentSpaceRange(CurrentStart, CurrentEnd);
-    
-    	this->GetEvaluator().SetParentSpaceRange(
-			FMath::Min(CurrentStart, StartParam),
-			FMath::Max(CurrentEnd, EndParam)
-		);
     
     	return P0Index / 4;
     }
diff --git a/Engine/Source/Runtime/GeometryCore/Public/Splines/Spline.h b/Engine/Source/Runtime/GeometryCore/Public/Splines/Spline.h
--- a/Engine/Source/Runtime/GeometryCore/Public/Splines/Spline.h
+++ b/Engine/Source/Runtime/GeometryCore/Public/Splines/Spline.h
@@ -6,6 +6,7 @@
 #include "BaseEvaluator.h"
 #include "ChannelTypeRegistry.h"
 #include "BuiltInAttributeTypes.h"
+#include "BoxTypes.h"
 
 namespace UE
 {
@@ -53,7 +54,6 @@ public:
 	using InterfaceType = TBaseEvaluator<ValueType>;
 
 	using SplineValueType = typename TChooseVectorType<SplineT, SplineD>::Type;
-	// Change from reference to pointer
 	const TBaseEvaluator<SplineValueType>* ParentEvaluator;
 
 public:
@@ -83,16 +83,15 @@ public:
 		return false;
 	}
 protected:
-	virtual float GetParentSpaceMapping(float NormalizedParam) const override
+	virtual float MapToParentSpace(float ChildParam) const override
 	{
-		return ParentEvaluator->MapToParentSpace(NormalizedParam);
+		return ParentEvaluator->MapToParentSpace(ChildParam);
 	}
 
-	virtual float GetNormalizedSpaceMapping(float ParentParam) const override
+	virtual float MapFromParentSpace(float ParentParam) const override
 	{
 		return ParentEvaluator->MapFromParentSpace(ParentParam);
 	}
-
 };
 
 /**
@@ -126,6 +125,12 @@ public:
 		{
 			AttributeChannels.Add(Pair.Key, Pair.Value->Clone());
 		}
+		
+		// After loading all channels, update their parent evaluator references
+        for (auto& Pair : AttributeChannels)
+        {
+        	Pair.Value->UpdateEvaluatorParent(&GetEvaluator());
+        }
 	}
     
 	TSplineBase& operator=(const TSplineBase& Other)
@@ -138,14 +143,62 @@ public:
 			{
 				AttributeChannels.Add(Pair.Key, Pair.Value->Clone());
 			}
+
+			// After loading all channels, update their parent evaluator references
+			for (auto& Pair : AttributeChannels)
+			{
+				Pair.Value->UpdateEvaluatorParent(&GetEvaluator());
+			}
 		}
 		return *this;
 	}
+	bool operator==(const TSplineBase& Other) const
+	{
+		// Compare values channel
+		if (Values != Other.Values)
+		{
+			return false;
+		}
+
+		// Compare attribute channels
+		if (AttributeChannels.Num() != Other.AttributeChannels.Num())
+		{
+			return false;
+		}
+
+		// Check each attribute channel by name and content
+		for (const auto& Pair : AttributeChannels)
+		{
+			const FName& ChannelName = Pair.Key;
+			const TUniquePtr<IAttributeChannelBase>& Channel = Pair.Value;
+
+			// Find corresponding channel in Other
+			const TUniquePtr<IAttributeChannelBase>* OtherChannelPtr = Other.AttributeChannels.Find(ChannelName);
+			if (!OtherChannelPtr || !(*OtherChannelPtr))
+			{
+				return false;
+			}
+
+			// Compare channel contents using the IsEqual method
+			const IAttributeChannelBase& OtherChannel = **OtherChannelPtr;
+			if (!Channel->IsEqual(OtherChannel))
+			{
+				return false;
+			}
+		}
+
+		return true;
+	}
+
+	bool operator!=(const TSplineBase& Other) const
+	{
+		return !(*this == Other);
+	}
 
 	// Move operations can be defaulted
 	TSplineBase(TSplineBase&&) noexcept = default;
 	TSplineBase& operator=(TSplineBase&&) noexcept = default;
-
+	
 	friend FArchive& operator<<(FArchive& Ar, TSplineBase& Spline)
 	{
 		// Ensure initialization before serialization/deserialization
@@ -224,10 +277,30 @@ public:
 	  * @tparam AttrType - Custom attribute type
 	  * @param Name - Name of the attribute channel
 	  * @param AttrEvaluator - Custom evaluator for the attribute channel
+	  * @return true if the channel was created or already exists with matching type, false otherwise
 	  */
 	template <typename AttrType>
-	void CreateAttributeChannel( const FName& Name,	TUniquePtr<TBaseEvaluator<AttrType>> AttrEvaluator)
+	bool CreateAttributeChannel( const FName& Name,	TUniquePtr<TBaseEvaluator<AttrType>> AttrEvaluator)
 	{
+		// Check if a channel with this name already exists
+		if (const TUniquePtr<IAttributeChannelBase>* ExistingChannel = AttributeChannels.Find(Name))
+		{
+			// Channel exists - verify it has the same type by comparing type names
+			const FName TypeName = TAttributeTypeTraits<AttrType>::Name;
+			const FName ExistingTypeName = (*ExistingChannel)->GetTypeName();
+        
+			if (TypeName != ExistingTypeName)
+			{
+				// Channel exists but with a different type
+				UE_LOG(LogSpline, Warning, TEXT("Failed to create attribute channel '%s': Channel already exists with type '%s', but requested type is '%s'."),
+					*Name.ToString(), *ExistingTypeName.ToString(), *TypeName.ToString());
+				return false;
+			}
+        
+			// Channel already exists with same type - nothing to do
+			return true;
+		}
+
 		// Ensure attribute interpolation degree doesn't exceed spline degree
 		const int32 SplineDegree = GetEvaluator().GetDegree();
 		const int32 AttrDegree = AttrEvaluator->GetDegree();
@@ -238,7 +311,7 @@ public:
 			UE_LOG(LogSpline, Error, TEXT("Attribute channel '%s' has a degree (%d) greater than the spline degree (%d). This is not allowed."),
 				*Name.ToString(), AttrDegree, SplineDegree);
 	
-			return;
+			return false;
 		}
 
 		// Ensure the registry itself is initialized
@@ -253,15 +326,17 @@ public:
 
 		TUniquePtr<TAttributeChannel<AttrType>> Channel = MakeUnique<TAttributeChannel<AttrType>>(MoveTemp(AttrEvaluator));
 		AttributeChannels.Add(Name, MoveTemp(Channel));
+		return true;
 	}
 
 	/**
 	 * Creates an attribute channel that follows the spline's evaluation pattern
 	 * @tparam AttrType - Custom attribute type
-	 * @param Name - Name of the attribute channel 
+	 * @param Name - Name of the attribute channel
+	 * @return true if the channel was created or already exists with matching type, false otherwise
 	 */
 	template <typename AttrType>
-	void CreateAttributeChannel(const FName& Name)
+	bool CreateAttributeChannel(const FName& Name)
 	{
 		using AttributeEvaluator = TParentBoundAttributeEvaluator<AttrType, ScalarType, Dimensions>;
 
@@ -271,7 +346,7 @@ public:
 		// Release ownership and create new unique ptr of base type
 		TUniquePtr<TBaseEvaluator<AttrType>> BaseEvaluator(ChannelEvaluator.Release());
 
-		CreateAttributeChannel(Name, MoveTemp(BaseEvaluator));
+		return CreateAttributeChannel(Name, MoveTemp(BaseEvaluator));
 	}
 
 	/**
@@ -279,29 +354,44 @@ public:
 	 * @tparam AttrType - Custom attribute type
 	 * @param Name - Name of the attribute channel
 	 * @param Value - Value to add
-	 * @param Parameter - Parameter value in parent spline space
+	 * @param ParentSpaceParameter - Parameter value in parent spline space
 	 * @return Index of the added value
 	 */
 	template <typename AttrType>
-	int32 AddAttributeValue(FName Name,	const AttrType& Value, float Parameter)
+	int32 AddAttributeValue(FName Name,	const AttrType& Value, float ParentSpaceParameter)
 	{
 		if (TAttributeChannel<AttrType>* Channel = GetAttributeChannel<AttrType>(Name))
 		{
-			return Channel->AddValue(Value, Parameter);
+			float ChildSpaceParameter = Channel->MapFromParentSpace(ParentSpaceParameter);
+			return Channel->AddValue(Value, ChildSpaceParameter);
 		}
 		return INDEX_NONE;
 	}
 
+	/**
+	 * Removes a value from an attribute channel
+	 * @param Name - Name of the attribute channel
+	 * @param Index - Index of the value to remove
+	 * @return true if the value was removed
+	 */
+	bool RemoveAttributeValue(FName Name, int32 Index)
+	{
+		if (const TUniquePtr<IAttributeChannelBase>* Channel = AttributeChannels.Find(Name))
+		{
+			return (*Channel)->RemoveValue(Index);
+		}
+		return false;
+	}
+
 	/**
 	 * Adds a new value to the spline's main storage channel.
 	 * Also expands the parameter range if necessary.
 	 * @param Value - The value to add
 	 * @param Parameter - The parameter value in parent spline space
-	 * @return Index of the added value
+     * @return Index of the added value
 	 */
 	virtual int32 AddValue(const ValueType& Value, float Parameter)
 	{
-		GetEvaluator().ExpandRangeToInclude(Parameter);
 		return Values.AddValue(Value, Parameter);		
 	}
 	
@@ -342,7 +432,7 @@ public:
 	{
 		return Values.Evaluate(Parameter);
 	}
-
+	
 	/**
 	 * Evaluates a derivative at the given parameter in parent space
 	 * @tparam Order Derivative order
@@ -382,6 +472,7 @@ public:
 	 * @tparam AttrType - Custom attribute type
 	 * @param Name - Name of the attribute channel
 	 * @param Parameter - Parameter value in parent spline space
+	 * @return Value at the given parameter
 	 */
 	template <typename AttrType>
 	AttrType EvaluateAttribute(FName Name, float Parameter) const
@@ -437,46 +528,55 @@ public:
 
 	/**
 	 * Retrieves the parameter value of a specific attribute at the given index
-	 * @param Index - Index of the value
 	 * @param Name - Name of the attribute channel
-	 * @return Parameter value
+	 * @param Index - Index of the value
+	 * @return Parameter value in parent spline space
 	 */
-	float GetAttributeParameter(int32 Index, const FName& Name) const
+	float GetAttributeParameter(const FName& Name, int32 Index) const
 	{
 		if (const TUniquePtr<IAttributeChannelBase>* Channel = AttributeChannels.Find(Name))
 		{
-			return (*Channel)->GetParameter(Index);
+			const float ChannelSpaceParam = (*Channel)->GetParameter(Index);
+			return (*Channel)->MapToParentSpace(ChannelSpaceParam);
 		}
 		return 0.f;
 	}
 
 	/**
-	 * Sets the parameter value of a specific attribute at the given index
-	 * @param Index - Index of the value
-	 * @param NewParameter - New parameter to set
+	 * Updates the parameter of an attribute by index, possibly changing its index to preserve the ordering of the attribute storage
 	 * @param Name - Name of the attribute channel
+	 * @param Index - Index of the attribute to modify
+	 * @param ParentSpaceParameter - Parameter to set in parent spline space
+	 * @return New index of the attribute originally at the given index after any reordering
 	 */
-	void SetAttributeParameter(int32 Index, float NewParameter, const FName& Name)
+	int32 SetAttributeParameter(const FName& Name, int32 Index, float ParentSpaceParameter)
 	{
 		if (const TUniquePtr<IAttributeChannelBase>* Channel = AttributeChannels.Find(Name))
 		{
-			(*Channel)->SetParameter(Index, NewParameter);
+			// Convert from parent space to channel's local space
+            const float Parameter = (*Channel)->MapFromParentSpace(ParentSpaceParameter);
+			return (*Channel)->SetParameter(Index, Parameter);
 		}
+
+		return Index;
 	}
 
 	/**
 	 * Retrieves the value of a specific attribute at the given index
 	 * @tparam AttrType - Custom attribute type
-	 * @param Index - Index of the value
 	 * @param Name - Name of the attribute channel
+	 * @param Index - Index of the value
 	 * @return Attribute value
 	 */
 	template <typename AttrType>
-	AttrType GetAttributeValue(int32 Index, const FName& Name) const
+	AttrType GetAttributeValue(const FName& Name, int32 Index) const
 	{
 		if (const TAttributeChannel<AttrType>* Channel = GetAttributeChannel<AttrType>(Name))
 		{
-			return Channel->GetStorage().GetEntries()[Index].Value;
+			if (Channel->GetStorage().GetEntries().IsValidIndex(Index))
+			{
+				return Channel->GetStorage().GetEntries()[Index].Value;
+			}
 		}
 
 		return AttrType();
@@ -485,40 +585,19 @@ public:
 	/**
 	 * Sets the value of a specific attribute at the given index
 	 * @tparam AttrType - Custom attribute type
+	 * @param Name - Name of the attribute channel
 	 * @param Index - Index of the value
 	 * @param Value - New value to set
-	 * @param Name - Name of the attribute channel
 	 */
 	template <typename AttrType>
-	void SetAttributeValue(int32 Index, const AttrType& Value, const FName& Name)
+	void SetAttributeValue(const FName& Name, int32 Index, const AttrType& Value)
 	{
-		if (const TAttributeChannel<AttrType>* Channel = GetAttributeChannel<AttrType>(Name))
+		if (TAttributeChannel<AttrType>* Channel = GetAttributeChannel<AttrType>(Name))
 		{
 			Channel->GetStorage().GetEntries()[Index].Value = Value;
 		}
 	}
 
-	/** 
-	 * Sets the parameter range for evaluation. This remaps how parameters are interpreted
-	 * but does not move control points.
-	 * @param InStart - New start parameter
-	 * @param InEnd - New end parameter 
-	 */
-	void SetParameterRange(float InStart, float InEnd)
-	{
-		GetEvaluator().SetParentSpaceRange(InStart, InEnd);
-	}
-
-	/**
-	 *  Get custom parameter range in parent space
-	 *  @param OutStart - Start of the range
-	 *	@param OutEnd - End of the range	
-	 */ 
-	void GetParameterRange(float& OutStart, float& OutEnd) const
-	{
-		GetEvaluator().GetParentSpaceRange(OutStart, OutEnd);
-	}
-
 	/** Set Pre-Infinity handling mode when out of bounds */ 
 	void SetPreInfinityMode(EOutOfBoundsHandlingMode InMode)
 	{
@@ -577,6 +656,11 @@ public:
 		return Values.GetStorage().GetEntries()[Index].Value;
 	}
 
+	void SetValueAt(int32 Index, const ValueType& Value)
+	{
+		Values.GetStorage().GetEntries()[Index].Value = Value;
+	}
+
 	float GetParameterAt(int32 Index) const
 	{
 		return Values.GetStorage().GetEntries()[Index].Parameter;
@@ -612,6 +696,29 @@ public:
 		TUniquePtr<IAttributeChannelBase>* Found = AttributeChannels.Find(Name);
 		return Found ? static_cast<TAttributeChannel<AttrType>*>(Found->Get()) : nullptr;
 	}
+
+	const TMap<FName, TUniquePtr<IAttributeChannelBase>>& GetAttributeChannels() const
+	{
+		return AttributeChannels;
+	}
+
+	template <typename AttrType>
+	TArray<FName> GetAttributeChannelNames() const
+	{
+		TArray<FName> Names;
+		const FName TypeName = TAttributeTypeTraits<AttrType>::Name;
+    
+		for (const auto& Pair : AttributeChannels)
+		{
+			// Compare type names directly instead of using dynamic_cast
+			if (Pair.Value->GetTypeName() == TypeName)
+			{
+				Names.Add(Pair.Key);
+			}
+		}
+		return Names;
+	}
+	
 protected:
 	TAttributeChannel<ValueType>& GetValuesChannel()
 	{
@@ -626,12 +733,18 @@ protected:
 	 * @return Number of segments
 	 */
 	virtual int32 GetNumSegments() const = 0;
+
+	/**
+	 * Get splines parameter range
+	 * @return Parameter range
+	 */ 
+	virtual FInterval1f GetParameterRange() const = 0;
 	
 	/**
 	 * Find nearest point on the spline to a target point
 	 * @param Point - Target point
 	 * @outparam OutSquaredDistance - Squared distance to the nearest point
-	 * @return Parameter value of the nearest point
+	 * @return Parameter value of the nearest point in parent space
 	 */
 	virtual float FindNearest(const ValueType& Point, float& OutSquaredDistance) const = 0;
 
@@ -640,7 +753,7 @@ protected:
 	 * @param Point - Target point
 	 * @param SegmentIndex - Index of segment to search
 	 * @outparam OutSquaredDistance - Squared distance to the nearest point
-	 * @return Parameter value of the nearest point
+	 * @return Parameter value of the nearest point in spline space
 	 */
 	virtual float FindNearestOnSegment(const ValueType& Point, int32 SegmentIndex, float& OutSquaredDistance) const = 0;
 
@@ -670,11 +783,6 @@ private:
 		// Create new instance of correct derived type
 		TSharedPtr<TSplineBase> NewSpline = CreateNewInstance();
         NewSpline->Values = Values;
-
-		// Copy parameter range settings
-		float Start, End;
-		GetParameterRange(Start, End);
-		NewSpline->SetParameterRange(Start, End);
 	
 		// Copy out-of-bounds handling mode
 		NewSpline->SetPostInfinityMode(GetPostInfinityMode());
diff --git a/Engine/Source/Runtime/GeometryCore/Public/Splines/TangentBezierSpline.h b/Engine/Source/Runtime/GeometryCore/Public/Splines/TangentBezierSpline.h
--- a/Engine/Source/Runtime/GeometryCore/Public/Splines/TangentBezierSpline.h
+++ b/Engine/Source/Runtime/GeometryCore/Public/Splines/TangentBezierSpline.h
@@ -77,13 +77,43 @@ struct FTangentBezierControlPoint
 		const FVector& InPosition,
 		const FVector& InTangentIn,
 		const FVector& InTangentOut,
-		ETangentMode InTangentMode = ETangentMode::Auto)
+		ETangentMode InTangentMode)
 		: Position(InPosition)
 		, TangentIn(InTangentIn)
 		, TangentOut(InTangentOut)
 		, TangentMode(InTangentMode)
 	{
 	}
+	
+	void Serialize(FArchive& Ar)
+	{
+		Ar << Position;
+		Ar << TangentIn;
+		Ar << TangentOut;
+		// Serialize as uint8 for compatibility
+		uint8 TangentModeValue = static_cast<uint8>(TangentMode);
+		Ar << TangentModeValue;
+        
+		if (Ar.IsLoading())
+		{
+			TangentMode = static_cast<ETangentMode>(TangentModeValue);
+		}
+		
+	}
+    
+	friend FArchive& operator<<(FArchive& Ar, FTangentBezierControlPoint& Point)
+	{
+		Point.Serialize(Ar);
+		return Ar;
+	}
+
+	bool operator==(const FTangentBezierControlPoint& Other) const
+	{
+		return Position == Other.Position
+			&& TangentIn == Other.TangentIn
+			&& TangentOut == Other.TangentOut
+			&& TangentMode == Other.TangentMode;
+	}
 };
 
 /**
@@ -93,16 +123,10 @@ struct FTangentBezierControlPoint
  */
 struct FTangentBezierSpline
 {
-    /** Tension parameter for auto-computed tangents [0,1] */
-    float Tension;
 
-	/** How tangents should be computed for each point */
-	TArray<ETangentMode> TangentModes;
-	
 protected:
 	/** 
-	 * Protected default constructor for serialization
-	 * Initializes with a minimal default segment
+	 ** Initializes with a minimal default segment
 	 */
 	FTangentBezierSpline()
 		: Tension(0.0f)
@@ -181,7 +205,8 @@ public:
 			}
 		}
         Ar << InternalSpline;
-                
+		Ar << bStationaryEndpoint;
+		
         return true;
     }
 
@@ -299,6 +324,11 @@ public:
         
         return Result;
     }
+	
+	bool IsClosedLoop() const
+	{
+		return InternalSpline.IsClosedLoop();
+	}
     
     void SetClosedLoop(bool bInClosedLoop)
     {
@@ -343,6 +373,9 @@ public:
 	            
 	            // Add the closing segment to the internal spline
 	            InternalSpline.AppendBezierSegment(P1, P2, P3);
+
+	        	// Add corresponding tangent mode
+	        	TangentModes.Add(ETangentMode::Auto);
 	        }
 	    }
 	    else
@@ -354,6 +387,9 @@ public:
 	    	{
 	    		// Remove the last segment
 	    		InternalSpline.RemoveSegment(NumSegments - 1);
+
+	    		// Remove last tangent mode
+	    		TangentModes.RemoveAt(TangentModes.Num() - 1);
 	    	}
         
 	    	// Now update the flag on the internal spline
@@ -382,74 +418,144 @@ public:
 	        return false;
 	    }
 	    
-	   // Convert to cubic Bezier control points for the internal spline
+	    // Convert to cubic Bezier control points for the internal spline
 	    TArray<FVector> BezierPoints;
 	    
-	    // Reserve space for Bezier control points (3n-2 for n points)
-	    BezierPoints.Reserve(3 * NumPoints - 2);
+	    // For each segment (between consecutive points), we need 4 control points
+	    // Reserve space based on number of segments
+	    const int32 NumSegments = NumPoints - 1;
+	    BezierPoints.Reserve(NumSegments * 4);
 	    
-	    // Handle the first point
-	    const FVector& FirstPoint = Points[0].Position;
-	    FVector FirstOutTangent = Points[0].TangentOut;
-	    
-		// Apply auto tangent calculation if needed
-		if (Points[0].TangentMode == ETangentMode::Auto || Points[0].TangentMode == ETangentMode::AutoClamped)
-		{
-			FirstOutTangent = (1.0f - Tension) * (Points[1].Position - FirstPoint);
-		}
-	    
-	    // Add first point and its outgoing control point
-	    BezierPoints.Add(FirstPoint);
-	    BezierPoints.Add(FirstPoint + FirstOutTangent / 3.0f);
-	    
-	    // Handle middle points
-	    for (int32 i = 1; i < NumPoints - 1; ++i)
+	    for (int32 i = 0; i < NumSegments; ++i)
 	    {
-	        const FVector& CurrentPoint = Points[i].Position;
+	        const FTangentBezierControlPoint& StartPoint = Points[i];
+	        const FTangentBezierControlPoint& EndPoint = Points[i + 1];
+	        
+	        FVector P0 = StartPoint.Position;
+	        FVector P3 = EndPoint.Position;
 	        
-	        // Get tangents based on mode
-	        FVector InTangent = Points[i].TangentIn;
-	        FVector OutTangent = Points[i].TangentOut;
+	        // Calculate tangents based on the points' tangent modes
+	        FVector OutTangent, InTangent;
 	        
-	        if (Points[i].TangentMode == ETangentMode::Auto || Points[i].TangentMode == ETangentMode::AutoClamped)
+	        // Process start point's outgoing tangent
+	        if (StartPoint.TangentMode == ETangentMode::Auto || 
+	            StartPoint.TangentMode == ETangentMode::AutoClamped)
 	        {
-	            // Compute a smooth tangent using surrounding points
-	            FVector AutoTangent = (1.0f - Tension) * (Points[i+1].Position - Points[i-1].Position);
-	            InTangent = AutoTangent;
-	            OutTangent = AutoTangent;
+	            // Compute auto tangent
+	            if (i == 0 && NumPoints > 2)
+	            {
+	                // First point - use forward difference
+	                OutTangent = (1.0f - Tension) * (Points[i + 1].Position - Points[i].Position);
+	            }
+	            else if (i == NumPoints - 2 && NumPoints > 2)
+	            {
+	                // Point before last - use central difference
+	                OutTangent = (1.0f - Tension) * (Points[i + 1].Position - Points[i - 1].Position);
+	            }
+	            else if (NumPoints > 2)
+	            {
+	                // Normal case - use central difference
+	                OutTangent = (1.0f - Tension) * (Points[i + 1].Position - Points[i - 1].Position);
+	            }
+	            else
+	            {
+	                // Only two points - use forward difference
+	                OutTangent = (1.0f - Tension) * (Points[i + 1].Position - Points[i].Position);
+	            }
+	            
+	            // Apply clamping for AutoClamped
+	            if (StartPoint.TangentMode == ETangentMode::AutoClamped)
+	            {
+	                const float SegmentLength = static_cast<float>((P3 - P0).Size());
+	                const float MaxTangentSize = SegmentLength * 0.33f;
+	                
+	                if (OutTangent.Size() > MaxTangentSize && !FMath::IsNearlyZero(OutTangent.Size()))
+	                {
+	                    OutTangent = OutTangent.GetSafeNormal() * MaxTangentSize;
+	                }
+	            }
 	        }
-	        else if (Points[i].TangentMode == ETangentMode::Linear)
+	        else if (StartPoint.TangentMode == ETangentMode::Linear)
 	        {
-	            // Linear tangents aligned with adjacent points
-	            InTangent = CurrentPoint - Points[i-1].Position;
-	            OutTangent = Points[i+1].Position - CurrentPoint;
+	            // Linear tangent
+	            OutTangent = (P3 - P0) * 0.33f;
 	        }
-	        else if (Points[i].TangentMode == ETangentMode::Constant)
+	        else if (StartPoint.TangentMode == ETangentMode::Constant)
 	        {
-	            // Zero tangents for constant interpolation
-	            InTangent = FVector::ZeroVector;
+	            // Zero tangent for constant interpolation
 	            OutTangent = FVector::ZeroVector;
 	        }
+	        else
+	        {
+	            // User or Broken mode - use provided tangent
+	            OutTangent = StartPoint.TangentOut;
+	        }
 	        
-	        // Add the control points
-	        BezierPoints.Add(CurrentPoint - InTangent / 3.0f);
-	        BezierPoints.Add(CurrentPoint);
-	        BezierPoints.Add(CurrentPoint + OutTangent / 3.0f);
-	    }
-	    
-	    // Handle the last point
-	    const FVector& LastPoint = Points.Last().Position;
-	    FVector LastInTangent = Points.Last().TangentIn;
-	    
-	    if (Points.Last().TangentMode == ETangentMode::Auto || Points.Last().TangentMode == ETangentMode::AutoClamped)
-	    {
-	        LastInTangent = (1.0f - Tension) * (LastPoint - Points[NumPoints - 2].Position);
+	        // Process end point's incoming tangent
+	        if (EndPoint.TangentMode == ETangentMode::Auto || 
+	            EndPoint.TangentMode == ETangentMode::AutoClamped)
+	        {
+	            // Compute auto tangent
+	            if (i + 1 == 0 && NumPoints > 2)
+	            {
+	                // First point - use central difference
+	                InTangent = (1.0f - Tension) * (Points[i + 2].Position - Points[i].Position);
+	            }
+	            else if (i + 1 == NumPoints - 1 && NumPoints > 2)
+	            {
+	                // Last point - use backward difference
+	                InTangent = (1.0f - Tension) * (Points[i + 1].Position - Points[i].Position);
+	            }
+	            else if (NumPoints > 2)
+	            {
+	                // Normal case - use central difference
+	                InTangent = (1.0f - Tension) * (Points[i + 2].Position - Points[i].Position);
+	            }
+	            else
+	            {
+	                // Only two points - use backward difference
+	                InTangent = (1.0f - Tension) * (Points[i + 1].Position - Points[i].Position);
+	            }
+	            
+	            // Apply clamping for AutoClamped
+	            if (EndPoint.TangentMode == ETangentMode::AutoClamped)
+	            {
+	                const float SegmentLength = static_cast<float>((P3 - P0).Size());
+	                const float MaxTangentSize = SegmentLength * 0.33f;
+	                
+	                if (InTangent.Size() > MaxTangentSize && !FMath::IsNearlyZero(InTangent.Size()))
+	                {
+	                    InTangent = InTangent.GetSafeNormal() * MaxTangentSize;
+	                }
+	            }
+	        }
+	        else if (EndPoint.TangentMode == ETangentMode::Linear)
+	        {
+	            // Linear tangent
+	            InTangent = (P3 - P0) * 0.33f;
+	        }
+	        else if (EndPoint.TangentMode == ETangentMode::Constant)
+	        {
+	            // Zero tangent for constant interpolation
+	            InTangent = FVector::ZeroVector;
+	        }
+	        else
+	        {
+	            // User or Broken mode - use provided tangent
+	            InTangent = EndPoint.TangentIn;
+	        }
+	        
+	        // Calculate bezier control points
+	        FVector P1 = P0 + OutTangent / 3.0f;
+	        FVector P2 = P3 - InTangent / 3.0f;
+	        
+	        // Add all 4 control points for this segment
+	        BezierPoints.Add(P0);
+	        BezierPoints.Add(P1);
+	        BezierPoints.Add(P2);
+	        BezierPoints.Add(P3);
 	    }
 	    
-	    // Add last incoming control point and the last point
-	    BezierPoints.Add(LastPoint - LastInTangent / 3.0f);
-	    BezierPoints.Add(LastPoint);
-	    
 	    // Set the control points in the internal spline
 	    InternalSpline.SetControlPoints(BezierPoints, EParameterizationPolicy::Centripetal);
 	    
@@ -460,7 +566,7 @@ public:
 	        TangentModes.Add(Point.TangentMode);
 	    }
 	    
-	    // Update auto tangents to ensure proper shape
+	    // Update tangents to ensure proper curve continuity
 	    UpdateTangents();
 	    
 	    // Reparameterize the spline for better parameter distribution
@@ -628,7 +734,7 @@ public:
      * Inserts a new point at the specified parameter along the spline.
      * Uses distance-based parameterization to preserve attribute positions.
      * 
-     * @param Parameter - Parameter value [0,1] where to insert the point
+     * @param Parameter - Parameter value where to insert the point (Spline Space)
      * @param Position - Position for the new point
 	 * @param InTangentMode - Tangent mode for the new point
      * @return The index of the newly created point
@@ -742,6 +848,8 @@ public:
                 UpdatePointTangents(NewIndex - 1, TangentModes[NewIndex - 1]);
             }
         }
+
+		InternalSpline.Reparameterize();
     }
 	/**
 	 * Gets the control point at the specified index
@@ -858,8 +966,8 @@ public:
         
         // Update tangents based on tangent mode
         UpdatePointTangents(Index, ControlPoint.TangentMode);
-        
-        // Update surrounding points if position changed and they're auto tangents
+
+		// Update surrounding points if position changed and they're auto tangents
         if (bPositionChanged)
         {
             if (Index > 0 && (TangentModes[Index-1] == ETangentMode::Auto || 
@@ -876,7 +984,7 @@ public:
         }
     }
 
-    /** Gets position at parameter [0,1] */
+    /** Gets position at specified parent space parameter */
     FVector GetPosition(float Parameter) const
     {
         return FVector(InternalSpline.Evaluate(Parameter));
@@ -915,7 +1023,7 @@ public:
     	return FVector::ZeroVector;
     }
 
-    /** Gets tangent at parameter [0,1] */
+    /** Gets tangent at parameter in spline space */
     FVector GetTangent(float Parameter) const
     {
         return FVector(InternalSpline.EvaluateDerivative<1>(Parameter));
@@ -1061,7 +1169,7 @@ public:
 	void SetPointTangentMode(int32 Index, ETangentMode Mode)
 	{
 		const int32 NumPoints = GetNumPoints();
-		if (Index < 0 || Index >= NumPoints)
+		if (Index < 0 || Index > NumPoints)
 		{
 			return;
 		}
@@ -1078,10 +1186,30 @@ public:
 		UpdatePointTangents(Index, Mode);
 	}
 
+	FVector GetValueAt(int32 Index) const
+	{
+		return InternalSpline.GetValueAt(Index);
+	}
+
+	void SetValueAt(int32 Index, const FVector& Value)
+	{
+		InternalSpline.SetValueAt(Index, Value);
+	}
+	
+	float GetParameterAt(int32 Index) const
+	{
+		return InternalSpline.GetParameterAt(Index);
+	}
+
+	FInterval1f GetParameterRange() const
+	{
+		return InternalSpline.GetParameterRange();
+	}
+	
     template <typename AttrType>
-    void CreateAttributeChannel(const FName& Name)
+    bool CreateAttributeChannel(const FName& Name)
     {
-        InternalSpline.CreateAttributeChannel<AttrType>(Name);
+        return InternalSpline.CreateAttributeChannel<AttrType>(Name);
     }
     
     template <typename AttrType>
@@ -1096,6 +1224,104 @@ public:
         return InternalSpline.AddAttributeValue<AttrType>(Name, Value, Parameter);
     }
 
+	/**
+	 * Checks if the spline has an attribute channel with the given name
+	 * @param Name - Name of the attribute channel 
+	 * @return true if the channel exists, false otherwise
+	 */
+	bool HasAttributeChannel(const FName& Name) const
+	{
+		return InternalSpline.HasAttributeChannel(Name);
+	}
+
+	/**
+	 * Retrieves the number of values in a specific attribute channel
+	 * @param Name - Name of the attribute channel
+	 * @return Number of values or zero if the channel does not exist
+	 */
+	int32 NumAttributeValues(const FName& Name) const
+	{
+		return InternalSpline.NumAttributeValues(Name);
+	}
+
+	/**
+	 * Retrieves the parameter value of a specific attribute at the given index
+	 * @param Name - Name of the attribute channel
+	 * @param Index - Index of the value
+	 * @return Parameter value
+	 */
+	float GetAttributeParameter(const FName& Name, int32 Index) const
+	{
+		return InternalSpline.GetAttributeParameter(Name, Index);
+	}
+
+	/**
+	 * Retrieves all attribute channels in the spline
+	 * @return Map of attribute channels by name
+	 */
+	const TMap<FName, TUniquePtr<IAttributeChannelBase>>& GetAttributeChannels() const
+	{
+		return InternalSpline.GetAttributeChannels();
+	}
+
+	/**
+	 * Updates the parameter of an attribute by index, possibly changing its index to preserve the ordering of the attribute storage
+	 * @param Name - Name of the attribute channel
+	 * @param Index - Index of the attribute to modify
+	 * @param ParentSpaceParameter - Parameter to set in parent spline space
+	 * @return New index of the attribute originally at the given index after any reordering
+	 */
+	int32 SetAttributeParameter(const FName& Name, int32 Index, float ParentSpaceParameter)
+	{
+		return InternalSpline.SetAttributeParameter(Name, Index, ParentSpaceParameter);
+	}
+	
+	/**
+	 * Retrieves the value of a specific attribute at the given index
+	 * @tparam AttrType - Custom attribute type
+	 * @param Name - Name of the attribute channel
+	 * @param Index - Index of the value
+	 * @return Attribute value
+	 */
+	template <typename AttrType>
+	AttrType GetAttributeValue(const FName& Name, int32 Index) const
+	{
+		return InternalSpline.GetAttributeValue<AttrType>(Name, Index);
+	}
+
+	/**
+	 * Removes a value from an attribute channel
+	 * @param Name - Name of the attribute channel
+	 * @param Index - Index of the value to remove
+	 * @return true if the value was removed, false otherwise
+	 */
+	bool RemoveAttributeValue(FName Name, int32 Index)
+	{
+		return InternalSpline.RemoveAttributeValue(Name, Index);
+	}
+
+	/**
+	 * Sets the value of a specific attribute at the given index
+	 * @tparam AttrType - Custom attribute type
+	 * @param Name - Name of the attribute channel
+	 * @param Index - Index of the value
+	 * @param Value - New value to set
+	 */
+	template <typename AttrType>
+	void SetAttributeValue(const FName& Name, int32 Index, const AttrType& Value)
+	{
+		InternalSpline.SetAttributeValue<AttrType>(Name, Index, Value);
+	}
+
+	/**
+	 * Clears all values from a specific attribute channel
+	 * @param Name - Name of the attribute channel to clear
+	 * @return true if the channel was found and cleared, false otherwise
+	 */
+	bool ClearAttributeChannel(const FName& Name)
+	{
+		return InternalSpline.ClearAttributeChannel(Name);
+	}
     /**
 	 * Updates tangents for all control points based on their tangent modes
 	 */
@@ -1106,18 +1332,10 @@ public:
 	    {
 	        return;
 	    }
-
 		// Update tangents for each point according to its mode
 		for (int32 i = 0; i < NumPoints; ++i)
 		{
-			if (TangentModes.IsValidIndex(i) && 
-				(TangentModes[i] == ETangentMode::Auto || 
-				 TangentModes[i] == ETangentMode::AutoClamped ||
-				 TangentModes[i] == ETangentMode::Linear ||
-				 TangentModes[i] == ETangentMode::Constant))
-			{
-				UpdatePointTangents(i, TangentModes[i]);
-			}
+			UpdatePointTangents(i, TangentModes[i]);
 		}
 	}
 
@@ -1128,110 +1346,138 @@ public:
 	 */
 	void UpdatePointTangents(int32 Index, ETangentMode Mode)
 	{
-	    const int32 NumPoints = GetNumPoints();
-	    if (Index < 0 || Index >= NumPoints)
-	    {
-	        return;
-	    }
+		const int32 NumPoints = GetNumPoints();
+		if (Index < 0 || Index >= NumPoints)
+		{
+			return;
+		}
 
-	    // Get the position of the current point
-	    const FVector& Current = GetPosition(Index);
-	    
-	    // Get previous and next points (handling closed loop case)
-	    const FVector& Prev = (Index > 0) ? GetPosition(Index-1) : 
-	        (InternalSpline.IsClosedLoop() ? GetPosition(NumPoints-1) : Current);
-	    
-	    const FVector& Next = (Index < NumPoints-1) ? GetPosition(Index+1) :
-	        (InternalSpline.IsClosedLoop() ? GetPosition(0) : Current);
-	    
-		// For User and Broken modes, we should preserve existing tangents
-		// So we fetch the current tangents first and only overwrite if needed
-		FVector InTangent = GetTangentIn(Index);
+		// Get the position of the current point
+		const FVector& Current = GetPosition(Index);
+
+		// Get previous and next points (handling closed loop case)
+		const FVector& Prev = (Index > 0) ? GetPosition(Index - 1) : (IsClosedLoop() ? GetPosition(NumPoints - 1) : Current);
+
+		const FVector& Next = (Index < NumPoints - 1) ? GetPosition(Index + 1) : (IsClosedLoop() ? GetPosition(0) : Current);
+
+		// Get actual parameter values for points
+		float PrevParam = (Index > 0)
+			                  ? static_cast<float>(Index - 1)
+			                  : (IsClosedLoop() ? static_cast<float>(NumPoints - 1) : static_cast<float>(Index));
+
+		float CurrentParam = static_cast<float>(Index);
+
+		float NextParam = (Index < NumPoints - 1) ? static_cast<float>(Index + 1) : (IsClosedLoop() ? 0.0f : static_cast<float>(Index));
+
+		// Adjust params for closed loop to ensure proper weighting
+		if (IsClosedLoop() && Index == 0)
+		{
+			PrevParam -= static_cast<float>(NumPoints); // Make previous parameter negative for proper delta
+		}
+		else if (IsClosedLoop() && Index == NumPoints - 1)
+		{
+			NextParam += static_cast<float>(NumPoints); // Make next parameter greater than NumPoints for proper delta
+		}
+
+		// Calculate tangents based on mode
+		FVector InTangent = GetTangentIn(Index); // Keep existing if not changing
 		FVector OutTangent = GetTangentOut(Index);
+		
+		switch (Mode)
+		{
+		case ETangentMode::Auto:
+		case ETangentMode::AutoClamped:
+			{
+				// Legacy-style auto tangent calculation with parameter weighting
+				float PrevToNextParamDiff = FMath::Max<float>(UE_KINDA_SMALL_NUMBER, NextParam - PrevParam);
 
-	    switch (Mode)
-	    {
-	    case ETangentMode::Auto:
-	    case ETangentMode::AutoClamped:
-	        {
-	            // Compute smooth tangent based on surrounding points
-	            FVector Tangent = (1.0f - Tension) * (Next - Prev);
-	            
-	            // Apply clamping for AutoClamped mode
-	            if (Mode == ETangentMode::AutoClamped)
-	            {
-	                // Clamp tangent to prevent overshoot
-	                const double PrevDistance = (Current - Prev).Size();
-	                const double NextDistance = (Next - Current).Size();
-	                
-	                if (PrevDistance > 0.0 && NextDistance > 0.0)
-	                {
-	                    const double MaxTangentSize = FMath::Min(PrevDistance, NextDistance) * 1.0;
-	                    const double TangentSize = Tangent.Size();
-	                    
-	                    if (TangentSize > MaxTangentSize && TangentSize > UE_KINDA_SMALL_NUMBER)
-	                    {
-	                        Tangent = Tangent * (MaxTangentSize / TangentSize);
-	                    }
-	                }
-	            }
-	            
-	            InTangent = Tangent;
-	            OutTangent = Tangent;
-	        }
-	        break;
-	        
-	    case ETangentMode::Linear:
-	        {
-	            // Compute tangents aligned with adjacent points
-	            if (Index > 0 || InternalSpline.IsClosedLoop())
-	            {
-	                InTangent = Current - Prev;
-	            }
-	            else 
-	            {
-	                InTangent = FVector::ZeroVector;
-	            }
-	            
-	            if (Index < NumPoints-1 || InternalSpline.IsClosedLoop())
-	            {
-	                OutTangent = Next - Current;
-	            }
-	            else
-	            {
-	                OutTangent = FVector::ZeroVector;
-	            }
-	        }
-	        break;
-	        
-	    case ETangentMode::Constant:
-	        {
-	            // Zero tangents for constant interpolation
-	            InTangent = FVector::ZeroVector;
-	            OutTangent = FVector::ZeroVector;
-	        }
-	        break;
-	        
-	    case ETangentMode::Broken:
-	    case ETangentMode::User:
-	        // For User and Broken modes, we're keeping the existing tangents
-	        // which we already fetched above, so no action needed here
-	        break;
-	        
-	    case ETangentMode::Unknown:
-	    default:
-	        // Default to Auto mode if unknown
-	        {
-	            FVector Tangent = (1.0f - Tension) * (Next - Prev);
-	            InTangent = Tangent;
-	            OutTangent = Tangent;
-	        }
-	        break;
-	    }
-	    
-	    // Apply the calculated tangents
-	    SetTangentIn(Index, InTangent);
-	    SetTangentOut(Index, OutTangent);
+				// Calculate basic tangent - matches legacy formula
+				FVector AutoTangent = (1.0f - Tension) * ((Next - Prev) / PrevToNextParamDiff);
+
+				// For AutoClamped mode, apply clamping like the legacy implementation
+				if (Mode == ETangentMode::AutoClamped)
+				{
+					float PrevToCurParamDiff = FMath::Max<float>(UE_KINDA_SMALL_NUMBER, CurrentParam - PrevParam);
+					float NextToCurParamDiff = FMath::Max<float>(UE_KINDA_SMALL_NUMBER, NextParam - CurrentParam);
+
+					FVector PrevToCurTangent = (Current - Prev) / PrevToCurParamDiff;
+					FVector NextToCurTangent = (Next - Current) / NextToCurParamDiff;
+
+					float PrevToCurLength = static_cast<float>(PrevToCurTangent.Size());
+					float NextToCurLength = static_cast<float>(NextToCurTangent.Size());
+
+					// Clamp tangent magnitude
+					float MaxScale = FMath::Min(PrevToCurLength, NextToCurLength);
+
+					if (AutoTangent.Size() > MaxScale && !FMath::IsNearlyZero(AutoTangent.Size()))
+					{
+						AutoTangent = AutoTangent.GetSafeNormal() * MaxScale;
+					}
+				}
+
+				// Handle stationary endpoints
+				if (bStationaryEndpoint && (Index == 0 || Index == NumPoints - 1) && !IsClosedLoop())
+				{
+					AutoTangent = FVector::ZeroVector;
+				}
+
+				InTangent = AutoTangent;
+				OutTangent = AutoTangent;
+			}
+			break;
+
+		case ETangentMode::Linear:
+			{
+				// Previous segment tangent (incoming)
+				if (Index > 0 || IsClosedLoop())
+				{
+					InTangent = Current - Prev;
+				}
+				else
+				{
+					InTangent = FVector::ZeroVector;
+				}
+
+				// Next segment tangent (outgoing)
+				if (Index < NumPoints - 1 || IsClosedLoop())
+				{
+					OutTangent = Next - Current;
+				}
+				else
+				{
+					OutTangent = FVector::ZeroVector;
+				}
+			}
+			break;
+
+		case ETangentMode::Constant:
+			{
+				// Zero tangents for constant interpolation
+				InTangent = FVector::ZeroVector;
+				OutTangent = FVector::ZeroVector;
+			}
+			break;
+
+		case ETangentMode::Broken:
+		case ETangentMode::User:
+			// For User and Broken modes, we're keeping the existing tangents
+			// which we already fetched above, so no action needed here
+			break;
+
+		case ETangentMode::Unknown:
+		default:
+			// Default to Auto mode if unknown
+			{
+				FVector Tangent = (1.0f - Tension) * (Next - Prev);
+				InTangent = Tangent;
+				OutTangent = Tangent;
+			}
+			break;
+		}
+
+		// Apply the calculated tangents
+		SetTangentIn(Index, InTangent);
+		SetTangentOut(Index, OutTangent);
 	}
     
     float FindNearest(const FVector& Point, float& OutSquaredDistance) const
@@ -1256,11 +1502,52 @@ public:
     {
         return InternalSpline;
     }
+	
+	float GetTension() const
+	{
+		return Tension;
+	}
 
-private:
-    /** Internal PolyBezier implementation */
-    FPolyBezierSpline3d InternalSpline;
+	void SetTension(const float InTension)
+	{
+		this->Tension = InTension;
+	}
+
+	ETangentMode GetTangentMode(int32 Index) const
+	{
+		if (Index < 0 || Index >= GetNumPoints())
+		{
+			return ETangentMode::Unknown;
+		}
+		return TangentModes[Index];
+	}
+	TArray<ETangentMode> GetTangentModes() const
+	{
+		return TangentModes;
+	}
+
+	void SetStationaryEndpoints(bool bInStationaryEndpoints)
+	{
+		bStationaryEndpoint = bInStationaryEndpoints;
+		UpdateTangents(); // Update tangents when this flag changes
+	}
+
+	bool IsStationaryEndpoints() const
+	{
+		return bStationaryEndpoint;
+	}
+protected:
+    /** Tension parameter for auto-computed tangents [0,1] */
+    float Tension = 0.0f;
+    
+    /** How tangents should be computed for each point */
+    TArray<ETangentMode> TangentModes;
 	
+	/** Internal PolyBezier implementation */
+    FPolyBezierSpline3d InternalSpline;
+
+	bool bStationaryEndpoint = false;
+
     /** Updates a Bezier control point in the internal spline */
     void UpdateBezierControlPoint(int32 Index, const FVector& NewValue)
     {
