Fix redeclaration of 'GExperimentalAllowPerInstanceChildActorProperties'

#rnx

diff --git a/Engine/Source/Editor/LevelEditor/Private/Elements/Actor/ActorElementLevelEditorSelectionCustomization.cpp b/Engine/Source/Editor/LevelEditor/Private/Elements/Actor/ActorElementLevelEditorSelectionCustomization.cpp
--- a/Engine/Source/Editor/LevelEditor/Private/Elements/Actor/ActorElementLevelEditorSelectionCustomization.cpp
+++ b/Engine/Source/Editor/LevelEditor/Private/Elements/Actor/ActorElementLevelEditorSelectionCustomization.cpp
@@ -50,8 +50,7 @@ FTypedElementHandle FActorElementLevelEditorSelectionCustomization::GetSelection
 {
 	if (AActor* ConsideredActor = ActorElementDataUtil::GetActorFromHandle(InElementSelectionHandle))
 	{
-		extern ENGINE_API int32 GExperimentalAllowPerInstanceChildActorProperties;
-		if (!GExperimentalAllowPerInstanceChildActorProperties)
+		if (!UChildActorComponent::ExperimentalAllowPerInstanceChildActorProperties())
 		{
 			while (ConsideredActor->IsChildActor())
 			{
diff --git a/Engine/Source/Runtime/Engine/Classes/Components/ChildActorComponent.h b/Engine/Source/Runtime/Engine/Classes/Components/ChildActorComponent.h
--- a/Engine/Source/Runtime/Engine/Classes/Components/ChildActorComponent.h
+++ b/Engine/Source/Runtime/Engine/Classes/Components/ChildActorComponent.h
@@ -160,6 +160,8 @@ private:
 
 public:
 
+	ENGINE_API static bool ExperimentalAllowPerInstanceChildActorProperties();
+
 	//~ Begin Object Interface.
 #if WITH_EDITOR
 	ENGINE_API virtual void PostEditChangeProperty(FPropertyChangedEvent& PropertyChangedEvent) override;
diff --git a/Engine/Source/Runtime/Engine/Private/Actor.cpp b/Engine/Source/Runtime/Engine/Private/Actor.cpp
--- a/Engine/Source/Runtime/Engine/Private/Actor.cpp
+++ b/Engine/Source/Runtime/Engine/Private/Actor.cpp
@@ -2121,8 +2121,7 @@ bool AActor::Modify( bool bAlwaysMarkDirty/*=true*/ )
 		return false;
 	}
 
-	extern int32 GExperimentalAllowPerInstanceChildActorProperties;
-	if (GExperimentalAllowPerInstanceChildActorProperties)
+	if (UChildActorComponent::ExperimentalAllowPerInstanceChildActorProperties())
 	{
 		if (AActor* ParentActor = GetParentActor())
 		{
@@ -5670,8 +5669,7 @@ bool AActor::IsSelectionChild() const
 
 AActor* AActor::GetSelectionParent() const
 {
-	extern int32 GExperimentalAllowPerInstanceChildActorProperties;
-	if (!GExperimentalAllowPerInstanceChildActorProperties && IsChildActor())
+	if (!UChildActorComponent::ExperimentalAllowPerInstanceChildActorProperties() && IsChildActor())
 	{
 		return GetParentActor();
 	}
diff --git a/Engine/Source/Runtime/Engine/Private/Components/ActorComponent.cpp b/Engine/Source/Runtime/Engine/Private/Components/ActorComponent.cpp
--- a/Engine/Source/Runtime/Engine/Private/Components/ActorComponent.cpp
+++ b/Engine/Source/Runtime/Engine/Private/Components/ActorComponent.cpp
@@ -1097,8 +1097,7 @@ bool UActorComponent::Modify( bool bAlwaysMarkDirty/*=true*/ )
 	// the Actor as modified so that we store of the transaction annotation that has the component properties stashed
 	if (MyOwner)
 	{
-		extern int32 GExperimentalAllowPerInstanceChildActorProperties;
-		if (IsCreatedByConstructionScript() || (GExperimentalAllowPerInstanceChildActorProperties && MyOwner->IsChildActor()))
+		if (IsCreatedByConstructionScript() || (UChildActorComponent::ExperimentalAllowPerInstanceChildActorProperties() && MyOwner->IsChildActor()))
 		{
 			return MyOwner->Modify(bAlwaysMarkDirty);
 		}
diff --git a/Engine/Source/Runtime/Engine/Private/Components/ChildActorComponent.cpp b/Engine/Source/Runtime/Engine/Private/Components/ChildActorComponent.cpp
--- a/Engine/Source/Runtime/Engine/Private/Components/ChildActorComponent.cpp
+++ b/Engine/Source/Runtime/Engine/Private/Components/ChildActorComponent.cpp
@@ -21,7 +21,7 @@
 
 DEFINE_LOG_CATEGORY_STATIC(LogChildActorComponent, Log, All);
 
-ENGINE_API int32 GExperimentalAllowPerInstanceChildActorProperties = 0;
+static int32 GExperimentalAllowPerInstanceChildActorProperties = 0;
 
 static FAutoConsoleVariableRef CVarExperimentalAllowPerInstanceChildActorProperties(
 	TEXT("cac.ExperimentalAllowPerInstanceChildActorProperties"),
@@ -30,6 +30,11 @@ static FAutoConsoleVariableRef CVarExperimentalAllowPerInstanceChildActorPropert
 	ECVF_Default
 );
 
+bool UChildActorComponent::ExperimentalAllowPerInstanceChildActorProperties()
+{
+	return GExperimentalAllowPerInstanceChildActorProperties != 0;
+}
+
 UChildActorComponent::UChildActorComponent(const FObjectInitializer& ObjectInitializer)
 	: Super(ObjectInitializer)
 	, ActorOuter(nullptr)
@@ -369,7 +374,7 @@ FChildActorComponentInstanceData::FChildActorComponentInstanceData(const UChildA
 			ChildActorName = ChildActor->GetFName();
 		}
 
-		if (GExperimentalAllowPerInstanceChildActorProperties)
+		if (UChildActorComponent::ExperimentalAllowPerInstanceChildActorProperties())
 		{
 			ActorInstanceData = MakeShared<FActorInstanceData>(ChildActor);
 			// If it is empty dump it
