Remove unsupported Replication of ALevelInstance::WorldAsset as it would cause a crash if ever creating a LevelInstance actor from PCG grammars.

diff --git a/Engine/Source/Runtime/Engine/Private/LevelInstance/LevelInstanceActor.cpp b/Engine/Source/Runtime/Engine/Private/LevelInstance/LevelInstanceActor.cpp
--- a/Engine/Source/Runtime/Engine/Private/LevelInstance/LevelInstanceActor.cpp
+++ b/Engine/Source/Runtime/Engine/Private/LevelInstance/LevelInstanceActor.cpp
@@ -66,12 +66,6 @@ void ALevelInstance::GetLifetimeReplicatedProps(TArray<FLifetimeProperty>& OutLi
 	FDoRepLifetimeParams Params;
 	Params.Condition = COND_InitialOnly;
 	DOREPLIFETIME_WITH_PARAMS_FAST(ALevelInstance, LevelInstanceSpawnGuid, Params);
-
-	// Generally the WorldAsset ref is set at edit time, but in cases where the LI is spawned during gameplay (ex: from PCG), the WorldAsset is set on spawn and needs to be replicated
-	// We can assume that these will never be changed during the lifetime of the LI and are assigned on spawn, so we only need the one initial rep
-#if WITH_EDITORONLY_DATA
-	DOREPLIFETIME_WITH_PARAMS_FAST(ALevelInstance, WorldAsset, Params);
-#endif
 	DOREPLIFETIME_WITH_PARAMS_FAST(ALevelInstance, CookedWorldAsset, Params);
 }
 
@@ -121,10 +115,13 @@ bool ALevelInstance::IsLoadingEnabled() const
 const TSoftObjectPtr<UWorld>& ALevelInstance::GetWorldAsset() const
 {
 #if WITH_EDITORONLY_DATA
-	return GetPropertyOverrideAsset() ? GetPropertyOverrideAsset()->GetWorldAsset() : WorldAsset;
-#else
-	return CookedWorldAsset;
+	TSoftObjectPtr<UWorld> Result = GetPropertyOverrideAsset() ? GetPropertyOverrideAsset()->GetWorldAsset() : WorldAsset;
+	if (!Result.IsNull())
+	{
+		return Result;
+	}
 #endif
+	return CookedWorldAsset;
 }
 
 const FLevelInstanceID& ALevelInstance::GetLevelInstanceID() const
diff --git a/Engine/Source/Runtime/Engine/Public/LevelInstance/LevelInstanceActor.h b/Engine/Source/Runtime/Engine/Public/LevelInstance/LevelInstanceActor.h
--- a/Engine/Source/Runtime/Engine/Public/LevelInstance/LevelInstanceActor.h
+++ b/Engine/Source/Runtime/Engine/Public/LevelInstance/LevelInstanceActor.h
@@ -25,7 +25,7 @@ public:
 protected:
 #if WITH_EDITORONLY_DATA
 	/** Level LevelInstance */
-	UPROPERTY(EditAnywhere, Replicated, Category = Level, Meta = (NoCreate, DisplayName="Level", DisableLevelInstancePropertyOverride))
+	UPROPERTY(EditAnywhere, Category = Level, Meta = (NoCreate, DisplayName="Level", DisableLevelInstancePropertyOverride))
 	TSoftObjectPtr<UWorld> WorldAsset;
 
 	UPROPERTY(VisibleAnywhere, Category = Override, meta = (EditInline, NoResetToDefault, EditCondition="PropertyOverrides != nullptr", EditConditionHides))
@@ -59,7 +59,7 @@ public:
 	ENGINE_API virtual void OnLevelInstanceLoaded() override;
 	ENGINE_API virtual bool IsLoadingEnabled() const override;
 	// End ILevelInstanceInterface
-			
+
 	ENGINE_API virtual void PostRegisterAllComponents() override;
 	ENGINE_API virtual void PostUnregisterAllComponents() override;
 	ENGINE_API virtual void Serialize(FArchive& Ar) override;
