Fixing crash in streamable manager when descending into child handles and the parrent has compelted.
#rnx

diff --git a/Engine/Source/Runtime/Engine/Private/StreamableManager.cpp b/Engine/Source/Runtime/Engine/Private/StreamableManager.cpp
--- a/Engine/Source/Runtime/Engine/Private/StreamableManager.cpp
+++ b/Engine/Source/Runtime/Engine/Private/StreamableManager.cpp
@@ -834,11 +834,11 @@ float FStreamableHandle::GetProgress(EStreamableProgressType Type) const
 			++DownloadCount;
 			if (Type == EStreamableProgressType::DownloadRelative)
 			{
-				TotalProgress += DownloadContext->InstallProgress.GetRelativeProgress();
+				TotalProgress += Current->DownloadContext->InstallProgress.GetRelativeProgress();
 			}
 			else
 			{
-				TotalProgress += DownloadContext->InstallProgress.GetAbsoluteProgress();
+				TotalProgress += Current->DownloadContext->InstallProgress.GetAbsoluteProgress();
 			}
 		}
 
