Address compile error that wasn't caught by any of our toolchains

diff --git a/Engine/Source/Runtime/Engine/Private/LevelInstance/LevelInstanceActor.cpp b/Engine/Source/Runtime/Engine/Private/LevelInstance/LevelInstanceActor.cpp
--- a/Engine/Source/Runtime/Engine/Private/LevelInstance/LevelInstanceActor.cpp
+++ b/Engine/Source/Runtime/Engine/Private/LevelInstance/LevelInstanceActor.cpp
@@ -115,7 +115,7 @@ bool ALevelInstance::IsLoadingEnabled() const
 const TSoftObjectPtr<UWorld>& ALevelInstance::GetWorldAsset() const
 {
 #if WITH_EDITORONLY_DATA
-	TSoftObjectPtr<UWorld> Result = GetPropertyOverrideAsset() ? GetPropertyOverrideAsset()->GetWorldAsset() : WorldAsset;
+	const TSoftObjectPtr<UWorld> &Result = GetPropertyOverrideAsset() ? GetPropertyOverrideAsset()->GetWorldAsset() : WorldAsset;
 	if (!Result.IsNull())
 	{
 		return Result;
