[SkeletalMesh] removed interface check when checking for halfedgebufferrequirement asset tag. The GetClass() call is not allowed when this function is called during cook sometimes.

diff --git a/Engine/Source/Runtime/Engine/Private/Rendering/SkeletalMeshHalfEdgeBufferAccessor.cpp b/Engine/Source/Runtime/Engine/Private/Rendering/SkeletalMeshHalfEdgeBufferAccessor.cpp
--- a/Engine/Source/Runtime/Engine/Private/Rendering/SkeletalMeshHalfEdgeBufferAccessor.cpp
+++ b/Engine/Source/Runtime/Engine/Private/Rendering/SkeletalMeshHalfEdgeBufferAccessor.cpp
@@ -23,7 +23,8 @@ bool SkeletalMeshHalfEdgeBufferAccessor::IsHalfEdgeRequired(TSoftObjectPtr<UObje
 	FAssetData AssetData;
 	AssetRegistryModule.TryGetAssetByObjectPath(InAssetSoftPtr.ToSoftObjectPath(), AssetData);
 
-	if (AssetData.GetClass()->ImplementsInterface(USkeletalMeshHalfEdgeBufferAccessor::StaticClass()))
+	
+	// if (AssetData.GetClass()->ImplementsInterface(USkeletalMeshHalfEdgeBufferAccessor::StaticClass()))
 	{
 		bool bIsHalfEdgeRequired = false;
 		if (AssetData.GetTagValue<bool>(GetHalfEdgeRequirementAssetTagName(), bIsHalfEdgeRequired))
