New material translator changes:
- Reworked and extended support for texture sampling. The API is now split between EmitTextureGather and four EmitTextureXXX() for the various various types of texture sampling modes (normal, miplevel, mipbias, and grad). The emitter automatically emits the appropriate instructions for both the pixel stage (which supports hardware partial derivatives) and the compute stage (which doesn't). Texture2D sampling is now feature parity with old translator.
- Added support for specifying that an operator argument is expected to be a float, for operators that exclusively run on floats (e.g. exp(), fmod(), etc). Now we can emit a "exp" operator on an int argument, as the int will be automatically cast to a float.
- Temporarily removed the BO_Xor operator. Will reintroduce it together with all remaining operators in coming changes.
- Some other minor improvements to the overall MIR codebase.

#rnx

diff --git a/Engine/Source/Runtime/Engine/Private/Materials/MaterialExpressionsIR.cpp b/Engine/Source/Runtime/Engine/Private/Materials/MaterialExpressionsIR.cpp
--- a/Engine/Source/Runtime/Engine/Private/Materials/MaterialExpressionsIR.cpp
+++ b/Engine/Source/Runtime/Engine/Private/Materials/MaterialExpressionsIR.cpp
@@ -914,6 +914,18 @@ void UMaterialExpressionTextureObject::Build(MIR::FEmitter& Emitter)
 	Emitter.Output(0, Emitter.EmitTextureObject(Texture, SamplerType));
 }
 
+static MIR::ETextureReadMode TextureGatherModeToMIR(::ETextureGatherMode Mode)
+{
+	switch (Mode)
+	{
+		case TGM_Red: return MIR::ETextureReadMode::GatherRed;
+		case TGM_Green: return MIR::ETextureReadMode::GatherGreen;
+		case TGM_Blue: return MIR::ETextureReadMode::GatherBlue;
+		case TGM_Alpha: return MIR::ETextureReadMode::GatherAlpha;
+		default: UE_MIR_UNREACHABLE();
+	}
+}
+
 static void BuildTextureSampleExpression(MIR::FValue* Texture, UMaterialExpressionTextureSample* Expr, MIR::FEmitter& Emitter)
 {
 	MIR::FValue* TexCoords = Emitter.TryInput(&Expr->Coordinates);
@@ -922,38 +934,61 @@ static void BuildTextureSampleExpression(MIR::FValue* Texture, UMaterialExpressi
 		TexCoords = Emitter.EmitExternalInput(MIR::TexCoordIndexToExternalInput(Expr->ConstCoordinate));
 	}
 
-	MIR::FValue* TexCoordsDdx{};
-	MIR::FValue* TexCoordsDdy{};
-	MIR::FValue* MipValue{};
-	if (Expr->MipValueMode == TMVM_MipLevel || Expr->MipValueMode == TMVM_MipBias)
+	MIR::FValue* TextureRead{};
+	if (Expr->GatherMode != TGM_None)
 	{
-		// Get the mip value level (either through the expression input or using the given constant if disconnected).
-		MipValue = Emitter.DefaultInputToIntScalar(&Expr->MipValue, Expr->ConstMipValue).InputScalar(&Expr->MipValue);
+		if (Expr->MipValueMode != TMVM_None)
+		{
+			Emitter.Errorf(TEXT("Texture gather does not support mipmap overrides (it implicitly accesses a specific mip)."));
+			return;
+		}
 
-		// Determine if automatic view mip bias should be used
-		bool bAutomaticViewMipBias = Emitter.DefaultInputToBoolScalar(&Expr->AutomaticViewMipBiasValue, Expr->AutomaticViewMipBias).InputConstantBoolScalar(&Expr->AutomaticViewMipBiasValue);
+		TextureRead = Emitter.EmitTextureGather(Texture, TexCoords, TextureGatherModeToMIR(Expr->GatherMode), Expr->SamplerSource, Expr->SamplerType);
+	} 
+	else
+	{
+		// Determine if automatic view mip bias should be used, by trying ot acquire a its input as a static boolean.
+		bool bAutomaticViewMipBias = Emitter
+			.DefaultInputToBoolScalar(&Expr->AutomaticViewMipBiasValue, Expr->AutomaticViewMipBias)
+			.InputConstantBoolScalar(&Expr->AutomaticViewMipBiasValue);
+	
 		UE_MIR_CHECKPOINT();
 
-		// If so, add "View.MaterialTextureViewBias" to the specified mip value
-		if (bAutomaticViewMipBias)
+		// Get the mip value level (either through the expression input or using the given constant if disconnected).
+		MIR::FValue* MipValue = (Expr->MipValueMode == TMVM_MipLevel || Expr->MipValueMode == TMVM_MipBias)
+			? Emitter.DefaultInputToFloatScalar(&Expr->MipValue, Expr->ConstMipValue).InputScalar(&Expr->MipValue)
+			: MIR::FPoison::Get();
+
+		switch (Expr->MipValueMode)
 		{
-			MipValue = Emitter.EmitOperator(MIR::BO_Add, MipValue, Emitter.EmitExternalInput(MIR::EExternalInput::ViewMaterialTextureMipBias));
+			case TMVM_None:
+				TextureRead = Emitter.EmitTextureSample(Texture, TexCoords, bAutomaticViewMipBias, Expr->SamplerSource, Expr->SamplerType);
+				break;
+		
+			case TMVM_MipBias:
+				TextureRead = Emitter.EmitTextureSampleBias(Texture, TexCoords, MipValue, bAutomaticViewMipBias, Expr->SamplerSource, Expr->SamplerType);
+				break;
+		
+			case TMVM_MipLevel:
+				TextureRead = Emitter.EmitTextureSampleLevel(Texture, TexCoords, MipValue, bAutomaticViewMipBias, Expr->SamplerSource, Expr->SamplerType);
+				break;
+
+			case TMVM_Derivative:
+			{
+				MIR::FValue* TexCoordsDdx = Emitter.EmitConstruct(TexCoords->Type, Emitter.Input(&Expr->CoordinatesDX));
+				MIR::FValue* TexCoordsDdy = Emitter.EmitConstruct(TexCoords->Type, Emitter.Input(&Expr->CoordinatesDY));
+				TextureRead = Emitter.EmitTextureSampleGrad(Texture, TexCoords, TexCoordsDdx, TexCoordsDdy, bAutomaticViewMipBias, Expr->SamplerSource, Expr->SamplerType);
+				break;
+			}
 		}
 	}
-	else if (Expr->MipValueMode == TMVM_Derivative)
-	{
-		check(Expr->MipValueMode = TMVM_Derivative);
-		TexCoordsDdx = Emitter.EmitConstruct(TexCoords->Type, Emitter.Input(&Expr->CoordinatesDX));
-		TexCoordsDdy = Emitter.EmitConstruct(TexCoords->Type, Emitter.Input(&Expr->CoordinatesDY));
-	}
 
-	MIR::FValue* OutputValue = Emitter.EmitTextureSample(Texture, TexCoords, TexCoordsDdx, TexCoordsDdy, MipValue, Expr->SamplerSource, Expr->MipValueMode, Expr->SamplerType);
-	Emitter.Output(0, Emitter.EmitSwizzle(OutputValue, MIR::FSwizzleMask::XYZ()));
-	Emitter.Output(1, Emitter.EmitSubscript(OutputValue, 0)); 
-	Emitter.Output(2, Emitter.EmitSubscript(OutputValue, 1)); 
-	Emitter.Output(3, Emitter.EmitSubscript(OutputValue, 2));
-	Emitter.Output(4, Emitter.EmitSubscript(OutputValue, 3));
-	Emitter.Output(5, OutputValue); 
+	Emitter.Output(0, Emitter.EmitSwizzle(TextureRead, MIR::FSwizzleMask::XYZ()));
+	Emitter.Output(1, Emitter.EmitSubscript(TextureRead, 0));
+	Emitter.Output(2, Emitter.EmitSubscript(TextureRead, 1)); 
+	Emitter.Output(3, Emitter.EmitSubscript(TextureRead, 2));
+	Emitter.Output(4, Emitter.EmitSubscript(TextureRead, 3));
+	Emitter.Output(5, TextureRead); 
 }
 
 // Returns this value's textures if it has one (nullptr otherwise).
diff --git a/Engine/Source/Runtime/Engine/Private/Materials/MaterialIR.cpp b/Engine/Source/Runtime/Engine/Private/Materials/MaterialIR.cpp
--- a/Engine/Source/Runtime/Engine/Private/Materials/MaterialIR.cpp
+++ b/Engine/Source/Runtime/Engine/Private/Materials/MaterialIR.cpp
@@ -48,7 +48,7 @@ const TCHAR* LexToString(EValueKind Kind)
 		case VK_Branch: return TEXT("Branch");
 		case VK_Subscript: return TEXT("Subscript");
 		case VK_Cast: return TEXT("Cast");
-		case VK_TextureSample: return TEXT("TextureSample");
+		case VK_TextureRead: return TEXT("TextureRead");
 		case VK_InlineHLSL: return TEXT("InlineHLSL");
 		case VK_StageSwitch: return TEXT("StageSwitch");
 		case VK_HardwarePartialDerivative: return TEXT("HardwarePartialDerivative");
@@ -105,7 +105,7 @@ uint32 FValue::GetSizeInBytes() const
 		case VK_Branch: return sizeof(FBranch);
 		case VK_Subscript: return sizeof(FSubscript);
 		case VK_Cast: return sizeof(FCast);
-		case VK_TextureSample: return sizeof(FTextureSample);
+		case VK_TextureRead: return sizeof(FTextureRead);
 		case VK_InlineHLSL: return sizeof(FInlineHLSL);
 		case VK_StageSwitch: return sizeof(FStageSwitch);
 		case VK_HardwarePartialDerivative: return sizeof(FHardwarePartialDerivative);
@@ -185,10 +185,10 @@ TConstArrayView<FValue*> FValue::GetUses() const
 			return { &This->Arg, FSubscript::NumStaticUses };
 		}
 
-		case VK_TextureSample:
+		case VK_TextureRead:
 		{
-			auto This = static_cast<const FTextureSample*>(this);
-			return { &This->TextureObject, FTextureSample::NumStaticUses };
+			auto This = static_cast<const FTextureRead*>(this);
+			return { &This->TextureObject, FTextureRead::NumStaticUses };
 		}
 
 		case VK_InlineHLSL:
@@ -427,6 +427,7 @@ const TCHAR* LexToString(EExternalInput Input)
 		case EExternalInput::TexCoord6_Ddy: return TEXT("TexCoord6_Ddy");
 		case EExternalInput::TexCoord7_Ddy: return TEXT("TexCoord7_Ddy");
 		case EExternalInput::ViewMaterialTextureMipBias: return TEXT("ViewMaterialTextureMipBias");
+		case EExternalInput::ViewMaterialTextureDerivativeMultiply: return TEXT("ViewMaterialTextureDerivativeMultiply");
 		default: UE_MIR_UNREACHABLE();
 	}
 }
@@ -441,6 +442,7 @@ const FType* GetExternalInputType(EExternalInput Id)
 	switch (Id)
 	{
 		case EExternalInput::ViewMaterialTextureMipBias: return FPrimitiveType::GetFloat1();
+		case EExternalInput::ViewMaterialTextureDerivativeMultiply: return FPrimitiveType::GetFloat1();
 		default: UE_MIR_UNREACHABLE();
 	}
 }
@@ -562,7 +564,7 @@ bool IsComparisonOperator(EOperator Op)
 
 bool IsBitOperator(EOperator Op)
 {
-	return Op == UO_Not || (Op >= BO_Or && Op <= BO_Xor);
+	return Op == UO_Not || (Op >= BO_Or && Op <= BO_And);
 }
 
 bool IsUnaryOperator(EOperator Op)
@@ -602,7 +604,6 @@ const TCHAR* LexToString(EOperator Op)
 		case BO_NotEquals: return TEXT("NotEquals");
 		case BO_Or: return TEXT("Or");
 		case BO_Subtract: return TEXT("Subtract");
-		case BO_Xor: return TEXT("Xor");
 		case TO_Clamp: return TEXT("Clamp");
 		case TO_Lerp: return TEXT("Lerp");
 		case TO_Select: return TEXT("Select");
@@ -637,6 +638,30 @@ const TCHAR* LexToString(EOperator Op)
 	return TEXT("???");
 }
 
+const TCHAR* LexToString(ETextureReadMode Mode)
+{
+	switch (Mode)
+	{
+		case ETextureReadMode::GatherRed: return TEXT("GatherRed");
+		case ETextureReadMode::GatherGreen: return TEXT("GatherGreen");
+		case ETextureReadMode::GatherBlue: return TEXT("GatherBlue");
+		case ETextureReadMode::GatherAlpha: return TEXT("GatherAlpha");
+		case ETextureReadMode::MipAuto: return TEXT("MipAuto");
+		case ETextureReadMode::MipLevel: return TEXT("MipLevel");
+		case ETextureReadMode::MipBias: return TEXT("MipBias");
+		case ETextureReadMode::Derivatives: return TEXT("Derivatives");
+		default: UE_MIR_UNREACHABLE();
+	}
+}
+
+void FStageSwitch::SetArgs(FValue* PixelStageArg, FValue* OtherStagesArg)
+{
+	for (uint32 i = 0; i < NumStages; ++i)
+	{
+		Args[i]= (i == Stage_Pixel) ? PixelStageArg : OtherStagesArg;
+	}
+}
+
 } // namespace MIR
 
 #endif // #if WITH_EDITOR
diff --git a/Engine/Source/Runtime/Engine/Private/Materials/MaterialIRDebug.cpp b/Engine/Source/Runtime/Engine/Private/Materials/MaterialIRDebug.cpp
--- a/Engine/Source/Runtime/Engine/Private/Materials/MaterialIRDebug.cpp
+++ b/Engine/Source/Runtime/Engine/Private/Materials/MaterialIRDebug.cpp
@@ -316,11 +316,11 @@ struct FDebugDumpIRState
 				Out.Appendf(TEXT(" Index=%d"), Subscript->Index);
 			}
 		}
-		else if (const FTextureSample* TextureSample = Instr->As<FTextureSample>())
+		else if (const FTextureRead* TextureRead = Instr->As<FTextureRead>())
 		{
-			Out.Appendf(TEXT(" SamplerSourceMode=\"%s\""), *StaticEnum<ESamplerSourceMode>()->GetDisplayNameTextByValue(TextureSample->SamplerSourceMode).ToString());
-			Out.Appendf(TEXT(" MipValueMode=\"%s\""), *StaticEnum<ETextureMipValueMode>()->GetDisplayNameTextByValue(TextureSample->MipValueMode).ToString());
-			Out.Appendf(TEXT(" SamplerType=\"%s\""), *StaticEnum<EMaterialSamplerType>()->GetDisplayNameTextByValue(TextureSample->SamplerType).ToString());
+			Out.Appendf(TEXT(" Mode=\"%s\""), MIR::LexToString(TextureRead->Mode));
+			Out.Appendf(TEXT(" SamplerSourceMode=\"%s\""), *StaticEnum<ESamplerSourceMode>()->GetDisplayNameTextByValue(TextureRead->SamplerSourceMode).ToString());
+			Out.Appendf(TEXT(" SamplerType=\"%s\""), *StaticEnum<EMaterialSamplerType>()->GetDisplayNameTextByValue(TextureRead->SamplerType).ToString());
 		}
 		else if (const FInlineHLSL* InlineHLSL = Instr->As<FInlineHLSL>())
 		{
diff --git a/Engine/Source/Runtime/Engine/Private/Materials/MaterialIREmitter.cpp b/Engine/Source/Runtime/Engine/Private/Materials/MaterialIREmitter.cpp
--- a/Engine/Source/Runtime/Engine/Private/Materials/MaterialIREmitter.cpp
+++ b/Engine/Source/Runtime/Engine/Private/Materials/MaterialIREmitter.cpp
@@ -926,7 +926,6 @@ static TInteger FoldScalarBitOperator(EOperator Operator, TInteger A, TInteger B
 		case UO_Not: return !A;
 		case BO_And: return A & B;
 		case BO_Or: return A | B;
-		case BO_Xor: return A ^ B;
 		default: UE_MIR_UNREACHABLE();
 	}
 }
@@ -1235,12 +1234,13 @@ enum EOperatorParameterTypeFilter
 	OPTF_Any = 0,							///< Any primitive type (no filters)
 	OPTF_Boolean = 1 << 0,					///< Boolean primitive type of any dimension
 	OPTF_Arithmetic = 1 << 1,				///< Arithmetic primitive type of any dimension (i.e. that supports arithmetic operations)
-	OPTF_AnyFloat = 1 << 2,					///< Floating point primitive type of any dimension
-	OPTF_AsFirstArgument = 1 << 3,			///< Must be the same type as the first argument
-	OPTF_NotMatrix = 1 << 4,				///< Any primitive type except matrices
-	OPTF_Vector3 = 1 << 5,					///< 3D vector of any scalar type
+	OPTF_NotMatrix = 1 << 2,				///< Any primitive type except matrices
+	OPTF_Vector3 = 1 << 3,					///< 3D vector of any scalar type
+	OPTF_CastToAnyFloat = 1 << 4,			///< Floating point primitive type of any dimension
+	OPTF_CastToFirstArgumentType = 1 << 5,	///< Must be the same type as the first argument
 	OPTF_CastToCommon = 1 << 6,	 			///< An argument with this flag will be automatically cast to the common type between arguments.
 	OPTF_CastToCommonArithmetic = OPTF_Arithmetic | OPTF_CastToCommon,
+	OPTF_CastToCommonFloat = OPTF_CastToAnyFloat | OPTF_CastToCommon,
 };
 
 inline EOperatorParameterTypeFilter operator|(EOperatorParameterTypeFilter A, EOperatorParameterTypeFilter B)
@@ -1270,35 +1270,36 @@ static const FOperatorSignature* GetOperatorSignature(EOperator Op)
 {
 	static const FOperatorSignature* Signatures = [] ()
 	{
-		const FOperatorSignature UnaryArithmetic = { { OPTF_Arithmetic }, ORT_FirstArgumentType };
+		const FOperatorSignature UnaryFloat = { { OPTF_Arithmetic | OPTF_CastToAnyFloat }, ORT_FirstArgumentType };
 		const FOperatorSignature BinaryArithmetic = { { OPTF_CastToCommonArithmetic, OPTF_CastToCommonArithmetic }, ORT_FirstArgumentType };
 		const FOperatorSignature TernaryArithmetic = { { OPTF_CastToCommonArithmetic, OPTF_CastToCommonArithmetic, OPTF_CastToCommonArithmetic }, ORT_FirstArgumentType };
+		const FOperatorSignature TernaryFloat = { { OPTF_CastToCommonArithmetic | OPTF_CastToAnyFloat, OPTF_CastToCommonArithmetic, OPTF_CastToCommonArithmetic }, ORT_FirstArgumentType };
 		const FOperatorSignature BinaryArithmeticComparison = { { OPTF_CastToCommonArithmetic, OPTF_CastToCommonArithmetic }, ORT_BooleanWithFirstArgumentDimensions };
 		const FOperatorSignature BinaryLogical = { { OPTF_Boolean | OPTF_CastToCommon, OPTF_Boolean | OPTF_CastToCommon }, ORT_FirstArgumentType };
 
 		static FOperatorSignature S[MIR::OperatorCount];
-		S[UO_Negate] 				= UnaryArithmetic;
+		S[UO_Negate] 				= { { OPTF_Arithmetic }, ORT_FirstArgumentType };
 		S[UO_Not] 					= { { OPTF_Boolean }, ORT_FirstArgumentType };
-		S[UO_Abs] 					= UnaryArithmetic;
-		S[UO_ACos] 					= UnaryArithmetic;
-		S[UO_ASin] 					= UnaryArithmetic;
-		S[UO_ATan] 					= UnaryArithmetic;
-		S[UO_Ceil] 					= UnaryArithmetic;
-		S[UO_Cos] 					= UnaryArithmetic;
-		S[UO_Exponential] 			= UnaryArithmetic;
-		S[UO_Exponential2] 			= UnaryArithmetic;
-		S[UO_Floor] 				= UnaryArithmetic;
-		S[UO_Frac]					= UnaryArithmetic;
-		S[UO_Length]				= { { OPTF_Arithmetic | OPTF_NotMatrix }, ORT_FirstArgumentTypeToScalar };
-		S[UO_Logarithm] 			= UnaryArithmetic;
-		S[UO_Logarithm2] 			= UnaryArithmetic;
-		S[UO_Logarithm10] 			= UnaryArithmetic;
-		S[UO_Round] 				= UnaryArithmetic;
-		S[UO_Saturate]				= UnaryArithmetic;
-		S[UO_Sin] 					= UnaryArithmetic;
-		S[UO_Sqrt]					= UnaryArithmetic;
-		S[UO_Tan] 					= UnaryArithmetic;
-		S[UO_Truncate] 				= UnaryArithmetic;
+		S[UO_Abs] 					= UnaryFloat;
+		S[UO_ACos] 					= UnaryFloat;
+		S[UO_ASin] 					= UnaryFloat;
+		S[UO_ATan] 					= UnaryFloat;
+		S[UO_Ceil] 					= UnaryFloat;
+		S[UO_Cos] 					= UnaryFloat;
+		S[UO_Exponential] 			= UnaryFloat;
+		S[UO_Exponential2] 			= UnaryFloat;
+		S[UO_Floor] 				= UnaryFloat;
+		S[UO_Frac]					= UnaryFloat;
+		S[UO_Length]				= { { OPTF_Arithmetic | OPTF_NotMatrix | OPTF_CastToAnyFloat }, ORT_FirstArgumentTypeToScalar };
+		S[UO_Logarithm] 			= UnaryFloat;
+		S[UO_Logarithm2] 			= UnaryFloat;
+		S[UO_Logarithm10] 			= UnaryFloat;
+		S[UO_Round] 				= UnaryFloat;
+		S[UO_Saturate]				= UnaryFloat;
+		S[UO_Sin] 					= UnaryFloat;
+		S[UO_Sqrt]					= UnaryFloat;
+		S[UO_Tan] 					= UnaryFloat;
+		S[UO_Truncate] 				= UnaryFloat;
 		S[BO_GreaterThan] 			= BinaryArithmeticComparison;
 		S[BO_GreaterThanOrEquals]	= BinaryArithmeticComparison;
 		S[BO_LessThan]				= BinaryArithmeticComparison;
@@ -1314,13 +1315,12 @@ static const FOperatorSignature* GetOperatorSignature(EOperator Op)
 		S[BO_Min] 					= BinaryArithmetic;
 		S[BO_Or] 					= BinaryLogical;
 		S[BO_And] 					= BinaryLogical;
-		S[BO_Xor] 					= BinaryLogical;
-		S[BO_Dot] 					= { { OPTF_Arithmetic | OPTF_NotMatrix, OPTF_AsFirstArgument }, ORT_FirstArgumentTypeToScalar };
-		S[BO_Cross] 				= { { OPTF_Arithmetic | OPTF_Vector3, OPTF_AsFirstArgument }, ORT_FirstArgumentType };
+		S[BO_Dot] 					= { { OPTF_Arithmetic | OPTF_NotMatrix, OPTF_CastToFirstArgumentType }, ORT_FirstArgumentTypeToScalar };
+		S[BO_Cross] 				= { { OPTF_Arithmetic | OPTF_Vector3, OPTF_CastToFirstArgumentType }, ORT_FirstArgumentType };
 		S[TO_Clamp]					= TernaryArithmetic;
-		S[TO_Lerp] 					= TernaryArithmetic;
+		S[TO_Lerp] 					= TernaryFloat;
 		S[TO_Select]				= { { OPTF_Boolean | OPTF_NotMatrix, OPTF_NotMatrix, OPTF_NotMatrix}, ORT_SecondArgumentType }; // Note: this is a special operator, which is handled manually in the Validate function
-		S[TO_Smoothstep]			= TernaryArithmetic;
+		S[TO_Smoothstep]			= TernaryFloat;
 		return S;
 	} ();
 	return &Signatures[Op];
@@ -1365,15 +1365,25 @@ static const FPrimitiveType* ValidateOperatorAndGetResultType(FEmitter* Emitter,
 		EOperatorParameterTypeFilter Filter = Signature->ParameterTypes[i];
 		check(Filter != OPTF_Unknown); // No signature specified for this operator.
 
-		if (Filter & OPTF_AsFirstArgument)
+		if (Filter & OPTF_CastToFirstArgumentType)
 		{
 			check(i > 0); // This check can't apply to the first argument.
 
 			// Cast this argument to the first argument primitive type
-			Arguments[i] = Emitter->EmitConstruct(FirstArgumentPrimitiveType, Arguments[i]);
+			ArgumentPrimitiveType = FirstArgumentPrimitiveType;
+			Arguments[i] = Emitter->EmitConstruct(ArgumentPrimitiveType, Arguments[i]);
 			bValid &= !Arguments[i]->IsPoison();
 		}
-
+		else if (Filter & OPTF_CastToAnyFloat)
+		{
+			if (!ScalarKindIsAnyFloat(ArgumentPrimitiveType->ScalarKind))
+			{
+				ArgumentPrimitiveType = ArgumentPrimitiveType->ToScalarKind(MIR::ScalarKind_Float);
+				Arguments[i] = Emitter->EmitConstruct(ArgumentPrimitiveType, Arguments[i]);
+				bValid &= !Arguments[i]->IsPoison();
+			}
+		}
+		
 		if (Filter & OPTF_Boolean)
 		{
 			if (ArgumentPrimitiveType->ScalarKind != ScalarKind_Bool)
@@ -1392,15 +1402,6 @@ static const FPrimitiveType* ValidateOperatorAndGetResultType(FEmitter* Emitter,
 			}
 		}
 		
-		if (Filter & OPTF_AnyFloat)
-		{
-			if (!ScalarKindIsAnyFloat(ArgumentPrimitiveType->ScalarKind))
-			{
-				Emitter->Errorf(TEXT("Operator '%s' %s argument component types expected to be floating point (float, double, etc)."), LexToString(Op), ArgumentsStr[i]);
-				bValid = false;
-			}
-		}
-		
 		if (Filter & OPTF_NotMatrix)
 		{
 			if (ArgumentPrimitiveType->IsMatrix())
@@ -1998,23 +1999,140 @@ FValue* FEmitter::EmitConstruct(const FType* Type, FValue* Initializer)
 	return Result;
 }
 
-FValue* FEmitter::EmitTextureSample(FValue* Texture, FValue* TexCoord, FValue* TexCoordDdx, FValue* TexCoordDdy, FValue* MipValue, ESamplerSourceMode SamplerSourceMode, ETextureMipValueMode MipValueMode, EMaterialSamplerType SamplerType)
+FValue* FEmitter::EmitTextureGather(FValue* Texture, FValue* TexCoord, ETextureReadMode GatherMode, ESamplerSourceMode SamplerSourceMode, EMaterialSamplerType SamplerType)
+{
+	check(GatherMode >= ETextureReadMode::GatherRed && GatherMode <= ETextureReadMode::GatherAlpha);
+
+	if (IsAnyAPoison(Texture, TexCoord))
+	{
+		return FPoison::Get();
+	}
+
+	FTextureRead Prototype = MakePrototype<FTextureRead>(FPrimitiveType::GetFloat4());
+	Prototype.TextureObject = Texture;
+	Prototype.TexCoord = TexCoord;
+	Prototype.Mode = GatherMode;
+	Prototype.SamplerSourceMode = SamplerSourceMode;
+	Prototype.SamplerType = SamplerType;
+
+	return EmitPrototype(this, Prototype);
+}
+
+FValue* FEmitter::EmitTextureSample(FValue* Texture, FValue* TexCoord, bool bAutomaticViewMipBias, ESamplerSourceMode SamplerSourceMode, EMaterialSamplerType SamplerType)
+{
+	if (IsAnyAPoison(Texture, TexCoord))
+	{
+		return FPoison::Get();
+	}
+
+	FTextureRead PrototypeHw = MakePrototype<FTextureRead>(FPrimitiveType::GetFloat4());
+	PrototypeHw.TextureObject = Texture;
+	PrototypeHw.TexCoord = TexCoord;
+	PrototypeHw.Mode = ETextureReadMode::MipAuto;
+	PrototypeHw.SamplerSourceMode = SamplerSourceMode;
+	PrototypeHw.SamplerType = SamplerType;
+
+	FTextureRead PrototypeAn = PrototypeHw;
+	PrototypeAn.Mode = ETextureReadMode::Derivatives;
+	PrototypeAn.TexCoordDdx = EmitAnalyticalPartialDerivative(TexCoord, EDerivativeAxis::X);
+	PrototypeAn.TexCoordDdy = EmitAnalyticalPartialDerivative(TexCoord, EDerivativeAxis::Y);	
+
+	if (bAutomaticViewMipBias)
+	{
+		FValue* ViewMaterialTextureMipBias = EmitExternalInput(EExternalInput::ViewMaterialTextureMipBias);
+		PrototypeHw.Mode = ETextureReadMode::MipBias;
+		PrototypeHw.MipValue = ViewMaterialTextureMipBias;
+
+		FValue* Exp2ViewMaterialTextureMipBias = EmitOperator(UO_Exponential2, ViewMaterialTextureMipBias);
+		PrototypeAn.TexCoordDdx = EmitOperator(BO_Multiply, PrototypeAn.TexCoordDdx, Exp2ViewMaterialTextureMipBias);
+		PrototypeAn.TexCoordDdy = EmitOperator(BO_Multiply, PrototypeAn.TexCoordDdy, Exp2ViewMaterialTextureMipBias);
+	}
+
+	FStageSwitch StageSwitch = MakePrototype<FStageSwitch>(PrototypeHw.Type);
+	StageSwitch.SetArgs(EmitPrototype(this, PrototypeHw), EmitPrototype(this, PrototypeAn));
+
+	return EmitPrototype(this, StageSwitch);
+}
+
+FValue* FEmitter::EmitTextureSampleLevel(FValue* Texture, FValue* TexCoord, FValue* MipLevel, bool bAutomaticViewMipBias, ESamplerSourceMode SamplerSourceMode, EMaterialSamplerType SamplerType)
+{
+	if (IsAnyAPoison(Texture, TexCoord, MipLevel))
+	{
+		return FPoison::Get();
+	}
+
+	FTextureRead Prototype = MakePrototype<FTextureRead>(FPrimitiveType::GetFloat4());
+	Prototype.TextureObject = Texture;
+	Prototype.TexCoord = TexCoord;
+	Prototype.MipValue = MipLevel;
+	Prototype.Mode = ETextureReadMode::MipLevel;
+	Prototype.SamplerSourceMode = SamplerSourceMode;
+	Prototype.SamplerType = SamplerType;
+
+	if (bAutomaticViewMipBias)
+	{
+		Prototype.MipValue = EmitOperator(BO_Add, MipLevel, EmitExternalInput(EExternalInput::ViewMaterialTextureMipBias));
+	}
+
+	return EmitPrototype(this, Prototype);
+}
+
+FValue* FEmitter::EmitTextureSampleBias(FValue* Texture, FValue* TexCoord, FValue* MipBias, bool bAutomaticViewMipBias, ESamplerSourceMode SamplerSourceMode, EMaterialSamplerType SamplerType)
 {
-	if (IsAnyNonNullAPoison(Texture, TexCoord, TexCoordDdx, TexCoordDdy, MipValue))
+	if (IsAnyAPoison(Texture, TexCoord, MipBias))
 	{
 		return FPoison::Get();
 	}
 
-	MIR::FTextureSample Prototype = MakePrototype<MIR::FTextureSample>(FPrimitiveType::GetFloat4());
+	if (bAutomaticViewMipBias)
+	{
+		MipBias = EmitOperator(BO_Add, MipBias, EmitExternalInput(EExternalInput::ViewMaterialTextureMipBias));
+	}
+
+	FTextureRead PrototypeHw = MakePrototype<FTextureRead>(FPrimitiveType::GetFloat4());
+	PrototypeHw.TextureObject = Texture;
+	PrototypeHw.TexCoord = TexCoord;
+	PrototypeHw.MipValue = MipBias;
+	PrototypeHw.Mode = ETextureReadMode::MipBias;
+	PrototypeHw.SamplerSourceMode = SamplerSourceMode;
+	PrototypeHw.SamplerType = SamplerType;
+
+	FTextureRead PrototypeAn = PrototypeHw;
+	PrototypeAn.Mode = ETextureReadMode::Derivatives;
+
+	FValue* Exp2MipBias = EmitOperator(UO_Exponential2, MipBias);
+	PrototypeAn.TexCoordDdx = EmitOperator(BO_Multiply, EmitAnalyticalPartialDerivative(TexCoord, EDerivativeAxis::X), Exp2MipBias);
+	PrototypeAn.TexCoordDdy = EmitOperator(BO_Multiply, EmitAnalyticalPartialDerivative(TexCoord, EDerivativeAxis::Y), Exp2MipBias);	
+
+	FStageSwitch StageSwitch = MakePrototype<FStageSwitch>(PrototypeHw.Type);
+	StageSwitch.SetArgs(EmitPrototype(this, PrototypeHw), EmitPrototype(this, PrototypeAn));
+
+	return EmitPrototype(this, StageSwitch);
+}
+
+FValue* FEmitter::EmitTextureSampleGrad(FValue* Texture, FValue* TexCoord, FValue* TexCoordDdx, FValue* TexCoordDdy, bool bAutomaticViewMipBias, ESamplerSourceMode SamplerSourceMode, EMaterialSamplerType SamplerType)
+{
+	if (IsAnyAPoison(Texture, TexCoord, TexCoordDdx, TexCoordDdy))
+	{
+		return FPoison::Get();
+	}
+
+	FTextureRead Prototype = MakePrototype<MIR::FTextureRead>(FPrimitiveType::GetFloat4());
 	Prototype.TextureObject = Texture;
 	Prototype.TexCoord = TexCoord;
 	Prototype.TexCoordDdx = TexCoordDdx;
 	Prototype.TexCoordDdy = TexCoordDdy;
-	Prototype.MipValue = MipValue;
+	Prototype.Mode = ETextureReadMode::Derivatives;
 	Prototype.SamplerSourceMode = SamplerSourceMode;
-	Prototype.MipValueMode = MipValueMode;
 	Prototype.SamplerType = SamplerType;
 
+	if (bAutomaticViewMipBias)
+	{
+		FValue* ViewMaterialTextureDerivativeMultiply = EmitExternalInput(EExternalInput::ViewMaterialTextureDerivativeMultiply);
+		Prototype.TexCoordDdx = EmitOperator(BO_Multiply, Prototype.TexCoordDdx, ViewMaterialTextureDerivativeMultiply);
+		Prototype.TexCoordDdy = EmitOperator(BO_Multiply, Prototype.TexCoordDdy, ViewMaterialTextureDerivativeMultiply);
+	}
+
 	return EmitPrototype(this, Prototype);
 }
 
@@ -2407,7 +2525,7 @@ FValue* FEmitter::EmitAnalyticalPartialDerivative(FValue* Value, EDerivativeAxis
 		// These values are uniform (constant), thus their value is always zero.
 		case MIR::VK_Constant:
 		case MIR::VK_UniformParameter:
-		case MIR::VK_TextureSample:
+		case MIR::VK_TextureRead:
 		case MIR::VK_InlineHLSL:
 			return EmitConstruct(ValuePrimitiveType, EmitConstantScalarZero(ValuePrimitiveType->ScalarKind));
 
diff --git a/Engine/Source/Runtime/Engine/Private/Materials/MaterialIRToHLSLTranslator.cpp b/Engine/Source/Runtime/Engine/Private/Materials/MaterialIRToHLSLTranslator.cpp
--- a/Engine/Source/Runtime/Engine/Private/Materials/MaterialIRToHLSLTranslator.cpp
+++ b/Engine/Source/Runtime/Engine/Private/Materials/MaterialIRToHLSLTranslator.cpp
@@ -371,7 +371,7 @@ struct FTranslator : FMaterialIRToHLSLTranslation
 			case MIR::VK_Operator: LowerOperator(static_cast<MIR::FOperator*>(Instr)); break;
 			case MIR::VK_Branch: LowerBranch(static_cast<MIR::FBranch*>(Instr)); break;
 			case MIR::VK_Subscript: LowerSubscript(static_cast<MIR::FSubscript*>(Instr)); break;
-			case MIR::VK_TextureSample: LowerTextureSample(static_cast<MIR::FTextureSample*>(Instr)); break;
+			case MIR::VK_TextureRead: LowerTextureRead(static_cast<MIR::FTextureRead*>(Instr)); break;
 			case MIR::VK_InlineHLSL: LowerInlineHLSL(static_cast<const MIR::FInlineHLSL*>(Instr)); break;
 			case MIR::VK_StageSwitch: LowerStageSwitch(static_cast<const MIR::FStageSwitch*>(Instr)); break;
 			case MIR::VK_HardwarePartialDerivative: LowerHardwarePartialDerivative(static_cast<const MIR::FHardwarePartialDerivative*>(Instr)); break;
@@ -418,6 +418,7 @@ struct FTranslator : FMaterialIRToHLSLTranslation
 			switch (ExternalInput->Id)
 			{
 				case MIR::EExternalInput::ViewMaterialTextureMipBias: Code = TEXT("View.MaterialTextureMipBias"); break;
+				case MIR::EExternalInput::ViewMaterialTextureDerivativeMultiply: Code = TEXT("View.MaterialTextureDerivativeMultiply"); break;
 				default: UE_MIR_UNREACHABLE();
 			}
 			Printer << Code;
@@ -684,7 +685,6 @@ struct FTranslator : FMaterialIRToHLSLTranslation
 				case MIR::BO_Max: OpString = TEXT("max"); break;
 				case MIR::BO_Min: OpString = TEXT("min"); break;
 				case MIR::BO_Or: OpString = TEXT("or"); break;
-				case MIR::BO_Xor: OpString = TEXT("xor"); break;
 				case MIR::TO_Clamp: OpString = TEXT("clamp"); break;
 				case MIR::TO_Lerp: OpString = TEXT("lerp"); break;
 				case MIR::TO_Select: OpString = TEXT("select"); break;
@@ -768,41 +768,44 @@ struct FTranslator : FMaterialIRToHLSLTranslation
 		}
 	}
 
-	void LowerTextureSample(const MIR::FTextureSample* TextureSample)
+	void LowerTextureRead(const MIR::FTextureRead* TextureRead)
 	{
-		bool bSamplerNeedsBrackets = LowerSamplerType(TextureSample->SamplerType);
+		bool bSamplerNeedsBrackets = LowerSamplerType(TextureRead->SamplerType);
 		if (bSamplerNeedsBrackets)
 		{
 			Printer << TEXT("(");
 		}
-
-		switch (TextureSample->TextureObject->Type->AsObject()->ObjectKind)
+		
+		switch (TextureRead->TextureObject->Type->AsObject()->ObjectKind)
 		{
-			case MIR::ObjectKind_Texture2D: Printer << TEXT("Texture2DSample"); break;
+			case MIR::ObjectKind_Texture2D: Printer << TEXT("Texture2D"); break;
 			default: UE_MIR_UNREACHABLE();
 		}
 
-		switch (TextureSample->MipValueMode)
+		switch (TextureRead->Mode)
 		{
-			case TMVM_None: break;
-			case TMVM_MipLevel: Printer << TEXT("Level"); break;
-			case TMVM_MipBias: Printer << TEXT("Bias"); break;
-			case TMVM_Derivative: Printer << TEXT("Grad"); break;
+			case MIR::ETextureReadMode::GatherRed: Printer << TEXT("GatherRed"); break;
+			case MIR::ETextureReadMode::GatherGreen: Printer << TEXT("GatherGreen"); break;
+			case MIR::ETextureReadMode::GatherBlue: Printer << TEXT("GatherBlue"); break;
+			case MIR::ETextureReadMode::GatherAlpha: Printer << TEXT("GatherAlpha"); break;
+			case MIR::ETextureReadMode::MipAuto: Printer << TEXT("Sample"); break;
+			case MIR::ETextureReadMode::MipLevel: Printer << TEXT("SampleLevel"); break;
+			case MIR::ETextureReadMode::MipBias: Printer << TEXT("SampleBias"); break;
+			case MIR::ETextureReadMode::Derivatives: Printer << TEXT("SampleGrad"); break;
 			default: UE_MIR_UNREACHABLE();
 		}
 
 		Printer << BeginArgs
-				<< ListSeparator << LowerValue(TextureSample->TextureObject)
-				<< ListSeparator << LowerTextureSamplerReference(TextureSample->TextureObject, TextureSample->SamplerSourceMode)
-				<< ListSeparator << LowerValue(TextureSample->TexCoord);
+				<< ListSeparator << LowerValue(TextureRead->TextureObject)
+				<< ListSeparator << LowerTextureSamplerReference(TextureRead->TextureObject, TextureRead->SamplerSourceMode)
+				<< ListSeparator << LowerValue(TextureRead->TexCoord);
 
-		switch (TextureSample->MipValueMode)
+		switch (TextureRead->Mode)
 		{
-			case TMVM_None: break;
-			case TMVM_MipLevel: Printer << ListSeparator << LowerValue(TextureSample->MipValue); break;
-			case TMVM_MipBias: Printer << ListSeparator << LowerValue(TextureSample->MipValue); break;
-			case TMVM_Derivative: Printer << ListSeparator << LowerValue(TextureSample->TexCoordDdx) << ListSeparator << LowerValue(TextureSample->TexCoordDdy); break;
-			default: UE_MIR_UNREACHABLE();
+			case MIR::ETextureReadMode::MipLevel: Printer << ListSeparator << LowerValue(TextureRead->MipValue); break;
+			case MIR::ETextureReadMode::MipBias: Printer << ListSeparator << LowerValue(TextureRead->MipValue); break;
+			case MIR::ETextureReadMode::Derivatives: Printer << ListSeparator << LowerValue(TextureRead->TexCoordDdx) << ListSeparator << LowerValue(TextureRead->TexCoordDdy); break;
+			default: break;
 		}
 
 		Printer << EndArgs;
@@ -870,15 +873,25 @@ struct FTranslator : FMaterialIRToHLSLTranslation
 
 	ENoOp LowerTextureSamplerReference(MIR::FValue* TextureValue, ESamplerSourceMode SamplerSource)
 	{
-		switch (SamplerSource)
+		if (SamplerSource != SSM_FromTextureAsset)
 		{
-		case SSM_FromTextureAsset:
-			LowerValue(TextureValue);
-			Printer << TEXT("Sampler");
-			break;
-
-		default:
-			UE_MIR_UNREACHABLE();
+			Printer << "GetMaterialSharedSampler(";
+		}
+		
+		Printer << LowerValue(TextureValue) << TEXT("Sampler");
+		
+		if (SamplerSource == SSM_Wrap_WorldGroupSettings)
+		{
+			Printer << ", View.MaterialTextureBilinearWrapedSampler)";
+		}
+		else if (SamplerSource == SSM_Clamp_WorldGroupSettings)
+		{
+			Printer << ", View.MaterialTextureBilinearClampedSampler)";
+		}
+		else
+		{
+			// SSM_TerrainWeightmapGroupSettings unsupported yet
+			check(SamplerSource == SSM_FromTextureAsset);
 		}
 
 		return NoOp;
diff --git a/Engine/Source/Runtime/Engine/Private/Materials/MaterialIRTypes.cpp b/Engine/Source/Runtime/Engine/Private/Materials/MaterialIRTypes.cpp
--- a/Engine/Source/Runtime/Engine/Private/Materials/MaterialIRTypes.cpp
+++ b/Engine/Source/Runtime/Engine/Private/Materials/MaterialIRTypes.cpp
@@ -380,6 +380,11 @@ const FPrimitiveType* FPrimitiveType::ToScalar() const
 	return FPrimitiveType::GetScalar(ScalarKind);
 }
 
+const FPrimitiveType* FPrimitiveType::ToScalarKind(EScalarKind InScalarKind) const
+{
+	return FPrimitiveType::Get(InScalarKind, NumRows, NumColumns);
+}
+
 const FType* FObjectType::GetTexture2D()
 {
 	static const FObjectType Instance { { TypeKind_Object }, ObjectKind_Texture2D };
diff --git a/Engine/Source/Runtime/Engine/Public/Materials/MaterialIR.h b/Engine/Source/Runtime/Engine/Public/Materials/MaterialIR.h
--- a/Engine/Source/Runtime/Engine/Public/Materials/MaterialIR.h
+++ b/Engine/Source/Runtime/Engine/Public/Materials/MaterialIR.h
@@ -106,7 +106,7 @@ enum EValueKind
 	VK_Branch,
 	VK_Subscript,
 	VK_Cast,
-	VK_TextureSample,
+	VK_TextureRead,
 	VK_InlineHLSL,
 	VK_StageSwitch,
 	VK_HardwarePartialDerivative,
@@ -311,14 +311,20 @@ static inline bool IsAnyNonNullAPoison()
 template <typename... TTail>
 bool IsAnyNonNullAPoison(const FValue* Head, TTail... Tail)
 {
-	if (Head && Head->IsPoison())
-	{
-		return true;
-	}
-	else
-	{
-		return IsAnyNonNullAPoison(Tail...);
-	}
+	return (Head && Head->IsPoison()) || IsAnyNonNullAPoison(Tail...);
+}
+
+static inline bool IsAnyAPoison()
+{
+	return false;
+}
+
+/** @return Whether there is a poison in the values specified among non-null values */
+template <typename... TTail>
+bool IsAnyAPoison(const FValue* Head, TTail... Tail)
+{
+	check(Head && "A non-null value must have been provided.");
+	return Head->IsPoison() || IsAnyAPoison(Tail...);
 }
 
 /** The integer type used inside MIR. */
@@ -410,6 +416,7 @@ enum class EExternalInput
 	TexCoord7_Ddy,
 
 	ViewMaterialTextureMipBias,
+	ViewMaterialTextureDerivativeMultiply,
 
 	Count,
 };
@@ -674,7 +681,6 @@ enum EOperator
 	/* Binary Logical Operators */
 	BO_Or,
 	BO_And,
-	BO_Xor,
 
 	/* Binary Intrinsic Operators */
 	BO_Dot,
@@ -783,9 +789,42 @@ struct FCast : TInstruction<VK_Cast, 1>
 };
 
 /**
- * @brief This instruction performs texture sampling.
+ * What texture gather mode to use in a texture read instruction (none indicates a sample).
  */
-struct FTextureSample : TInstruction<VK_TextureSample, 5>
+enum class ETextureReadMode
+{
+	/** Gather the four red components in a 2x2 pixel block. */
+	GatherRed,
+
+	/** Gather the four green components in a 2x2 pixel block. */
+	GatherGreen,
+
+	/** Gather the four blue components in a 2x2 pixel block. */
+	GatherBlue,
+
+	/** Gather the four alpha components in a 2x2 pixel block. */
+	GatherAlpha,
+
+	/** Texture gather with automatically calculated mip level */
+	MipAuto,
+
+	/** Texture gather with user specified mip level */
+	MipLevel,
+
+	/** Texture gather with automatically calculated mip level plus user specified bias */
+	MipBias,
+
+	/** Texture gather using automatically caluclated mip level based on user provided partial derivatives */
+	Derivatives,
+};
+
+/** Returns the string representation of given mode. */
+const TCHAR* LexToString(ETextureReadMode Mode);
+
+/**
+ * @brief This instruction performs texture read operaation (sample or gather).
+ */
+struct FTextureRead : TInstruction<VK_TextureRead, 5>
 {
 	/** The texture object to sample. */
 	FValue* TextureObject;
@@ -802,14 +841,14 @@ struct FTextureSample : TInstruction<VK_TextureSample, 5>
 	/** Optional. The analytical partial derivative of the coordinates along the Y axis. */
 	FValue* TexCoordDdy;
 
-	/** The sampler source mode to use for sampling. */
-	ESamplerSourceMode SamplerSourceMode;
-
 	/** The mip value mode to use for sampling. */
-	ETextureMipValueMode MipValueMode;
+	ETextureReadMode Mode : 8;
+
+	/** The sampler source mode to use for sampling. */
+	ESamplerSourceMode SamplerSourceMode : 8;
 
 	/** The sampler type to use for sampling. */
-	EMaterialSamplerType SamplerType;
+	EMaterialSamplerType SamplerType : 8;
 };
 
 /**
@@ -817,7 +856,14 @@ struct FTextureSample : TInstruction<VK_TextureSample, 5>
  */
 struct FStageSwitch : TInstruction<VK_StageSwitch, 1>
 {
+	/** The argument for to be bypassed in each stage. e*/
 	FValue* Args[NumStages];
+
+	/**
+	 * Use the specified value argument in the pixel stage, and another specified
+	 * argument for other stages.
+	 */
+	void SetArgs(FValue* PixelStageArg, FValue* OtherStagesArg);
 };
 
 /**
diff --git a/Engine/Source/Runtime/Engine/Public/Materials/MaterialIREmitter.h b/Engine/Source/Runtime/Engine/Public/Materials/MaterialIREmitter.h
--- a/Engine/Source/Runtime/Engine/Public/Materials/MaterialIREmitter.h
+++ b/Engine/Source/Runtime/Engine/Public/Materials/MaterialIREmitter.h
@@ -190,20 +190,63 @@ public:
 	FValue* EmitConstruct(const FType* Type, FValue* Initializer);
 
 	/**
-	 * @brief Emits a texture sampling instruction.
+	 * @brief Emits a texture gather instruction.
 	 * 
 	 * @param Texture The texture to sample. Must be of type MIR::FObjectType and represent a valid texture.
 	 * @param TexCoord The UV coordinates at which to sample the texture.
-	 * @param TexCoordDdx Optional. The manually provided partial derivative of TexCoord along the X axis. Required if MipValueMode is Derivative.
-	 * @param TexCoordDdy Optional. The manually provided partial derivative of TexCoord along the Y axis. Required if MipValueMode is Derivative.
-	 * @param MipValue Optional. The mip level or mip bias to apply. Required if MipValueMode is Level or Bias.
-	 * @param SamplerSourceMode The sampler source mode to use.
-	 * @param MipValueMode The mip value mode to use.
-	 * @param SamplerType The sampler type to use.
+	 * @param GatherMode What channel to gather (any ::GatherXXX value).
+	 * @param SamplerSourceMode Which sampler source mode to use.
+	 * @param SamplerType Which sampler type to use.
+	 */
+	FValue* EmitTextureGather(FValue* Texture, FValue* TexCoord, ETextureReadMode GatherMode, ESamplerSourceMode SamplerSourceMode, EMaterialSamplerType SamplerType);
+
+	/**
+	 * @brief Emits a texture sample instruction.
+	 * 
+	 * @param Texture The texture to sample. Must be of type MIR::FObjectType and represent a valid texture.
+	 * @param TexCoord The UV coordinates at which to sample the texture.
+	 * @param bAddViewMaterialMipBias Whether to add the view material texture mip bias to the mip level.
+	 * @param SamplerSourceMode Which sampler source mode to use.
+	 * @param SamplerType Which sampler type to use.
+	 */
+	FValue* EmitTextureSample(FValue* Texture, FValue* TexCoord, bool bAddViewMaterialMipBias, ESamplerSourceMode SamplerSourceMode, EMaterialSamplerType SamplerType);
+
+	/**
+	 * @brief Emits a texture sample instruction with a manually specified mip level.
 	 * 
-	 * @return A pointer to the emitted texture sample instruction.
+	 * @param Texture The texture to sample. Must be of type MIR::FObjectType and represent a valid texture.
+	 * @param TexCoord The UV coordinates at which to sample the texture.
+	 * @param MipLevel Which mip level to use for this sample.
+	 * @param bAddViewMaterialMipBias Whether to add the view material texture mip bias to the mip level.
+	 * @param SamplerSourceMode Which sampler source mode to use.
+	 * @param SamplerType Which sampler type to use.
+	 */
+	FValue* EmitTextureSampleLevel(FValue* Texture, FValue* TexCoord, FValue* MipLevel, bool bAddViewMaterialMipBias, ESamplerSourceMode SamplerSourceMode, EMaterialSamplerType SamplerType);
+
+	/**
+	 * @brief Emits a texture sample instruction with bias added to the automatically selected mip level.
+	 * 
+	 * @param Texture The texture to sample. Must be of type MIR::FObjectType and represent a valid texture.
+	 * @param TexCoord The UV coordinates at which to sample the texture.
+	 * @param MipBias The bias to add to the automatically selected mip level.
+	 * @param bAddViewMaterialMipBias Whether to add the view material texture mip bias to the mip level.
+	 * @param SamplerSourceMode Which sampler source mode to use.
+	 * @param SamplerType Which sampler type to use.
+	 */
+	FValue* EmitTextureSampleBias(FValue* Texture, FValue* TexCoord, FValue* MipBias, bool bAddViewMaterialMipBias, ESamplerSourceMode SamplerSourceMode, EMaterialSamplerType SamplerType);
+	
+	/**
+	 * @brief Emits a texture sample instruction using user provided partial derivatives.
+	 * 
+	 * @param Texture The texture to sample. Must be of type MIR::FObjectType and represent a valid texture.
+	 * @param TexCoord The UV coordinates at which to sample the texture.
+	 * @param TexCoordDdx The partial derivative of the texture coordinates along the screen space x axis.
+	 * @param TexCoordDdx The partial derivative of the texture coordinates along the screen space y axis.
+	 * @param bAddViewMaterialMipBias Whether to add the view material texture mip bias to the mip level.
+	 * @param SamplerSourceMode Which sampler source mode to use.
+	 * @param SamplerType Which sampler type to use.
 	 */
-	FValue* EmitTextureSample(FValue* Texture, FValue* TexCoord, FValue* TexCoordDdx, FValue* TexCoordDdy, FValue* MipValue, ESamplerSourceMode SamplerSourceMode, ETextureMipValueMode MipValueMode, EMaterialSamplerType SamplerType);
+	FValue* EmitTextureSampleGrad(FValue* Texture, FValue* TexCoord, FValue* TexCoordDdx, FValue* TexCoordDdy, bool bAddViewMaterialMipBias, ESamplerSourceMode SamplerSourceMode, EMaterialSamplerType SamplerType);
 	
 	FValue* EmitInlineHLSL(const FType* Type, FString Code, TConstArrayView<FValue*> InputValues, EValueFlags ValueFlags = EValueFlags::None, EGraphProperties UsedGraphProperties = EGraphProperties::None);
 	FValue* EmitInlineHLSL(const FType* Type, FString Code, EValueFlags ValueFlags = EValueFlags::None, EGraphProperties UsedGraphProperties = EGraphProperties::None);
diff --git a/Engine/Source/Runtime/Engine/Public/Materials/MaterialIRTypes.h b/Engine/Source/Runtime/Engine/Public/Materials/MaterialIRTypes.h
--- a/Engine/Source/Runtime/Engine/Public/Materials/MaterialIRTypes.h
+++ b/Engine/Source/Runtime/Engine/Public/Materials/MaterialIRTypes.h
@@ -137,6 +137,7 @@ struct FPrimitiveType : FType
 	bool IsVector() const { return NumRows > 1 && NumColumns == 1; }
 	bool IsMatrix() const { return NumRows > 1 && NumColumns > 1; }
 	const FPrimitiveType* ToScalar() const;
+	const FPrimitiveType* ToScalarKind(EScalarKind InScalarKind) const;
 };
 
 struct FObjectType : FType
