NetTests
* Added a PreTickAll delegate function

SceneGraph
* Added a way to automatically tick the ExecutionManager when calling FTestWorlds::TickAll instead of having to do it manually

#rnx

diff --git a/Engine/Source/Runtime/Engine/Private/Tests/NetTestHelpers.cpp b/Engine/Source/Runtime/Engine/Private/Tests/NetTestHelpers.cpp
--- a/Engine/Source/Runtime/Engine/Private/Tests/NetTestHelpers.cpp
+++ b/Engine/Source/Runtime/Engine/Private/Tests/NetTestHelpers.cpp
@@ -349,6 +349,8 @@ bool FTestWorlds::WaitForClientConnect(FTestWorldInstance& Client)
 
 void FTestWorlds::TickAll(int32 NumTicks)
 {
+	PreTickAllDelegate.Broadcast();
+
 	for (int32 i = 0; i < NumTicks; ++i)
 	{
 		TickServer();
diff --git a/Engine/Source/Runtime/Engine/Public/Tests/NetTestHelpers.h b/Engine/Source/Runtime/Engine/Public/Tests/NetTestHelpers.h
--- a/Engine/Source/Runtime/Engine/Public/Tests/NetTestHelpers.h
+++ b/Engine/Source/Runtime/Engine/Public/Tests/NetTestHelpers.h
@@ -195,6 +195,10 @@ struct FTestWorlds
 	template<class PredicateT>
 	bool TickAllUntil(const PredicateT& Predicate, float DeltaSeconds = 0.0166f, int32 MaxTicks = 60);
 
+	DECLARE_MULTICAST_DELEGATE(FPreTickAll);
+	/** Delegate called before a TickAll is executed */
+	FPreTickAll PreTickAllDelegate;
+
 	/** Ticks all server & client worlds until the passed in client world has a valid client PlayerController. */
 	ENGINE_API bool WaitForClientConnect(FTestWorldInstance& Client);
 
@@ -259,6 +263,8 @@ private:
 template<class PredicateT>
 inline bool FTestWorlds::TickAllUntil(const PredicateT& Predicate, float DeltaSeconds, int32 MaxTicks)
 {
+	PreTickAllDelegate.Broadcast();
+
 	int32 TickCount = 0;
 	bool bPredicateResult = Predicate();
 		
