World Partition Handle: avoid per-actor individual calls to FlushAsyncLoading when flushing actor loading deferred contexts by gathering all async load IDs and flushing once.
In a map with ~40K actors to load, saves ~25% of the time spent flushing async load requests (~8.6s instead of ~11.6s).

diff --git a/Engine/Source/Runtime/Engine/Private/WorldPartition/WorldPartitionHandle.cpp b/Engine/Source/Runtime/Engine/Private/WorldPartition/WorldPartitionHandle.cpp
--- a/Engine/Source/Runtime/Engine/Private/WorldPartition/WorldPartitionHandle.cpp
+++ b/Engine/Source/Runtime/Engine/Private/WorldPartition/WorldPartitionHandle.cpp
@@ -158,6 +158,17 @@ FWorldPartitionLoadingContext::FDeferred::~FDeferred()
 			{
 				ActorList.Empty(SourceList.Num());
 
+				TArray<int32> AllAsyncLoadIds;
+				AllAsyncLoadIds.Reserve(SourceList.Num());
+				for (FWorldPartitionActorDescInstance* ActorDescInstance : SourceList)
+				{
+					if (int32 AsyncLoadId = ActorDescInstance->GetAsyncLoadID(); AsyncLoadId != INDEX_NONE)
+					{
+						AllAsyncLoadIds.Add(AsyncLoadId);
+					}
+				}
+				FlushAsyncLoading(AllAsyncLoadIds);
+
 				ULevel* Level = nullptr;
 				for (FWorldPartitionActorDescInstance* ActorDescInstance : SourceList)
 				{
diff --git a/Engine/Source/Runtime/Engine/Public/WorldPartition/WorldPartitionActorDescInstance.h b/Engine/Source/Runtime/Engine/Public/WorldPartition/WorldPartitionActorDescInstance.h
--- a/Engine/Source/Runtime/Engine/Public/WorldPartition/WorldPartitionActorDescInstance.h
+++ b/Engine/Source/Runtime/Engine/Public/WorldPartition/WorldPartitionActorDescInstance.h
@@ -68,6 +68,8 @@ public:
 	ENGINE_API virtual FBox GetEditorBounds() const;
 	ENGINE_API virtual FBox GetRuntimeBounds() const;
 
+	ENGINE_API int32 GetAsyncLoadID() const { return AsyncLoadID; }
+
 	virtual bool GetProperty(FName PropertyName, FName* PropertyValue) const { return ActorDesc->GetProperty(PropertyName, PropertyValue); }
 	virtual bool HasProperty(FName PropertyName) const { return ActorDesc->HasProperty(PropertyName); }
 
