Disable OnlyFirstMip for non pow 2 textures due to issues with tiling platforms.

#rnx

diff --git a/Engine/Source/Runtime/Engine/Private/Texture2D.cpp b/Engine/Source/Runtime/Engine/Private/Texture2D.cpp
--- a/Engine/Source/Runtime/Engine/Private/Texture2D.cpp
+++ b/Engine/Source/Runtime/Engine/Private/Texture2D.cpp
@@ -759,21 +759,29 @@ void UTexture2D::UpdateResourceWithParams(EUpdateResourceFlags InFlags)
 	const bool bLoadOnlyFirstMip = UDeviceProfileManager::Get().GetActiveProfile()->GetTextureLODSettings()->GetMipLoadOptions(this) == ETextureMipLoadOptions::OnlyFirstMip;
 	if (bLoadOnlyFirstMip && GetPlatformData() && GetPlatformData()->Mips.Num() > 0 && FPlatformProperties::RequiresCookedData())
 	{
-		const int32 NumMipsInTail = GetPlatformData()->GetNumMipsInTail();
-		const int32 MipTailBaseIndex = FMath::Max(0, NumMipsInTail > 0 ? (GetPlatformData()->Mips.Num() - NumMipsInTail) : (GetPlatformData()->Mips.Num() - 1));
-
-		const int32 FirstMip = FMath::Min(FMath::Max(0, GetCachedLODBias()), MipTailBaseIndex);
-		if (FirstMip < MipTailBaseIndex)
+		// If we aren't pow2 there are failure modes here on some platforms that will cause crashes if we
+		// strip tail mips without updating padding.
+		// If you are using this to get around that fact that LODBias doesn't work with NoMipMaps, see the Downscale property.
+		FTexturePlatformData* PlatformData = GetPlatformData();
+		bool bIsPow2 = FMath::IsPowerOfTwo(PlatformData->Mips[0].SizeX) && FMath::IsPowerOfTwo(PlatformData->Mips[0].SizeY);
+		if (bIsPow2)
 		{
-			// Remove any mips after the first mip.
-			GetPlatformData()->Mips.RemoveAt(FirstMip + 1, GetPlatformData()->Mips.Num() - FirstMip - 1);
-			GetPlatformData()->OptData.NumMipsInTail = 0;
+			const int32 NumMipsInTail = PlatformData->GetNumMipsInTail();
+			const int32 MipTailBaseIndex = FMath::Max(0, NumMipsInTail > 0 ? (PlatformData->Mips.Num() - NumMipsInTail) : (PlatformData->Mips.Num() - 1));
+
+			const int32 FirstMip = FMath::Min(FMath::Max(0, GetCachedLODBias()), MipTailBaseIndex);
+			if (FirstMip < MipTailBaseIndex)
+			{
+				// Remove any mips after the first mip.
+				PlatformData->Mips.RemoveAt(FirstMip + 1, PlatformData->Mips.Num() - FirstMip - 1);
+				PlatformData->OptData.NumMipsInTail = 0;
+			}
+			// Remove any mips before the first mip.
+			PlatformData->Mips.RemoveAt(0, FirstMip);
+			// Update the texture size for the memory usage metrics.
+			PlatformData->SizeX = PlatformData->Mips[0].SizeX;
+			PlatformData->SizeY = PlatformData->Mips[0].SizeY;
 		}
-		// Remove any mips before the first mip.
-		GetPlatformData()->Mips.RemoveAt(0, FirstMip);
-		// Update the texture size for the memory usage metrics.
-		GetPlatformData()->SizeX = GetPlatformData()->Mips[0].SizeX;
-		GetPlatformData()->SizeY = GetPlatformData()->Mips[0].SizeY;
 	}
 #endif // #if WITH_EDITOR
 
diff --git a/Engine/Source/Runtime/Engine/Private/TextureDerivedData.cpp b/Engine/Source/Runtime/Engine/Private/TextureDerivedData.cpp
--- a/Engine/Source/Runtime/Engine/Private/TextureDerivedData.cpp
+++ b/Engine/Source/Runtime/Engine/Private/TextureDerivedData.cpp
@@ -2865,6 +2865,24 @@ static void SerializePlatformData(
 		{
 			FirstMipToSerialize = FMath::Clamp(FirstMipToSerialize, 0, FMath::Max((int32)PlatformData->VTData->GetNumMips() - 1, 0));
 		}
+
+		// We can't reliably strip the mip tail on non-pow2 textures after tiling on all platforms, so this gets ignored for nonpow2
+		// at runtime. Warn here to catch it earlier.
+		ETextureMipLoadOptions MipLoadOptions = Ar.CookingTarget()->GetTextureLODSettings().GetMipLoadOptions(Texture);
+		if (MipLoadOptions == ETextureMipLoadOptions::OnlyFirstMip)
+		{
+			// This property only applies to physical 2d textures and there's no point in warning if there's only 1 mip.
+			if (!bIsVirtual && Texture->GetTextureClass() == ETextureClass::TwoD && NumMips > 1)
+			{
+				if (!FMath::IsPowerOfTwo(PlatformData->Mips[FirstMipToSerialize].SizeX) ||
+					!FMath::IsPowerOfTwo(PlatformData->Mips[FirstMipToSerialize].SizeY))
+				{
+					// If you are here because you're trying to LODBias a texture with no mips and using this as a workaround,
+					// look at the Downscale setting.
+					UE_LOG(LogTexture, Warning, TEXT("MipLoadOption OnlyFirstMip can't be applied to non pow2 textures, see Downscale option if this is being used as a workaround for LODBias on NoMipMaps. Texture: %s"), *Texture->GetName());
+				}
+			}
+		}
 	}
 #endif // #if WITH_EDITORONLY_DATA
 
