Allow HV component to correctly cast shadow when also marked hidden in game.

diff --git a/Engine/Source/Runtime/Engine/Private/Components/HeterogeneousVolumeComponent.cpp b/Engine/Source/Runtime/Engine/Private/Components/HeterogeneousVolumeComponent.cpp
--- a/Engine/Source/Runtime/Engine/Private/Components/HeterogeneousVolumeComponent.cpp
+++ b/Engine/Source/Runtime/Engine/Private/Components/HeterogeneousVolumeComponent.cpp
@@ -645,7 +645,8 @@ void UHeterogeneousVolumeComponent::TickComponent(float DeltaTime, ELevelTick Ti
 		}
 	}
 
-	if (ShouldRender() && DoesRuntimeSupportHeterogeneousVolumes(GetWorld()->Scene->GetShaderPlatform()) && MaterialInstanceDynamic)
+	// Aligning renderable state behavior with Instanced Static Mesh and Skinned Mesh. HV is currently not supporoted in indirect lighting and ray-tracing far field.
+	if ((ShouldRender() || bCastHiddenShadow) && DoesRuntimeSupportHeterogeneousVolumes(GetWorld()->Scene->GetShaderPlatform()) && MaterialInstanceDynamic)
 	{
 		const int32 SVTParameterIndex = 0;
 		FName SVTParameterName;
