Adding an ensure when a physics body is being terminated if it has other bodies welded to it.
Adding to track down a crash.

diff --git a/Engine/Source/Runtime/Engine/Private/PhysicsEngine/BodyInstance.cpp b/Engine/Source/Runtime/Engine/Private/PhysicsEngine/BodyInstance.cpp
--- a/Engine/Source/Runtime/Engine/Private/PhysicsEngine/BodyInstance.cpp
+++ b/Engine/Source/Runtime/Engine/Private/PhysicsEngine/BodyInstance.cpp
@@ -1842,6 +1842,21 @@ void FBodyInstance::TermBody(bool bNeverDeferRelease)
 {
 	SCOPE_CYCLE_COUNTER(STAT_TermBody);
 
+	// It's bad if we're being terminated but there's still stuff welded to us
+	if (ShapeToBodiesMap != nullptr && !ShapeToBodiesMap->IsEmpty())
+	{
+		TStringBuilder<256> Builder;
+		int32 BodyIndex = 0;
+		for (TMap<FPhysicsShapeHandle, FWeldInfo>::TRangedForIterator It = ShapeToBodiesMap->begin(); It; ++It)
+		{
+			const FBodyInstance* ChildBody = It->Value.ChildBI;
+			const FString Name = ChildBody != nullptr ? ChildBody->GetBodyDebugName() : FString("null");
+			Builder.Appendf(TEXT("\tBody[%d]: \'%s'\n"), BodyIndex, *Name);
+			++BodyIndex;
+		}
+		ensureMsgf(false, TEXT("Body %s is being terminated but has %d welded bodies:\n%s"), *OwnerComponent->GetName(), ShapeToBodiesMap->Num(), Builder.ToString());
+	}
+
 	if (UPrimitiveComponent* PrimComp = OwnerComponent.Get())
 	{
 		if (FPhysScene_Chaos* PhysScene = PrimComp->GetWorld()->GetPhysicsScene())
