Make subtitles plugin use UMG.

#virtualized

diff --git a/Engine/Source/Runtime/Engine/Private/GameViewportClient.cpp b/Engine/Source/Runtime/Engine/Private/GameViewportClient.cpp
--- a/Engine/Source/Runtime/Engine/Private/GameViewportClient.cpp
+++ b/Engine/Source/Runtime/Engine/Private/GameViewportClient.cpp
@@ -32,7 +32,6 @@
 #include "GameFramework/HUD.h"
 #include "FXSystem.h"
 #include "SubtitleManager.h"
-#include "Subtitles/SubtitlesAndClosedCaptionsDelegates.h"
 #include "ImageUtils.h"
 #include "SceneViewExtension.h"
 #include "IHeadMountedDisplay.h"
@@ -2029,7 +2028,6 @@ void UGameViewportClient::Draw(FViewport* InViewport, FCanvas* SceneCanvas)
 							const uint32 SizeY = SceneCanvas->GetRenderTarget()->GetSizeXY().Y;
 							FIntRect SubtitleRegion(FMath::TruncToInt(SizeX * MinPos.X), FMath::TruncToInt(SizeY * MinPos.Y), FMath::TruncToInt(SizeX * MaxPos.X), FMath::TruncToInt(SizeY * MaxPos.Y));
 							FSubtitleManager::GetSubtitleManager()->DisplaySubtitles( SceneCanvas, SubtitleRegion, MyWorld->GetAudioTimeSeconds() );
-							FSubtitlesAndClosedCaptionsDelegates::DisplaySubtitle.ExecuteIfBound(*SceneCanvas);
 							bDisplayedSubtitles = true;
 						}
 					}
