Added log for the case when the audio engine is unable to initialize due to intentional audio disablement (-nosound or an application mode without audio backend support)
[RN] Added log for the case when the audio engine is unable to initialize due to intentional audio disablement (-nosound or an application mode without audio backend support)

diff --git a/Engine/Source/Runtime/Engine/Private/AudioDeviceManager.cpp b/Engine/Source/Runtime/Engine/Private/AudioDeviceManager.cpp
--- a/Engine/Source/Runtime/Engine/Private/AudioDeviceManager.cpp
+++ b/Engine/Source/Runtime/Engine/Private/AudioDeviceManager.cpp
@@ -742,6 +742,8 @@ FAudioDeviceManager* FAudioDeviceManager::GetOrCreate()
 			static bool bDoOnce = false;
 			if (!bDoOnce)
 			{
+				UE_LOG(LogAudio, Display, TEXT("Audio Device Manager not initializing due to all audio being disabled. If this is not intentional, please check command line arguments for \"-nosound\"."));
+
 				Audio::Analytics::RecordEvent_Usage(TEXT("AllAudioDisabled"));
 				bDoOnce = true;
 			}
