[Audio Insights] Home Screen - First level sound list with grouping (no children sounds)

diff --git a/Engine/Plugins/AudioInsights/Source/AudioInsights/Private/Messages/SoundTraceMessages.cpp b/Engine/Plugins/AudioInsights/Source/AudioInsights/Private/Messages/SoundTraceMessages.cpp
new file mode 100644
--- /dev/null
+++ b/Engine/Plugins/AudioInsights/Source/AudioInsights/Private/Messages/SoundTraceMessages.cpp
@@ -0,0 +1,74 @@
+// Copyright Epic Games, Inc. All Rights Reserved.
+#include "Messages/SoundTraceMessages.h"
+
+namespace UE::Audio::Insights
+{
+	namespace SoundClassNames
+	{
+		const FString MetaSoundSource = TEXT("MetaSoundSource");
+		const FString SoundWaveProcedural = TEXT("SoundWaveProcedural");
+		const FString SoundCue = TEXT("SoundCue");
+		const FString SoundWave = TEXT("SoundWave");
+		const FString SoundCueTemplate = TEXT("SoundCueTemplate");
+	}
+
+	FSoundMessageBase::FSoundMessageBase(const Trace::IAnalyzer::FOnEventContext& InContext)
+	{
+		const Trace::IAnalyzer::FEventData& EventData = InContext.EventData;
+
+		DeviceId = static_cast<::Audio::FDeviceId>(EventData.GetValue<uint32>("DeviceId"));
+		Timestamp = InContext.EventTime.AsSeconds(EventData.GetValue<uint64>("Timestamp"));
+		PlayOrder = EventData.GetValue<uint32>("PlayOrder");
+	}
+
+	FSoundStartMessage::FSoundStartMessage(const Trace::IAnalyzer::FOnEventContext& InContext)
+		: FSoundMessageBase(InContext)
+	{
+		const Trace::IAnalyzer::FEventData& EventData = InContext.EventData;
+		EventData.GetString("Name", Name);
+
+		FString SoundClassName;
+		FString SoundSuperClassName;
+
+		EventData.GetString("SoundClassName", SoundClassName);
+		EventData.GetString("SoundSuperClassName", SoundSuperClassName);
+
+		if (SoundClassName == SoundClassNames::SoundCue)
+		{
+			EntryType = ESoundDashboardEntryType::SoundCue;
+		}
+		else if (SoundClassName == SoundClassNames::SoundWave)
+		{
+			EntryType = ESoundDashboardEntryType::SoundWave;
+		}
+		else if (SoundClassName == SoundClassNames::MetaSoundSource)
+		{
+			EntryType = ESoundDashboardEntryType::MetaSound;
+		}
+		else if (SoundClassName == SoundClassNames::SoundWaveProcedural
+			|| SoundSuperClassName == SoundClassNames::SoundWaveProcedural)
+		{
+			EntryType = ESoundDashboardEntryType::ProceduralSource;
+		}
+		else if (SoundClassName == SoundClassNames::SoundCueTemplate
+			|| SoundSuperClassName == SoundClassNames::SoundCueTemplate)
+		{
+			EntryType = ESoundDashboardEntryType::SoundCueTemplate;
+		}
+		else
+		{
+			EntryType = ESoundDashboardEntryType::None;
+		}
+	}
+
+	FSoundDashboardEntry::FSoundDashboardEntry()
+	{
+		constexpr uint32 DataPointsCapacity = 256;
+
+		PriorityDataPoints.SetCapacity(DataPointsCapacity);
+		DistanceDataPoints.SetCapacity(DataPointsCapacity);
+		AmplitudeDataPoints.SetCapacity(DataPointsCapacity);
+		VolumeDataPoints.SetCapacity(DataPointsCapacity);
+		PitchDataPoints.SetCapacity(DataPointsCapacity);
+	}
+}
\ No newline at end of file
diff --git a/Engine/Plugins/AudioInsights/Source/AudioInsights/Private/Views/SoundDashboardViewFactory.cpp b/Engine/Plugins/AudioInsights/Source/AudioInsights/Private/Views/SoundDashboardViewFactory.cpp
--- a/Engine/Plugins/AudioInsights/Source/AudioInsights/Private/Views/SoundDashboardViewFactory.cpp
+++ b/Engine/Plugins/AudioInsights/Source/AudioInsights/Private/Views/SoundDashboardViewFactory.cpp
@@ -201,7 +201,25 @@ namespace UE::Audio::Insights
 				"AudioInsights.Icon.SoundDashboard.ProceduralSource",
 				FText::GetEmpty(),
 				FLinearColor(0.98f, 0.32f, 0.006f),
-				FilterCategory)
+				FilterCategory),
+
+				MakeShared<FSoundDashboardFilter>(
+				ESoundDashboardFilterFlags::SoundWave,
+				"SoundWave",
+				LOCTEXT("SoundDashboard_Filter_SoundWaveNameText", "Sound Wave"),
+				"AudioInsights.Icon.SoundDashboard.SoundWave",
+				FText::GetEmpty(),
+				FLinearColor(FColor::Purple),
+				FilterCategory),
+
+				MakeShared<FSoundDashboardFilter>(
+				ESoundDashboardFilterFlags::SoundCue,
+				"SoundCueTemplate",
+				LOCTEXT("SoundDashboard_Filter_SoundCueTemplateNameText", "Sound Cue Template"),
+				"AudioInsights.Icon.SoundDashboard.SoundCue",
+				FText::GetEmpty(),
+				FLinearColor(FColor::Red),
+				FilterCategory),
 			};
 
 			SAssignNew(SoundsFilterBar, SAudioFilterBar<ESoundDashboardFilterFlags>)
@@ -285,6 +303,12 @@ namespace UE::Audio::Insights
 				case ESoundDashboardEntryType::ProceduralSource:
 					return FSlateColor(FColor::Orange);
 
+				case ESoundDashboardEntryType::SoundWave:
+					return FSlateColor(FColor::Purple);
+
+				case ESoundDashboardEntryType::SoundCueTemplate:
+					return FSlateColor(FColor::Red);
+
 				case ESoundDashboardEntryType::None:
 				default:
 					return FSlateColor(FColor::White);
@@ -596,6 +620,10 @@ namespace UE::Audio::Insights
 								return FName("AudioInsights.Icon.SoundDashboard.SoundCue");
 							case ESoundDashboardEntryType::ProceduralSource:
 								return FName("AudioInsights.Icon.SoundDashboard.ProceduralSource");
+							case ESoundDashboardEntryType::SoundWave:
+								return FName("AudioInsights.Icon.SoundDashboard.SoundWave");
+							case ESoundDashboardEntryType::SoundCueTemplate:
+								return FName("AudioInsights.Icon.SoundDashboard.SoundCue");
 							case ESoundDashboardEntryType::None:
 							default:
 								break;
@@ -830,7 +858,8 @@ namespace UE::Audio::Insights
 
 				const FName SoundAssetDisplayName = FName(SoundAssetDashboardEntry.GetDisplayName().ToString());
 
-				if (SoundAssetDashboardEntry.EntryType == ESoundDashboardEntryType::SoundCue)
+				if (SoundAssetDashboardEntry.EntryType == ESoundDashboardEntryType::SoundCue
+					|| SoundAssetDashboardEntry.EntryType == ESoundDashboardEntryType::SoundCueTemplate)
 				{
 					AudioDebugger.SetMuteSoundCue(SoundAssetDisplayName, true /* bInOnOff */);
 				}
@@ -868,7 +897,8 @@ namespace UE::Audio::Insights
 
 				const FName SoundAssetDisplayName = FName(SoundAssetDashboardEntry.GetDisplayName().ToString());
 
-				if (SoundAssetDashboardEntry.EntryType == ESoundDashboardEntryType::SoundCue)
+				if (SoundAssetDashboardEntry.EntryType == ESoundDashboardEntryType::SoundCue
+					|| SoundAssetDashboardEntry.EntryType == ESoundDashboardEntryType::SoundCueTemplate)
 				{
 					AudioDebugger.SetSoloSoundCue(SoundAssetDisplayName, true /* bInOnOff */);
 				}
diff --git a/Engine/Plugins/AudioInsights/Source/AudioInsights/Public/Messages/SoundTraceMessages.h b/Engine/Plugins/AudioInsights/Source/AudioInsights/Public/Messages/SoundTraceMessages.h
--- a/Engine/Plugins/AudioInsights/Source/AudioInsights/Public/Messages/SoundTraceMessages.h
+++ b/Engine/Plugins/AudioInsights/Source/AudioInsights/Public/Messages/SoundTraceMessages.h
@@ -21,20 +21,15 @@ namespace UE::Audio::Insights
 		MetaSound,
 		SoundCue,
 		ProceduralSource,
+		SoundWave,
+		SoundCueTemplate
 	};
 
 	// Trace messages
 	struct FSoundMessageBase
 	{
 		FSoundMessageBase() = default;
-		FSoundMessageBase(const Trace::IAnalyzer::FOnEventContext& InContext)
-		{
-			const Trace::IAnalyzer::FEventData& EventData = InContext.EventData;
-
-			DeviceId  = static_cast<::Audio::FDeviceId>(EventData.GetValue<uint32>("DeviceId"));
-			Timestamp = InContext.EventTime.AsSeconds(EventData.GetValue<uint64>("Timestamp"));
-			PlayOrder = EventData.GetValue<uint32>("PlayOrder");
-		}
+		FSoundMessageBase(const Trace::IAnalyzer::FOnEventContext& InContext);
 
 		::Audio::FDeviceId DeviceId = INDEX_NONE;
 		uint32 PlayOrder = INDEX_NONE;
@@ -44,53 +39,14 @@ namespace UE::Audio::Insights
 	struct FSoundStartMessage : public FSoundMessageBase
 	{
 		FSoundStartMessage() = default;
-
-		FSoundStartMessage(const Trace::IAnalyzer::FOnEventContext& InContext)
-			: FSoundMessageBase(InContext)
-		{
-			const Trace::IAnalyzer::FEventData& EventData = InContext.EventData;
-			EventData.GetString("Name", Name);
-		}
+		FSoundStartMessage(const Trace::IAnalyzer::FOnEventContext& InContext);
 
 		FString Name;
-		ESoundDashboardEntryType EntryType = ESoundDashboardEntryType::None;
+		ESoundDashboardEntryType EntryType;
 	};
 
 	using FSoundStopMessage = FSoundMessageBase;
 
-	struct FMetaSoundStartMessage : public FSoundStartMessage
-	{
-		FMetaSoundStartMessage() = default;
-
-		FMetaSoundStartMessage(const UE::Trace::IAnalyzer::FOnEventContext& InContext)
-			: FSoundStartMessage(InContext)
-		{
-			EntryType = ESoundDashboardEntryType::MetaSound;
-		}
-	};
-
-	struct FSoundCueStartMessage : public FSoundStartMessage
-	{
-		FSoundCueStartMessage() = default;
-
-		FSoundCueStartMessage(const UE::Trace::IAnalyzer::FOnEventContext& InContext)
-			: FSoundStartMessage(InContext)
-		{
-			EntryType = ESoundDashboardEntryType::SoundCue;
-		}
-	};
-
-	struct FProceduralSourceStartMessage : public FSoundStartMessage
-	{
-		FProceduralSourceStartMessage() = default;
-
-		FProceduralSourceStartMessage(const UE::Trace::IAnalyzer::FOnEventContext& InContext)
-			: FSoundStartMessage(InContext)
-		{
-			EntryType = ESoundDashboardEntryType::ProceduralSource;
-		}
-	};
-
 	struct FSoundPriorityMessage : public FSoundMessageBase
 	{
 		FSoundPriorityMessage() = default;
@@ -197,17 +153,7 @@ namespace UE::Audio::Insights
 	class FSoundDashboardEntry : public IObjectTreeDashboardEntry
 	{
 	public:
-		FSoundDashboardEntry()
-		{
-			constexpr uint32 DataPointsCapacity = 256;
-
-			PriorityDataPoints.SetCapacity(DataPointsCapacity);
-			DistanceDataPoints.SetCapacity(DataPointsCapacity);
-			AmplitudeDataPoints.SetCapacity(DataPointsCapacity);
-			VolumeDataPoints.SetCapacity(DataPointsCapacity);
-			PitchDataPoints.SetCapacity(DataPointsCapacity);
-		}
-
+		FSoundDashboardEntry();
 		virtual ~FSoundDashboardEntry() = default;
 
 		virtual TObjectPtr<UObject> GetObject() override { return FSoftObjectPath(Name).ResolveObject(); }
diff --git a/Engine/Plugins/AudioInsights/Source/AudioInsights/Public/Views/SoundDashboardViewFactory.h b/Engine/Plugins/AudioInsights/Source/AudioInsights/Public/Views/SoundDashboardViewFactory.h
--- a/Engine/Plugins/AudioInsights/Source/AudioInsights/Public/Views/SoundDashboardViewFactory.h
+++ b/Engine/Plugins/AudioInsights/Source/AudioInsights/Public/Views/SoundDashboardViewFactory.h
@@ -15,6 +15,7 @@ namespace UE::Audio::Insights
 		MetaSound = 1 << 0,
 		SoundCue  = 1 << 1,
 		ProceduralSource = 1 << 2,
+		SoundWave = 1 << 3,
 	};
 
 	ENUM_CLASS_FLAGS(ESoundDashboardFilterFlags);
diff --git a/Engine/Source/Runtime/Engine/Private/AudioDevice.cpp b/Engine/Source/Runtime/Engine/Private/AudioDevice.cpp
--- a/Engine/Source/Runtime/Engine/Private/AudioDevice.cpp
+++ b/Engine/Source/Runtime/Engine/Private/AudioDevice.cpp
@@ -240,6 +240,21 @@ FAutoConsoleVariableRef CVarEnableRelativeRenderCostVoiceLimit(
 	ECVF_Default);
 
 #if UE_AUDIO_PROFILERTRACE_ENABLED
+UE_TRACE_EVENT_BEGIN(Audio, SoundStart)
+	UE_TRACE_EVENT_FIELD(uint32, DeviceId)
+	UE_TRACE_EVENT_FIELD(uint64, Timestamp)
+	UE_TRACE_EVENT_FIELD(uint32, PlayOrder)
+	UE_TRACE_EVENT_FIELD(UE::Trace::WideString, Name)
+	UE_TRACE_EVENT_FIELD(UE::Trace::WideString, SoundClassName)
+	UE_TRACE_EVENT_FIELD(UE::Trace::WideString, SoundSuperClassName)
+UE_TRACE_EVENT_END()
+
+UE_TRACE_EVENT_BEGIN(Audio, SoundStop)
+	UE_TRACE_EVENT_FIELD(uint32, DeviceId)
+	UE_TRACE_EVENT_FIELD(uint64, Timestamp)
+	UE_TRACE_EVENT_FIELD(uint32, PlayOrder)
+UE_TRACE_EVENT_END()
+
 UE_TRACE_EVENT_BEGIN(Audio, VirtualLoopStopOrRealize)
 	UE_TRACE_EVENT_FIELD(uint32, DeviceId)
 	UE_TRACE_EVENT_FIELD(uint64, Timestamp)
@@ -2223,6 +2238,16 @@ bool FAudioDevice::HandleSetDynamicSoundCommand(const TCHAR* Cmd, FOutputDevice&
 #if UE_AUDIO_PROFILERTRACE_ENABLED
 void FAudioDevice::OnTraceStarted(FTraceAuxiliary::EConnectionType TraceType, const FString& TraceDestination)
 {
+	for (const FActiveSound* ActiveSound : ActiveSounds)
+	{
+		if (ActiveSound == nullptr)
+		{
+			continue;
+		}
+
+		SendSoundStartedTraceMessage(*ActiveSound);
+	}
+
 	for (auto& [ActiveSound, VirtualLoop] : VirtualLoops)
 	{
 		VirtualLoop.OnTraceStarted();
@@ -5007,6 +5032,36 @@ void FAudioDevice::UpdateAudioEngineSubsystems()
 	});
 }
 
+#if UE_AUDIO_PROFILERTRACE_ENABLED
+void FAudioDevice::SendSoundStartedTraceMessage(const FActiveSound& InActiveSound) const
+{
+	const bool bChannelEnabled = UE_TRACE_CHANNELEXPR_IS_ENABLED(AudioMixerChannel);
+	TObjectPtr<USoundBase> SoundBasePtr = InActiveSound.GetSound();
+
+	if (bChannelEnabled && SoundBasePtr)
+	{
+		TObjectPtr<UObject> SoundObjectPtr = FSoftObjectPath(InActiveSound.GetSound()->GetPathName()).ResolveObject();
+		if (SoundObjectPtr == nullptr)
+		{
+			return;
+		}
+
+		const FString SoundClassName = SoundObjectPtr->GetClass()->GetName();
+
+		const TObjectPtr<UClass> SoundSuperClass = SoundObjectPtr->GetClass()->GetSuperClass();
+		const FString SuperClassName = SoundSuperClass ? SoundSuperClass->GetName() : FString();
+
+		UE_TRACE_LOG(Audio, SoundStart, AudioMixerChannel)
+			<< SoundStart.DeviceId(InActiveSound.AudioDevice->DeviceID)
+			<< SoundStart.Timestamp(FPlatformTime::Cycles64())
+			<< SoundStart.PlayOrder(InActiveSound.PlayOrder)
+			<< SoundStart.Name(*SoundObjectPtr->GetPathName())
+			<< SoundStart.SoundClassName(*SoundClassName)
+			<< SoundStart.SoundSuperClassName(*SuperClassName);
+	}
+}
+#endif // UE_AUDIO_PROFILERTRACE_ENABLED
+
 FDelegateHandle FAudioDevice::AddPreRenderDelegate(const FOnAudioDevicePreRender::FDelegate& InDelegate)
 {
 	FScopeLock LockCallbacks(&RenderStateCallbackListCritSec);
@@ -5334,6 +5389,10 @@ void FAudioDevice::AddNewActiveSoundInternal(const FActiveSound& InNewActiveSoun
 	}
 
 	NotifySubsystemsActiveSoundCreated(*ActiveSound);
+
+#if UE_AUDIO_PROFILERTRACE_ENABLED
+	SendSoundStartedTraceMessage(*ActiveSound);
+#endif // UE_AUDIO_PROFILERTRACE_ENABLED
 }
 
 void FAudioDevice::InitSoundParams(FActiveSound& InOutActiveSound, TArray<FAudioParameter> InDefaultParams, bool bInIsVirtualLoopRealizing) const
@@ -5845,6 +5904,17 @@ void FAudioDevice::RemoveActiveSound(FActiveSound* ActiveSound)
 	const uint64 ComponentID = ActiveSound->GetAudioComponentID();
 	UAudioComponent::PlaybackCompleted(ComponentID, false);
 
+#if UE_AUDIO_PROFILERTRACE_ENABLED
+	const bool bChannelEnabled = UE_TRACE_CHANNELEXPR_IS_ENABLED(AudioMixerChannel);
+	if (bChannelEnabled)
+	{
+		UE_TRACE_LOG(Audio, SoundStop, AudioMixerChannel)
+			<< SoundStop.DeviceId(ActiveSound->AudioDevice->DeviceID)
+			<< SoundStop.Timestamp(FPlatformTime::Cycles64())
+			<< SoundStop.PlayOrder(ActiveSound->PlayOrder);
+	}
+#endif // UE_AUDIO_PROFILERTRACE_ENABLED
+
 	const int32 NumRemoved = ActiveSounds.RemoveSwap(ActiveSound);
 	if (!ensureMsgf(NumRemoved > 0, TEXT("Attempting to remove an already removed ActiveSound '%s'"), ActiveSound->Sound ? *ActiveSound->Sound->GetName() : TEXT("N/A")))
 	{
diff --git a/Engine/Source/Runtime/Engine/Public/AudioDevice.h b/Engine/Source/Runtime/Engine/Public/AudioDevice.h
--- a/Engine/Source/Runtime/Engine/Public/AudioDevice.h
+++ b/Engine/Source/Runtime/Engine/Public/AudioDevice.h
@@ -1702,6 +1702,11 @@ private:
 	/** Updates audio engine subsystems on this device. */
 	void UpdateAudioEngineSubsystems();
 
+#if UE_AUDIO_PROFILERTRACE_ENABLED
+	/** Notify Audio Insights that an active sound has started playing. */
+	void SendSoundStartedTraceMessage(const FActiveSound& InActiveSound) const;
+#endif // UE_AUDIO_PROFILERTRACE_ENABLED
+
 public:
 
 	/**
