Fix debug draw text location for VisualLogger, RewindDebugger and DebugRenderSceneProxy by considering DPI scale and aspect ratio requirement (black bars)

diff --git a/Engine/Plugins/Animation/GameplayInsights/Source/RewindDebuggerVLog/Private/VisualLogEntryRenderer.cpp b/Engine/Plugins/Animation/GameplayInsights/Source/RewindDebuggerVLog/Private/VisualLogEntryRenderer.cpp
--- a/Engine/Plugins/Animation/GameplayInsights/Source/RewindDebuggerVLog/Private/VisualLogEntryRenderer.cpp
+++ b/Engine/Plugins/Animation/GameplayInsights/Source/RewindDebuggerVLog/Private/VisualLogEntryRenderer.cpp
@@ -26,7 +26,7 @@ namespace
 				const FSceneView* View = Canvas->SceneView;
 				if (View->GetCullingFrustum().IntersectPoint(Position))
 				{
-					const FVector ScreenLoc = Canvas->Project(Position);
+					const FVector3f ScreenLoc = UE::DebugDrawHelper::GetScaleAdjustedScreenLocation(Canvas, Position);
 					Canvas->SetDrawColor(Color);
 					const FString PrintString = Count == 1 ? ElementToDraw->Description : FString::Printf(TEXT("%s_%d"), *ElementToDraw->Description, Index);
 					Canvas->Canvas->DrawShadowedString(ScreenLoc.X, ScreenLoc.Y, PrintString, Font, Color);
diff --git a/Engine/Source/Developer/LogVisualizer/Private/VisualLoggerCanvasRenderer.cpp b/Engine/Source/Developer/LogVisualizer/Private/VisualLoggerCanvasRenderer.cpp
--- a/Engine/Source/Developer/LogVisualizer/Private/VisualLoggerCanvasRenderer.cpp
+++ b/Engine/Source/Developer/LogVisualizer/Private/VisualLoggerCanvasRenderer.cpp
@@ -3,6 +3,7 @@
 #include "VisualLoggerCanvasRenderer.h"
 #include "EngineGlobals.h"
 #include "CanvasItem.h"
+#include "DebugRenderSceneProxy.h"
 #include "Engine/Canvas.h"
 #include "Engine/Engine.h"
 #include "SceneView.h"
@@ -47,7 +48,7 @@ namespace LogVisualizer
 	{
 		if (PointInFrustum(Canvas, WorldLocation))
 		{
-			const FVector3f ScreenLocation = UE::LWC::NarrowWorldPositionChecked(Canvas->Project(WorldLocation));
+			const FVector3f ScreenLocation = UE::DebugDrawHelper::GetScaleAdjustedScreenLocation (Canvas, WorldLocation);
 			float TextXL = 0.f;
 			float TextYL = 0.f;
 			Canvas->StrLen(Font, TextToDraw, TextXL, TextYL);
diff --git a/Engine/Source/Runtime/Engine/Private/DebugRenderSceneProxy.cpp b/Engine/Source/Runtime/Engine/Private/DebugRenderSceneProxy.cpp
--- a/Engine/Source/Runtime/Engine/Private/DebugRenderSceneProxy.cpp
+++ b/Engine/Source/Runtime/Engine/Private/DebugRenderSceneProxy.cpp
@@ -27,6 +27,15 @@ static TAutoConsoleVariable<float> CVarDebugRenderOverrideFarClippingPlane(
 	TEXT("Allows to override the far clipping plane for debug shapes and text (in cm), only effective if > 0.0 and if r.DebugRender.AllowFrustumCulling is enabled"),
 	ECVF_RenderThreadSafe);
 
+// UE::DebugDrawHelper
+FVector3f UE::DebugDrawHelper::GetScaleAdjustedScreenLocation(const TNotNull<UCanvas*> Canvas, const FVector WorldLocation)
+{
+	const FVector3f ScreenLocation = LWC::NarrowWorldPositionChecked(Canvas->Project(WorldLocation));
+	const float InvDPIScale = 1.f / Canvas->GetDPIScale();
+	const FIntRect Delta = (Canvas->SceneView->UnconstrainedViewRect - Canvas->SceneView->UnscaledViewRect);
+	return {ScreenLocation.X + (0.5f * InvDPIScale * Delta.Width()), ScreenLocation.Y + (0.5f * InvDPIScale * Delta.Height()), ScreenLocation.Z};
+}
+
 // FDebugRenderSceneProxy::FMesh
 FDebugRenderSceneProxy::FMesh::FMesh() = default;
 FDebugRenderSceneProxy::FMesh::~FMesh() = default;
@@ -169,7 +178,7 @@ void FDebugDrawDelegateHelper::DrawDebugLabels(UCanvas* Canvas, APlayerControlle
 	{
 		if (!bAllowFrustumCulling || FDebugRenderSceneProxy::PointInFrustum(It->Location, AdjustedFrustum))
 		{
-			const FVector ScreenLoc = Canvas->Project(It->Location);
+			const FVector3f ScreenLoc = UE::DebugDrawHelper::GetScaleAdjustedScreenLocation(Canvas, It->Location);
 			const FFontRenderInfo& FontInfo = TextWithoutShadowDistance >= 0 ? (FDebugRenderSceneProxy::PointInRange(It->Location, View, TextWithoutShadowDistance) ? FontRenderInfoWithShadow : FontRenderInfo) : FontRenderInfo;
 			Canvas->SetDrawColor(It->Color);
 			Canvas->DrawText(RenderFont, It->Text, ScreenLoc.X, ScreenLoc.Y, 1, 1, FontInfo);
diff --git a/Engine/Source/Runtime/Engine/Public/DebugRenderSceneProxy.h b/Engine/Source/Runtime/Engine/Public/DebugRenderSceneProxy.h
--- a/Engine/Source/Runtime/Engine/Public/DebugRenderSceneProxy.h
+++ b/Engine/Source/Runtime/Engine/Public/DebugRenderSceneProxy.h
@@ -28,6 +28,12 @@ class UPrimitiveComponent;
 
 DECLARE_DELEGATE_TwoParams(FDebugDrawDelegate, UCanvas*, APlayerController*);
 
+namespace UE::DebugDrawHelper
+{
+/** Projects provided 3D world-space location into 2D screen coordinates considering DPI scale and aspect ratio requirement (black bars) */
+ENGINE_API FVector3f GetScaleAdjustedScreenLocation(TNotNull<UCanvas*> Canvas, FVector WorldLocation);
+}
+
 class FDebugRenderSceneProxy : public FPrimitiveSceneProxy
 {
 public:
