[HLODs] Fix Build HLOD Loc Strings

#tests Buildling HLODs
#rnx

diff --git a/Engine/Source/Editor/UnrealEd/Private/WorldPartition/WorldPartitionHLODsBuilder.cpp b/Engine/Source/Editor/UnrealEd/Private/WorldPartition/WorldPartitionHLODsBuilder.cpp
--- a/Engine/Source/Editor/UnrealEd/Private/WorldPartition/WorldPartitionHLODsBuilder.cpp
+++ b/Engine/Source/Editor/UnrealEd/Private/WorldPartition/WorldPartitionHLODsBuilder.cpp
@@ -33,6 +33,8 @@
 #include "WorldPartition/HLOD/HLODProviderInterface.h"
 #include "WorldPartition/IWorldPartitionEditorModule.h"
 
+#define LOCTEXT_NAMESPACE "WorldPartitionHLODsBuilder"
+
 DEFINE_LOG_CATEGORY_STATIC(LogWorldPartitionHLODsBuilder, All, All);
 
 static const FString DistributedBuildWorkingDirName = TEXT("HLODTemp");
@@ -410,7 +412,7 @@ bool UWorldPartitionHLODsBuilder::BuildHLODActors()
 
 				AWorldPartitionHLOD* HLODActor = CastChecked<AWorldPartitionHLOD>(ActorRef.GetActor());
 
-				UE_LOG(LogWorldPartitionHLODsBuilder, Display, TEXT("[%d / %d] Building HLOD actor %s..."), CurrentActor + 1, HLODActorsToBuild.Num(), *HLODActor->GetActorLabel());
+				UE_LOG(LogWorldPartitionHLODsBuilder, Display, TEXT("[%d / %d] %s %s..."), CurrentActor + 1, HLODActorsToBuild.Num(), *LOCTEXT("BuildingHLODActor", "Building HLOD actor").ToString(), *HLODActor->GetActorLabel());
 
 				// Simulate an engine tick to make sure engine & render resources that are queued for deletion are processed.
 				FWorldPartitionHelpers::FakeEngineTick(World);
@@ -1086,3 +1088,5 @@ bool UWorldPartitionHLODsBuilder::AddBuildProducts(const TArray<FString>& BuildP
 	}
 	return bRet;
 }
+
+#undef LOCTEXT_NAMESPACE
\ No newline at end of file
diff --git a/Engine/Source/Runtime/Engine/Private/WorldPartition/RuntimeHashSet/WorldPartitionRuntimeHashSetHLODGeneration.cpp b/Engine/Source/Runtime/Engine/Private/WorldPartition/RuntimeHashSet/WorldPartitionRuntimeHashSetHLODGeneration.cpp
--- a/Engine/Source/Runtime/Engine/Private/WorldPartition/RuntimeHashSet/WorldPartitionRuntimeHashSetHLODGeneration.cpp
+++ b/Engine/Source/Runtime/Engine/Private/WorldPartition/RuntimeHashSet/WorldPartitionRuntimeHashSetHLODGeneration.cpp
@@ -21,6 +21,8 @@
 #include "UObject/SavePackage.h"
 #include "HAL/PlatformFileManager.h"
 
+#define LOCTEXT_NAMESPACE "WorldPartitionRuntimeHashSetHLODGeneration"
+
 #if WITH_EDITOR
 class FHLODStreamingGenerationContext : public IStreamingGenerationContext
 {
@@ -208,7 +210,7 @@ bool UWorldPartitionRuntimeHashSet::SetupHLODActors(const IStreamingGenerationCo
 
 				const FCellUniqueId CellUniqueId = GetCellUniqueId(CellDescInstance);
 
-				UE_LOG(LogWorldPartition, Display, TEXT("[%d / %d] Processing cell %s..."), ++CellDescInstanceIndex, CellDescInstances.Num(), *CellUniqueId.Name);
+				UE_LOG(LogWorldPartition, Display, TEXT("[%d / %d] %s %s..."), ++CellDescInstanceIndex, CellDescInstances.Num(), *LOCTEXT("ProcessingCell", "Processing cell").ToString(), *CellUniqueId.Name);
 
 				FBox CellBounds;
 				TArray<IStreamingGenerationContext::FActorInstance> ActorInstances;
@@ -293,7 +295,7 @@ bool UWorldPartitionRuntimeHashSet::SetupHLODActors(const IStreamingGenerationCo
 			HLODStreamingGenerationContext->ActorSetInstanceList.Reserve(HLODActorGuids.Num());
 			HLODStreamingGenerationContext->ActorSetContainerInstance.ActorDescViewMap = &HLODStreamingGenerationContext->ActorDescViewMap;
 
-			UE_LOG(LogWorldPartition, Log, TEXT("Creating HLOD context:"));
+			UE_LOG(LogWorldPartition, Log, TEXT("%s"), *LOCTEXT("CreatingHLODContext:", "Creating HLOD context:").ToString());
 			for (const FGuid& HLODActorGuid : HLODActorGuids)
 			{
 				FWorldPartitionActorDescInstance* HLODActorDescInstance = WorldPartition->GetActorDescInstance(HLODActorGuid);
@@ -342,3 +344,5 @@ bool UWorldPartitionRuntimeHashSet::SetupHLODActors(const IStreamingGenerationCo
 	return false;
 }
 #endif
+
+#undef LOCTEXT_NAMESPACE
\ No newline at end of file
diff --git a/Engine/Source/Runtime/Engine/Private/WorldPartition/RuntimeSpatialHash/RuntimeSpatialHashHLOD.cpp b/Engine/Source/Runtime/Engine/Private/WorldPartition/RuntimeSpatialHash/RuntimeSpatialHashHLOD.cpp
--- a/Engine/Source/Runtime/Engine/Private/WorldPartition/RuntimeSpatialHash/RuntimeSpatialHashHLOD.cpp
+++ b/Engine/Source/Runtime/Engine/Private/WorldPartition/RuntimeSpatialHash/RuntimeSpatialHashHLOD.cpp
@@ -28,6 +28,8 @@
 #include "Misc/ScopedSlowTask.h"
 #include "Modules/ModuleManager.h"
 
+#define LOCTEXT_NAMESPACE "RuntimeSpatialHashHLOD"
+
 DEFINE_LOG_CATEGORY_STATIC(LogWorldPartitionRuntimeSpatialHashHLOD, Log, All);
 
 static void SavePackage(UPackage* Package, ISourceControlHelper* SourceControlHelper)
@@ -150,9 +152,9 @@ static TArray<FGuid> GenerateHLODActorsForGrid(UWorldPartition* WorldPartition,
 		}
 	});
 
-	UE_LOG(LogWorldPartitionRuntimeSpatialHashHLOD, Verbose, TEXT("Building HLODs for grid %s..."), *RuntimeGrid.GridName.ToString());
+	UE_LOG(LogWorldPartitionRuntimeSpatialHashHLOD, Verbose, TEXT("%s %s..."), *LOCTEXT("BuildingHLODsForGrid", "Building HLODs for grid").ToString(), *RuntimeGrid.GridName.ToString());
 
-	FScopedSlowTask SlowTask(NbCellsToProcess, FText::FromString(FString::Printf(TEXT("Building HLODs for grid %s..."), *RuntimeGrid.GridName.ToString())));
+	FScopedSlowTask SlowTask(NbCellsToProcess, FText::FromString(FString::Printf(TEXT("%s %s..."), *LOCTEXT("BuildingHLODsForGrid", "Building HLODs for grid").ToString(), *RuntimeGrid.GridName.ToString())));
 	SlowTask.MakeDialog();
 
 	TArray<FGuid> GridHLODActors;
@@ -179,11 +181,11 @@ static TArray<FGuid> GenerateHLODActorsForGrid(UWorldPartition* WorldPartition,
 				const FGuid CellGuid = UWorldPartitionRuntimeSpatialHash::GetCellGuid(RuntimeGrid.GridName, RuntimeGrid.CellSize, CellGlobalCoord, GridCellDataChunk.GetDataLayersID(), GridCellDataChunk.GetContentBundleID());
 				const FString CellName = UWorldPartitionRuntimeSpatialHash::GetCellNameString(OuterWorld, RuntimeGrid.GridName, CellGlobalCoord, GridCellDataChunk.GetDataLayersID(), GridCellDataChunk.GetContentBundleID());
 
-				UE_LOG(LogWorldPartitionRuntimeSpatialHashHLOD, Verbose, TEXT("Creating HLOD for cell %s at %s..."), *CellName, *CellCoord.ToString());
+				UE_LOG(LogWorldPartitionRuntimeSpatialHashHLOD, Verbose, TEXT("%s %s -> %s..."), *CellName, *LOCTEXT("CreatingHLODCell", "Creating HLOD for cell").ToString() , *CellCoord.ToString());
 
 				TRACE_CPUPROFILER_EVENT_SCOPE_TEXT(*CellCoord.ToString());
 
-				UE_LOG(LogWorldPartitionRuntimeSpatialHashHLOD, Display, TEXT("[%d / %d] Processing cell %s..."), (int32)SlowTask.CompletedWork + 1, (int32)SlowTask.TotalAmountOfWork, *CellName);
+				UE_LOG(LogWorldPartitionRuntimeSpatialHashHLOD, Display, TEXT("[%d / %d] %s %s..."), (int32)SlowTask.CompletedWork + 1, (int32)SlowTask.TotalAmountOfWork, *LOCTEXT("ProcessingCell", "Processing cell").ToString(), *CellName);
 
 				// We know the cell's 2D bound but must figure out the bounds in Z
 				FBox CellBounds = FBox(FVector(CellBounds2D.Min, MAX_dbl), FVector(CellBounds2D.Max, -MAX_dbl));
@@ -673,3 +675,5 @@ void UWorldPartitionRuntimeSpatialHash::PostSetupHLODActors(const UWorldPartitio
 }
 
 #endif // #if WITH_EDITOR
+
+#undef LOCTEXT_NAMESPACE
