[AutoRTFM] Update LightweightInstanceSubsystem for transactional safety.

A RWLock has been replaced with a safe version of the lock, and a lock that was used for "do-once" purposes was replaced with a static initializer (where the thread-safety is handled natively by the compiler/AutoRTFM).

This system is technically obsolete code and should probably be deleted entirely: https://epicgames.slack.com/archives/C04325JU9RA/p1743427253088319?thread_ts=1743187160.398099&cid=C04325JU9RA

diff --git a/Engine/Source/Runtime/Engine/Classes/GameFramework/LightWeightInstanceSubsystem.h b/Engine/Source/Runtime/Engine/Classes/GameFramework/LightWeightInstanceSubsystem.h
--- a/Engine/Source/Runtime/Engine/Classes/GameFramework/LightWeightInstanceSubsystem.h
+++ b/Engine/Source/Runtime/Engine/Classes/GameFramework/LightWeightInstanceSubsystem.h
@@ -4,6 +4,7 @@
 
 #include "UObject/ObjectMacros.h"
 #include "GameFramework/LightWeightInstanceManager.h"
+#include "Misc/TransactionallySafeRWLock.h"
 #include "Templates/SharedPointer.h"
 
 
@@ -13,19 +14,13 @@ class FAutoConsoleVariableRef;
 
 struct FLightWeightInstanceSubsystem
 {
-	static ENGINE_API FCriticalSection GetFunctionCS;
-
 	static FLightWeightInstanceSubsystem& Get()
 	{
-		if (!LWISubsystem)
+		static TSharedPtr<FLightWeightInstanceSubsystem, ESPMode::ThreadSafe> LWISubsystemSingleton = []
 		{
-			FScopeLock Lock(&GetFunctionCS);
-			if (!LWISubsystem)
-			{
-				LWISubsystem = MakeShareable(new FLightWeightInstanceSubsystem());
-			}
-		}
-		return *LWISubsystem;
+			return MakeShareable(new FLightWeightInstanceSubsystem());
+		}();
+		return *LWISubsystemSingleton;
 	}
 
 	ENGINE_API FLightWeightInstanceSubsystem();
@@ -90,14 +85,11 @@ protected:
 	ENGINE_API const ALightWeightInstanceManager* GetManagerAt(int32 Index) const;
 
 private:
-	/** Application singleton */
-	static ENGINE_API TSharedPtr<FLightWeightInstanceSubsystem, ESPMode::ThreadSafe> LWISubsystem;
-
 	// TODO: preallocate the size of this based on a config variable
 	TArray<ALightWeightInstanceManager*> LWInstanceManagers;
 
 	// Mutex to make sure we don't change the LWInstanceManagers array while reading/writing it.
-	mutable FRWLock LWIManagersRWLock;
+	mutable FTransactionallySafeRWLock LWIManagersRWLock;
 
 #ifdef WITH_EDITOR
 private:
diff --git a/Engine/Source/Runtime/Engine/Private/LightWeightInstanceSubsystem.cpp b/Engine/Source/Runtime/Engine/Private/LightWeightInstanceSubsystem.cpp
--- a/Engine/Source/Runtime/Engine/Private/LightWeightInstanceSubsystem.cpp
+++ b/Engine/Source/Runtime/Engine/Private/LightWeightInstanceSubsystem.cpp
@@ -8,9 +8,6 @@
 
 DEFINE_LOG_CATEGORY(LogLightWeightInstance);
 
-TSharedPtr<FLightWeightInstanceSubsystem> FLightWeightInstanceSubsystem::LWISubsystem;
-FCriticalSection FLightWeightInstanceSubsystem::GetFunctionCS;
-
 FLightWeightInstanceSubsystem::FLightWeightInstanceSubsystem()
 {
 #if WITH_EDITOR
@@ -47,7 +44,7 @@ FLightWeightInstanceSubsystem::~FLightWeightInstanceSubsystem()
 
 int32 FLightWeightInstanceSubsystem::GetManagerIndex(const ALightWeightInstanceManager* Manager) const
 {
-	FReadScopeLock Lock(LWIManagersRWLock);
+	UE::TReadScopeLock Lock(LWIManagersRWLock);
 	for (int32 Idx = 0; Idx < LWInstanceManagers.Num(); ++Idx)
 	{
 		if (LWInstanceManagers[Idx] == Manager)
@@ -61,19 +58,19 @@ int32 FLightWeightInstanceSubsystem::GetManagerIndex(const ALightWeightInstanceM
 
 const ALightWeightInstanceManager* FLightWeightInstanceSubsystem::GetManagerAt(int32 Index) const
 {
-	FReadScopeLock Lock(LWIManagersRWLock);
+	UE::TReadScopeLock Lock(LWIManagersRWLock);
 	return LWInstanceManagers.IsValidIndex(Index) ? LWInstanceManagers[Index] : nullptr;
 }
 
 bool FLightWeightInstanceSubsystem::AddManager(ALightWeightInstanceManager* Manager)
 {
-	FWriteScopeLock Lock(LWIManagersRWLock);
+	UE::TWriteScopeLock Lock(LWIManagersRWLock);
 	return Get().LWInstanceManagers.Add(Manager) ? true : false;
 }
 
 bool FLightWeightInstanceSubsystem::RemoveManager(ALightWeightInstanceManager* Manager)
 {
-	FWriteScopeLock Lock(LWIManagersRWLock);
+	UE::TWriteScopeLock Lock(LWIManagersRWLock);
 	return Get().LWInstanceManagers.Remove(Manager) ? true : false;
 }
 
@@ -89,7 +86,7 @@ ALightWeightInstanceManager* FLightWeightInstanceSubsystem::FindLightWeightInsta
 		const UClass* CachedActorClass = CachedActor->GetClass();
 		CA_ASSUME(CachedActorClass);
 
-		FReadScopeLock Lock(LWIManagersRWLock);
+		UE::TReadScopeLock Lock(LWIManagersRWLock);
 		// see if we already have a match
 		for (ALightWeightInstanceManager* LWInstance : LWInstanceManagers)
 		{
@@ -136,7 +133,7 @@ ALightWeightInstanceManager* FLightWeightInstanceSubsystem::FindOrAddLightWeight
 
 ALightWeightInstanceManager* FLightWeightInstanceSubsystem::FindLightWeightInstanceManager(UClass& ActorClass, UWorld& World, const FVector& InPos, const UDataLayerInstance* DataLayer) const
 {
-	FReadScopeLock Lock(LWIManagersRWLock);
+	UE::TReadScopeLock Lock(LWIManagersRWLock);
 
 	for (ALightWeightInstanceManager* InstanceManager : LWInstanceManagers)
 	{
diff --git a/Engine/Source/Runtime/Engine/Private/Tests/AutoRTFM/AutoRTFMActorComponentTests.cpp b/Engine/Source/Runtime/Engine/Private/Tests/AutoRTFM/AutoRTFMActorComponentTests.cpp
--- a/Engine/Source/Runtime/Engine/Private/Tests/AutoRTFM/AutoRTFMActorComponentTests.cpp
+++ b/Engine/Source/Runtime/Engine/Private/Tests/AutoRTFM/AutoRTFMActorComponentTests.cpp
@@ -17,6 +17,7 @@
 #include "Engine/LevelStreamingPersistent.h"
 #include "Engine/SkeletalMesh.h"
 #include "Engine/World.h"
+#include "GameFramework/LightWeightInstanceSubsystem.h"
 #include "PBDRigidsSolver.h"
 #include "Physics/Experimental/PhysScene_Chaos.h"
 #include "PhysicsProxy/SingleParticlePhysicsProxy.h"
@@ -638,6 +639,19 @@ AUTORTFM_ACTOR_COMPONENT_TEST(ForEachObjectOfClass_Reconstructed)
 	TestTrueExpr(CountAAutoRTFMTestActors() == InitialCount + 1u);
 }
 
+AUTORTFM_ACTOR_COMPONENT_TEST(LightWeightInstanceSubsystemTest)
+{
+	AutoRTFM::Testing::Commit([&]
+	{
+		FLightWeightInstanceSubsystem& LWI = FLightWeightInstanceSubsystem::Get();
+
+		FLWIData InitData = {};
+		FActorInstanceHandle Handle = LWI.CreateNewLightWeightInstance(Actor->GetClass(), &InitData, nullptr, World);
+
+		LWI.DeleteInstance(Handle);
+	});
+}
+
 }  // anonymous namespace
 
 #undef AUTORTFM_ACTOR_COMPONENT_TEST
