[AssetManager] Fix state transition. HasLoadComplete can be set to true even if CompleteCallback was not called yet.

#rnx
#tests editor tests

diff --git a/Engine/Source/Runtime/Engine/Private/AssetManager.cpp b/Engine/Source/Runtime/Engine/Private/AssetManager.cpp
--- a/Engine/Source/Runtime/Engine/Private/AssetManager.cpp
+++ b/Engine/Source/Runtime/Engine/Private/AssetManager.cpp
@@ -1955,15 +1955,19 @@ TSharedPtr<FStreamableHandle> UAssetManager::ChangeBundleStateForPrimaryAssets(
 				{
 					AssetListLoadParams.DownloadParams = CopyTemp(*LoadParams.DownloadParams);
 				}
-				AssetListLoadParams.OnComplete.BindUObject(this, &UAssetManager::OnAssetStateChangeCompleted, PrimaryAssetId);
-				AssetListLoadParams.OnCancel.BindUObject(this, &UAssetManager::OnAssetStateChangeCompleted, PrimaryAssetId);
-	
+
 				NewHandle = LoadAssetList(MoveTemp(PathsToLoad), MoveTemp(AssetListLoadParams), FString(DebugName), CopyTemp(Location));
-				// save handle so possible duplicate requests will not override this one
-				// set only if loading is not completed yet because completion callback might already reset the PendingState
 				if (NewHandle && !NewHandle->HasLoadCompleted())
 				{
+					// save handle so possible duplicate requests will not override this one
+					// set only if loading is not completed yet because completion callback might already reset the PendingState
 					NameData->PendingState.Handle = NewHandle;
+					NewHandle->BindCompleteDelegate(FStreamableDelegateWithHandle::CreateUObject(this, &UAssetManager::OnAssetStateChangeCompleted, PrimaryAssetId));
+					NewHandle->BindCancelDelegate(FStreamableDelegateWithHandle::CreateUObject(this, &UAssetManager::OnAssetStateChangeCompleted, PrimaryAssetId));
+				}
+				else
+				{
+					OnAssetStateChangeCompleted(NewHandle, PrimaryAssetId);
 				}
 			}
 
