Fixed unbuilt ray tracing passthrough uniform buffers. Optimized the path to skip building the buffer if nothing has changed.
- Keep current frame 0 for ray tracing to avoid invalidating the UB.
- Skip updating stream slots if nothing changed.
- Avoid cache miss hitting the source vertex factory for the tangent SRV.
- Skip emitting granular events if draw events are disabled. Draw events account for half of the cost.

diff --git a/Engine/Source/Runtime/Engine/Private/GPUSkinCache.cpp b/Engine/Source/Runtime/Engine/Private/GPUSkinCache.cpp
--- a/Engine/Source/Runtime/Engine/Private/GPUSkinCache.cpp
+++ b/Engine/Source/Runtime/Engine/Private/GPUSkinCache.cpp
@@ -677,13 +677,13 @@ public:
 			return IntermediateAccumulatedTangentBuffer;
 		}
 
-		void UpdateVertexFactoryDeclaration(FRHICommandListBase& RHICmdList)
+		void UpdateVertexFactoryDeclaration(FRHICommandListBase& RHICmdList, EGPUSkinCacheEntryMode InMode)
 		{
 			FRHIShaderResourceView* CurrentPositionSRV = PositionBuffer->Buffer.SRV;
 			FRHIShaderResourceView* PreviousPositionSRV = PreviousPositionBuffer ? PreviousPositionBuffer->Buffer.SRV.GetReference() : CurrentPositionSRV;
 
 			FGPUSkinPassthroughVertexFactory::FAddVertexAttributeDesc Desc;
-			Desc.FrameNumber = SourceVertexFactory->GetShaderData().UpdatedFrameNumber;
+			Desc.FrameNumber = InMode == EGPUSkinCacheEntryMode::Raster ? SourceVertexFactory->GetShaderData().UpdatedFrameNumber : 0;
 			Desc.StreamBuffers[FGPUSkinPassthroughVertexFactory::EVertexAttribute::VertexPosition] = PositionBuffer->Buffer.Buffer;
 			Desc.StreamBuffers[FGPUSkinPassthroughVertexFactory::EVertexAttribute::VertexTangent] = TangentBuffer->Buffer.Buffer;
 			Desc.SRVs[FGPUSkinPassthroughVertexFactory::EShaderResource::Position] = CurrentPositionSRV;
@@ -695,11 +695,7 @@ public:
 
 	void UpdateVertexFactoryDeclaration(FRHICommandListBase& RHICmdList, int32 Section)
 	{
-		// Inline ray tracing can skip building this uniform buffer for ray tracing entries.
-		if (GRHISupportsRayTracingShaders || Mode == EGPUSkinCacheEntryMode::Raster)
-		{
-			DispatchData[Section].UpdateVertexFactoryDeclaration(RHICmdList);
-		}
+		DispatchData[Section].UpdateVertexFactoryDeclaration(RHICmdList, Mode);
 	}
 
 	inline FCachedGeometry::Section GetCachedGeometry(int32 SectionIndex) const
@@ -1372,7 +1368,7 @@ void FGPUSkinCache::DispatchUpdateSkinTangentsVertexPass(FRHICommandList& RHICmd
 {
 	FGPUSkinCacheEntry::FSectionDispatchData& DispatchData = Entry->DispatchData[SectionIndex];
 
-	SCOPED_DRAW_EVENTF(RHICmdList, SkinTangents_PerVertexPass, TEXT("%sTangentsVertex Mesh=%s, LOD=%d, Chunk=%d, InputStreamStart=%d, OutputStreamStart=%d, Vert=%d")
+	SCOPED_CONDITIONAL_DRAW_EVENTF(RHICmdList, SkinTangents_PerVertexPass, GetEmitDrawEvents(), TEXT("%sTangentsVertex Mesh=%s, LOD=%d, Chunk=%d, InputStreamStart=%d, OutputStreamStart=%d, Vert=%d")
 		, RHI_BREADCRUMB_FORCE_STRING_LITERAL(Entry->Mode == EGPUSkinCacheEntryMode::RayTracing ? TEXT("[RT]") : TEXT(""))
 		, GetSkeletalMeshObjectDebugName(Entry->GPUSkin)
 		, Entry->LOD
@@ -1477,7 +1473,7 @@ void FGPUSkinCache::DispatchUpdateSkinTangentsTrianglePass(FRHICommandList& RHIC
 	uint32 NumTriangles = DispatchData.NumTriangles;
 	uint32 ThreadGroupCountValue = FMath::DivideAndRoundUp(NumTriangles, FBaseRecomputeTangentsPerTriangleShader::ThreadGroupSizeX);
 
-	SCOPED_DRAW_EVENTF(RHICmdList, SkinTangents_PerTrianglePass, TEXT("%sTangentsTri  Mesh=%s, LOD=%d, Chunk=%d, IndexStart=%d Tri=%d BoneInfluenceType=%d UVPrecision=%d")
+	SCOPED_CONDITIONAL_DRAW_EVENTF(RHICmdList, SkinTangents_PerTrianglePass, GetEmitDrawEvents(), TEXT("%sTangentsTri  Mesh=%s, LOD=%d, Chunk=%d, IndexStart=%d Tri=%d BoneInfluenceType=%d UVPrecision=%d")
 		, RHI_BREADCRUMB_FORCE_STRING_LITERAL(Entry->Mode == EGPUSkinCacheEntryMode::RayTracing ? TEXT("[RT]") : TEXT(""))
 		, GetSkeletalMeshObjectDebugName(Entry->GPUSkin)
 		, LODIndex
@@ -1976,7 +1972,7 @@ void FGPUSkinCache::DispatchPassExecute(FTaskData& TaskData)
 			FGPUSkinCacheEntry::FSectionDispatchData& DispatchData = Entry->DispatchData[DispatchEntry.Section];
 			FGPUBaseSkinVertexFactory::FShaderDataType& ShaderData = DispatchData.SourceVertexFactory->GetShaderData();
 
-			SCOPED_DRAW_EVENTF(RHICmdList, SkinCacheDispatch, TEXT("%sSkinning%d%d%d%d Mesh=%s LOD=%d Chunk=%d InStreamStart=%d OutStart=%d Vert=%d Morph=%d/%d")
+			SCOPED_CONDITIONAL_DRAW_EVENTF(RHICmdList, SkinCacheDispatch, GetEmitDrawEvents(), TEXT("%sSkinning%d%d%d%d Mesh=%s LOD=%d Chunk=%d InStreamStart=%d OutStart=%d Vert=%d Morph=%d/%d")
 				, RHI_BREADCRUMB_FORCE_STRING_LITERAL(Entry->Mode == EGPUSkinCacheEntryMode::RayTracing ? TEXT("[RT]") : TEXT(""))
 				, (int32)Entry->bUse16BitBoneIndex
 				, (int32)Entry->bUse16BitBoneWeight
diff --git a/Engine/Source/Runtime/Engine/Private/GPUSkinVertexFactory.cpp b/Engine/Source/Runtime/Engine/Private/GPUSkinVertexFactory.cpp
--- a/Engine/Source/Runtime/Engine/Private/GPUSkinVertexFactory.cpp
+++ b/Engine/Source/Runtime/Engine/Private/GPUSkinVertexFactory.cpp
@@ -1418,6 +1418,7 @@ void FGPUSkinPassthroughVertexFactory::InitRHI(FRHICommandListBase& RHICmdList)
 	if (Data.TangentsSRV)
 	{
 		TangentFormat = Data.TangentsSRV->GetDesc().Common.Format;
+		PreSkinnedTangentSRV = Data.TangentsSRV;
 	}
 
 	FLocalVertexFactory::InitRHI(RHICmdList);
@@ -1449,7 +1450,7 @@ void FGPUSkinPassthroughVertexFactory::UpdateLooseUniformBuffer(FRHICommandListB
 	Parameters.FrameNumber = InFrameNumber;
 	Parameters.PositionBuffer = PositionSRV;
 	Parameters.PreviousPositionBuffer = PrevPositionSRV;
-	Parameters.PreSkinnedTangentBuffer = InSourceVertexFactory->GetTangentsSRV();
+	Parameters.PreSkinnedTangentBuffer = PreSkinnedTangentSRV;
 	LooseParametersUniformBuffer.UpdateUniformBufferImmediate(RHICmdList, Parameters);
 }
 
@@ -1477,11 +1478,14 @@ void FGPUSkinPassthroughVertexFactory::SetVertexAttributes(FRHICommandListBase&
 		}
 	}
 
-	for (int32 Index = 0; Index < EVertexAttribute::NumAttributes; ++Index)
+	if (bNeedUniformBufferUpdate || bNeedLooseUniformBufferUpdate)
 	{
-		if (FRHIStreamSourceSlot* Slot = StreamSourceSlots[Index])
+		for (int32 Index = 0; Index < EVertexAttribute::NumAttributes; ++Index)
 		{
-			RHICmdList.UpdateStreamSourceSlot(Slot, InDesc.StreamBuffers[Index] ? InDesc.StreamBuffers[Index] : SourceStreamBuffers[Index]);
+			if (FRHIStreamSourceSlot* Slot = StreamSourceSlots[Index])
+			{
+				RHICmdList.UpdateStreamSourceSlot(Slot, InDesc.StreamBuffers[Index] ? InDesc.StreamBuffers[Index] : SourceStreamBuffers[Index]);
+			}
 		}
 	}
 
diff --git a/Engine/Source/Runtime/Engine/Public/GPUSkinVertexFactory.h b/Engine/Source/Runtime/Engine/Public/GPUSkinVertexFactory.h
--- a/Engine/Source/Runtime/Engine/Public/GPUSkinVertexFactory.h
+++ b/Engine/Source/Runtime/Engine/Public/GPUSkinVertexFactory.h
@@ -822,6 +822,7 @@ private:
 	TStaticArray<FRHIBuffer*, EVertexAttribute::NumAttributes> SourceStreamBuffers{ InPlace, nullptr };
 	TStaticArray<TRefCountPtr<FRHIStreamSourceSlot>, EVertexAttribute::NumAttributes> StreamSourceSlots;
 	TStaticArray<FRHIShaderResourceView*, EShaderResource::NumShaderResources> SRVs{ InPlace, nullptr };
+	FRHIShaderResourceView* PreSkinnedTangentSRV = nullptr;
 	uint32 UpdatedFrameNumber = ~0u;
 	EVertexAttributeFlags VertexAttributesRequested;
 	EVertexAttributeFlags VertexAttributesToBind = EVertexAttributeFlags::None;
