Allow SkeletalMeshComponent SetEnableAnimation to immediately fire during owning actor construction scripts.

diff --git a/Engine/Source/Runtime/Engine/Private/Components/SkeletalMeshComponent.cpp b/Engine/Source/Runtime/Engine/Private/Components/SkeletalMeshComponent.cpp
--- a/Engine/Source/Runtime/Engine/Private/Components/SkeletalMeshComponent.cpp
+++ b/Engine/Source/Runtime/Engine/Private/Components/SkeletalMeshComponent.cpp
@@ -3470,8 +3470,10 @@ void USkeletalMeshComponent::SetAnimInstanceClass(class UClass* NewClass)
 
 void USkeletalMeshComponent::SetEnableAnimation(bool bNewEnableAnimation)
 {
-	// If we're not registered yet, we can set the value straight away.
-	if (IsRegistered() == false)
+	// If we're not registered yet or our parent actor is running construction scripts, we can set the value straight away.
+	AActor* OwningActor = GetOwner();
+	const bool bRunningConstructionScript = OwningActor && OwningActor->IsRunningUserConstructionScript();
+	if (IsRegistered() == false || bRunningConstructionScript)
 	{
 		bEnableAnimation = bNewEnableAnimation;
 		PrimaryComponentTick.bRunOnAnyThread = !bEnableAnimation;
