Fix ensure when closing editor preview windows with the TaskSyncManager enabled, it was using the world context index in an unsafe way

diff --git a/Engine/Source/Runtime/Engine/Private/TaskSyncManager.cpp b/Engine/Source/Runtime/Engine/Private/TaskSyncManager.cpp
--- a/Engine/Source/Runtime/Engine/Private/TaskSyncManager.cpp
+++ b/Engine/Source/Runtime/Engine/Private/TaskSyncManager.cpp
@@ -61,7 +61,8 @@ FWorldContextId FWorldContextId::GetContextIdForWorld(const UWorld* World)
 				// Behavior is currently unclear for preview/test worlds, some are ticked separately and some are not with no explicit setting
 				else if (Context.WorldType == EWorldType::Game || Context.WorldType == EWorldType::GamePreview || Context.WorldType == EWorldType::EditorPreview)
 				{
-					return FWorldContextId(CurrentWorldIndex + TestWorldStart);
+					// Use the context handle number as this always increases for a new context
+					return FWorldContextId(TestWorldStart + Context.ContextHandle.GetNumber());
 				}
 				else
 				{
