Skeletal Mesh: Remove the async flag on the preallocated mesh description, as well as the bulk data, to avoid a GC hold.

#rnx

diff --git a/Engine/Source/Runtime/Engine/Classes/Engine/SkeletalMeshSourceModel.h b/Engine/Source/Runtime/Engine/Classes/Engine/SkeletalMeshSourceModel.h
--- a/Engine/Source/Runtime/Engine/Classes/Engine/SkeletalMeshSourceModel.h
+++ b/Engine/Source/Runtime/Engine/Classes/Engine/SkeletalMeshSourceModel.h
@@ -18,6 +18,8 @@ class USkeletalMeshDescriptionBulkData : public UMeshDescriptionBaseBulkData
 	GENERATED_BODY()
 public:
 	ENGINE_API USkeletalMeshDescriptionBulkData();
+
+	void Initialize();
 };
 
 
diff --git a/Engine/Source/Runtime/Engine/Private/SkeletalMeshSourceModel.cpp b/Engine/Source/Runtime/Engine/Private/SkeletalMeshSourceModel.cpp
--- a/Engine/Source/Runtime/Engine/Private/SkeletalMeshSourceModel.cpp
+++ b/Engine/Source/Runtime/Engine/Private/SkeletalMeshSourceModel.cpp
@@ -21,17 +21,31 @@ USkeletalMeshDescriptionBulkData::USkeletalMeshDescriptionBulkData()
 #endif
 }
 
+
+void USkeletalMeshDescriptionBulkData::Initialize()
+{
+	// If the object was created on a non-game thread, clear the async flag immediately, on both self, and the owned subobject, 
+	// so that it can be garbage collected in the future. 
+	(void)AtomicallyClearInternalFlags(EInternalObjectFlags::Async);
+#if WITH_EDITOR	
+	(void)PreallocatedMeshDescription->AtomicallyClearInternalFlags(EInternalObjectFlags::Async);
+#endif
+}
+
+
 FSkeletalMeshSourceModel::FSkeletalMeshSourceModel() :
 	Bounds(ForceInitToZero)
 {
 	
 }
 
+
 FSkeletalMeshSourceModel::FSkeletalMeshSourceModel(FSkeletalMeshSourceModel&& InSource)
 {
 	*this = MoveTemp(InSource);
 }
 
+
 FSkeletalMeshSourceModel& FSkeletalMeshSourceModel::operator=(FSkeletalMeshSourceModel&& InSource)
 {
 #if WITH_EDITOR
@@ -60,10 +74,7 @@ void FSkeletalMeshSourceModel::Initialize(USkeletalMesh* InOwner)
 		// so use a GCScopeGuard to prevent it happening at the same time.
 		FGCScopeGuard Guard;
 		MeshDescriptionBulkData = NewObject<USkeletalMeshDescriptionBulkData>(InOwner, NAME_None, RF_Transactional);
-
-		// If the object was created on a non-game thread, clear the async flag immediately, so that it can be
-		// garbage collected in the future. 
-		(void)MeshDescriptionBulkData->AtomicallyClearInternalFlags(EInternalObjectFlags::Async);
+		MeshDescriptionBulkData->Initialize();
 	}
 
 	// We should not have a mesh description at this point.
