[ChaosVD] Added more details to the Chaos Scene Debug Name.

This is currently used by CVD to show as solver name in the timeline controls and the scene outliner. The problem is that when you are working with multiple servers (like DSTM) from multiple sessions/files, everything named as Server Physics does not help.

Also, if you start a recording in PIE but have asset editors open, they also appear in CVD as Client physics which is not correct. This CL adds the world type as name, which helps to differentiate which solver might be from a Asset Editor.

Long term solution will be trace metadata about each specific solver once instead of relying on its name.

diff --git a/Engine/Source/Runtime/Engine/Private/World.cpp b/Engine/Source/Runtime/Engine/Private/World.cpp
--- a/Engine/Source/Runtime/Engine/Private/World.cpp
+++ b/Engine/Source/Runtime/Engine/Private/World.cpp
@@ -6373,7 +6373,15 @@ bool UWorld::AreActorsInitialized() const
 void UWorld::CreatePhysicsScene(const AWorldSettings* Settings)
 {
 #if CHAOS_DEBUG_NAME
-	const FName PhysicsName = IsNetMode(NM_DedicatedServer) ? TEXT("ServerPhysics") : TEXT("ClientPhysics");
+	TStringBuilder<128> PhysicsSceneNameBuilder;
+ 	PhysicsSceneNameBuilder.Appendf(TEXT("%s | World Type [%s]"), IsNetMode(NM_DedicatedServer) ? TEXT("Server Physics") : TEXT("Client Physics"), LexToString(WorldType));
+ 
+	#if UE_WITH_REMOTE_OBJECT_HANDLE
+ 		PhysicsSceneNameBuilder.Appendf(TEXT(" | Remote Server ID [%u]"), UE::RemoteObject::GetGlobalServerId().GetIdNumber()); 
+	#endif
+
+	const FName PhysicsName = FName(PhysicsSceneNameBuilder);
+ 
 	FPhysScene* NewScene = new FPhysScene(nullptr, PhysicsName);
 #else
 	FPhysScene* NewScene = new FPhysScene(nullptr);
