Fixed FActorPrimitiveColorHandler not properly handling unset ActivateFunc/DeactivateFunc Funcs.

diff --git a/Engine/Source/Runtime/Engine/Private/ActorPrimitiveColorHandler.cpp b/Engine/Source/Runtime/Engine/Private/ActorPrimitiveColorHandler.cpp
--- a/Engine/Source/Runtime/Engine/Private/ActorPrimitiveColorHandler.cpp
+++ b/Engine/Source/Runtime/Engine/Private/ActorPrimitiveColorHandler.cpp
@@ -56,7 +56,7 @@ void FActorPrimitiveColorHandler::UnregisterPrimitiveColorHandler(FName InHandle
 	
 	if (InHandlerName == ActivePrimitiveColorHandlerName)
 	{
-		if (ActivePrimitiveColorHandler)
+		if (ActivePrimitiveColorHandler && ActivePrimitiveColorHandler->DeactivateFunc.IsSet())
 		{
 			ActivePrimitiveColorHandler->DeactivateFunc();
 		}
@@ -70,16 +70,22 @@ void FActorPrimitiveColorHandler::UnregisterPrimitiveColorHandler(FName InHandle
 bool FActorPrimitiveColorHandler::SetActivePrimitiveColorHandler(FName InHandlerName, UWorld* InWorld)
 {
 #if ENABLE_ACTOR_PRIMITIVE_COLOR_HANDLER
-	if (FPrimitiveColorHandler* NewActivePrimitiveColorHandler = Handlers.Find(InHandlerName); NewActivePrimitiveColorHandler && (NewActivePrimitiveColorHandler != ActivePrimitiveColorHandler))
+	FPrimitiveColorHandler* NewActivePrimitiveColorHandler = Handlers.Find(InHandlerName);
+	if (NewActivePrimitiveColorHandler != ActivePrimitiveColorHandler)
 	{
-		if (ActivePrimitiveColorHandler)
+		if (ActivePrimitiveColorHandler && ActivePrimitiveColorHandler->DeactivateFunc.IsSet())
 		{
 			ActivePrimitiveColorHandler->DeactivateFunc();
 		}
 
-		ActivePrimitiveColorHandlerName = InHandlerName;
 		ActivePrimitiveColorHandler = NewActivePrimitiveColorHandler;
-		NewActivePrimitiveColorHandler->ActivateFunc();
+		ActivePrimitiveColorHandlerName = ActivePrimitiveColorHandler ? InHandlerName : NAME_None;
+
+		if (ActivePrimitiveColorHandler && ActivePrimitiveColorHandler->ActivateFunc.IsSet())
+		{
+			ActivePrimitiveColorHandler->ActivateFunc();
+		}
+
 		RefreshPrimitiveColorHandler(InHandlerName, InWorld);
 		return true;
 	}
