[MIR] Integrate "Period" parameter to trigonometric material expressions in new translator.

#rnx

diff --git a/Engine/Source/Runtime/Engine/Private/Materials/MaterialExpressionsIR.cpp b/Engine/Source/Runtime/Engine/Private/Materials/MaterialExpressionsIR.cpp
--- a/Engine/Source/Runtime/Engine/Private/Materials/MaterialExpressionsIR.cpp
+++ b/Engine/Source/Runtime/Engine/Private/Materials/MaterialExpressionsIR.cpp
@@ -527,6 +527,24 @@ static void BuildUnaryArithmeticOperatorFromInput(
 	BuildUnaryArithmeticOperatorFromValue(Emitter, Op, Value);
 }
 
+static void BuildPeriodicTrigonometricOperatorFromInput(
+	MIR::FEmitter& Emitter,
+	MIR::EOperator Op,
+	FExpressionInput* Input,
+	float Period)
+{
+	// Get input after checking it has primitive type.
+	MIR::FValue* Value = Emitter.Input(Input);
+
+	if (Period > 0.0f)
+	{
+		Value = Emitter.EmitOperator(MIR::BO_Multiply, Value, Emitter.EmitConstantFloatScalar(2.0f * (float)UE_PI / Period));
+	}
+
+	BuildUnaryArithmeticOperatorFromValue(Emitter, Op, Value);
+}
+
+
 // TODO: Add the missing expressions
 
 void UMaterialExpressionAbs::Build(MIR::FEmitter& Emitter)
@@ -607,12 +625,12 @@ void UMaterialExpressionArctangent::Build(MIR::FEmitter& Emitter)
 
 void UMaterialExpressionCosine::Build(MIR::FEmitter& Emitter)
 {
-	BuildUnaryArithmeticOperatorFromInput(Emitter, MIR::UO_Cos, &Input);
+	BuildPeriodicTrigonometricOperatorFromInput(Emitter, MIR::UO_Cos, &Input, Period);
 }
 
 void UMaterialExpressionSine::Build(MIR::FEmitter& Emitter)
 {
-	BuildUnaryArithmeticOperatorFromInput(Emitter, MIR::UO_Sin, &Input);
+	BuildPeriodicTrigonometricOperatorFromInput(Emitter, MIR::UO_Sin, &Input, Period);
 }
 
 void UMaterialExpressionSaturate::Build(MIR::FEmitter& Emitter)
@@ -627,7 +645,7 @@ void UMaterialExpressionSquareRoot::Build(MIR::FEmitter& Emitter)
 
 void UMaterialExpressionTangent::Build(MIR::FEmitter& Emitter)
 {
-	BuildUnaryArithmeticOperatorFromInput(Emitter, MIR::UO_Tan, &Input);
+	BuildPeriodicTrigonometricOperatorFromInput(Emitter, MIR::UO_Tan, &Input, Period);
 }
 
 static MIR::FValue* EmitInlineHLSL(const UMaterialExpressionExternalCodeBase& InExternalCodeExpression, MIR::FEmitter& Emitter, int32 InExternalCodeIdentifierIndex, TConstArrayView<MIR::FValue*> InArguments)
