PR #13029: Removed additional multiply from URectLightComponent::ComputeLightBrightness

The extra multiply caused Rect lights configured to be in EV units to have the wrong intensity. In particular, switching units back and forth could cause the light intensity to change.

Since this affects backwards compatibility, introduce a serialization change with additional logic to maintain the old logic so that older scenes can keep their existing look, even in cases where intensity is driven externally. In these cases, creating a new light will be required to get the correct behavior.

[RN] A bug affecting rect lights in EV units has been fixed. For backwards compatibility with existing scenes, the old behavior is maintained until the light is re-created.

diff --git a/Engine/Source/Runtime/Core/Public/UObject/FortniteMainBranchObjectVersions.inl b/Engine/Source/Runtime/Core/Public/UObject/FortniteMainBranchObjectVersions.inl
--- a/Engine/Source/Runtime/Core/Public/UObject/FortniteMainBranchObjectVersions.inl
+++ b/Engine/Source/Runtime/Core/Public/UObject/FortniteMainBranchObjectVersions.inl
@@ -619,5 +619,8 @@ FORTNITE_MAIN_VERSION(UpgradeWidgetBlueprintLegacySequencePlayer, 201)
 // Changed clockwise detection algorithm for PCGSplineDirection node with the correct one, but add a version to not break previous nodes.
 FORTNITE_MAIN_VERSION(PCGSplineDirectionClockwiseFix, 202)
 
+// Rect Lights set in EV units had the wrong intensity (older files need a flag set to keep the old look)
+FORTNITE_MAIN_VERSION(RectLightFixedEVUnitConversion, 203)
+
 // -----<new versions can be added above this line>-------------------------------------------------
 #undef FORTNITE_MAIN_VERSION
\ No newline at end of file
diff --git a/Engine/Source/Runtime/Engine/Classes/Components/RectLightComponent.h b/Engine/Source/Runtime/Engine/Classes/Components/RectLightComponent.h
--- a/Engine/Source/Runtime/Engine/Classes/Components/RectLightComponent.h
+++ b/Engine/Source/Runtime/Engine/Classes/Components/RectLightComponent.h
@@ -69,6 +69,10 @@ class URectLightComponent : public ULocalLightComponent
 	UPROPERTY(EditAnywhere, BlueprintReadOnly, Category = Light, meta = (UIMin = "0.0", UIMax = "1.0"), AdvancedDisplay)
 	FVector2f SourceTextureOffset;
 
+	/** Maintain compatibility with lights created before an inconsistency in the EV lighting unit was fixed */
+	UPROPERTY(meta = (Hidden))
+	bool bLightRequiresBrokenEVMath = false;
+
 	UFUNCTION(BlueprintCallable, Category = "Rendering|Lighting")
 	ENGINE_API void SetSourceTexture(UTexture* NewValue);
 
@@ -103,4 +107,6 @@ public:
 	ENGINE_API virtual void PostEditChangeProperty(FPropertyChangedEvent& PropertyChangedEvent) override;
 #endif // WITH_EDITOR
 	//~ End UObject Interface
+
+	ENGINE_API virtual void Serialize(FArchive& Ar) override;
 };
diff --git a/Engine/Source/Runtime/Engine/Private/Components/RectLightComponent.cpp b/Engine/Source/Runtime/Engine/Private/Components/RectLightComponent.cpp
--- a/Engine/Source/Runtime/Engine/Private/Components/RectLightComponent.cpp
+++ b/Engine/Source/Runtime/Engine/Private/Components/RectLightComponent.cpp
@@ -7,6 +7,7 @@
 #include "Components/RectLightComponent.h"
 #include "SceneInterface.h"
 #include "UObject/ConstructorHelpers.h"
+#include "UObject/FortniteMainBranchObjectVersion.h"
 #include "Engine/Texture2D.h"
 #include "RectLightSceneProxy.h"
 #include "SceneView.h"
@@ -186,7 +187,18 @@ float URectLightComponent::ComputeLightBrightness() const
 	}
 	else if (IntensityUnits == ELightUnits::EV)
 	{
-		LightBrightness *= (100.f * 100.f) * EV100ToLuminance(LightBrightness);
+		if (bLightRequiresBrokenEVMath)
+		{
+			// The code below is a typo, but to preserve legacy content, we need to maintain it so that old scenes
+			// keep working even in cases with blueprint logic, sequencer animations, etc ... which cannot be fixed
+			// trivially via serialization.
+			LightBrightness *= EV100ToLuminance(LightBrightness) * (100.f * 100.f);
+		}
+		else
+		{
+			// This is the correct formula
+			LightBrightness = EV100ToLuminance(LightBrightness) * (100.f * 100.f);
+		}
 	}
 	else
 	{
@@ -285,6 +297,29 @@ void URectLightComponent::PostEditChangeProperty(FPropertyChangedEvent& Property
 }
 #endif // WITH_EDITOR
 
+
+void URectLightComponent::Serialize(FArchive& Ar)
+{
+	Ar.UsingCustomVersion(FFortniteMainBranchObjectVersion::GUID);
+
+	Super::Serialize(Ar);
+
+	if (Ar.IsLoading() && Ar.CustomVer(FFortniteMainBranchObjectVersion::GUID) < FFortniteMainBranchObjectVersion::RectLightFixedEVUnitConversion)
+	{
+		// Before this version, the lights contained a subtly wrong interpretation of EV units (see ComputeLightBrighness() above). To preserve
+		// backwards compatibility, we cannot simply change the intensity here (as it would not address other way the intensity can be set such
+		// as from blueprints, sequencer, etc ...). Instead, make sure that older lights that come in with EV units just apply the old formula.
+		// Limit this fix to lights with units that were explicitly configured to use EV so that older lights will get the correct behavior if
+		// their units are changed later. Technically a light that is saved on disk in one unit and dynamically changed to EV in blueprint code
+		// will be broken, but this seems like a rare enough case and minimizing the number of files that have this workaround boolean set is
+		// preferable.
+		if (IntensityUnits == ELightUnits::EV)
+		{
+			bLightRequiresBrokenEVMath = true;
+		}
+	}
+}
+
 FRectLightSceneProxy::FRectLightSceneProxy(const URectLightComponent* Component)
 	: FLocalLightSceneProxy(Component)
 	, SourceWidth(Component->SourceWidth)
