[ue5] - Fix loss of precision in UKismetMathLibrary FTimeSpawn functions taking or returning float when they should use doubles. Modified from contrib #13055.

#tests compiled and ran game

diff --git a/Engine/Source/Runtime/Engine/Classes/Kismet/KismetMathLibrary.h b/Engine/Source/Runtime/Engine/Classes/Kismet/KismetMathLibrary.h
--- a/Engine/Source/Runtime/Engine/Classes/Kismet/KismetMathLibrary.h
+++ b/Engine/Source/Runtime/Engine/Classes/Kismet/KismetMathLibrary.h
@@ -3501,43 +3501,43 @@ class UKismetMathLibrary : public UBlueprintFunctionLibrary
 
 	/** Returns the total number of days in A */
 	UFUNCTION(BlueprintPure, meta=(DisplayName="Get Total Days"), Category="Math|Timespan")
-	static ENGINE_API float GetTotalDays( FTimespan A );
+	static ENGINE_API double GetTotalDays( FTimespan A );
 
 	/** Returns the total number of hours in A */
 	UFUNCTION(BlueprintPure, meta=(DisplayName="Get Total Hours"), Category="Math|Timespan")
-	static ENGINE_API float GetTotalHours( FTimespan A );
+	static ENGINE_API double GetTotalHours( FTimespan A );
 
 	/** Returns the total number of milliseconds in A */
 	UFUNCTION(BlueprintPure, meta=(DisplayName="Get Total Milliseconds"), Category="Math|Timespan")
-	static ENGINE_API float GetTotalMilliseconds( FTimespan A );
+	static ENGINE_API double GetTotalMilliseconds( FTimespan A );
 
 	/** Returns the total number of minutes in A */
 	UFUNCTION(BlueprintPure, meta=(DisplayName="Get Total Minutes"), Category="Math|Timespan")
-	static ENGINE_API float GetTotalMinutes( FTimespan A );
+	static ENGINE_API double GetTotalMinutes( FTimespan A );
 
 	/** Returns the total number of seconds in A */
 	UFUNCTION(BlueprintPure, meta=(DisplayName="Get Total Seconds"), Category="Math|Timespan")
-	static ENGINE_API float GetTotalSeconds( FTimespan A );
+	static ENGINE_API double GetTotalSeconds( FTimespan A );
 
 	/** Returns a time span that represents the specified number of days */
 	UFUNCTION(BlueprintPure, meta=(DisplayName="From Days"), Category="Math|Timespan")
-	static ENGINE_API FTimespan FromDays( float Days );
+	static ENGINE_API FTimespan FromDays( double Days );
 
 	/** Returns a time span that represents the specified number of hours */
 	UFUNCTION(BlueprintPure, meta=(DisplayName="From Hours"), Category="Math|Timespan")
-	static ENGINE_API FTimespan FromHours( float Hours );
+	static ENGINE_API FTimespan FromHours(double Hours );
 
 	/** Returns a time span that represents the specified number of milliseconds */
 	UFUNCTION(BlueprintPure, meta=(DisplayName="From Milliseconds"), Category="Math|Timespan")
-	static ENGINE_API FTimespan FromMilliseconds( float Milliseconds );
+	static ENGINE_API FTimespan FromMilliseconds(double Milliseconds );
 
 	/** Returns a time span that represents the specified number of minutes */
 	UFUNCTION(BlueprintPure, meta=(DisplayName="From Minutes"), Category="Math|Timespan")
-	static ENGINE_API FTimespan FromMinutes( float Minutes );
+	static ENGINE_API FTimespan FromMinutes(double Minutes );
 
 	/** Returns a time span that represents the specified number of seconds */
 	UFUNCTION(BlueprintPure, meta=(DisplayName="From Seconds"), Category="Math|Timespan")
-	static ENGINE_API FTimespan FromSeconds( float Seconds );
+	static ENGINE_API FTimespan FromSeconds(double Seconds );
 
 	/** Returns the ratio between two time spans (A / B), handles zero values */
 	UFUNCTION(BlueprintPure, meta=(DisplayName="Timespan Ratio"), Category="Math|Timespan")
diff --git a/Engine/Source/Runtime/Engine/Classes/Kismet/KismetMathLibrary.inl b/Engine/Source/Runtime/Engine/Classes/Kismet/KismetMathLibrary.inl
--- a/Engine/Source/Runtime/Engine/Classes/Kismet/KismetMathLibrary.inl
+++ b/Engine/Source/Runtime/Engine/Classes/Kismet/KismetMathLibrary.inl
@@ -2566,33 +2566,33 @@ int32 UKismetMathLibrary::GetSeconds(FTimespan A)
 }
 
 KISMET_MATH_FORCEINLINE
-float UKismetMathLibrary::GetTotalDays(FTimespan A)
+double UKismetMathLibrary::GetTotalDays(FTimespan A)
 {
-	return static_cast<float>(A.GetTotalDays());
+	return A.GetTotalDays();
 }
 
 KISMET_MATH_FORCEINLINE
-float UKismetMathLibrary::GetTotalHours(FTimespan A)
+double UKismetMathLibrary::GetTotalHours(FTimespan A)
 {
-	return static_cast<float>(A.GetTotalHours());
+	return A.GetTotalHours();
 }
 
 KISMET_MATH_FORCEINLINE
-float UKismetMathLibrary::GetTotalMilliseconds(FTimespan A)
+double UKismetMathLibrary::GetTotalMilliseconds(FTimespan A)
 {
-	return static_cast<float>(A.GetTotalMilliseconds());
+	return A.GetTotalMilliseconds();
 }
 
 KISMET_MATH_FORCEINLINE
-float UKismetMathLibrary::GetTotalMinutes(FTimespan A)
+double UKismetMathLibrary::GetTotalMinutes(FTimespan A)
 {
-	return static_cast<float>(A.GetTotalMinutes());
+	return A.GetTotalMinutes();
 }
 
 KISMET_MATH_FORCEINLINE
-float UKismetMathLibrary::GetTotalSeconds(FTimespan A)
+double UKismetMathLibrary::GetTotalSeconds(FTimespan A)
 {
-	return static_cast<float>(A.GetTotalSeconds());
+	return A.GetTotalSeconds();
 }
 
 KISMET_MATH_FORCEINLINE
diff --git a/Engine/Source/Runtime/Engine/Private/KismetMathLibrary.cpp b/Engine/Source/Runtime/Engine/Private/KismetMathLibrary.cpp
--- a/Engine/Source/Runtime/Engine/Private/KismetMathLibrary.cpp
+++ b/Engine/Source/Runtime/Engine/Private/KismetMathLibrary.cpp
@@ -836,7 +836,7 @@ void UKismetMathLibrary::BreakTimespan2(FTimespan InTimespan, int32& Days, int32
 }
 
 
-FTimespan UKismetMathLibrary::FromDays(float Days)
+FTimespan UKismetMathLibrary::FromDays(double Days)
 {
 	if (Days < FTimespan::MinValue().GetTotalDays())
 	{
@@ -857,7 +857,7 @@ FTimespan UKismetMathLibrary::FromDays(float Days)
 }
 
 
-FTimespan UKismetMathLibrary::FromHours(float Hours)
+FTimespan UKismetMathLibrary::FromHours(double Hours)
 {
 	if (Hours < FTimespan::MinValue().GetTotalHours())
 	{
@@ -878,7 +878,7 @@ FTimespan UKismetMathLibrary::FromHours(float Hours)
 }
 
 
-FTimespan UKismetMathLibrary::FromMinutes(float Minutes)
+FTimespan UKismetMathLibrary::FromMinutes(double Minutes)
 {
 	if (Minutes < FTimespan::MinValue().GetTotalMinutes())
 	{
@@ -899,7 +899,7 @@ FTimespan UKismetMathLibrary::FromMinutes(float Minutes)
 }
 
 
-FTimespan UKismetMathLibrary::FromSeconds(float Seconds)
+FTimespan UKismetMathLibrary::FromSeconds(double Seconds)
 {
 	if (Seconds < FTimespan::MinValue().GetTotalSeconds())
 	{
@@ -920,7 +920,7 @@ FTimespan UKismetMathLibrary::FromSeconds(float Seconds)
 }
 
 
-FTimespan UKismetMathLibrary::FromMilliseconds(float Milliseconds)
+FTimespan UKismetMathLibrary::FromMilliseconds(double Milliseconds)
 {
 	if (Milliseconds < FTimespan::MinValue().GetTotalMilliseconds())
 	{
