Introduce a shared time budget for ProcessAsyncLoading and UpdateLevelStreaming that can be enabled with s.UseUnifiedTimeBudgetForStreaming 1.
When this is set, it runs the async asset and level streaming at the end of the frame from HandleUnifiedStreaming which also handles high priority streaming.
UpdateLevelStreaming will have less time if there are hitches in ProcessAsyncLoading, and time unused by UpdateLevelStreaming will be used to process more loaded assets.
This also includes performance and timing fixes for functions like RemoveFromWorld which were not correctly computing the elapsed time in some cases.

diff --git a/Engine/Source/Editor/UnrealEd/Private/EditorEngine.cpp b/Engine/Source/Editor/UnrealEd/Private/EditorEngine.cpp
--- a/Engine/Source/Editor/UnrealEd/Private/EditorEngine.cpp
+++ b/Engine/Source/Editor/UnrealEd/Private/EditorEngine.cpp
@@ -1750,7 +1750,7 @@ void UEditorEngine::Tick( float DeltaSeconds, bool bIdleMode )
 	// Update the timer manager
 	TimerManager->Tick(DeltaSeconds);
 
-	// Update subsystems.
+	if (!GUseUnifiedTimeBudgetForStreaming)
 	{
 		// This assumes that UObject::StaticTick only calls ProcessAsyncLoading.
 		StaticTick(DeltaSeconds, !!GAsyncLoadingUseFullTimeLimit, GAsyncLoadingTimeLimit / 1000.f);
@@ -2121,7 +2121,7 @@ void UEditorEngine::Tick( float DeltaSeconds, bool bIdleMode )
 				UpdateTransitionType(PlayWorld);
 
 				// Update streaming for dedicated servers in PIE
-				if (PieContext.RunAsDedicated)
+				if (PieContext.RunAsDedicated && !GUseUnifiedTimeBudgetForStreaming)
 				{
 					SCOPE_CYCLE_COUNTER(STAT_UpdateLevelStreaming);
 					PlayWorld->UpdateLevelStreaming();
@@ -2360,6 +2360,11 @@ void UEditorEngine::Tick( float DeltaSeconds, bool bIdleMode )
 		}
 	}
 
+	if (GUseUnifiedTimeBudgetForStreaming)
+	{
+		HandleUnifiedStreaming(DeltaSeconds);
+	}
+
 	// Rendering resources are normally flushed when a 3D viewport is drawn. If no viewports are updated (because the editor is hidden, or no realtime viewports are visible),
 	// we need to force-flush resources here, since nothing else will. Note that this condition checks if any *level* viewports have been drawn in the block above; other
 	// editors can contain 3D viewports and refreshing them will also flush resources, but it's difficult to know when 3D rendering has happened in general. Instead, we
@@ -2607,8 +2612,11 @@ bool UEditorEngine::UpdateSingleViewportClient(FEditorViewportClient* InViewport
 			// this ourselves!
 			if( PlayWorld != NULL && bIsSimulatingInEditor && InViewportClient->IsSimulateInEditorViewport() )
 			{
-				// Update level streaming.
-				InViewportClient->GetWorld()->UpdateLevelStreaming();
+				if (!GUseUnifiedTimeBudgetForStreaming)
+				{
+					// Update level streaming.
+					InViewportClient->GetWorld()->UpdateLevelStreaming();
+				}
 
 				// Also make sure hit proxies are refreshed for SIE viewports, as the user may be trying to grab an object or widget manipulator that's moving!
 				if( InViewportClient->IsRealtime() && (GInvalidateHitProxiesEachSIEFrameCVar->GetInt() != 0))
diff --git a/Engine/Source/Runtime/Engine/Classes/Components/ActorComponent.h b/Engine/Source/Runtime/Engine/Classes/Components/ActorComponent.h
--- a/Engine/Source/Runtime/Engine/Classes/Components/ActorComponent.h
+++ b/Engine/Source/Runtime/Engine/Classes/Components/ActorComponent.h
@@ -17,6 +17,7 @@
 #include "ComponentInstanceDataCache.h"
 #include "Experimental/ConcurrentLinearAllocator.h"
 #include "Misc/ScopeRWLock.h"
+#include "Misc/Timeout.h"
 #include "Misc/TransactionallySafeRWLock.h"
 #include "Physics/Experimental/AsyncPhysicsStateProcessorInterface.h"
 #include "ActorComponent.generated.h"
@@ -44,12 +45,12 @@ public:
 	FRegisterComponentContext(UWorld* InWorld);
 	~FRegisterComponentContext();
 
-	FRegisterComponentContext& SetIncrementalUpdateTimeLimitMs(double Value);
+	FRegisterComponentContext& SetIncrementalUpdateTimeout(const UE::FTimeout& InTimeout);
 	FRegisterComponentContext& SetIncrementalUpdateAddPrimitiveBatchSize(int32 Value);
 	FRegisterComponentContext& SetIncrementalUpdateNumComponentsToUpdate(int32 Value);
 	FRegisterComponentContext& SetAsyncRegisterLevelContext(FAsyncRegisterLevelContext* Value);
 
-	double GetIncrementalUpdateTimeLimitMs() const { return IncrementalUpdateTimeLimitMs; }
+	const UE::FTimeout& GetIncrementalUpdateTimeout() const { return Timeout; }
 	int32 GetIncrementalUpdateAddPrimitiveBatchSize() const { return IncrementalUpdateAddPrimitiveBatchSize; }
 	int32 GetIncrementalUpdateNumComponentsToUpdate() const { return IncrementalUpdateNumComponentsToUpdate; }
 
@@ -67,7 +68,7 @@ private:
 	bool OnIncrementalRegisterComponentsDone();
 
 	UWorld* World;
-	double	IncrementalUpdateTimeLimitMs;
+	UE::FTimeout Timeout;
 	int32	IncrementalUpdateAddPrimitiveBatchSize;
 	int32	IncrementalUpdateNumComponentsToUpdate;
 	FAsyncRegisterLevelContext* AsyncRegisterLevelContext;
diff --git a/Engine/Source/Runtime/Engine/Classes/Components/UnregisterComponentContext.h b/Engine/Source/Runtime/Engine/Classes/Components/UnregisterComponentContext.h
--- a/Engine/Source/Runtime/Engine/Classes/Components/UnregisterComponentContext.h
+++ b/Engine/Source/Runtime/Engine/Classes/Components/UnregisterComponentContext.h
@@ -2,6 +2,8 @@
 
 #pragma once
 
+#include "Misc/Timeout.h"
+
 class UWorld;
 
 class FUnregisterComponentContext
@@ -9,21 +11,21 @@ class FUnregisterComponentContext
 public:
 	FUnregisterComponentContext(UWorld* InWorld)
 		: World(InWorld)
-		, IncrementalRemoveTimeLimitMs(0)
+		, Timeout(UE::FTimeout::AlwaysExpired())
 	{}
 
-	FUnregisterComponentContext& SetIncrementalRemoveTimeLimitMs(double Value)
+	FUnregisterComponentContext& SetIncrementalRemoveTimeout(const UE::FTimeout& InTimeout)
 	{
-		IncrementalRemoveTimeLimitMs = Value;
+		Timeout = InTimeout;
 		return *this;
 	}
 	
-	double GetIncrementalRemoveTimeLimitMs() const
+	const UE::FTimeout& GetIncrementalRemoveTimeout() const
 	{
-		return IncrementalRemoveTimeLimitMs;
+		return Timeout;
 	}
 
 private:
 	UWorld* World;
-	double	IncrementalRemoveTimeLimitMs;
+	UE::FTimeout Timeout;
 };
diff --git a/Engine/Source/Runtime/Engine/Classes/Engine/CoreSettings.h b/Engine/Source/Runtime/Engine/Classes/Engine/CoreSettings.h
--- a/Engine/Source/Runtime/Engine/Classes/Engine/CoreSettings.h
+++ b/Engine/Source/Runtime/Engine/Classes/Engine/CoreSettings.h
@@ -121,6 +121,8 @@ protected:
 	//~ End UObject Interface
 };
 
+/** Whether to use unified TimeBudget to manage ProcessAsyncLoading and UpdateLevelStreaming */
+extern ENGINE_API int32 GUseUnifiedTimeBudgetForStreaming;
 /** Whether to allow background level streaming. */
 extern ENGINE_API int32 GUseBackgroundLevelStreaming;
 /** Maximum amount of time to spend doing asynchronous loading (ms per frame). */
diff --git a/Engine/Source/Runtime/Engine/Classes/Engine/Engine.h b/Engine/Source/Runtime/Engine/Classes/Engine/Engine.h
--- a/Engine/Source/Runtime/Engine/Classes/Engine/Engine.h
+++ b/Engine/Source/Runtime/Engine/Classes/Engine/Engine.h
@@ -1672,10 +1672,6 @@ public:
 	UPROPERTY()
 	FColor C_BrushShape;
 
-	/** Fudge factor for tweaking the distance based miplevel determination. No longer used. */
-	UE_DEPRECATED(5.2, "This setting is no longer used.")
-	float StreamingDistanceFactor;
-
 	/** The save directory for newly created screenshots */
 	UPROPERTY(config, EditAnywhere, Category = Screenshots)
 	FDirectoryPath GameScreenshotSaveDirectory;
@@ -2655,6 +2651,20 @@ public:
 		return true;
 	}
 
+	/** 
+	 * Computes the amount of time in seconds that should be used for unified asset/level streaming for this frame.
+	 * This is only used if UseUnifiedTimeBudgetForStreaming is enabled, and can be overridden by a game engine.
+	 * @param DeltaSeconds		How long the last frame tick took
+	 * @param bHighPriority		True if this is considered a high priority load such as during a loading screen
+	 */
+	ENGINE_API virtual double GetUnifiedTimeBudgetForStreaming(float DeltaSeconds, bool bHighPriority);
+
+	/**
+	 * Called once per frame to execute unified async asset and level streaming using the time budget from GetUnifiedTimeBudgetForStreaming.
+	 * This is only used if UseUnifiedTimeBudgetForStreaming is enabled, and can be overridden by a game engine.
+	 */
+	ENGINE_API virtual void HandleUnifiedStreaming(float DeltaSeconds);
+
 	/**
 	 * @return true if level streaming should prefer to stream levels from disk instead of duplicating them from editor world
 	 */
@@ -2855,12 +2865,6 @@ public:
 	 */
 	ENGINE_API void AddTextureStreamingLoc(FVector InLoc, float BoostFactor, bool bOverrideLocation, float OverrideDuration);
 
-	UE_DEPRECATED(5.1, "This is deprecated to follow inclusive naming rules. Use AddTextureStreamingLoc() instead.")
-	void AddTextureStreamingSlaveLoc(FVector InLoc, float BoostFactor, bool bOverrideLocation, float OverrideDuration)
-	{
-		AddTextureStreamingLoc(InLoc, BoostFactor, bOverrideLocation, OverrideDuration);
-	}
-
 	/** 
 	 * Obtain a world object pointer from an object with has a world context.
 	 *
diff --git a/Engine/Source/Runtime/Engine/Classes/Engine/LevelStreaming.h b/Engine/Source/Runtime/Engine/Classes/Engine/LevelStreaming.h
--- a/Engine/Source/Runtime/Engine/Classes/Engine/LevelStreaming.h
+++ b/Engine/Source/Runtime/Engine/Classes/Engine/LevelStreaming.h
@@ -10,6 +10,7 @@
 #include "LatentActions.h"
 #include "ProfilingDebugging/ProfilingHelpers.h"
 #include "GameFramework/UpdateLevelVisibilityLevelInfo.h"
+#include "Misc/Timeout.h"
 
 #if WITH_EDITOR
 #include "Folder.h"
@@ -396,7 +397,7 @@ private:
 	ENGINE_API bool UpdateTargetState();
 
 	/** Update the load process of the streaming level. Out parameters instruct calling code how to proceed. */
-	ENGINE_API void UpdateStreamingState(bool& bOutUpdateAgain, bool& bOutRedetermineTarget);
+	ENGINE_API void UpdateStreamingState(bool& bOutUpdateAgain, bool& bOutRedetermineTarget, const TOptional<const UE::FTimeout>& InExternalTimeout);
 
 	/** Update internal variables when the level is added from the streaming levels array */
 	ENGINE_API void OnLevelAdded();
@@ -791,7 +792,7 @@ private:
 	/** Determine what the streaming levels target state should be. Returns whether the streaming level should be in the consider list. */
 	static bool UpdateTargetState(ULevelStreaming* StreamingLevel) { return StreamingLevel->UpdateTargetState(); }
 	/** Update the load process of the streaming level. Out parameters instruct calling code how to proceed. */
-	static void UpdateStreamingState(ULevelStreaming* StreamingLevel, bool& bOutUpdateAgain, bool& bOutRedetermineTarget) { StreamingLevel->UpdateStreamingState(bOutUpdateAgain, bOutRedetermineTarget); }
+	static void UpdateStreamingState(ULevelStreaming* StreamingLevel, bool& bOutUpdateAgain, bool& bOutRedetermineTarget, const TOptional<const UE::FTimeout>& InExternalTimeout) { StreamingLevel->UpdateStreamingState(bOutUpdateAgain, bOutRedetermineTarget, InExternalTimeout); }
 
 	/** Friend classes to manipulate the streaming level more extensively */
 	friend class UEngine;
diff --git a/Engine/Source/Runtime/Engine/Classes/Engine/World.h b/Engine/Source/Runtime/Engine/Classes/Engine/World.h
--- a/Engine/Source/Runtime/Engine/Classes/Engine/World.h
+++ b/Engine/Source/Runtime/Engine/Classes/Engine/World.h
@@ -2993,31 +2993,46 @@ public:
 	void PropagateLightingScenarioChange();
 
 	/**
-	 * Associates the passed in level with the world. The work to make the level visible is spread across several frames and this
-	 * function has to be called till it returns true for the level to be visible/ associated with the world and no longer be in
-	 * a limbo state.
+	 * Associates the passed in level with the world. If using a time limit the work to make the level visible can be spread across several frames
+	 * so this function has to be called multiple time for the level to be visible/ associated with the world and no longer be in a limbo state.
 	 *
 	 * @param Level				Level object we should add
 	 * @param LevelTransform	Transformation to apply to each actor in the level
-	 * @param bConsiderTimeLimie optional bool indicating if we should consider timelimit or not, default is true
+	 * @param bConsiderTimeLimit optional bool indicating if time-slicing will be used to limit the amount of work done, defaults to false
+	 * @param ExternalTimeout optional timeout that will override the default behavior if bConsiderTimeLimit is true
 	 * @param TransactionId optional parameter that carries the current transaction id associated with calls updating LevelVisibility used when communicating level visibility with server
 	 * @param OwningLevelStreaming optional parameter, the ULevelStreaming object driving this level's presence in the world
 	 */
-	void AddToWorld(ULevel* Level, const FTransform& LevelTransform = FTransform::Identity, bool bConsiderTimeLimit = true, FNetLevelVisibilityTransactionId TransactionId = FNetLevelVisibilityTransactionId(), ULevelStreaming* OwningLevelStreaming = nullptr);
+	void AddToWorld(ULevel* Level, const FTransform& LevelTransform = FTransform::Identity, bool bConsiderTimeLimit = false, const TOptional<const UE::FTimeout>& ExternalTimeout = TOptional<const UE::FTimeout>(), FNetLevelVisibilityTransactionId TransactionId = FNetLevelVisibilityTransactionId(), ULevelStreaming* OwningLevelStreaming = nullptr);
+
+	UE_DEPRECATED(5.6, "Use version that requires passing the ExternalTimout, bConsiderTimeLimit now defaults to false to match common usage")
+	void AddToWorld(ULevel* Level, const FTransform& LevelTransform, bool bConsiderTimeLimit, FNetLevelVisibilityTransactionId TransactionId, ULevelStreaming* OwningLevelStreaming = nullptr)
+	{
+		AddToWorld(Level, LevelTransform, bConsiderTimeLimit, TOptional<const UE::FTimeout>(), TransactionId, OwningLevelStreaming);
+	}
 
 	/** 
-	 * Dissociates the passed in level from the world. The removal is blocking.
+	 * Dissociates the passed in level from the world. The removal is blocking unless bAllowIncrementalRemoval is set
 	 *
 	 * @param Level			Level object we should remove
+	 * @param bAllowIncrementalRemoval optional bool indicating if time-slicing will be used to remove objects from scene, defaults to false
+	 * @param ExternalTimeout optional timeout that will override the default behavior if bAllowIncrementalRemoval is true
 	 * @param TransactionId optional parameter that carries the current transaction id associated with calls updating LevelVisibility used when communicating level visibility with server
 	 * @param OwningLevelStreaming optional parameter, the ULevelStreaming object driving this level's presence in the world
 	 */
-	void RemoveFromWorld(ULevel* Level, bool bAllowIncrementalRemoval = false, FNetLevelVisibilityTransactionId TransactionId = FNetLevelVisibilityTransactionId(), ULevelStreaming* OwningLevelStreaming = nullptr);
+	void RemoveFromWorld(ULevel* Level, bool bAllowIncrementalRemoval = false, const TOptional<const UE::FTimeout>& ExternalTimeout = TOptional<const UE::FTimeout>(), FNetLevelVisibilityTransactionId TransactionId = FNetLevelVisibilityTransactionId(), ULevelStreaming* OwningLevelStreaming = nullptr);
 
+	UE_DEPRECATED(5.6, "Use version that requires passing the optional ExternalTimout")
+	void RemoveFromWorld(ULevel* Level, bool bAllowIncrementalRemoval, FNetLevelVisibilityTransactionId TransactionId, ULevelStreaming* OwningLevelStreaming = nullptr)
+	{
+		RemoveFromWorld(Level, bAllowIncrementalRemoval, TOptional<const UE::FTimeout>(), TransactionId, OwningLevelStreaming);
+	}
 	/**
 	 * Updates sub-levels (load/unload/show/hide) using streaming levels current state
+	 *
+	 * @param ExternalTimeout Optional soft limit to the time this function is allowed to take [time-slicing].
 	 */
-	void UpdateLevelStreaming();
+	void UpdateLevelStreaming(const TOptional<const UE::FTimeout>& ExternalTimeout = TOptional<const UE::FTimeout>());
 
 	/** Releases PhysicsScene manually */
 	void ReleasePhysicsScene();
diff --git a/Engine/Source/Runtime/Engine/Private/Actor.cpp b/Engine/Source/Runtime/Engine/Private/Actor.cpp
--- a/Engine/Source/Runtime/Engine/Private/Actor.cpp
+++ b/Engine/Source/Runtime/Engine/Private/Actor.cpp
@@ -85,6 +85,7 @@
 #include "WorldPartition/DataLayer/DataLayerManager.h"
 #include "WorldPartition/ContentBundle/ContentBundlePaths.h"
 #include "WorldPartition/ActorInstanceGuids.h"
+#include "Engine/CoreSettings.h"
 
 DEFINE_LOG_CATEGORY(LogActor);
 
@@ -6090,19 +6091,7 @@ bool AActor::IncrementalRegisterComponents(int32 NumComponentsToRegister, FRegis
 	const int32 ComponentCount = Components.Num();
 	TSet<UActorComponent*> RegisteredParents;
 
-	const double StartTime = FPlatformTime::Seconds();
-	const double TimeLimitMs = Context ? Context->GetIncrementalUpdateTimeLimitMs() : 0;
-	const bool bConsiderTimeLimit = Context && TimeLimitMs >= 0;
-
-	ON_SCOPE_EXIT
-	{
-		if (bConsiderTimeLimit)
-		{
-			const double ElapsedTimeMs = (FPlatformTime::Seconds() - StartTime) * 1000;
-			const double IncrementalUpdateTimeLimitMs = bConsiderTimeLimit ? FMath::Max<double>(0.0, TimeLimitMs - ElapsedTimeMs) : -1;
-			Context->SetIncrementalUpdateTimeLimitMs(IncrementalUpdateTimeLimitMs);
-		}
-	};
+	const UE::FTimeout* Timeout = Context ? &Context->GetIncrementalUpdateTimeout() : nullptr;
 
 	for (int32 CompIdx = 0; CompIdx < ComponentCount; CompIdx++)
 	{
@@ -6143,12 +6132,14 @@ bool AActor::IncrementalRegisterComponents(int32 NumComponentsToRegister, FRegis
 		if ((CompIdx + 1 < ComponentCount) && (NumRegisteredComponentsThisRun >= NumComponentsToRegister))
 		{
 			NumRegisteredComponentsThisRun = 0;
-			if (!bConsiderTimeLimit)
+			if (Timeout == nullptr)
 			{
+				// If there is no time limit, stop at NumComponentsToRegister
 				break;
 			}
-			else if (const double ElapseTimeMs = (FPlatformTime::Seconds() - StartTime) * 1000; ElapseTimeMs > TimeLimitMs)
+			else if (Timeout->IsExpired())
 			{
+				// If there is a time limit, only stop if we're out of time
 				break;
 			}
 		}
diff --git a/Engine/Source/Runtime/Engine/Private/Components/ActorComponent.cpp b/Engine/Source/Runtime/Engine/Private/Components/ActorComponent.cpp
--- a/Engine/Source/Runtime/Engine/Private/Components/ActorComponent.cpp
+++ b/Engine/Source/Runtime/Engine/Private/Components/ActorComponent.cpp
@@ -127,7 +127,7 @@ FAutoConsoleVariableRef CVarPrecachePSOsOnComponentRecreateRenderContext(
 
 FRegisterComponentContext::FRegisterComponentContext(UWorld* InWorld)
 	: World(InWorld)
-	, IncrementalUpdateTimeLimitMs(0)
+	, Timeout(UE::FTimeout::AlwaysExpired())
 	, IncrementalUpdateAddPrimitiveBatchSize(0)
 	, IncrementalUpdateNumComponentsToUpdate(0)
 	, AsyncRegisterLevelContext(nullptr)
@@ -183,9 +183,9 @@ void FRegisterComponentContext::SendRenderDynamicData(FRegisterComponentContext*
 	}
 }
 
-FRegisterComponentContext& FRegisterComponentContext::SetIncrementalUpdateTimeLimitMs(double Value)
+FRegisterComponentContext& FRegisterComponentContext::SetIncrementalUpdateTimeout(const UE::FTimeout& InTimeout)
 {
-	IncrementalUpdateTimeLimitMs = Value;
+	Timeout = InTimeout;
 	return *this;
 }
 
diff --git a/Engine/Source/Runtime/Engine/Private/CoreSettings.cpp b/Engine/Source/Runtime/Engine/Private/CoreSettings.cpp
--- a/Engine/Source/Runtime/Engine/Private/CoreSettings.cpp
+++ b/Engine/Source/Runtime/Engine/Private/CoreSettings.cpp
@@ -9,6 +9,7 @@
 
 DEFINE_LOG_CATEGORY_STATIC(LogCoreSettings, Log, All);
 
+int32 GUseUnifiedTimeBudgetForStreaming = 0;
 int32 GUseBackgroundLevelStreaming = 1;
 float GAsyncLoadingTimeLimit = 5.0f;
 int32 GAsyncLoadingUseFullTimeLimit = 0;
@@ -28,24 +29,33 @@ int32 GLevelStreamingForceVerifyLevelsGotRemovedByGC = 0;
 int32 GLevelStreamingForceRouteActorInitializeNextFrame = 0;
 int32 GLevelStreamingLowMemoryPendingPurgeCount = MAX_int32;
 
+static FAutoConsoleVariableRef CVarUseUnifiedTimeBudgetForStreaming(
+	TEXT("s.UseUnifiedTimeBudgetForStreaming"),
+	GUseUnifiedTimeBudgetForStreaming,
+	TEXT("Whether to use unified time budgets and logic for async loading and level streaming.\n")
+	TEXT("If 0 (current default), cvars like s.AsyncLoadingTimeLimit and s.LevelStreamingActorsUpdateTimeLimit will separately control each phase and can be modified by game-specific code.\n")
+	TEXT("If 1, the time budgets are added together in GetUnifiedTimeBudgetForStreaming, and asset streaming will run immediately before level streaming with a shared budget."),
+	ECVF_Default
+	);
+
 static FAutoConsoleVariableRef CVarUseBackgroundLevelStreaming(
 	TEXT("s.UseBackgroundLevelStreaming"),
 	GUseBackgroundLevelStreaming,
-	TEXT("Whether to allow background level streaming."),
+	TEXT("Whether to allow background level streaming. If this is 0, all level loading will block the game thread until it completes."),
 	ECVF_Default
 	);
 
 static FAutoConsoleVariableRef CVarAsyncLoadingTimeLimit(
 	TEXT("s.AsyncLoadingTimeLimit"),
 	GAsyncLoadingTimeLimit,
-	TEXT("Maximum amount of time to spend doing asynchronous loading (ms per frame)."),
+	TEXT("Maximum amount of time to spend doing asynchronous asset loading (ms per frame)."),
 	ECVF_Default
 	);
 
 static FAutoConsoleVariableRef CVarAsyncLoadingUseFullTimeLimit(
 	TEXT("s.AsyncLoadingUseFullTimeLimit"),
 	GAsyncLoadingUseFullTimeLimit,
-	TEXT("Whether to use the entire time limit even if blocked on I/O."),
+	TEXT("Whether to use the entire time limit for async asset loading even if blocked on I/O."),
 	ECVF_Default
 	);
 
diff --git a/Engine/Source/Runtime/Engine/Private/GameEngine.cpp b/Engine/Source/Runtime/Engine/Private/GameEngine.cpp
--- a/Engine/Source/Runtime/Engine/Private/GameEngine.cpp
+++ b/Engine/Source/Runtime/Engine/Private/GameEngine.cpp
@@ -1821,7 +1821,7 @@ void UGameEngine::Tick( float DeltaSeconds, bool bIdleMode )
 	}
 #endif
 
-	// Update subsystems.
+	if (!GUseUnifiedTimeBudgetForStreaming)
 	{
 		// This assumes that UObject::StaticTick only calls ProcessAsyncLoading.
 		SCOPE_TIME_GUARD(TEXT("UGameEngine::Tick - StaticTick"));
@@ -1918,7 +1918,7 @@ void UGameEngine::Tick( float DeltaSeconds, bool bIdleMode )
 		}
 
 		// streamingServer
-		if( GIsServer == true )
+		if( GIsServer == true && !GUseUnifiedTimeBudgetForStreaming)
 		{
 			SCOPE_CYCLE_COUNTER(STAT_UpdateLevelStreaming);
 			Context.World()->UpdateLevelStreaming();
@@ -1986,14 +1986,22 @@ void UGameEngine::Tick( float DeltaSeconds, bool bIdleMode )
 	}
 
 	const bool bRenderingSuspended = IsRenderingSuspended();
+	const bool bIsRenderingScene = !bIdleMode && !IsRunningDedicatedServer() && !IsRunningCommandlet() && FEmbeddedCommunication::IsAwakeForRendering();
 
-	if (!bIdleMode && !IsRunningDedicatedServer() && !IsRunningCommandlet() && FEmbeddedCommunication::IsAwakeForRendering())
+	if (bIsRenderingScene && !bRenderingSuspended)
+	{
+		RedrawViewports();
+	}
+
+	if (GUseUnifiedTimeBudgetForStreaming)
+	{
+		HandleUnifiedStreaming(DeltaSeconds);
+	}
+
+	if (bIsRenderingScene)
 	{
 		if (!bRenderingSuspended)
 		{
-			// Render everything.
-			RedrawViewports();
-
 			// CPU/GPU synchronization is achieved by calling EndDrawingViewport. If no viewports are updated (because the game is hidden),
 			// we need to explicitly wait for the GPU to finish here, to prevent the CPU from submitting work faster than the GPU
 			// can process it, which leads to an unbounded accumulation of resources.
@@ -2010,9 +2018,10 @@ void UGameEngine::Tick( float DeltaSeconds, bool bIdleMode )
 				GetRendererModule().PostRenderAllViewports();
 			}
 		}
-		else
+		else if (!GUseUnifiedTimeBudgetForStreaming)
 		{
 			// Still need to call UpdateLevelStreaming() even when not rendering
+			// unless we are using the unified streaming budget
 			if (GameViewport && GameViewport->Viewport)
 			{
 				if (FViewportClient* ViewportClient = GameViewport->Viewport->GetClient())
diff --git a/Engine/Source/Runtime/Engine/Private/GameViewportClient.cpp b/Engine/Source/Runtime/Engine/Private/GameViewportClient.cpp
--- a/Engine/Source/Runtime/Engine/Private/GameViewportClient.cpp
+++ b/Engine/Source/Runtime/Engine/Private/GameViewportClient.cpp
@@ -49,6 +49,7 @@
 #include "Logging/MessageLog.h"
 #include "Blueprint/UserWidget.h"
 #include "GameFramework/GameUserSettings.h"
+#include "Engine/CoreSettings.h"
 #include "Engine/UserInterfaceSettings.h"
 #include "Slate/SceneViewport.h"
 #include "Slate/SGameLayerManager.h"
@@ -1785,8 +1786,11 @@ void UGameViewportClient::Draw(FViewport* InViewport, FCanvas* SceneCanvas)
 
 	FinalizeViews(&ViewFamily, PlayerViewMap);
 
-	// Update level streaming.
-	MyWorld->UpdateLevelStreaming();
+	if (!GUseUnifiedTimeBudgetForStreaming)
+	{
+		// Update level streaming.
+		MyWorld->UpdateLevelStreaming();
+	}
 
 	// Find largest rectangle bounded by all rendered views.
 	uint32 MinX=InViewport->GetSizeXY().X, MinY=InViewport->GetSizeXY().Y, MaxX=0, MaxY=0;
diff --git a/Engine/Source/Runtime/Engine/Private/Level.cpp b/Engine/Source/Runtime/Engine/Private/Level.cpp
--- a/Engine/Source/Runtime/Engine/Private/Level.cpp
+++ b/Engine/Source/Runtime/Engine/Private/Level.cpp
@@ -1819,7 +1819,7 @@ void ULevel::IncrementalUpdateComponents(int32 NumComponentsToUpdate, bool bReru
 
 	FRegisterComponentContext LocalContext(OwningWorld);
 	FRegisterComponentContext& Context = InContext ? *InContext : LocalContext;
-	Context.SetIncrementalUpdateTimeLimitMs(InContext ? InContext->GetIncrementalUpdateTimeLimitMs() : 0);
+	Context.SetIncrementalUpdateTimeout(InContext ? InContext->GetIncrementalUpdateTimeout() : UE::FTimeout::AlwaysExpired());
 	Context.SetIncrementalUpdateNumComponentsToUpdate(NumComponentsToUpdate);
 	Context.SetIncrementalUpdateAddPrimitiveBatchSize(InContext ? InContext->GetIncrementalUpdateAddPrimitiveBatchSize() : 0);
 	Context.SetAsyncRegisterLevelContext(AsyncRegisterLevelContext.Get());
@@ -1915,20 +1915,8 @@ bool ULevel::IncrementalPreRegisterComponents(FRegisterComponentContext& Context
 		}
 	};
 
-	const double StartTime = FPlatformTime::Seconds();
-	const double TimeLimitMs = Context.GetIncrementalUpdateTimeLimitMs();
-	const bool bConsiderTimeLimit = TimeLimitMs >= 0;
+	const UE::FTimeout& Timeout = Context.GetIncrementalUpdateTimeout();
 	
-	ON_SCOPE_EXIT
-	{
-		if (bConsiderTimeLimit)
-		{
-			const double ElapsedTimeMs = (FPlatformTime::Seconds() - StartTime) * 1000;
-			const double IncrementalUpdateTimeLimitMs = bConsiderTimeLimit ? FMath::Max<double>(0.0, TimeLimitMs - ElapsedTimeMs) : -1;
-			Context.SetIncrementalUpdateTimeLimitMs(IncrementalUpdateTimeLimitMs);
-		}
-	};
-
 	check(OwningWorld);
 	
 	while (CurrentActorIndexForIncrementalUpdate < Actors.Num())
@@ -1940,12 +1928,9 @@ bool ULevel::IncrementalPreRegisterComponents(FRegisterComponentContext& Context
 			FActorPreRegistrationAccessor::IncrementalPreRegisterComponents(Actor);
 		}
 
-		if (bConsiderTimeLimit)
+		if (Timeout.IsExpired())
 		{
-			if (const double ElapseTimeMs = (FPlatformTime::Seconds() - StartTime) * 1000; ElapseTimeMs > TimeLimitMs)
-			{
-				break;
-			}
+			break;
 		}
 	}
 
@@ -2110,23 +2095,11 @@ bool ULevel::IncrementalPreUnregisterComponents(FUnregisterComponentContext& Con
 			}
 		}
 	};
-
-	const double StartTime = FPlatformTime::Seconds();
-	const double TimeLimitMs = Context.GetIncrementalRemoveTimeLimitMs();
-	const bool bConsiderTimeLimit = TimeLimitMs >= 0;
-
-	ON_SCOPE_EXIT
-	{
-		if (bConsiderTimeLimit)
-		{
-			const double ElapsedTimeMs = (FPlatformTime::Seconds() - StartTime) * 1000;
-			const double IncrementalUpdateTimeLimitMs = bConsiderTimeLimit ? FMath::Max<double>(0.0, TimeLimitMs - ElapsedTimeMs) : -1;
-			Context.SetIncrementalRemoveTimeLimitMs(IncrementalUpdateTimeLimitMs);
-		}
-	};
-
+	
 	check(OwningWorld);
 
+	const UE::FTimeout& Timeout = Context.GetIncrementalRemoveTimeout();
+
 	while (CurrentActorIndexForUnregisterComponents < Actors.Num())
 	{
 		AActor* Actor = Actors[CurrentActorIndexForUnregisterComponents];
@@ -2136,12 +2109,9 @@ bool ULevel::IncrementalPreUnregisterComponents(FUnregisterComponentContext& Con
 			FActorPreRegistrationAccessor::IncrementalPreUnregisterComponents(Actor);
 		}
 
-		if (bConsiderTimeLimit)
+		if (Timeout.IsExpired())
 		{
-			if (const double ElapseTimeMs = (FPlatformTime::Seconds() - StartTime) * 1000; ElapseTimeMs > TimeLimitMs)
-			{
-				break;
-			}
+			break;
 		}
 	}
 
diff --git a/Engine/Source/Runtime/Engine/Private/LevelStreaming.cpp b/Engine/Source/Runtime/Engine/Private/LevelStreaming.cpp
--- a/Engine/Source/Runtime/Engine/Private/LevelStreaming.cpp
+++ b/Engine/Source/Runtime/Engine/Private/LevelStreaming.cpp
@@ -978,7 +978,7 @@ void ULevelStreaming::AckNetVisibilityTransaction(FNetLevelVisibilityTransaction
 	}
 }
 
-void ULevelStreaming::UpdateStreamingState(bool& bOutUpdateAgain, bool& bOutRedetermineTarget)
+void ULevelStreaming::UpdateStreamingState(bool& bOutUpdateAgain, bool& bOutRedetermineTarget, const TOptional<const UE::FTimeout>& InExternalTimeout)
 {
 	FScopeCycleCounterUObject ContextScope(this);
 
@@ -1043,7 +1043,7 @@ void ULevelStreaming::UpdateStreamingState(bool& bOutUpdateAgain, bool& bOutRede
 					break;
 				}
 
-				World->AddToWorld(LoadedLevel, LevelTransform, !bShouldBlockOnLoad, TransactionId, this);
+				World->AddToWorld(LoadedLevel, LevelTransform, !bShouldBlockOnLoad, InExternalTimeout, TransactionId, this);
 
 				if (LoadedLevel->bIsVisible)
 				{
@@ -1089,7 +1089,7 @@ void ULevelStreaming::UpdateStreamingState(bool& bOutUpdateAgain, bool& bOutRede
 			TransactionId.SetTransactionIndex(NetVisibilityState.ServerRequestIndex);
 
 			// Hide loaded level, incrementally if necessary
-			World->RemoveFromWorld(LoadedLevel, !ShouldBlockOnUnload() && World->IsGameWorld(), TransactionId, this);
+			World->RemoveFromWorld(LoadedLevel, !ShouldBlockOnUnload() && World->IsGameWorld(), InExternalTimeout, TransactionId, this);
 
 			// Hide loaded level immediately if bRequireFullVisibilityToRender is set
 			const bool LevelBecameInvisible = bWasVisible && !LoadedLevel->bIsVisible;
diff --git a/Engine/Source/Runtime/Engine/Private/LevelTick.cpp b/Engine/Source/Runtime/Engine/Private/LevelTick.cpp
--- a/Engine/Source/Runtime/Engine/Private/LevelTick.cpp
+++ b/Engine/Source/Runtime/Engine/Private/LevelTick.cpp
@@ -1371,10 +1371,14 @@ void UWorld::Tick( ELevelTick TickType, float DeltaSeconds )
 	// give the async loading code more time if we're performing a high priority load or are in seamless travel
 	if (Info->bHighPriorityLoading || Info->bHighPriorityLoadingLocal || IsInSeamlessTravel())
 	{
-		CSV_SCOPED_SET_WAIT_STAT(AsyncLoading);
-		TRACE_CPUPROFILER_EVENT_SCOPE(HighPriorityAsyncLoading)
-		// Force it to use the entire time slice, even if blocked on I/O
-		ProcessAsyncLoading(true, true, GPriorityAsyncLoadingExtraTime / 1000.0f);
+		// GUseUnifiedTimeBudgetForStreaming handles this in HandleUnifiedStreaming
+		if (!GUseUnifiedTimeBudgetForStreaming)
+		{
+			CSV_SCOPED_SET_WAIT_STAT(AsyncLoading);
+			TRACE_CPUPROFILER_EVENT_SCOPE(HighPriorityAsyncLoading)
+			// Force it to use the entire time slice, even if blocked on I/O
+			ProcessAsyncLoading(true, true, GPriorityAsyncLoadingExtraTime / 1000.0f);
+		}
 	}
 
 	// Translate world origin if requested
diff --git a/Engine/Source/Runtime/Engine/Private/UnrealClient.cpp b/Engine/Source/Runtime/Engine/Private/UnrealClient.cpp
--- a/Engine/Source/Runtime/Engine/Private/UnrealClient.cpp
+++ b/Engine/Source/Runtime/Engine/Private/UnrealClient.cpp
@@ -34,6 +34,7 @@
 #include "Stats/StatsTrace.h"
 #include "Stats/ThreadIdleStats.h"
 #include "RHIUtilities.h"
+#include "Engine/CoreSettings.h"
 
 DEFINE_LOG_CATEGORY_STATIC(LogClient, Log, All);
 
@@ -1726,8 +1727,11 @@ void FViewport::Draw( bool bShouldPresent /*= true */)
 		// if this is a game viewport, and game rendering is disabled, then we don't want to actually draw anything
 		if ( World && World->IsGameWorld() && !bIsGameRenderingEnabled)
 		{
-			// since we aren't drawing the viewport, we still need to update streaming
-			World->UpdateLevelStreaming();
+			if (!GUseUnifiedTimeBudgetForStreaming)
+			{
+				// since we aren't drawing the viewport, we still need to update streaming
+				World->UpdateLevelStreaming(); 
+			}
 		}
 		else
 		{
@@ -1808,7 +1812,7 @@ void FViewport::Draw( bool bShouldPresent /*= true */)
 
 				GInputLatencyTimer.GameThreadTrigger = false;
 			}
-			else if (World && World->IsGameWorld())
+			else if (World && World->IsGameWorld() && !GUseUnifiedTimeBudgetForStreaming)
 			{
 				// since we aren't drawing the viewport, we still need to update streaming
 				World->UpdateLevelStreaming();
diff --git a/Engine/Source/Runtime/Engine/Private/UnrealEngine.cpp b/Engine/Source/Runtime/Engine/Private/UnrealEngine.cpp
--- a/Engine/Source/Runtime/Engine/Private/UnrealEngine.cpp
+++ b/Engine/Source/Runtime/Engine/Private/UnrealEngine.cpp
@@ -11853,11 +11853,77 @@ void UEngine::SetMaxFPS(const float MaxFPS)
 	ConsoleVariable->Set(MaxFPS, ThisSetReason);
 }
 
-/**
-* Enables or disables the ScreenSaver (desktop only)
-*
-* @param bEnable	If true the enable the screen saver, if false disable it.
-*/
+double UEngine::GetUnifiedTimeBudgetForStreaming(float DeltaSeconds, bool bHighPriority)
+{
+	double UnifiedTime = (GAsyncLoadingTimeLimit + GLevelStreamingActorsUpdateTimeLimit) / 1000.0f;
+	if (bHighPriority)
+	{
+		UnifiedTime += (GPriorityAsyncLoadingExtraTime + GPriorityLevelStreamingActorsUpdateExtraTime) / 1000.0f;
+	}
+
+	// TODO handle FPS scaling
+	return UnifiedTime;
+}
+
+void UEngine::HandleUnifiedStreaming(float DeltaSeconds)
+{
+	SCOPED_NAMED_EVENT(Streaming, FColor::Silver);
+	SCOPE_TIME_GUARD(TEXT("UEngine::HandleUnifiedStreaming"));
+
+	// Give the async loading code more time if we're performing a high priority load or are in seamless travel
+	bool bIsHighPriorityLoad = false;
+
+	for (int32 WorldIdx = 0; WorldIdx < WorldList.Num(); ++WorldIdx)
+	{
+		FWorldContext& Context = WorldList[WorldIdx];
+		if (Context.World() != nullptr && Context.World()->ShouldTick())
+		{
+			AWorldSettings* WorldSettings = Context.World()->GetWorldSettings();
+			if (WorldSettings && (WorldSettings->bHighPriorityLoadingLocal || WorldSettings->bHighPriorityLoading || Context.World()->IsInSeamlessTravel()))
+			{
+				bIsHighPriorityLoad = true;
+				break;
+			}
+		}
+	}
+
+	// Start the timeout before first async load
+	double UnifiedBudget = GetUnifiedTimeBudgetForStreaming(DeltaSeconds, bIsHighPriorityLoad);
+	UE::FTimeout StreamingTimeout(UnifiedBudget);
+
+	double AsyncLoadBudget = (GAsyncLoadingTimeLimit / 1000.0);
+	if (bIsHighPriorityLoad)
+	{
+		AsyncLoadBudget += (GPriorityAsyncLoadingExtraTime / 1000.0);
+	}
+
+	// Handle async loading on the game thread right before level streaming so they can share a budget
+	// Use the full time during high priority loads to match the existing code in UWorld::Tick
+	StaticTick(DeltaSeconds, (GAsyncLoadingUseFullTimeLimit || bIsHighPriorityLoad), AsyncLoadBudget);
+
+	// If the process async loading did not use the full time budget it will be granted to UpdateLevelStreaming for all ticking worlds
+	// If process async loading exceeded the budget - UpdateLevelStreaming will amortize the hitch
+	for (int32 WorldIdx = 0; WorldIdx < WorldList.Num(); ++WorldIdx)
+	{
+		FWorldContext& Context = WorldList[WorldIdx];
+		if (Context.World() != nullptr && Context.World()->ShouldTick())
+		{
+			Context.World()->UpdateLevelStreaming(StreamingTimeout);
+		}
+	}
+
+	// If the Update LevelStreaming did not use full time budget, use the rest for async asset loading
+	constexpr double ExtraProcessingThreshold = 0.05 * 0.001;
+	double RemainingTime = StreamingTimeout.GetRemainingSeconds();
+
+	if (RemainingTime > ExtraProcessingThreshold)
+	{
+		constexpr bool bUseTimeLimit = true;
+		constexpr bool bUseFullTimeLimit = false;
+		ProcessAsyncLoading(bUseTimeLimit, bUseFullTimeLimit, RemainingTime);
+	}
+}
+
 void UEngine::EnableScreenSaver( bool bEnable )
 {
 #if PLATFORM_DESKTOP
diff --git a/Engine/Source/Runtime/Engine/Private/World.cpp b/Engine/Source/Runtime/Engine/Private/World.cpp
--- a/Engine/Source/Runtime/Engine/Private/World.cpp
+++ b/Engine/Source/Runtime/Engine/Private/World.cpp
@@ -3141,25 +3141,36 @@ DEFINE_STAT(STAT_UpdateStreamingState);
  * @param	CurrentTask		Description of last task performed
  * @param	StartTime		StartTime, used to calculate time passed
  * @param	Level			Level work has been performed on
- * @param	TimeLimit		The amount of time that is allowed to be used
+ * @param	Timeout			The amount of time that is allowed to be used
  *
  * @return true if time limit has been exceeded, false otherwise
  */
-static bool IsTimeLimitExceeded( const TCHAR* CurrentTask, double StartTime, ULevel* Level, double TimeLimit )
+static bool IsTimeLimitExceeded(const UE::FTimeout& Timeout, const TCHAR* CurrentTask, ULevel* Level)
 {
-	bool bIsTimeLimitExceed = false;
+	const bool bIsTimeLimitExceed = Timeout.IsExpired();
 
-	double CurrentTime	= FPlatformTime::Seconds();
-	// Delta time in ms.
-	double DeltaTime	= (CurrentTime - StartTime) * 1000;
-	if( DeltaTime > TimeLimit )
+	if (UE_LOG_ACTIVE(LogStreaming, Display))
 	{
-		// Log if a single event took way too much time.
-		if( DeltaTime > 20.0 )
+		// We only want to warn for large gaps in time limit checks so store the previous elapsed time
+		static double LastElapsed = 0.0;
+		if (bIsTimeLimitExceed)
+		{
+			if (LastElapsed > 0.0)
+			{
+				const double DeltaMS = (Timeout.GetElapsedSeconds() - LastElapsed) * 1000.0;
+				constexpr double WarningLimit = 20.0;
+				if (DeltaMS > WarningLimit)
+				{
+					// Log if a single event took way too much time, this could also activate after an unrelated stall
+					UE_LOG(LogStreaming, Display, TEXT("UWorld::AddToWorld: %s for %s took (less than) %5.2f ms"), CurrentTask, *Level->GetOutermost()->GetName(), float(DeltaMS) );
+				}
+				LastElapsed = 0.0;
+			}
+		}
+		else
 		{
-			UE_LOG(LogStreaming, Display, TEXT("UWorld::AddToWorld: %s for %s took (less than) %5.2f ms"), CurrentTask, *Level->GetOutermost()->GetName(), DeltaTime );
+			LastElapsed = Timeout.GetElapsedSeconds();
 		}
-		bIsTimeLimitExceed = true;
 	}
 
 	return bIsTimeLimitExceed;
@@ -3359,7 +3370,7 @@ void UWorld::SetLevelPendingVisibilityState(ULevel* InLevel, ELevelPendingVisibi
 	}
 }
 
-void UWorld::AddToWorld( ULevel* Level, const FTransform& LevelTransform, bool bConsiderTimeLimit, FNetLevelVisibilityTransactionId TransactionId, ULevelStreaming* InOwningLevelStreaming)
+void UWorld::AddToWorld(ULevel* Level, const FTransform& LevelTransform, bool bConsiderTimeLimit, const TOptional<const UE::FTimeout>& ExternalTimeout, FNetLevelVisibilityTransactionId TransactionId, ULevelStreaming* InOwningLevelStreaming)
 {
 	TRACE_CPUPROFILER_EVENT_SCOPE(AddToWorld);
 	SCOPE_CYCLE_COUNTER(STAT_AddToWorldTime);
@@ -3395,34 +3406,43 @@ void UWorld::AddToWorld( ULevel* Level, const FTransform& LevelTransform, bool b
 
 	// Don't consider the time limit if the match hasn't started as we need to ensure that the levels are fully loaded
 	bConsiderTimeLimit &= bMatchStarted && bIsGameWorld;
-	double TimeLimit = 0.0;
 
+	UE::FTimeout Timeout = UE::FTimeout::Never();
 	if (bExecuteNextStep && bConsiderTimeLimit)
 	{
-		float ExtraTimeLimit = 0.0f;
-		// Give the actor initialization code more time if we're performing a high priority load or are in seamless travel
-		if (AWorldSettings* WorldSettings = GetWorldSettings(false, false))
-		{
-			if (WorldSettings->bHighPriorityLoading || WorldSettings->bHighPriorityLoadingLocal || IsInSeamlessTravel())
-			{
-				ExtraTimeLimit = GPriorityLevelStreamingActorsUpdateExtraTime;
-			}
-		}
-		if (GAdaptiveAddToWorld.IsEnabled())
+		if (ExternalTimeout.IsSet())
 		{
-			GAdaptiveAddToWorld.SetExtraTimeSlice(ExtraTimeLimit);
-			TimeLimit = GAdaptiveAddToWorld.GetTimeSlice();
+			// Use the external timeout if specified, usually set from GetUnifiedTimeBudgetForStreaming
+			Timeout = ExternalTimeout.GetValue();
 		}
 		else
 		{
-			TimeLimit = GLevelStreamingActorsUpdateTimeLimit + ExtraTimeLimit;
-		}
-
-		// Remove cumulated time since UpdateLevelStreaming
-		TimeLimit = FMath::Max<double>(0.0, TimeLimit - GAddToWorldTimeCumul);
+			double TimeLimit = 0.0;
+			float ExtraTimeLimit = 0.0f;
+			// Give the actor initialization code more time if we're performing a high priority load or are in seamless travel
+			if (AWorldSettings* WorldSettings = GetWorldSettings(false, false))
+			{
+				if (WorldSettings->bHighPriorityLoading || WorldSettings->bHighPriorityLoadingLocal || IsInSeamlessTravel())
+				{
+					ExtraTimeLimit = GPriorityLevelStreamingActorsUpdateExtraTime;
+				}
+			}
+			if (GAdaptiveAddToWorld.IsEnabled())
+			{
+				GAdaptiveAddToWorld.SetExtraTimeSlice(ExtraTimeLimit);
+				TimeLimit = GAdaptiveAddToWorld.GetTimeSlice();
+			}
+			else
+			{
+				TimeLimit = GLevelStreamingActorsUpdateTimeLimit + ExtraTimeLimit;
+			}
 
+			// Remove cumulated time since UpdateLevelStreaming
+			TimeLimit = FMath::Max<double>(0.0, TimeLimit - GAddToWorldTimeCumul);
+			Timeout.SetTimeoutSeconds(TimeLimit / 1000.0);
+		}
 		// Make sure to test time limit (which includes GAddToWorldTimeCumul) before allowing to begin making visible a level
-		bExecuteNextStep = (!bConsiderTimeLimit || !IsTimeLimitExceeded(TEXT("computing time limit"), StartTime, Level, TimeLimit));
+		bExecuteNextStep = !IsTimeLimitExceeded(Timeout, TEXT("computing time limit"), Level);
 	}
 
 	// Don't make this level visible if it's currently being made invisible
@@ -3458,7 +3478,7 @@ void UWorld::AddToWorld( ULevel* Level, const FTransform& LevelTransform, bool b
 		PerformLastStepTime = 0.0;
 #endif
 
-		bExecuteNextStep = (!bConsiderTimeLimit || !IsTimeLimitExceeded(TEXT("begin making visible"), StartTime, Level, TimeLimit));
+		bExecuteNextStep = (!bConsiderTimeLimit || !IsTimeLimitExceeded(Timeout, TEXT("begin making visible"), Level));
 	}
 
 	int32 StartWorkUnitsRemaining = Level->GetEstimatedAddToWorldWorkUnitsRemaining();
@@ -3475,7 +3495,7 @@ void UWorld::AddToWorld( ULevel* Level, const FTransform& LevelTransform, bool b
 		FLevelUtils::ApplyLevelTransform(TransformParams);
 
 		Level->bAlreadyMovedActors = true;
-		bExecuteNextStep = (!bConsiderTimeLimit || !IsTimeLimitExceeded( TEXT("moving actors"), StartTime, Level, TimeLimit));
+		bExecuteNextStep = (!bConsiderTimeLimit || !IsTimeLimitExceeded( Timeout, TEXT("moving actors"), Level));
 	}
 
 	if( bExecuteNextStep && !Level->bAlreadyShiftedActors )
@@ -3492,7 +3512,7 @@ void UWorld::AddToWorld( ULevel* Level, const FTransform& LevelTransform, bool b
 		}
 		
 		Level->bAlreadyShiftedActors = true;
-		bExecuteNextStep = (!bConsiderTimeLimit || !IsTimeLimitExceeded( TEXT("shifting actors"), StartTime, Level, TimeLimit ));
+		bExecuteNextStep = (!bConsiderTimeLimit || !IsTimeLimitExceeded( Timeout, TEXT("shifting actors"), Level));
 	}
 
 #if WITH_EDITOR
@@ -3568,10 +3588,8 @@ void UWorld::AddToWorld( ULevel* Level, const FTransform& LevelTransform, bool b
 		// Incremental update components
 		{
 			// Prepare context used to store batch/parallelize AddPrimitive calls
-			const double ElapsedTimeMs = (FPlatformTime::Seconds() - StartTime) * 1000;
-			const double IncrementalUpdateTimeLimitMs = bConsiderTimeLimit ? FMath::Max<double>(0.0, TimeLimit - ElapsedTimeMs) : -1;
 			FRegisterComponentContext Context = FRegisterComponentContext(this)
-				.SetIncrementalUpdateTimeLimitMs(IncrementalUpdateTimeLimitMs)
+				.SetIncrementalUpdateTimeout(Timeout)
 				.SetIncrementalUpdateAddPrimitiveBatchSize(FMath::Max<int32>(0, GLevelStreamingAddPrimitiveGranularity));
 
 			// Incrementally update components.
@@ -3585,7 +3603,7 @@ void UWorld::AddToWorld( ULevel* Level, const FTransform& LevelTransform, bool b
 				{
 					break;
 				}
-				bTimeLimitExceeded = IsTimeLimitExceeded(TEXT("updating components"), StartTime, Level, TimeLimit);
+				bTimeLimitExceeded = IsTimeLimitExceeded(Timeout, TEXT("updating components"), Level);
 				bHasRunningTasks = AsyncRegisterLevelContext ? AsyncRegisterLevelContext->IsRunningAsync() : false;
 				if (AsyncRegisterLevelContext && bTimeLimitExceeded && bHasRunningTasks)
 				{
@@ -3604,7 +3622,7 @@ void UWorld::AddToWorld( ULevel* Level, const FTransform& LevelTransform, bool b
 
 		// We are done once all components are attached.
 		Level->bAlreadyUpdatedComponents	= Level->bAreComponentsCurrentlyRegistered;
-		bExecuteNextStep					= Level->bAreComponentsCurrentlyRegistered && (!bConsiderTimeLimit || !IsTimeLimitExceeded(TEXT("updating components"), StartTime, Level, TimeLimit));
+		bExecuteNextStep					= Level->bAreComponentsCurrentlyRegistered && (!bConsiderTimeLimit || !IsTimeLimitExceeded(Timeout, TEXT("updating components"), Level));
 	}
 
 #if WITH_EDITOR
@@ -3636,8 +3654,11 @@ void UWorld::AddToWorld( ULevel* Level, const FTransform& LevelTransform, bool b
 				Level->ClearActorsSeamlessTraveledFlag();
 			}
 
-			const double PreventNextStepTimeLimit = GLevelStreamingForceRouteActorInitializeNextFrame ? 0.0 : TimeLimit;
-			bExecuteNextStep = (!bConsiderTimeLimit || !IsTimeLimitExceeded( TEXT("initializing network actors"), StartTime, Level, PreventNextStepTimeLimit ));
+			if (bConsiderTimeLimit)
+			{
+				// Only check GLevelStreamingForceRouteActorInitializeNextFrame if this is an incremental add
+				bExecuteNextStep = !GLevelStreamingForceRouteActorInitializeNextFrame && !IsTimeLimitExceeded(Timeout, TEXT("initializing network actors"), Level);
+			}
 		}
 
 		// Route various initialization functions and set volumes.
@@ -3652,10 +3673,10 @@ void UWorld::AddToWorld( ULevel* Level, const FTransform& LevelTransform, bool b
 			do 
 			{
 				Level->RouteActorInitialize(NumActorsToProcess);
-			} while (!Level->IsFinishedRouteActorInitialization() && !IsTimeLimitExceeded(TEXT("routing Initialize on actors"), StartTime, Level, TimeLimit));
+			} while (!Level->IsFinishedRouteActorInitialization() && !IsTimeLimitExceeded(Timeout, TEXT("routing Initialize on actors"), Level));
 			bStartup = 0;
 
-			bExecuteNextStep = Level->IsFinishedRouteActorInitialization() && (!bConsiderTimeLimit || !IsTimeLimitExceeded( TEXT("routing Initialize on actors"), StartTime, Level, TimeLimit ));
+			bExecuteNextStep = Level->IsFinishedRouteActorInitialization() && (!bConsiderTimeLimit || !IsTimeLimitExceeded(Timeout,  TEXT("routing Initialize on actors"), Level ));
 		}
 		
 		if( bExecuteNextStep )
@@ -3669,7 +3690,7 @@ void UWorld::AddToWorld( ULevel* Level, const FTransform& LevelTransform, bool b
 
 				Level->SortActorList();
 				Level->bAlreadySortedActorList = true;
-				bExecuteNextStep = (!bConsiderTimeLimit || !IsTimeLimitExceeded( TEXT("sorting actor list"), StartTime, Level, TimeLimit ));
+				bExecuteNextStep = (!bConsiderTimeLimit || !IsTimeLimitExceeded(Timeout, TEXT("sorting actor list"), Level));
 			}
 
 			bPerformedLastStep = true;
@@ -3801,9 +3822,9 @@ void UWorld::AddToWorld( ULevel* Level, const FTransform& LevelTransform, bool b
 }
 
 // Cumulated time doing IncrementalUnregisterComponents in UWorld::RemoveFromWorld since last call to UWorld::UpdateLevelStreaming.
-static double GRemoveFromWorldUnregisterComponentTimeCumul = 0.0;
+static double GRemoveFromWorldTimeCumul = 0.0;
 
-void UWorld::RemoveFromWorld( ULevel* Level, bool bAllowIncrementalRemoval, FNetLevelVisibilityTransactionId TransactionId, ULevelStreaming* InOwningLevelStreaming)
+void UWorld::RemoveFromWorld(ULevel* Level, bool bAllowIncrementalRemoval, const TOptional<const UE::FTimeout>& ExternalTimeout, FNetLevelVisibilityTransactionId TransactionId, ULevelStreaming* InOwningLevelStreaming)
 {
 	TRACE_CPUPROFILER_EVENT_SCOPE(RemoveFromWorld);
 	SCOPE_CYCLE_COUNTER(STAT_RemoveFromWorldTime);
@@ -3854,45 +3875,41 @@ void UWorld::RemoveFromWorld( ULevel* Level, bool bAllowIncrementalRemoval, FNet
 	// To be removed from the world a world must be visible and not pending being made visible (this may be redundent, but for safety)
 	if ( (!HasAnyLevelMakingVisible() || (bIsCandidateForRemoval && !IsLevelMakingVisible(Level))) && (Level->bIsVisible || Level->bIsBeingRemoved) )
 	{
-#if PERF_TRACK_DETAILED_ASYNC_STATS
-		// Keep track of timing.
-		double PerfTrackStartTime = FPlatformTime::Seconds();
-#endif
+		double StartTime = FPlatformTime::Seconds();
 		bool bWaitForCompletion = false;
 		bool bFinishRemovingLevel = true;
-		if ( bAllowIncrementalRemoval && GLevelStreamingUnregisterComponentsTimeLimit > 0.0 )
+
+		// Can't do an incremental removal if there is no time limit
+		bAllowIncrementalRemoval = bAllowIncrementalRemoval && (ExternalTimeout.IsSet() || GLevelStreamingUnregisterComponentsTimeLimit > 0.0);
+
+		if (bAllowIncrementalRemoval)
 		{
+			// Check external timeout or time limit (which includes GRemoveFromWorldTimeCumul) before allowing to begin making invisible a level
+			const UE::FTimeout Timeout = ExternalTimeout.IsSet() ? ExternalTimeout.GetValue()
+				: UE::FTimeout(FMath::Max<double>(0.0, GLevelStreamingUnregisterComponentsTimeLimit - GRemoveFromWorldTimeCumul) / 1000.0);
+
 			bFinishRemovingLevel = false;
 			if (!IsLevelMakingInvisible(Level) && (MakingInvisibleLevels.Num() < FMath::Max<int32>(1, GMaximumMakingInvisibleLevels)))
 			{
-				// Make sure to test time limit (which includes GRemoveFromWorldUnregisterComponentTimeCumul) before allowing to begin making invisible a level
-				double TimeLimit = FMath::Max<double>(0.0, GLevelStreamingUnregisterComponentsTimeLimit - GRemoveFromWorldUnregisterComponentTimeCumul);
-				if (TimeLimit > 0.0)
+				if (!Timeout.IsExpired())
 				{
 					TRACE_BEGIN_REGION(*WriteToString<256>(TEXT("RemoveFromWorld: "), Level->GetPackage()->GetName()));
 
-					double StartTime = FPlatformTime::Seconds();
 					// Mark level as being in process of being made invisible. 
 					// This will prevent this level from being unloaded or made visible in the meantime
 					SetLevelPendingVisibilityState(Level, ELevelPendingVisibilityState::MakingInvisible);
 
 					BeginRemoval();
-
-					// Take into consideration time taken by BeginRemoval
-					double DeltaTime = (FPlatformTime::Seconds() - StartTime) * 1000;
-					GRemoveFromWorldUnregisterComponentTimeCumul += DeltaTime;
 				}
 			}
 
 			if (IsLevelMakingInvisible(Level))
 			{
-				double TimeLimit = FMath::Max<double>(0.0, GLevelStreamingUnregisterComponentsTimeLimit - GRemoveFromWorldUnregisterComponentTimeCumul);
-				if (TimeLimit > 0.0)
+				if (!Timeout.IsExpired())
 				{
 					FUnregisterComponentContext Context = FUnregisterComponentContext(this)
-						.SetIncrementalRemoveTimeLimitMs(TimeLimit);
+						.SetIncrementalRemoveTimeout(Timeout);
 
-					double StartTime = FPlatformTime::Seconds();
 					// Incrementally unregister actor components. 
 					// This avoids spikes on the renderthread and gamethread when we subsequently call ClearLevelComponents() further down
 					check(bIsGameWorld);
@@ -3908,10 +3925,7 @@ void UWorld::RemoveFromWorld( ULevel* Level, bool bAllowIncrementalRemoval, FNet
 						{
 							break;
 						}
-					} while (!IsTimeLimitExceeded(TEXT("unregistering components"), StartTime, Level, TimeLimit));
-
-					double DeltaTime = (FPlatformTime::Seconds() - StartTime) * 1000;
-					GRemoveFromWorldUnregisterComponentTimeCumul += DeltaTime;
+					} while (!IsTimeLimitExceeded(Timeout, TEXT("unregistering components"), Level));
 				}
 			}
 		}
@@ -3932,7 +3946,7 @@ void UWorld::RemoveFromWorld( ULevel* Level, bool bAllowIncrementalRemoval, FNet
 
 		{
 			TRACE_CPUPROFILER_EVENT_SCOPE(RemoveFromWorld_Extension);
-			check(bWaitForCompletion == !(bAllowIncrementalRemoval && GLevelStreamingUnregisterComponentsTimeLimit > 0.0));
+			check(bWaitForCompletion == !bAllowIncrementalRemoval);
 			bool bOutHasCompleted = true;
 			RemoveLevelFromWorldExtensionEvent.Broadcast(Level, bWaitForCompletion, bOutHasCompleted);
 			if (!bOutHasCompleted)
@@ -3946,7 +3960,6 @@ void UWorld::RemoveFromWorld( ULevel* Level, bool bAllowIncrementalRemoval, FNet
 			// The level can be removed from the pending visibility state
 			SetLevelPendingVisibilityState(Level, ELevelPendingVisibilityState::None);
 			TRACE_CPUPROFILER_EVENT_SCOPE(RemoveFromWorld_FinishRemovingLevel);
-			double StartTime = FPlatformTime::Seconds();
 
 			for (int32 ActorIdx = 0; ActorIdx < Level->Actors.Num(); ActorIdx++)
 			{
@@ -4029,20 +4042,19 @@ void UWorld::RemoveFromWorld( ULevel* Level, bool bAllowIncrementalRemoval, FNet
 
 			Level->bIsBeingRemoved = false;
 
-			// Keep track of time spent completing the incremental removal of the level
-			if (bAllowIncrementalRemoval)
-			{
-				double DeltaTime = (FPlatformTime::Seconds() - StartTime) * 1000;
-				GRemoveFromWorldUnregisterComponentTimeCumul += DeltaTime;
-			}
-
 			TRACE_END_REGION(*WriteToString<256>(TEXT("RemoveFromWorld: "), Level->GetPackage()->GetName()));
 		} // if ( bFinishRemovingLevel )
 
-		Level->bIsDisassociatingLevel = false;
+		// Keep track of time spent for entire incremental removal
+		if (bAllowIncrementalRemoval)
+		{
+			double DeltaTime = (FPlatformTime::Seconds() - StartTime) * 1000;
+			GRemoveFromWorldTimeCumul += DeltaTime;
+		}
 
+		Level->bIsDisassociatingLevel = false;
 #if PERF_TRACK_DETAILED_ASYNC_STATS
-		UE_LOG(LogStreaming, Display, TEXT("UWorld::RemoveFromWorld for %s took %5.2f ms"), *Level->GetOutermost()->GetName(), (FPlatformTime::Seconds() - PerfTrackStartTime) * 1000.0);
+		UE_LOG(LogStreaming, Display, TEXT("UWorld::RemoveFromWorld for %s took %5.2f ms"), *Level->GetOutermost()->GetName(), (FPlatformTime::Seconds() - StartTime) * 1000.0);
 #endif // PERF_TRACK_DETAILED_ASYNC_STATS
 	}
 }
@@ -4675,7 +4687,7 @@ void UWorld::FlushAsyncLevelRequests()
 	}
 }
 
-void UWorld::UpdateLevelStreaming()
+void UWorld::UpdateLevelStreaming(const TOptional<const UE::FTimeout>& ExternalTimeout)
 {
 	TRACE_CPUPROFILER_EVENT_SCOPE(UWorld::UpdateLevelStreaming);
 	SCOPE_CYCLE_COUNTER(STAT_UpdateLevelStreamingTime);
@@ -4688,7 +4700,7 @@ void UWorld::UpdateLevelStreaming()
 
 	// Reset counters
 	GAddToWorldTimeCumul = 0.0;
-	GRemoveFromWorldUnregisterComponentTimeCumul = 0.0;
+	GRemoveFromWorldTimeCumul = 0.0;
 
 	// do nothing if level streaming is frozen
 	if (bIsLevelStreamingFrozen)
@@ -4750,7 +4762,7 @@ void UWorld::UpdateLevelStreaming()
 				while (bUpdateAgain && bShouldContinueToConsider)
 				{
 					bool bRedetermineTarget = false;
-					FStreamingLevelPrivateAccessor::UpdateStreamingState(StreamingLevel, bUpdateAgain, bRedetermineTarget);
+					FStreamingLevelPrivateAccessor::UpdateStreamingState(StreamingLevel, bUpdateAgain, bRedetermineTarget, ExternalTimeout);
 
 					if (bRedetermineTarget)
 					{
