[DSTM] Fix lock usage issues with Physics State migration

#rnx

diff --git a/Engine/Source/Runtime/Engine/Private/PhysicsEngine/ConstraintInstance.cpp b/Engine/Source/Runtime/Engine/Private/PhysicsEngine/ConstraintInstance.cpp
--- a/Engine/Source/Runtime/Engine/Private/PhysicsEngine/ConstraintInstance.cpp
+++ b/Engine/Source/Runtime/Engine/Private/PhysicsEngine/ConstraintInstance.cpp
@@ -716,7 +716,7 @@ const FPhysicsConstraintHandle& FConstraintInstance::GetPhysicsConstraintRef() c
 void FConstraintInstance::SerializeState(Chaos::FSerializedDataBuffer& OutSerializedState, Chaos::ESerializedDataContext DataContext)
 {
 	const EPhysicsInterfaceScopedThreadContext LockContext = DataContext == Chaos::ESerializedDataContext::Internal ? EPhysicsInterfaceScopedThreadContext::Internal : EPhysicsInterfaceScopedThreadContext::External;
-	FPhysicsCommand::ExecuteWrite(ConstraintHandle, [&OutSerializedState, DataContext](const FPhysicsConstraintHandle& Constraint)
+	FPhysicsCommand::ExecuteRead(ConstraintHandle, [&OutSerializedState, DataContext](const FPhysicsConstraintHandle& Constraint)
 	{
 		FChaosEngineInterface::SerializeConstraintStateToBuffer(Constraint, OutSerializedState, DataContext);
 	}, LockContext);
@@ -725,7 +725,7 @@ void FConstraintInstance::SerializeState(Chaos::FSerializedDataBuffer& OutSerial
 void FConstraintInstance::ApplySerializedState(Chaos::FSerializedDataBufferPtr&& InSerializedState, Chaos::ESerializedDataContext DataContext)
 {
 	const EPhysicsInterfaceScopedThreadContext LockContext = DataContext == Chaos::ESerializedDataContext::Internal ? EPhysicsInterfaceScopedThreadContext::Internal : EPhysicsInterfaceScopedThreadContext::External;
-	FPhysicsCommand::ExecuteRead(ConstraintHandle, [&InSerializedState, this, DataContext](const FPhysicsConstraintHandle& Constraint)
+	FPhysicsCommand::ExecuteWrite(ConstraintHandle, [&InSerializedState, this, DataContext](const FPhysicsConstraintHandle& Constraint)
 	{
 		FChaosEngineInterface::ApplySerializedStateToConstraint(Constraint, MoveTemp(InSerializedState), DataContext);
 	}, LockContext);
diff --git a/Engine/Source/Runtime/Engine/Private/PhysicsEngine/Experimental/ChaosScopedSceneLock.cpp b/Engine/Source/Runtime/Engine/Private/PhysicsEngine/Experimental/ChaosScopedSceneLock.cpp
--- a/Engine/Source/Runtime/Engine/Private/PhysicsEngine/Experimental/ChaosScopedSceneLock.cpp
+++ b/Engine/Source/Runtime/Engine/Private/PhysicsEngine/Experimental/ChaosScopedSceneLock.cpp
@@ -85,18 +85,7 @@ FScopedSceneLock_Chaos::FScopedSceneLock_Chaos(FPhysicsConstraintHandle const* I
 	: Solver(nullptr),
 	LockType(InLockType)
 {
-	if (InConstraintHandle)
-	{
-		FChaosScene* Scene = GetSceneForActor(InConstraintHandle);
-		Solver = Scene ? Scene->GetSolver() : nullptr;
-	}
-#if CHAOS_CHECKED
-	if (!Solver)
-	{
-		UE_LOG(LogPhysics, Warning, TEXT("Failed to find Scene for constraint. Skipping lock"));
-	}
-#endif
-	LockScene();
+	LockSceneForConstraint(InConstraintHandle);
 }
 
 FScopedSceneLock_Chaos::FScopedSceneLock_Chaos(USkeletalMeshComponent* InSkelMeshComp, EPhysicsInterfaceScopedLockType InLockType)
@@ -274,13 +263,35 @@ Chaos::FPhysSceneLock* FScopedSceneLock_Chaos::GetSolverLock(Chaos::FPBDRigidsSo
 #endif
 }
 
-FScopedSceneLockWithContext_Chaos::FScopedSceneLockWithContext_Chaos(const FPhysicsActorHandle& InActorHandle, EPhysicsInterfaceScopedLockType InLockType, EPhysicsInterfaceScopedThreadContext InThreadContext) : FScopedSceneLock_Chaos(InActorHandle, InLockType)
+void FScopedSceneLock_Chaos::LockSceneForConstraint(FPhysicsConstraintHandle const* InConstraintHandle)
 {
-	ThreadContext = InThreadContext;
+	if (InConstraintHandle)
+	{
+		FChaosScene* Scene = GetSceneForActor(InConstraintHandle);
+		Solver = Scene ? Scene->GetSolver() : nullptr;
+	}
+#if CHAOS_CHECKED
+	if (!Solver)
+	{
+		UE_LOG(LogPhysics, Warning, TEXT("Failed to find Scene for constraint. Skipping lock"));
+	}
+#endif
+	LockScene();
 }
 
+FScopedSceneLockWithContext_Chaos::FScopedSceneLockWithContext_Chaos(const FPhysicsActorHandle& InActorHandle, EPhysicsInterfaceScopedLockType InLockType, EPhysicsInterfaceScopedThreadContext InThreadContext)
+	: FScopedSceneLock_Chaos(InLockType)
+{
+	ThreadContext = InThreadContext;
+
+	Solver = GetSolverForActor(InActorHandle);
+	LockScene();
+}
 
-FScopedSceneLockWithContext_Chaos::FScopedSceneLockWithContext_Chaos(FPhysicsConstraintHandle const* InConstraintHandle, EPhysicsInterfaceScopedLockType InLockType, EPhysicsInterfaceScopedThreadContext InThreadContext) : FScopedSceneLock_Chaos(InConstraintHandle, InLockType)
+FScopedSceneLockWithContext_Chaos::FScopedSceneLockWithContext_Chaos(FPhysicsConstraintHandle const* InConstraintHandle, EPhysicsInterfaceScopedLockType InLockType, EPhysicsInterfaceScopedThreadContext InThreadContext)
+	: FScopedSceneLock_Chaos(InLockType)
 {
 	ThreadContext = InThreadContext;
+
+	LockSceneForConstraint(InConstraintHandle);
 }
diff --git a/Engine/Source/Runtime/Engine/Public/Physics/Experimental/ChaosScopedSceneLock.h b/Engine/Source/Runtime/Engine/Public/Physics/Experimental/ChaosScopedSceneLock.h
--- a/Engine/Source/Runtime/Engine/Public/Physics/Experimental/ChaosScopedSceneLock.h
+++ b/Engine/Source/Runtime/Engine/Public/Physics/Experimental/ChaosScopedSceneLock.h
@@ -47,7 +47,13 @@ struct FScopedSceneLock_Chaos
 
 	ENGINE_API void Release();
 
-private:
+protected:
+
+	FScopedSceneLock_Chaos(EPhysicsInterfaceScopedLockType InLockType) : Solver(nullptr), LockType(InLockType)
+	{
+	}
+	
+	void LockSceneForConstraint(FPhysicsConstraintHandle const* InConstraintHandle);
 
 	ENGINE_API void LockScene();
 	ENGINE_API void UnlockScene();
@@ -61,8 +67,8 @@ private:
 	Chaos::FPBDRigidsSolver* Solver;
 	EPhysicsInterfaceScopedLockType LockType;
 
-protected:
 	EPhysicsInterfaceScopedThreadContext ThreadContext = EPhysicsInterfaceScopedThreadContext::External;
+
 };
 
 struct UE_INTERNAL FScopedSceneLockWithContext_Chaos : FScopedSceneLock_Chaos
