FSpline: Support text import/export (and therefore copy/paste & duplication).

Heavily inspired by UDynamicMesh import/export. We use the existing serialize functionality to get a binary representation of the spline data which is encoded in base 64.

The base64 encoding is probably not necessary, and we may be better off eventually doing a proper non-binary import/export for debugging purposes (as the exported text is just in the clipboard if we copy the actor holding the spline), but this works for now and should let us work with large enough splines in the meantime.

diff --git a/Engine/Source/Runtime/Engine/Classes/Curves/Spline.h b/Engine/Source/Runtime/Engine/Classes/Curves/Spline.h
--- a/Engine/Source/Runtime/Engine/Classes/Curves/Spline.h
+++ b/Engine/Source/Runtime/Engine/Classes/Curves/Spline.h
@@ -20,7 +20,7 @@ USTRUCT()
 struct FSpline
 {
 	GENERATED_BODY()
-	
+
 	ENGINE_API FSpline();
 	ENGINE_API FSpline(const FSpline& Other);
 	ENGINE_API FSpline& operator=(const FSpline& Other);
@@ -102,6 +102,10 @@ struct FSpline
 		return Ar;
 	}
 	ENGINE_API bool Serialize(FArchive& Ar);
+	ENGINE_API void SerializeLoad(FArchive& Ar);
+	ENGINE_API void SerializeSave(FArchive& Ar) const;
+	ENGINE_API bool ExportTextItem(FString& ValueStr, FSpline const& DefaultValue, class UObject* Parent, int32 PortFlags, class UObject* ExportRootScope) const;
+	ENGINE_API bool ImportTextItem(const TCHAR*& Buffer, int32 PortFlags, UObject* Parent, FOutputDevice* ErrorText);
 
 	uint32 GetVersion() const { return Version; }
 	
@@ -188,9 +192,7 @@ struct TStructOpsTypeTraits<FSpline> : public TStructOpsTypeTraitsBase2<FSpline>
 	{
 		WithSerializer				= true, // Enables the use of a custom Serialize method.
 		WithIdenticalViaEquality	= true, // Enables the use of a custom equality operator.
-
-		// todo: need to implement for copy/paste and duplicate.
-		//WithExportTextItem			= true, // Enables the use of a custom ExportTextItem method.
-		//WithImportTextItem			= true, // Enables the use of a custom ImportTextItem method.
+		WithExportTextItem			= true, // Enables the use of a custom ExportTextItem method.
+		WithImportTextItem			= true, // Enables the use of a custom ImportTextItem method.
 	};
 };
\ No newline at end of file
diff --git a/Engine/Source/Runtime/Engine/Private/Curves/Spline.cpp b/Engine/Source/Runtime/Engine/Private/Curves/Spline.cpp
--- a/Engine/Source/Runtime/Engine/Private/Curves/Spline.cpp
+++ b/Engine/Source/Runtime/Engine/Private/Curves/Spline.cpp
@@ -3,8 +3,11 @@
 #include "Curves/Spline.h"
 
 #include "Components/SplineComponent.h"	// for FSplineCurves, ideally removed
-#include "UObject/FortniteMainBranchObjectVersion.h"
+#include "Misc/Base64.h"
+#include "Serialization/MemoryReader.h"
+#include "Serialization/MemoryWriter.h"
 #include "Splines/TangentBezierSpline.h"
+#include "UObject/FortniteMainBranchObjectVersion.h"
 
 #include UE_INLINE_GENERATED_CPP_BY_NAME(Spline)
 
@@ -894,78 +897,143 @@ bool FSpline::Serialize(FArchive& Ar)
 	
 	if (Ar.IsLoading())
 	{
-		// using intermediate because I am scared of serializing a 2 bit field.
-		int8 PreviousImpl;
-		Ar << PreviousImpl;
+		SerializeLoad(Ar);
+	}
+	else
+	{
+		SerializeSave(Ar);
+	}
+
+	return true;
+}
+
+void FSpline::SerializeLoad(FArchive& Ar)
+{
+	int8 PreviousImpl;
+	Ar << PreviousImpl;
 #if WITH_EDITOR
-		PreviousImplementation = PreviousImpl;
+	PreviousImplementation = PreviousImpl;
 #endif
 
-		auto WasEnabled = [PreviousImpl]()
-			{ return PreviousImpl != 0; };
+	auto WasEnabled = [PreviousImpl]()
+	{ return PreviousImpl != 0; };
 
-		auto WasLegacy = [PreviousImpl]()
-        	{ return PreviousImpl == 1; };
+	auto WasLegacy = [PreviousImpl]()
+	{ return PreviousImpl == 1; };
 		
-		if (WasEnabled())
+	if (WasEnabled())
+	{
+		if (WasLegacy())
 		{
-			if (WasLegacy())
+			if (IsLegacy())
 			{
-				if (IsLegacy())
-				{
-					LegacyData = MakeShared<FLegacySpline>();
-					Ar << *LegacyData;
-				}
-				else if (IsNew())
-				{
-					FLegacySpline IntermediateData;
-					Ar << IntermediateData;
-					NewData = MakeShared<FNewSpline>(IntermediateData);
-				}
-				else if (!IsEnabled())
-				{
-					// Intentionally doing nothing with data
-					FLegacySpline IntermediateData;
-					Ar << IntermediateData;
-				}
+				LegacyData = MakeShared<FLegacySpline>();
+				Ar << *LegacyData;
 			}
-			else
+			else if (IsNew())
 			{
-				if (IsLegacy())
-				{
-					FNewSpline IntermediateData;
-					Ar << IntermediateData;
-					LegacyData = MakeShared<FLegacySpline>(IntermediateData);
-				}
-				else if (IsNew())
-				{
-					NewData = MakeShared<FNewSpline>();
-					Ar << *NewData;
-				}
-				else if (!IsEnabled())
-				{
-					// Intentionally doing nothing with data
-					FNewSpline IntermediateData;
-					Ar << IntermediateData;
-				}
+				FLegacySpline IntermediateData;
+				Ar << IntermediateData;
+				NewData = MakeShared<FNewSpline>(IntermediateData);
+			}
+			else if (!IsEnabled())
+			{
+				// Intentionally doing nothing with data
+				FLegacySpline IntermediateData;
+				Ar << IntermediateData;
 			}
 		}
-	}
-	else
-	{
-		uint8 CurrentImpl = CurrentImplementation;
-		Ar << CurrentImpl;
-		
-		if (IsLegacy())
+		else
 		{
-			Ar << *LegacyData;
+			if (IsLegacy())
+			{
+				FNewSpline IntermediateData;
+				Ar << IntermediateData;
+				LegacyData = MakeShared<FLegacySpline>(IntermediateData);
+			}
+			else if (IsNew())
+			{
+				NewData = MakeShared<FNewSpline>();
+				Ar << *NewData;
+			}
+			else if (!IsEnabled())
+			{
+				// Intentionally doing nothing with data
+				FNewSpline IntermediateData;
+				Ar << IntermediateData;
+			}
 		}
-		else if (IsNew())
+	}
+}
+
+void FSpline::SerializeSave(FArchive& Ar) const
+{
+	uint8 CurrentImpl = CurrentImplementation;
+	Ar << CurrentImpl;
+		
+	if (IsLegacy())
+	{
+		Ar << *LegacyData;
+	}
+	else if (IsNew())
+	{
+		Ar << *NewData;
+	}
+}
+
+bool FSpline::ExportTextItem(FString& ValueStr, FSpline const& DefaultValue, class UObject* Parent, int32 PortFlags, class UObject* ExportRootScope) const
+{
+	// serialize our spline
+    TArray<uint8> SplineWriteBuffer;
+    FMemoryWriter MemWriter(SplineWriteBuffer);
+    SerializeSave(MemWriter);
+
+    FString Base64String = FBase64::Encode(SplineWriteBuffer);
+
+    // Base64 encoding uses the '/' character, but T3D interprets '//' as some kind of
+    // terminator (?). If it occurs then the string passed to ImportTextItem() will
+    // come back as full of nullptrs. So we will swap in '-' here, and swap back to '/' in ImportTextItem()
+	Base64String.ReplaceCharInline('/', '-');
+
+	ValueStr = FString::Printf(TEXT("SplineData SplineDataLen=%d SplineData=%s\r\n"), Base64String.Len(), *Base64String);
+	
+	return true;
+}
+
+bool FSpline::ImportTextItem(const TCHAR*& SourceText, int32 PortFlags, UObject* Parent, FOutputDevice* ErrorText)
+{
+	if (FParse::Command(&SourceText, TEXT("SplineData")))	
+	{
+		static const TCHAR SplineDataLenToken[] = TEXT("SplineDataLen=");
+		const TCHAR* FoundSplineDataLenStart = FCString::Strifind(SourceText, SplineDataLenToken);
+		if (FoundSplineDataLenStart)
 		{
-			Ar << *NewData;
+			SourceText = FoundSplineDataLenStart + FCString::Strlen(SplineDataLenToken);
+			int32 SplineDataLen = FCString::Atoi(SourceText);
+
+			static const TCHAR SplineDataToken[] = TEXT("SplineData=");
+			const TCHAR* FoundSplineDataStart = FCString::Strifind(SourceText, SplineDataToken);
+			if (FoundSplineDataStart)
+			{
+				SourceText = FoundSplineDataStart + FCString::Strlen(SplineDataToken);
+				FString SplineData = FString::ConstructFromPtrSize(SourceText, SplineDataLen);
+
+				// fix-up the hack applied to the Base64-encoded string in ExportTextItem()
+				SplineData.ReplaceCharInline('-', '/');
+
+				TArray<uint8> SplineReadBuffer;
+				bool bDecoded = FBase64::Decode(SplineData, SplineReadBuffer);
+				if (bDecoded)
+				{
+					FMemoryReader MemReader(SplineReadBuffer);
+					SerializeLoad(MemReader);
+				}
+			}
 		}
-	}
 
+
+	}
+	
 	return true;
 }
 
