Add ENGINE_API to GetMeshDeformerInstance
#tests none

diff --git a/Engine/Source/Runtime/Engine/Classes/Components/SkinnedMeshComponent.h b/Engine/Source/Runtime/Engine/Classes/Components/SkinnedMeshComponent.h
--- a/Engine/Source/Runtime/Engine/Classes/Components/SkinnedMeshComponent.h
+++ b/Engine/Source/Runtime/Engine/Classes/Components/SkinnedMeshComponent.h
@@ -340,7 +340,7 @@ private:
 
 public:
 	UFUNCTION(BlueprintPure, Category = "Components|SkinnedMesh")
-	UMeshDeformerInstance* GetMeshDeformerInstance() const;
+	ENGINE_API UMeshDeformerInstance* GetMeshDeformerInstance() const;
 
 	/** 
 	 * Gets the MeshDeformer for the given LOD.
