Fix Skinned Mesh Component validation

- Check for array out of bouds when reading the from the list of available deformers

-jira

diff --git a/Engine/Source/Runtime/Engine/Private/Components/SkinnedMeshComponent.cpp b/Engine/Source/Runtime/Engine/Private/Components/SkinnedMeshComponent.cpp
--- a/Engine/Source/Runtime/Engine/Private/Components/SkinnedMeshComponent.cpp
+++ b/Engine/Source/Runtime/Engine/Private/Components/SkinnedMeshComponent.cpp
@@ -1479,10 +1479,32 @@ EDataValidationResult USkinnedMeshComponent::IsDataValid(class FDataValidationCo
 	if (USkinnedAsset* Asset = GetSkinnedAsset())
 	{
 		FMeshDeformerSet MeshDeformers = GetActiveMeshDeformers();
-		
+
 		for (int32 LODIndex = 0; LODIndex < MeshDeformers.DeformerIndexForLOD.Num(); LODIndex++)
 		{
 			int32 DeformerIndex = MeshDeformers.DeformerIndexForLOD[LODIndex];
+
+			if (DeformerIndex == INDEX_NONE)
+			{
+				// No deformer is needed for this LOD
+				continue;
+			}
+			else if (!MeshDeformers.DeformerIndexForLOD.IsValidIndex(DeformerIndex))
+			{
+				// Deformer index for this LOD is invalid
+				Context.AddWarning(
+					FText::Format(
+						LOCTEXT("InvalidDeformerIndexForLOD",
+								"The deformer index {0} is invalid for Skeletal Mesh Asset {1} as "
+								"there are only {2} deformers available."),
+						DeformerIndex,
+						FText::FromName(Asset->GetPackage()->GetFName()),
+						MeshDeformers.Deformers.Num()
+					));
+				Result = EDataValidationResult::Invalid;
+				break;
+			}
+
 			UMeshDeformer* Deformer = MeshDeformers.Deformers[DeformerIndex];
 			if (Cast<ISkeletalMeshHalfEdgeBufferAccessor>(Deformer))
 			{
