[HWRT] Only use streaming code path for FStaticMeshRayTracingProxyLOD if there's data to stream.
- prevents an issue with the code path that skipped streaming in RequestRayTracingGeometryStreamIn(...) resulting in invalid state when using FRayTracingStreamableAsset.

diff --git a/Engine/Source/Runtime/Engine/Private/Rendering/RayTracingGeometryManager.cpp b/Engine/Source/Runtime/Engine/Private/Rendering/RayTracingGeometryManager.cpp
--- a/Engine/Source/Runtime/Engine/Private/Rendering/RayTracingGeometryManager.cpp
+++ b/Engine/Source/Runtime/Engine/Private/Rendering/RayTracingGeometryManager.cpp
@@ -375,6 +375,8 @@ void FRayTracingGeometryManager::SetRayTracingGeometryStreamingData(const FRayTr
 	RegisteredGeometry.StreamableAsset = &StreamableAsset;
 	RegisteredGeometry.StreamableBVHSize = StreamableAsset.GetRequestSizeBVH();
 	RegisteredGeometry.StreamableBuffersSize = StreamableAsset.GetRequestSizeBuffers();
+
+	checkf(RegisteredGeometry.StreamableBVHSize > 0 || RegisteredGeometry.StreamableBuffersSize > 0, TEXT("FRayTracingStreamableAsset should have data to stream."));
 }
 
 
@@ -928,18 +930,10 @@ bool FRayTracingGeometryManager::RequestRayTracingGeometryStreamIn(FRHICommandLi
 		}
 	}
 
-	TResourceArray<uint8> RawData;
-	RawData.SetAllowCPUAccess(true);
-
-	FResourceArrayInterface* OfflineData = nullptr;
+	checkf(RegisteredGeometry.StreamableBVHSize > 0 || RegisteredGeometry.StreamableBuffersSize > 0, TEXT("FRayTracingStreamableAsset should have data to stream."));
 
-	if (RegisteredGeometry.StreamableBuffersSize == 0 && RegisteredGeometry.StreamableBVHSize == 0)
-	{
-		// no offline data -> build from VB/IB at runtime
+	// TODO: Support DDC streaming
 
-		RegisteredGeometry.Status = FRegisteredGeometry::FStatus::StreamedIn;
-	}
-	else
 	{
 		if (NumStreamingRequests >= GRayTracingStreamingMaxPendingRequests)
 		{
@@ -978,27 +972,6 @@ bool FRayTracingGeometryManager::RequestRayTracingGeometryStreamIn(FRHICommandLi
 		RegisteredGeometry.Status = FRegisteredGeometry::FStatus::Streaming;
 	}
 
-	if (RegisteredGeometry.Status == FRegisteredGeometry::FStatus::StreamedIn)
-	{
-		{
-			FRHIResourceReplaceBatcher Batcher(RHICmdList, 1);
-			FRayTracingGeometryInitializer IntermediateInitializer = Geometry->Initializer;
-			IntermediateInitializer.Type = ERayTracingGeometryInitializerType::StreamingSource;
-			IntermediateInitializer.OfflineData = OfflineData;
-
-			FRayTracingGeometryRHIRef IntermediateRayTracingGeometry = RHICmdList.CreateRayTracingGeometry(IntermediateInitializer);
-
-			Geometry->SetRequiresBuild(IntermediateInitializer.OfflineData == nullptr || IntermediateRayTracingGeometry->IsCompressed());
-
-			Geometry->InitRHIForStreaming(IntermediateRayTracingGeometry, Batcher);
-
-			// When Batcher goes out of scope it will add commands to copy the BLAS buffers on RHI thread.
-			// We need to do it before we build the current geometry (also on RHI thread).
-		}
-
-		Geometry->RequestBuildIfNeeded(RHICmdList, ERTAccelerationStructureBuildPriority::Normal);
-	}
-
 	return true;
 }
 
diff --git a/Engine/Source/Runtime/Engine/Private/StaticMesh.cpp b/Engine/Source/Runtime/Engine/Private/StaticMesh.cpp
--- a/Engine/Source/Runtime/Engine/Private/StaticMesh.cpp
+++ b/Engine/Source/Runtime/Engine/Private/StaticMesh.cpp
@@ -2713,7 +2713,7 @@ void FStaticMeshRenderData::InitResources(ERHIFeatureLevel::Type InFeatureLevel,
 					{
 						RayTracingLOD.RayTracingGeometry->GroupHandle = RayTracingGeometryGroupHandle;
 
-						const bool bHasStreamableData = RayTracingLOD.StreamableData.GetBulkDataSize() > 0;
+						const bool bHasStreamableData = RayTracingLOD.StreamableData.GetBulkDataSize() > 0 && RayTracingLOD.GetRequestSize() > 0;
 
 						if (bHasStreamableData)
 						{
