Added s.LevelStreamingRouteActorEndPlayForRemoveFromWorldGranularity to allow UWorld::RemoveFromWorld to call EndPlay incrementally and check for time limit.
For backward compatibility, default value was set to 0 for now.

diff --git a/Engine/Source/Runtime/Engine/Classes/Engine/CoreSettings.h b/Engine/Source/Runtime/Engine/Classes/Engine/CoreSettings.h
--- a/Engine/Source/Runtime/Engine/Classes/Engine/CoreSettings.h
+++ b/Engine/Source/Runtime/Engine/Classes/Engine/CoreSettings.h
@@ -143,6 +143,8 @@ extern ENGINE_API int32 GLevelStreamingAddPrimitiveGranularity;
 extern ENGINE_API int32 GLevelStreamingComponentsUnregistrationGranularity;
 /** Batching granularity used to initialize actors during level streaming. If this is zero, we process all actors and stages in one pass. */
 extern ENGINE_API int32 GLevelStreamingRouteActorInitializationGranularity;
+/** Batching granularity used to call EndPlay on actors when removing streaming level. If this is zero, we process all actors in one pass. */
+extern ENGINE_API int32 GLevelStreamingRouteActorEndPlayForRemoveFromWorldGranularity;
 /** Maximum allowed time to spend for actor unregistration steps during level streaming (ms per frame). If this is 0.0 then we don't timeslice.*/
 extern ENGINE_API float GLevelStreamingUnregisterComponentsTimeLimit;
 /** Whether to force a GC after levels are streamed out to instantly reclaim the memory at the expensive of a hitch. */
diff --git a/Engine/Source/Runtime/Engine/Classes/Engine/Level.h b/Engine/Source/Runtime/Engine/Classes/Engine/Level.h
--- a/Engine/Source/Runtime/Engine/Classes/Engine/Level.h
+++ b/Engine/Source/Runtime/Engine/Classes/Engine/Level.h
@@ -871,6 +871,7 @@ private:
 	TSet<UActorComponent*> PreRegisteringActorComponents;
 	friend struct FLevelRegistrationAccessor;
 	friend struct FLevelPendingVisibilityStateAccessor;
+	friend struct FLevelRemoveFromWorldAccessor;
 
 protected:
 
@@ -916,6 +917,7 @@ private:
 	};
 	ERouteActorInitializationState RouteActorInitializationState;
 	int32 RouteActorInitializationIndex;
+	int32 RouteActorEndPlayForRemoveFromWorldIndex;
 
 public:
 	// Used internally to determine which actors should go on the world's NetworkActor list
@@ -1486,6 +1488,10 @@ public:
 #endif
 
 private:
+	void RouteActorEndPlayForRemoveFromWorld(int32 NumActorsToProcess);
+	void ResetRouteActorEndPlayForRemoveFromWorld();
+	bool IsFinishedRouteActorEndPlayForRemoveFromWorld() const;
+
 	bool IncrementalPreRegisterComponents(FRegisterComponentContext& Context);
 	bool IncrementalRegisterComponents(FRegisterComponentContext& Context);
 	bool IncrementalPreUnregisterComponents(FUnregisterComponentContext& Context);
@@ -1549,6 +1555,28 @@ private:
 #endif // WITH_EDITOR
 };
 
+/** Accessor used by World when removing level */
+struct FLevelRemoveFromWorldAccessor
+{
+private:
+	ENGINE_API static void RouteActorEndPlayForRemoveFromWorld(ULevel* InLevel, int32 InNumActorsToProcess)
+	{
+		InLevel->RouteActorEndPlayForRemoveFromWorld(InNumActorsToProcess);
+	}
+
+	ENGINE_API static void ResetRouteActorEndPlayForRemoveFromWorld(ULevel* InLevel)
+	{
+		InLevel->ResetRouteActorEndPlayForRemoveFromWorld();
+	}
+
+	ENGINE_API static bool IsFinishedRouteActorEndPlayForRemoveFromWorld(ULevel* InLevel)
+	{ 
+		return InLevel->IsFinishedRouteActorEndPlayForRemoveFromWorld();
+	}
+
+	friend class UWorld;
+};
+
 /** Accessor used by World/ActorComponent for component registration/unregistration. */
 struct FLevelRegistrationAccessor
 {
diff --git a/Engine/Source/Runtime/Engine/Private/CoreSettings.cpp b/Engine/Source/Runtime/Engine/Private/CoreSettings.cpp
--- a/Engine/Source/Runtime/Engine/Private/CoreSettings.cpp
+++ b/Engine/Source/Runtime/Engine/Private/CoreSettings.cpp
@@ -21,6 +21,7 @@ int32 GLevelStreamingComponentsRegistrationGranularity = 10;
 int32 GLevelStreamingAddPrimitiveGranularity = 120;
 int32 GLevelStreamingComponentsUnregistrationGranularity = 5;
 int32 GLevelStreamingRouteActorInitializationGranularity = 10;
+int32 GLevelStreamingRouteActorEndPlayForRemoveFromWorldGranularity = 0;
 int32 GLevelStreamingForceGCAfterLevelStreamedOut = 1;
 int32 GLevelStreamingContinuouslyIncrementalGCWhileLevelsPendingPurge = 1;
 int32 GLevelStreamingAllowLevelRequestsWhileAsyncLoadingInMatch = 1;
@@ -115,6 +116,13 @@ static FAutoConsoleVariableRef CVarLevelStreamingRouteActorInitializationGranula
 	ECVF_Default
 );
 
+static FAutoConsoleVariableRef CVarLevelStreamingRouteActorEndPlayForRemoveFromWorldGranularity(
+	TEXT("s.LevelStreamingRouteActorEndPlayForRemoveFromWorldGranularity"),
+	GLevelStreamingRouteActorEndPlayForRemoveFromWorldGranularity,
+	TEXT("Batching granularity used to call EndPlay on actors when removing streaming level. If this is zero, we process all actors in one pass."),
+	ECVF_Default
+);
+
 static FAutoConsoleVariableRef CVarForceGCAfterLevelStreamedOut(
 	TEXT("s.ForceGCAfterLevelStreamedOut"),
 	GLevelStreamingForceGCAfterLevelStreamedOut,
diff --git a/Engine/Source/Runtime/Engine/Private/Level.cpp b/Engine/Source/Runtime/Engine/Private/Level.cpp
--- a/Engine/Source/Runtime/Engine/Private/Level.cpp
+++ b/Engine/Source/Runtime/Engine/Private/Level.cpp
@@ -426,6 +426,7 @@ ULevel::ULevel( const FObjectInitializer& ObjectInitializer )
 	,	PrecomputedVolumetricLightmap(new FPrecomputedVolumetricLightmap())
 	,	RouteActorInitializationState(ERouteActorInitializationState::Preinitialize)
 	,	RouteActorInitializationIndex(0)
+	,   RouteActorEndPlayForRemoveFromWorldIndex(0)
 {
 #if WITH_EDITORONLY_DATA
 	LevelColor = FLinearColor::White;
@@ -3881,6 +3882,35 @@ void ULevel::RouteActorInitialize(int32 NumActorsToProcess)
 	}
 }
 
+void ULevel::RouteActorEndPlayForRemoveFromWorld(int32 NumActorsToProcess)
+{
+	const bool bFullProcessing = (NumActorsToProcess <= 0);
+
+	while (RouteActorEndPlayForRemoveFromWorldIndex < Actors.Num())
+	{
+		if (AActor* const Actor = Actors[RouteActorEndPlayForRemoveFromWorldIndex])
+		{
+			Actor->RouteEndPlay(EEndPlayReason::RemovedFromWorld);
+		}
+
+		++RouteActorEndPlayForRemoveFromWorldIndex;
+		if (!bFullProcessing && (--NumActorsToProcess == 0))
+		{
+			return;
+		}
+	}
+}
+
+void ULevel::ResetRouteActorEndPlayForRemoveFromWorld()
+{
+	RouteActorEndPlayForRemoveFromWorldIndex = 0;
+}
+
+bool ULevel::IsFinishedRouteActorEndPlayForRemoveFromWorld() const
+{
+	return RouteActorEndPlayForRemoveFromWorldIndex >= Actors.Num();
+}
+
 UPackage* ULevel::CreateMapBuildDataPackage() const
 {
 	FString PackageName = GetOutermost()->GetName() + TEXT("_BuiltData");
diff --git a/Engine/Source/Runtime/Engine/Private/World.cpp b/Engine/Source/Runtime/Engine/Private/World.cpp
--- a/Engine/Source/Runtime/Engine/Private/World.cpp
+++ b/Engine/Source/Runtime/Engine/Private/World.cpp
@@ -3830,6 +3830,16 @@ void UWorld::RemoveFromWorld(ULevel* Level, bool bAllowIncrementalRemoval, const
 	SCOPE_CYCLE_COUNTER(STAT_RemoveFromWorldTime);
 	CSV_SCOPED_TIMING_STAT_EXCLUSIVE(RemoveFromWorld);
 
+	UE::FTimeout Timeout = UE::FTimeout::Never();
+	// Can't do an incremental removal if there is no time limit
+	bAllowIncrementalRemoval = bAllowIncrementalRemoval && (ExternalTimeout.IsSet() || GLevelStreamingUnregisterComponentsTimeLimit > 0.0);
+	if (bAllowIncrementalRemoval)
+	{
+		// Check external timeout or time limit (which includes GRemoveFromWorldTimeCumul) before allowing to begin making invisible a level
+		Timeout = ExternalTimeout.IsSet() ? ExternalTimeout.GetValue()
+			: UE::FTimeout(FMath::Max<double>(0.0, GLevelStreamingUnregisterComponentsTimeLimit - GRemoveFromWorldTimeCumul) / 1000.0);
+	}
+
 	FScopeCycleCounterUObject ScopeContext(Level);
 	check(IsValid(Level));
 	check(!Level->IsUnreachable());
@@ -3843,6 +3853,8 @@ void UWorld::RemoveFromWorld(ULevel* Level, bool bAllowIncrementalRemoval, const
 
 	auto BeginRemoval = [Level, bIsGameWorld, this, OwningLevelStreaming]()
 	{
+		check(!Level->bIsBeingRemoved);
+
 		FWorldDelegates::PreLevelRemovedFromWorld.Broadcast(Level, this);
 		Level->bIsBeingRemoved = true;
 
@@ -3879,15 +3891,8 @@ void UWorld::RemoveFromWorld(ULevel* Level, bool bAllowIncrementalRemoval, const
 		bool bWaitForCompletion = false;
 		bool bFinishRemovingLevel = true;
 
-		// Can't do an incremental removal if there is no time limit
-		bAllowIncrementalRemoval = bAllowIncrementalRemoval && (ExternalTimeout.IsSet() || GLevelStreamingUnregisterComponentsTimeLimit > 0.0);
-
 		if (bAllowIncrementalRemoval)
 		{
-			// Check external timeout or time limit (which includes GRemoveFromWorldTimeCumul) before allowing to begin making invisible a level
-			const UE::FTimeout Timeout = ExternalTimeout.IsSet() ? ExternalTimeout.GetValue()
-				: UE::FTimeout(FMath::Max<double>(0.0, GLevelStreamingUnregisterComponentsTimeLimit - GRemoveFromWorldTimeCumul) / 1000.0);
-
 			bFinishRemovingLevel = false;
 			if (!IsLevelMakingInvisible(Level) && (MakingInvisibleLevels.Num() < FMath::Max<int32>(1, GMaximumMakingInvisibleLevels)))
 			{
@@ -3937,7 +3942,7 @@ void UWorld::RemoveFromWorld(ULevel* Level, bool bAllowIncrementalRemoval, const
 			bWaitForCompletion = true;
 			check(bFinishRemovingLevel);
 		}
-		else
+		else if (!Level->bIsBeingRemoved)
 		{
 			BeginRemoval();
 			bWaitForCompletion = true;
@@ -3957,92 +3962,101 @@ void UWorld::RemoveFromWorld(ULevel* Level, bool bAllowIncrementalRemoval, const
 
 		if (bFinishRemovingLevel)
 		{
-			// The level can be removed from the pending visibility state
-			SetLevelPendingVisibilityState(Level, ELevelPendingVisibilityState::None);
 			TRACE_CPUPROFILER_EVENT_SCOPE(RemoveFromWorld_FinishRemovingLevel);
 
-			for (int32 ActorIdx = 0; ActorIdx < Level->Actors.Num(); ActorIdx++)
+			// Route various initialization functions and set volumes.
+			if (!Timeout.IsExpired() && !FLevelRemoveFromWorldAccessor::IsFinishedRouteActorEndPlayForRemoveFromWorld(Level))
 			{
-				if (AActor* Actor = Level->Actors[ActorIdx])
+				TRACE_CPUPROFILER_EVENT_SCOPE(RemoveFromWorld_RouteActorEndPlay);
+				const int32 NumActorsToProcess = (!bAllowIncrementalRemoval || !bIsGameWorld || IsRunningCommandlet()) ? 0 : GLevelStreamingRouteActorEndPlayForRemoveFromWorldGranularity;
+				do
 				{
-					Actor->RouteEndPlay(EEndPlayReason::RemovedFromWorld);
-				}
+					FLevelRemoveFromWorldAccessor::RouteActorEndPlayForRemoveFromWorld(Level, NumActorsToProcess);
+				} while (!FLevelRemoveFromWorldAccessor::IsFinishedRouteActorEndPlayForRemoveFromWorld(Level) && !IsTimeLimitExceeded(Timeout, TEXT("routing actor EndPlay"), Level));
 			}
 
-			// Remove any pawns from the pawn list that are about to be streamed out
-			for (APawn* Pawn : TActorRange<APawn>(this))
+			if (!Timeout.IsExpired() && FLevelRemoveFromWorldAccessor::IsFinishedRouteActorEndPlayForRemoveFromWorld(Level))
 			{
-				if (Pawn->IsInLevel(Level))
+				// The level can be removed from the pending visibility state
+				SetLevelPendingVisibilityState(Level, ELevelPendingVisibilityState::None);
+
+				// Remove any pawns from the pawn list that are about to be streamed out
+				for (APawn* Pawn : TActorRange<APawn>(this))
 				{
-					AController* Controller = Pawn->GetController();
-					// This should have happened as part of the RouteEndPlay above, but ensuring to validate this assumption and maintain behavior
-					// with RemovePawn having been deprecated
-					if (!ensure(Controller == nullptr || (Controller->GetPawn() == Pawn)))
+					if (Pawn->IsInLevel(Level))
 					{
-						Controller->UnPossess();
+						AController* Controller = Pawn->GetController();
+						// This should have happened as part of the RouteEndPlay above, but ensuring to validate this assumption and maintain behavior
+						// with RemovePawn having been deprecated
+						if (!ensure(Controller == nullptr || (Controller->GetPawn() == Pawn)))
+						{
+							Controller->UnPossess();
+						}
+					}
+					else if (UCharacterMovementComponent* CharacterMovement = Cast<UCharacterMovementComponent>(Pawn->GetMovementComponent()))
+					{
+						// otherwise force floor check in case the floor was streamed out from under it
+						CharacterMovement->bForceNextFloorCheck = true;
 					}
 				}
-				else if (UCharacterMovementComponent* CharacterMovement = Cast<UCharacterMovementComponent>(Pawn->GetMovementComponent()))
-				{
-					// otherwise force floor check in case the floor was streamed out from under it
-					CharacterMovement->bForceNextFloorCheck = true;
-				}
-			}
 
-			Level->ReleaseRenderingResources();
+				Level->ReleaseRenderingResources();
 
-			// Remove from the world's level array and destroy actor components.
-			IStreamingManager::Get().RemoveLevel(Level);
+				// Remove from the world's level array and destroy actor components.
+				IStreamingManager::Get().RemoveLevel(Level);
 
-			Level->ClearLevelComponents();
-
-			if (bIsGameWorld && !Level->bClientOnlyVisible)
-			{
-				FUpdateLevelVisibilityLevelInfo LevelVisibility(Level, false);
-				const FName UnmappedPackageName = LevelVisibility.PackageName;
+				Level->ClearLevelComponents();
 
-				// notify server that the client has removed this level
-				for (FLocalPlayerIterator It(GEngine, this); It; ++It)
+				if (bIsGameWorld && !Level->bClientOnlyVisible)
 				{
-					if (APlayerController* LocalPlayerController = It->GetPlayerController(this))
+					FUpdateLevelVisibilityLevelInfo LevelVisibility(Level, false);
+					const FName UnmappedPackageName = LevelVisibility.PackageName;
+
+					// notify server that the client has removed this level
+					for (FLocalPlayerIterator It(GEngine, this); It; ++It)
 					{
-						LevelVisibility.PackageName = LocalPlayerController->NetworkRemapPath(UnmappedPackageName, false);
-						LevelVisibility.VisibilityRequestId = TransactionId;
-						LocalPlayerController->ServerUpdateLevelVisibility(LevelVisibility);
+						if (APlayerController* LocalPlayerController = It->GetPlayerController(this))
+						{
+							LevelVisibility.PackageName = LocalPlayerController->NetworkRemapPath(UnmappedPackageName, false);
+							LevelVisibility.VisibilityRequestId = TransactionId;
+							LocalPlayerController->ServerUpdateLevelVisibility(LevelVisibility);
+						}
 					}
 				}
-			}
 
-			// We expect level that use the bRequireFullVisibilityToRender flag to already be !visible at this point
-			check(!Level->bRequireFullVisibilityToRender || !Level->bIsVisible);
-			Level->bIsVisible = false;
+				// We expect level that use the bRequireFullVisibilityToRender flag to already be !visible at this point
+				check(!Level->bRequireFullVisibilityToRender || !Level->bIsVisible);
+				Level->bIsVisible = false;
 
-			// Notify world composition: will place a level at original position
-			if (WorldComposition)
-			{
-				WorldComposition->OnLevelRemovedFromWorld(Level);
-			}
+				// Notify world composition: will place a level at original position
+				if (WorldComposition)
+				{
+					WorldComposition->OnLevelRemovedFromWorld(Level);
+				}
 
-			// Make sure level always has OwningWorld in the editor
-			if (bIsGameWorld)
-			{
-				Levels.Remove(Level);
-				Level->OwningWorld = nullptr;
-			}
+				// Make sure level always has OwningWorld in the editor
+				if (bIsGameWorld)
+				{
+					Levels.Remove(Level);
+					Level->OwningWorld = nullptr;
+				}
 
-			// let the universe know we have removed a level
-			FWorldDelegates::LevelRemovedFromWorld.Broadcast(Level, this);
-			BroadcastLevelsChanged();
+				// let the universe know we have removed a level
+				FWorldDelegates::LevelRemovedFromWorld.Broadcast(Level, this);
+				BroadcastLevelsChanged();
 
-			// If the level requires full visibility to be rendered, we already made it non visible in BeginRemoval()
-			if (!Level->bRequireFullVisibilityToRender)
-			{
-				ULevelStreaming::BroadcastLevelVisibleStatus(this, Level->GetOutermost()->GetFName(), false);
-			}
+				// If the level requires full visibility to be rendered, we already made it non visible in BeginRemoval()
+				if (!Level->bRequireFullVisibilityToRender)
+				{
+					ULevelStreaming::BroadcastLevelVisibleStatus(this, Level->GetOutermost()->GetFName(), false);
+				}
 
-			Level->bIsBeingRemoved = false;
+				FLevelRemoveFromWorldAccessor::ResetRouteActorEndPlayForRemoveFromWorld(Level);
+				Level->bIsBeingRemoved = false;
+
+				TRACE_END_REGION(*WriteToString<256>(TEXT("RemoveFromWorld: "), Level->GetPackage()->GetName()));
+			}
 
-			TRACE_END_REGION(*WriteToString<256>(TEXT("RemoveFromWorld: "), Level->GetPackage()->GetName()));
 		} // if ( bFinishRemovingLevel )
 
 		// Keep track of time spent for entire incremental removal
