[Waveform Editor] - Changes to cue points aren't reflected in the Sound Wave data/during playback
Exposing Marker transformation cue points to other systems in unreal

diff --git a/Engine/Plugins/Editor/WaveformEditor/Source/WaveformTransformations/Private/WaveformTransformationMarkers.cpp b/Engine/Plugins/Editor/WaveformEditor/Source/WaveformTransformations/Private/WaveformTransformationMarkers.cpp
--- a/Engine/Plugins/Editor/WaveformEditor/Source/WaveformTransformations/Private/WaveformTransformationMarkers.cpp
+++ b/Engine/Plugins/Editor/WaveformEditor/Source/WaveformTransformations/Private/WaveformTransformationMarkers.cpp
@@ -65,6 +65,12 @@ void UWaveformTransformationMarkers::OverwriteSoundWaveData(USoundWave& InOutSou
 	}
 }
 
+void UWaveformTransformationMarkers::GetTransformationInfo(FWaveformTransformationInfo& InOutTransformationInfo) const
+{
+	check(Markers != nullptr);
+	InOutTransformationInfo.AllCuePoints.Append(Markers->WaveCues);
+}
+
 void UWaveCueArray::PostEditChangeProperty(FPropertyChangedEvent& PropertyChangedEvent)
 {
 	Super::PostEditChangeProperty(PropertyChangedEvent);
diff --git a/Engine/Plugins/Editor/WaveformEditor/Source/WaveformTransformations/Public/WaveformTransformationMarkers.h b/Engine/Plugins/Editor/WaveformEditor/Source/WaveformTransformations/Public/WaveformTransformationMarkers.h
--- a/Engine/Plugins/Editor/WaveformEditor/Source/WaveformTransformations/Public/WaveformTransformationMarkers.h
+++ b/Engine/Plugins/Editor/WaveformEditor/Source/WaveformTransformations/Public/WaveformTransformationMarkers.h
@@ -52,5 +52,6 @@ public:
 
 #if WITH_EDITOR
 	void OverwriteSoundWaveData(USoundWave& InOutSoundWave) override;
+	void GetTransformationInfo(FWaveformTransformationInfo& InOutTransformationInfo) const override;
 #endif //WITH_EDITOR
 };
diff --git a/Engine/Source/Editor/AudioEditor/Private/Factories/SoundFactory.cpp b/Engine/Source/Editor/AudioEditor/Private/Factories/SoundFactory.cpp
--- a/Engine/Source/Editor/AudioEditor/Private/Factories/SoundFactory.cpp
+++ b/Engine/Source/Editor/AudioEditor/Private/Factories/SoundFactory.cpp
@@ -667,7 +667,11 @@ UObject* USoundFactory::CreateObject
 		// Store the current file path and timestamp for re-import purposes
 		Sound->AssetImportData->Update(CurrentFilename);
 
+#if WITH_EDITORONLY_DATA
 		// Setup the cue points
+		Sound->CuePointOrigin = ESoundWaveCuePointOrigin::WaveFile;
+#endif //WITH_EDITORONLY_DATA
+
 		int TotalNum = WaveInfo.WaveCues.Num() + WaveInfo.WaveSampleLoops.Num();
 		Sound->CuePoints.Reset(TotalNum);
 		TSet<int32> WaveCuePointIDs;
diff --git a/Engine/Source/Runtime/AudioExtensions/Public/IWaveformTransformation.h b/Engine/Source/Runtime/AudioExtensions/Public/IWaveformTransformation.h
--- a/Engine/Source/Runtime/AudioExtensions/Public/IWaveformTransformation.h
+++ b/Engine/Source/Runtime/AudioExtensions/Public/IWaveformTransformation.h
@@ -102,6 +102,14 @@ struct FWaveTransformUObjectConfiguration
 	TArray<FSoundWaveCuePoint> WaveCues; // List of cues parsed from the wave file
 };
 
+#if WITH_EDITORONLY_DATA
+// Information to be retrieved from each transformation (Add members as needed when new transformation types are added)
+struct FWaveformTransformationInfo
+{
+	TArray<FSoundWaveCuePoint> AllCuePoints; // Cue Points and Loop Regions
+};
+#endif //WITH_EDITORONLY_DATA
+
 // Base class to hold editor configurable properties for an arbitrary transformation of audio waveform data
 UCLASS(Abstract, EditInlineNew, MinimalAPI)
 class UWaveformTransformationBase : public UObject
@@ -115,6 +123,7 @@ public:
 	
 #if WITH_EDITOR
 	virtual void OverwriteSoundWaveData(USoundWave& InOutSoundWave) {};
+	virtual void GetTransformationInfo(FWaveformTransformationInfo& InOutTransformationInfo) const {};
 #endif //WITH_EDITOR
 
 	virtual bool IsEditorOnly() const override { return true; }
diff --git a/Engine/Source/Runtime/Engine/Classes/Sound/SoundWave.h b/Engine/Source/Runtime/Engine/Classes/Sound/SoundWave.h
--- a/Engine/Source/Runtime/Engine/Classes/Sound/SoundWave.h
+++ b/Engine/Source/Runtime/Engine/Classes/Sound/SoundWave.h
@@ -65,6 +65,17 @@ inline const TCHAR* LexToString(const EDecompressionType InType)
 #undef CASE_TO_STRING
 }
 
+/**
+ * Dictates whether to use the CuePoints and Loop Regions from the .wav file or from the waveform editor
+ * Marker transformations during playback
+ */
+UENUM()
+enum class ESoundWaveCuePointOrigin : uint8
+{
+	WaveFile = 0,
+	MarkerTransformation,
+};
+
 /** Precache states */
 enum class ESoundWavePrecacheState
 {
@@ -606,6 +617,7 @@ public:
 	ENGINE_API const FWaveTransformUObjectConfiguration& UpdateTransformations();
 	ENGINE_API TArray<uint8> CreateNewCueData(const char ChunkName[4]);
 	ENGINE_API void OverwriteSoundWaveDataWithTransformationData();
+	ENGINE_API FWaveformTransformationInfo GetWaveformTransformationInfo() const;
 
 #endif
 
@@ -770,9 +782,15 @@ protected:
 	UPROPERTY(Category = Info, AssetRegistrySearchable, VisibleAnywhere)
 	int32 ImportedSampleRate;
 
-	/** Cue point data parsed fro the .wav file. Contains "Loop Regions" as cue points as well! */
+	/** Cue point data parsed from the .wav file. Contains "Loop Regions" as cue points as well! */
 	UPROPERTY(Category = Info, VisibleAnywhere, BlueprintGetter = GetCuePoints)
 	TArray<FSoundWaveCuePoint> CuePoints;
+
+	/** Dictates whether to use the CuePoints and Loop Regions from the .wav file or from the waveform editor 
+	 *	Marker transformations during playback.
+	 */
+	UPROPERTY(Category = Info, EditAnywhere)
+	ESoundWaveCuePointOrigin CuePointOrigin;
 #endif
 
 	ENGINE_API virtual void SerializeCuePoints(FArchive& Ar, const bool bIsLoadingFromCookedArchive);
diff --git a/Engine/Source/Runtime/Engine/Private/SoundWave.cpp b/Engine/Source/Runtime/Engine/Private/SoundWave.cpp
--- a/Engine/Source/Runtime/Engine/Private/SoundWave.cpp
+++ b/Engine/Source/Runtime/Engine/Private/SoundWave.cpp
@@ -348,6 +348,7 @@ void FSoundWaveData::SetAllCuePoints(const TArray<FSoundWaveCuePoint>& InCuePoin
 {
 	CuePoints.Reset();
 	LoopRegions.Reset();
+
 	for (const FSoundWaveCuePoint& CuePoint : InCuePoints)
 	{
 		if (!CuePoint.IsLoopRegion())
@@ -963,6 +964,8 @@ USoundWave::USoundWave(const FObjectInitializer& ObjectInitializer)
 	EnvelopeFollowerAttackTime = 10;
 	EnvelopeFollowerReleaseTime = 100;
 	bEnableCloudStreaming = false;
+
+	CuePointOrigin = ESoundWaveCuePointOrigin::WaveFile;
 #endif
 
 	bCachedSampleRateFromPlatformSettings = false;
@@ -2520,8 +2523,23 @@ ISoundWaveLoadingBehaviorUtil::FClassData USoundWave::GetOwnerLoadingBehavior(co
 
 TArray<FSoundWaveCuePoint> USoundWave::GetCuePointsScaledForSampleRate(const float InSampleRate) const
 {
-	TArray<FSoundWaveCuePoint> CuePointsScaled = CuePoints;
+	TArray<FSoundWaveCuePoint> CuePointsScaled;
+
+	switch (CuePointOrigin)
+	{
+	case ESoundWaveCuePointOrigin::WaveFile:
+		CuePointsScaled = CuePoints;
+		break;
+	case ESoundWaveCuePointOrigin::MarkerTransformation:
+		CuePointsScaled = GetWaveformTransformationInfo().AllCuePoints;
+		break;
+	default:
+		CuePointsScaled = CuePoints;
+		break;
+	}
+
 	ScaleCuePointsForSampleRate(InSampleRate, CuePointsScaled);
+
 	return CuePointsScaled;
 }
 
@@ -4678,6 +4696,20 @@ void USoundWave::OverwriteSoundWaveDataWithTransformationData()
 	RawData.UpdatePayload(Buffer);
 }
 
+FWaveformTransformationInfo USoundWave::GetWaveformTransformationInfo() const
+{
+	FWaveformTransformationInfo TransformationInfo;
+	for (TObjectPtr<UWaveformTransformationBase> Transformation : Transformations)
+	{
+		if (Transformation)
+		{
+			Transformation->GetTransformationInfo(TransformationInfo);
+		}
+	}
+
+	return TransformationInfo;
+}
+
 #endif // #if WITH_EDITOR
 
 void USoundWave::CacheInheritedLoadingBehavior() const
@@ -4743,7 +4775,19 @@ void USoundWave::SerializeCuePoints(FArchive& Ar, const bool bIsLoadingFromCooke
 #if WITH_EDITORONLY_DATA
 	if (Ar.IsCooking())
 	{
-		PlatformCuePoints = CuePoints;
+		switch (CuePointOrigin)
+		{
+		case ESoundWaveCuePointOrigin::WaveFile:
+			PlatformCuePoints = CuePoints;
+			break;
+		case ESoundWaveCuePointOrigin::MarkerTransformation:
+			PlatformCuePoints = GetWaveformTransformationInfo().AllCuePoints;
+			break;
+		default:
+			PlatformCuePoints = CuePoints;
+			break;
+		}
+
 		float ResampleRatio = -1.0f;
 		const ITargetPlatform* CookingTarget = Ar.CookingTarget();
 		if (ensure(CookingTarget))
@@ -4763,6 +4807,7 @@ void USoundWave::SerializeCuePoints(FArchive& Ar, const bool bIsLoadingFromCooke
 	else if (Ar.IsLoading() && !bIsLoadingFromCookedArchive)
 	{
 		SoundWaveDataPtr->SetAllCuePoints(GetCuePointsScaledForSampleRate(GetSampleRateForCurrentPlatform()));
+
 		return;
 	}
 #endif
