Submix Send - Hide irrelevant properties based on Send Level Control Method
- Also clean up tooltips to be more readable

diff --git a/Engine/Source/Runtime/Engine/Classes/Sound/SoundSubmixSend.h b/Engine/Source/Runtime/Engine/Classes/Sound/SoundSubmixSend.h
--- a/Engine/Source/Runtime/Engine/Classes/Sound/SoundSubmixSend.h
+++ b/Engine/Source/Runtime/Engine/Classes/Sound/SoundSubmixSend.h
@@ -88,8 +88,8 @@ struct FSoundSubmixSendInfoBase
 
 	/*
 		Manual: Use Send Level only
-		Linear: Interpolate between Min and Max Send Levels based on listener distance (between Distance Min and Distance Max)
-		Custom Curve: Use the float curve to map Send Level to distance (0.0-1.0 on curve maps to Distance Min - Distance Max)
+		Linear: Interpolate between Min and Max Send Levels based on listener distance (between Min/Max Send Distance)
+		Custom Curve: Use the float curve to map Send Level to distance (0.0-1.0 on curve maps to Min/Max Send Distance)
 	*/
 	UPROPERTY(EditAnywhere, BlueprintReadWrite, Category = SubmixSend)
 	ESendLevelControlMethod SendLevelControlMethod;
@@ -98,42 +98,46 @@ struct FSoundSubmixSendInfoBase
 	UPROPERTY(EditAnywhere, BlueprintReadWrite, Category = SubmixSend)
 	TObjectPtr<USoundSubmixBase> SoundSubmix;
 
-	// The amount of audio to send
-	UPROPERTY(EditAnywhere, BlueprintReadWrite, Category = SubmixSend, meta = (DisplayName = "Manual Send Level"))
+	// Manually set the amount of audio to send
+	UPROPERTY(EditAnywhere, BlueprintReadWrite, Category = SubmixSend, meta = (DisplayName = "Manual Send Level", EditCondition = "SendLevelControlMethod == ESendLevelControlMethod::Manual", EditConditionHides))
 	float SendLevel;
 
-	// Whether to disable the 0-1 clamp for manual SendLevel control
-	UPROPERTY(EditAnywhere, BlueprintReadWrite, Category = SubmixSend)
+	// Whether to disable the internal 0-1 clamp for Manual Send Level control
+	UPROPERTY(EditAnywhere, BlueprintReadWrite, Category = SubmixSend, meta = (EditCondition = "SendLevelControlMethod == ESendLevelControlMethod::Manual", EditConditionHides))
 	bool DisableManualSendClamp;
 
-	// The amount to send to master when sound is located at a distance equal to value specified in the min send distance.
-	UPROPERTY(EditAnywhere, BlueprintReadWrite, Category = SubmixSend)
+	// The amount to send to the Submix when sound is located at a distance less than or equal to value specified in the Min Send Distance
+	UPROPERTY(EditAnywhere, BlueprintReadWrite, Category = SubmixSend, meta = (EditCondition = "SendLevelControlMethod == ESendLevelControlMethod::Linear", EditConditionHides))
 	float MinSendLevel;
 
-	// The amount to send to master when sound is located at a distance equal to value specified in the max send distance.
-	UPROPERTY(EditAnywhere, BlueprintReadWrite, Category = SubmixSend)
+	// The amount to send to the Submix when sound is located at a distance greater than or equal to value specified in the Max Send Distance
+	UPROPERTY(EditAnywhere, BlueprintReadWrite, Category = SubmixSend, meta = (EditCondition = "SendLevelControlMethod == ESendLevelControlMethod::Linear", EditConditionHides))
 	float MaxSendLevel;
 
-	// The min distance to send to the master
-	UPROPERTY(EditAnywhere, BlueprintReadWrite, Category = SubmixSend)
+	/** The distance at which to start mapping between to Min/Max Send Level
+	 *  Distances LESS than this will result in a clamped Min Send Level
+	 */ 
+	UPROPERTY(EditAnywhere, BlueprintReadWrite, Category = SubmixSend, meta = (EditCondition = "SendLevelControlMethod != ESendLevelControlMethod::Manual", EditConditionHides))
 	float MinSendDistance;
 
-	// The max distance to send to the master
-	UPROPERTY(EditAnywhere, BlueprintReadWrite, Category = SubmixSend)
+	/** The distance at which to stop mapping between Min/Max Send Level
+	 *  Distances GREATER than this will result in a clamped Max Send Level
+	 */
+	UPROPERTY(EditAnywhere, BlueprintReadWrite, Category = SubmixSend, meta = (EditCondition = "SendLevelControlMethod != ESendLevelControlMethod::Manual", EditConditionHides))
 	float MaxSendDistance;
 
-	// The custom reverb send curve to use for distance-based send level.
-	UPROPERTY(EditAnywhere, BlueprintReadWrite, Category = SubmixSend)
+	// The custom reverb send curve to use for distance-based send level. (0.0-1.0 on curve maps to Min/Max Send Distance)
+	UPROPERTY(EditAnywhere, BlueprintReadWrite, Category = SubmixSend, meta = (EditCondition = "SendLevelControlMethod == ESendLevelControlMethod::CustomCurve", EditConditionHides))
 	FRuntimeFloatCurve CustomSendLevelCurve;
 };
 
 UENUM(BlueprintType)
 enum class ESubmixSendStage : uint8
 {
-	// Whether to do the send pre distance attenuation
+	// Whether to do the send post distance attenuation
 	PostDistanceAttenuation,
 
-	// Whether to do the send post distance attenuation
+	// Whether to do the send pre distance attenuation
 	PreDistanceAttenuation,
 };
 
