[Skeletal mesh LOD streaming] Add support for syncing LOD levels in an attachment tree when streaming is enabled. Old behavior can be restored with r.SkinnedMesh.SyncStreamingLODClamp=0.

#tests tested in editor and on a console

diff --git a/Engine/Source/Runtime/Engine/Classes/Components/SkinnedMeshComponent.h b/Engine/Source/Runtime/Engine/Classes/Components/SkinnedMeshComponent.h
--- a/Engine/Source/Runtime/Engine/Classes/Components/SkinnedMeshComponent.h
+++ b/Engine/Source/Runtime/Engine/Classes/Components/SkinnedMeshComponent.h
@@ -738,6 +738,9 @@ protected:
 	/** Record of the tick rate we are using when externally controlled */
 	uint8 ExternalTickRate;
 
+	/** The highest streaming LOD index clamp among components with SyncAttachParentLOD enabled in the attachment tree */
+	uint8 MaxFirstLODIdxInAttachmentTree;
+
 protected:
 	/** used to cache previous bone transform or not */
 	uint8 bHasValidBoneTransform:1;
@@ -1523,6 +1526,9 @@ protected:
 	/** Helper function for UpdateLODStatus, called with a valid index for InLeaderPoseComponentPredictedLODLevel when updating LOD status for follower components */
 	ENGINE_API bool UpdateLODStatus_Internal(int32 InLeaderPoseComponentPredictedLODLevel, bool bRequestedByLeaderPoseComponent = false);
 
+	/** Traverse the attachment tree and compute the highest streaming LOD index clamp among components with SyncAttachParentLOD enabled */
+	uint8 ComputeMaxFirstLODIdxInAttachmentTree() const;
+
 public:
 	/**
 	 * Tick Pose, this function ticks and do whatever it needs to do in this frame, should be called before RefreshBoneTransforms
diff --git a/Engine/Source/Runtime/Engine/Private/Components/SkinnedMeshComponent.cpp b/Engine/Source/Runtime/Engine/Private/Components/SkinnedMeshComponent.cpp
--- a/Engine/Source/Runtime/Engine/Private/Components/SkinnedMeshComponent.cpp
+++ b/Engine/Source/Runtime/Engine/Private/Components/SkinnedMeshComponent.cpp
@@ -112,6 +112,12 @@ static FAutoConsoleVariableRef CVarSkinnedMeshRenderNanite(
 	)
 );
 
+static TAutoConsoleVariable<int32> CVarSkinnedMeshSyncStreamingLODClamp(
+	TEXT("r.SkinnedMesh.SyncStreamingLODClamp"),
+	1,
+	TEXT("Whether to sync streaming LOD index clamp among components with SyncAttachParentLOD enabled in the attachment tree."),
+	ECVF_Default);
+
 bool ShouldRenderNaniteSkinnedMeshes()
 {
 	return NaniteSkinnedMeshesSupported() && GSkinnedMeshRenderNanite != 0;
@@ -481,6 +487,7 @@ USkinnedMeshComponent::USkinnedMeshComponent(const FObjectInitializer& ObjectIni
 	bExternalEvaluationRateLimited = false;
 	bExternalTickRateControlled = false;
 
+	MaxFirstLODIdxInAttachmentTree = 0;
 	bMipLevelCallbackRegistered = false;
 
 	bFollowerShouldTickPose = false;
@@ -4441,6 +4448,24 @@ bool USkinnedMeshComponent::UpdateLODStatus_Internal(int32 InLeaderPoseComponent
 	const int32 OldPredictedLODLevel = GetPredictedLODLevel();
 	int32 NewPredictedLODLevel = OldPredictedLODLevel;
 
+	if (CVarSkinnedMeshSyncStreamingLODClamp.GetValueOnAnyThread() != 0)
+	{
+		// Pull parent value which has already been updated if this can use it. Otherwise, recompute using this as root
+		if (bSyncAttachParentLOD && GetAttachParent() && GetAttachParent()->IsA<USkinnedMeshComponent>())
+		{
+			const USkinnedMeshComponent* ParentComponent = Cast<USkinnedMeshComponent>(GetAttachParent());
+			MaxFirstLODIdxInAttachmentTree = ParentComponent->MaxFirstLODIdxInAttachmentTree;
+		}
+		else
+		{
+			MaxFirstLODIdxInAttachmentTree = ComputeMaxFirstLODIdxInAttachmentTree();
+		}
+	}
+	else
+	{
+		MaxFirstLODIdxInAttachmentTree = 0;
+	}
+
 	if (GetSkinnedAsset() != nullptr)
 	{
 #if WITH_EDITOR
@@ -4472,7 +4497,8 @@ bool USkinnedMeshComponent::UpdateLODStatus_Internal(int32 InLeaderPoseComponent
 			}
 			else if (bSyncAttachParentLOD && GetAttachParent() && GetAttachParent()->IsA(USkinnedMeshComponent::StaticClass()))
 			{
-				NewPredictedLODLevel = FMath::Clamp(CastChecked<USkinnedMeshComponent>(GetAttachParent())->GetPredictedLODLevel(), 0, MaxLODIndex);
+				const USkinnedMeshComponent* ParentComponent = CastChecked<USkinnedMeshComponent>(GetAttachParent());
+				NewPredictedLODLevel = FMath::Min(FMath::Max(ParentComponent->GetPredictedLODLevel(), (int32)MaxFirstLODIdxInAttachmentTree), MaxLODIndex);
 			}
 			else if (MeshObject)
 			{
@@ -4565,6 +4591,31 @@ bool USkinnedMeshComponent::UpdateLODStatus_Internal(int32 InLeaderPoseComponent
 	return bLODChanged;
 }
 
+uint8 USkinnedMeshComponent::ComputeMaxFirstLODIdxInAttachmentTree() const
+{
+	const TArray<TObjectPtr<USceneComponent>>& LocalAttachChildren = GetAttachChildren();
+	uint8 Result = 0;
+	
+	for (const TObjectPtr<USceneComponent>& AttachChild : LocalAttachChildren)
+	{
+		const USkinnedMeshComponent* AttachSkinnedMeshComponent = Cast<USkinnedMeshComponent>(AttachChild);
+		if (AttachSkinnedMeshComponent && AttachSkinnedMeshComponent->bSyncAttachParentLOD)
+		{
+			const uint8 ChildResult = AttachSkinnedMeshComponent->ComputeMaxFirstLODIdxInAttachmentTree();
+			Result = FMath::Max(Result, ChildResult);
+		}
+	}
+
+	if (GetSkinnedAsset() && GetSkinnedAsset()->IsStreamable() && MeshObject)
+	{
+		const FSkeletalMeshRenderData& RenderData = MeshObject->GetSkeletalMeshRenderData();
+		Result = FMath::Max(Result, RenderData.PendingFirstLODIdx);
+		Result = FMath::Max(Result, RenderData.CurrentFirstLODIdx);
+	}
+
+	return Result;
+}
+
 void USkinnedMeshComponent::FinalizeBoneTransform()
 {
 	FlipEditableSpaceBases();
