FNewSpline: Override GetNumPoints to ignore additional point for closing the spline

Also introduces GetNumSegments to complete this set of overrides.

diff --git a/Engine/Source/Runtime/Engine/Private/Curves/Spline.cpp b/Engine/Source/Runtime/Engine/Private/Curves/Spline.cpp
--- a/Engine/Source/Runtime/Engine/Private/Curves/Spline.cpp
+++ b/Engine/Source/Runtime/Engine/Private/Curves/Spline.cpp
@@ -138,6 +138,8 @@ public:
 
 	float GetSegmentLength(const int32 Index, const float Param, const FVector& Scale3D = FVector(1.0f)) const;
 	float GetSplineLength() const;
+	int32 GetNumPoints() const;
+	int32 GetNumSegments() const;
 	
     void SetClosedLoop(bool bInClosedLoop);
 
@@ -1393,6 +1395,16 @@ float FNewSpline::GetSplineLength() const
     return EvaluateAttribute<float>(ParamToLengthAttrName, 1.f);
 }
 
+int32 FNewSpline::GetNumPoints() const
+{
+	return IsClosedLoop() ? FTangentBezierSpline::GetNumPoints() - 1 : FTangentBezierSpline::GetNumPoints();
+}
+
+int32 FNewSpline::GetNumSegments() const
+{
+	return FTangentBezierSpline::GetNumPoints() - 1;
+}
+
 void FNewSpline::SetClosedLoop(bool bInClosedLoop)
 {
     FTangentBezierSpline::SetClosedLoop(bInClosedLoop);
