Added a check for FSkinWeightProfileManager requests so we can skip some requested mesh LODs if they aren't loaded and not scheduled for streaming

#tests RR PF
#rnx

diff --git a/Engine/Source/Runtime/Engine/Private/Animation/SkinWeightProfileManager.cpp b/Engine/Source/Runtime/Engine/Private/Animation/SkinWeightProfileManager.cpp
--- a/Engine/Source/Runtime/Engine/Private/Animation/SkinWeightProfileManager.cpp
+++ b/Engine/Source/Runtime/Engine/Private/Animation/SkinWeightProfileManager.cpp
@@ -42,6 +42,14 @@ static TAutoConsoleVariable<int32> CVarForceSkinWeightProfileManagerAsyncTaskOnG
 	ECVF_Default | ECVF_RenderThreadSafe
 );
 
+static int32 GSkipLODWithoutDataOrStreamingRequest = 1;
+static FAutoConsoleVariableRef CVarSkipLODWithoutDataOrStreamingRequest(
+	TEXT("SkinWeightProfileManager.SkipLODWithoutDataOrStreamingRequest"),
+	GSkipLODWithoutDataOrStreamingRequest,
+	TEXT(""),
+	ECVF_Default | ECVF_RenderThreadSafe
+);
+
 FSkinWeightProfileManager::FSkinWeightProfileManager(UWorld* InWorld) : LastGamethreadProfileIndex(INDEX_NONE), WeakWorld(nullptr)
 {
 	if (InWorld)
@@ -219,7 +227,7 @@ void FSkinWeightProfileManager::DoTick(float DeltaTime, ENamedThreads::Type Curr
 		{
 			QUICK_SCOPE_CYCLE_COUNTER(STAT_SkinWeightProfileManager_TickMesh);
 
-			const FSetProfileRequest& Request = PendingSetProfileRequests[RequestIndex];
+			FSetProfileRequest& Request = PendingSetProfileRequests[RequestIndex];
 			
 			// Remove if canceled or objects have gone stale 
 			bool bRemove = Request.IdentifyingObject.IsStale() || Request.WeakSkeletalMesh.IsStale() ||
@@ -236,12 +244,21 @@ void FSkinWeightProfileManager::DoTick(float DeltaTime, ENamedThreads::Type Curr
 			if (!bRemove)
 			{	
 				bool bAllBuffersReady = true;
-				for (const int32 LODIndex : Request.LODIndices)
+				for (int32 LODProfileArrayIndex = 0; LODProfileArrayIndex < Request.LODIndices.Num(); ++LODProfileArrayIndex)
 				{
-					if (!RenderData->LODRenderData[LODIndex].SkinWeightProfilesData.ContainsOverrideBuffer(Request.ProfileStack))
+					int32 LODRenderDataArrayIndex = Request.LODIndices[LODProfileArrayIndex];
+					const FSkinWeightProfilesData& SkinWeightProfilesData = RenderData->LODRenderData[LODRenderDataArrayIndex].SkinWeightProfilesData;
+					if (!SkinWeightProfilesData.ContainsOverrideBuffer(Request.ProfileStack))
 					{
-						bAllBuffersReady = false;
-						break;
+						// Check if this LOD is loaded or scheduled to be streamed in, if not then we can ignore it and remove from Request.LODIndices
+						if (!ShouldIgnoreLOD(SkinWeightProfilesData, Request.WeakSkeletalMesh))
+						{
+							bAllBuffersReady = false;
+							break;
+						}
+
+						Request.LODIndices.RemoveAtSwap(LODProfileArrayIndex);
+						--LODProfileArrayIndex;
 					}
 				}
 
@@ -325,6 +342,31 @@ void FSkinWeightProfileManager::CleanupRequest(const FSetProfileRequest& Request
 	
 }
 
+bool FSkinWeightProfileManager::ShouldIgnoreLOD(const FSkinWeightProfilesData& SkinWeightProfilesData, const TWeakObjectPtr<USkeletalMesh>& WeakSkeletalMesh) const
+{
+	if (!GSkipLODWithoutDataOrStreamingRequest)
+	{
+		return false;
+	}
+
+	// Are LOD's vertices loaded? 
+	FSkinWeightVertexBuffer* BaseBuffer = SkinWeightProfilesData.BaseBuffer;
+	if (BaseBuffer && BaseBuffer->GetDataVertexBuffer()->VertexBufferRHI->GetSize())
+	{
+		return false;
+	}
+
+	// Are they scheduled to be streamed in?
+	const USkeletalMesh* SkeletalMesh = WeakSkeletalMesh.Get();
+	if (SkeletalMesh && SkeletalMesh->HasPendingInitOrStreaming())
+	{
+		return false;
+	}
+
+	// This mesh LOD isn't loaded and not scheduled for streaming, we can ignore it
+	return true;
+}
+
 bool FSkinWeightProfileManager::IsTickableWhenPaused() const
 {
 	return false;
diff --git a/Engine/Source/Runtime/Engine/Public/Animation/SkinWeightProfileManager.h b/Engine/Source/Runtime/Engine/Public/Animation/SkinWeightProfileManager.h
--- a/Engine/Source/Runtime/Engine/Public/Animation/SkinWeightProfileManager.h
+++ b/Engine/Source/Runtime/Engine/Public/Animation/SkinWeightProfileManager.h
@@ -117,6 +117,7 @@ public:
 	void DoTick(float DeltaTime, ENamedThreads::Type CurrentThread, const FGraphEventRef& MyCompletionGraphEvent);
 protected:	
 	void CleanupRequest(const FSetProfileRequest& Request);
+	bool ShouldIgnoreLOD(const FSkinWeightProfilesData& SkinWeightProfilesData, const TWeakObjectPtr<USkeletalMesh>& WeakSkeletalMesh) const;
 protected:
 	TArray<FSetProfileRequest, TInlineAllocator<4>> CanceledRequest;
 	TArray<FSetProfileRequest> PendingSetProfileRequests;
