Check for overflow when adding LevelStreamingPriorityBias to StreamingPriority.

#rnx

diff --git a/Engine/Source/Runtime/Engine/Private/LevelStreaming.cpp b/Engine/Source/Runtime/Engine/Private/LevelStreaming.cpp
--- a/Engine/Source/Runtime/Engine/Private/LevelStreaming.cpp
+++ b/Engine/Source/Runtime/Engine/Private/LevelStreaming.cpp
@@ -1754,7 +1754,13 @@ bool ULevelStreaming::RequestLevel(UWorld* PersistentWorld, bool bAllowLevelLoad
 		}
 #endif
 		LastAsyncLoadPackagePath = PackagePath;
-		AsyncRequestIDs.Add(LoadPackageAsync(PackagePath, DesiredPackageName, FLoadPackageAsyncDelegate::CreateUObject(this, &ULevelStreaming::AsyncLevelLoadComplete), PackageFlags, PIEInstanceID, LevelStreamingCVars::GLevelStreamingPriorityBias + GetPriority(), InstancingContextPtr));
+
+		int32 OutAsyncLoadPriority;
+		if (!FMath::AddAndCheckForOverflow(LevelStreamingCVars::GLevelStreamingPriorityBias, GetPriority(), OutAsyncLoadPriority))
+		{
+			OutAsyncLoadPriority = MAX_int32;
+		}
+		AsyncRequestIDs.Add(LoadPackageAsync(PackagePath, DesiredPackageName, FLoadPackageAsyncDelegate::CreateUObject(this, &ULevelStreaming::AsyncLevelLoadComplete), PackageFlags, PIEInstanceID, OutAsyncLoadPriority, InstancingContextPtr));
 
 		// streamingServer: server loads everything?
 		// Editor immediately blocks on load and we also block if background level streaming is disabled.
