Fixed IncrementalPreRegisterComponents not properly handling cases where the CS of BP would have already created and registered components.

diff --git a/Engine/Source/Runtime/Engine/Private/Actor.cpp b/Engine/Source/Runtime/Engine/Private/Actor.cpp
--- a/Engine/Source/Runtime/Engine/Private/Actor.cpp
+++ b/Engine/Source/Runtime/Engine/Private/Actor.cpp
@@ -6001,9 +6001,15 @@ void AActor::IncrementalPreRegisterComponents()
 	TInlineComponentArray<UActorComponent*> Components;
 	for (UActorComponent* Component : GetComponents())
 	{
-		if (Component->ShouldIncrementalPreRegister(World))
+#if WITH_EDITOR
+		// Ignore registered components (can happen if construction script adds component(s), which happens before this)
+		if (!Component->IsRegistered())
+#endif
 		{
-			Components.Add(Component);
+			if (Component->ShouldIncrementalPreRegister(World))
+			{
+				Components.Add(Component);
+			}
 		}
 	}
 
@@ -6016,8 +6022,14 @@ void AActor::IncrementalPreRegisterComponents()
 	// Pre-register RootComponent first so all other children components can reliably use it (i.e., call GetLocation) when they pre-register
 	if (RootComponent && RootComponent->bAutoRegister && !RootComponent->IsPreRegistered())
 	{
-		check(!RootComponent->IsPreRegistering());
-		RootComponent->PreRegisterComponentWithWorld(World);
+#if WITH_EDITOR
+		// Ignore registered components (can happen if construction script adds component(s), which happens before this)
+		if (!RootComponent->IsRegistered())
+#endif 
+		{
+			check(!RootComponent->IsPreRegistering());
+			RootComponent->PreRegisterComponentWithWorld(World);
+		}
 	}
 
 	TSet<UActorComponent*> PreRegisteredParents;
