Add support for shader permutations that are only built with ODSC.
These are identified by checking for the EShaderPermutationFlags::IsODSCOnly in ShouldCompilePermutation().
We only build these shaders when running the editor in (the default) ODSC mode or running cooked builds attached to an ODSC server.
Note that the editor doesn't use ODSC mode when cooking or when running commandlets so the shaders aren't built/available in those contexts.
The main intended use case is for debug shaders that we don't want in cooked builds.
Using the new flag has some differences compared to EShaderPermutationFlags::HasEditorOnlyData which is currently used for debug shaders.
* It won't compile the selected shaders at a commandlet start up.
* We can use the selected shaders in cooked builds attached to ODSC servers.
* Shaders won't be part of cooked editor shader maps.

diff --git a/Engine/Source/Runtime/Engine/Private/Materials/MaterialShader.cpp b/Engine/Source/Runtime/Engine/Private/Materials/MaterialShader.cpp
--- a/Engine/Source/Runtime/Engine/Private/Materials/MaterialShader.cpp
+++ b/Engine/Source/Runtime/Engine/Private/Materials/MaterialShader.cpp
@@ -3095,7 +3095,9 @@ bool FMaterialShaderMap::IsComplete(const FMaterial* Material, bool bSilent)
 	const EShaderPlatform Platform = LocalContent->GetShaderPlatform();
 	const FMaterialShaderParameters MaterialParameters(Material);
 
-	const FMaterialShaderMapLayout& Layout = AcquireMaterialShaderMapLayout(Platform, ShaderMapId.GetPermutationFlags(), MaterialParameters);
+	// IsODSCOnly shaders aren't in shader map so don't include when gathering for IsComplete().
+	EShaderPermutationFlags Flags = ShaderMapId.GetPermutationFlags() & ~EShaderPermutationFlags::IsODSCOnly;
+	const FMaterialShaderMapLayout& Layout = AcquireMaterialShaderMapLayout(Platform, Flags, MaterialParameters);
 	if (Layout.ShaderMapHash == LocalContent->ShaderContentHash)
 	{
 		return true;
diff --git a/Engine/Source/Runtime/Engine/Private/ShaderCompiler/ShaderCompilerJobCache.cpp b/Engine/Source/Runtime/Engine/Private/ShaderCompiler/ShaderCompilerJobCache.cpp
--- a/Engine/Source/Runtime/Engine/Private/ShaderCompiler/ShaderCompilerJobCache.cpp
+++ b/Engine/Source/Runtime/Engine/Private/ShaderCompiler/ShaderCompilerJobCache.cpp
@@ -11,6 +11,7 @@
 #include "DerivedDataRequestOwner.h"
 #include "Math/UnitConversion.h"
 #include "Misc/FileHelper.h"
+#include "ODSC/ODSCManager.h"
 #include "ProfilingDebugging/CookStats.h"
 #include "ProfilingDebugging/DiagnosticTable.h"
 #include "ProfilingDebugging/LoadTimeTracker.h"
@@ -147,6 +148,14 @@ namespace ShaderCompiler
 } // namespace ShaderCompiler
 
 
+/** The ODSC server is a special cook commandlet where we don't want to use the material map DDC. */
+static bool IsRunningCookCommandletAndNotODSCServer()
+{
+	static const bool bIsCookCommandlet = IsRunningCookCommandlet();
+	static const bool bIsODSCServer = FParse::Param(FCommandLine::Get(), TEXT("odsc"));
+	return bIsCookCommandlet && !bIsODSCServer;
+}
+
 bool IsShaderJobCacheDDCEnabled()
 {
 #if WITH_EDITOR
@@ -171,8 +180,22 @@ bool IsShaderJobCacheDDCEnabled()
 
 bool IsMaterialMapDDCEnabled()
 {
-	// If we are loading individual shaders from the shader job cache for ODSC, don't attempt to load full material maps.  Always load/cache material maps in cooks.
-	return (IsShaderJobCacheDDCEnabled() == false) || IsRunningCookCommandlet();
+	// If we are loading individual shaders from the shader job cache for ODSC, don't attempt to load full material maps. 
+	// Otherwise always load/cache material maps in cooks.
+	return (IsShaderJobCacheDDCEnabled() == false) || IsRunningCookCommandletAndNotODSCServer();
+}
+
+bool ShouldCompileODSCOnlyShaders()
+{
+#if WITH_EDITOR
+	static const bool bIsODSCEditor = !IsMaterialMapDDCEnabled();
+	return bIsODSCEditor;
+#elif WITH_ODSC
+	static const bool bIsODSCClient = FODSCManager::IsODSCActive();
+	return bIsODSCClient;
+#else
+	return false;
+#endif
 }
 
 static bool IsShaderJobCacheDDCRemotePolicyEnabled()
diff --git a/Engine/Source/Runtime/Engine/Public/ShaderCompiler.h b/Engine/Source/Runtime/Engine/Public/ShaderCompiler.h
--- a/Engine/Source/Runtime/Engine/Public/ShaderCompiler.h
+++ b/Engine/Source/Runtime/Engine/Public/ShaderCompiler.h
@@ -56,6 +56,12 @@ bool CreateShadersOnLoad();
 extern ENGINE_API bool IsShaderJobCacheDDCEnabled();
 extern ENGINE_API bool IsMaterialMapDDCEnabled();
 
+/**
+ * Returns true if we should compile shaders that are only compiled in on demand shader compilation
+ * modes and not stored in cooked shader maps. This can be useful for debug shaders.
+ */
+extern ENGINE_API bool ShouldCompileODSCOnlyShaders();
+
 struct FJobObjectLimitationInfo;
 
 struct FShaderJobCacheStoredOutput;
diff --git a/Engine/Source/Runtime/Launch/Private/LaunchEngineLoop.cpp b/Engine/Source/Runtime/Launch/Private/LaunchEngineLoop.cpp
--- a/Engine/Source/Runtime/Launch/Private/LaunchEngineLoop.cpp
+++ b/Engine/Source/Runtime/Launch/Private/LaunchEngineLoop.cpp
@@ -135,7 +135,6 @@
 	#include "Settings/EditorExperimentalSettings.h"
 	#include "PIEPreviewDeviceProfileSelectorModule.h"
 	#include "Serialization/BulkDataRegistry.h"
-	#include "ShaderCompiler.h"
 	#include "Virtualization/VirtualizationSystem.h"
 	#include "AnimationUtils.h"
 	#include "UObject/ICookInfo.h"
@@ -3049,6 +3048,19 @@ int32 FEngineLoop::PreInitPreStartupScreen(const TCHAR* CmdLine)
 		}
 	}
 
+	{
+		EShaderPermutationFlags Flags = EShaderPermutationFlags::None;
+		
+		bool bSupportCookedEditor = false;
+		GConfig->GetBool(TEXT("CookedEditorSettings"), TEXT("bSupportCookedEditor"), bSupportCookedEditor, GGameIni);
+		Flags |= bSupportCookedEditor ? EShaderPermutationFlags::HasEditorOnlyData : EShaderPermutationFlags::None;
+		
+		const bool bShouldCompileODSCOnlyShaders = ShouldCompileODSCOnlyShaders();
+		Flags |= bShouldCompileODSCOnlyShaders ? EShaderPermutationFlags::IsODSCOnly : EShaderPermutationFlags::None;
+
+		SetAdditionalShaderPermutationFlags(Flags);
+	}
+
 	{
 		SCOPED_BOOT_TIMING("GetRendererModule");
 		// Cache the renderer module in the main thread so that we can safely retrieve it later from the rendering thread.
diff --git a/Engine/Source/Runtime/RenderCore/Private/GlobalShader.cpp b/Engine/Source/Runtime/RenderCore/Private/GlobalShader.cpp
--- a/Engine/Source/Runtime/RenderCore/Private/GlobalShader.cpp
+++ b/Engine/Source/Runtime/RenderCore/Private/GlobalShader.cpp
@@ -653,7 +653,9 @@ bool FGlobalShaderMap::IsEmpty() const
 bool FGlobalShaderMap::IsComplete(const ITargetPlatform* TargetPlatform) const
 {
 	FGlobalShaderMapId ShaderMapId(Platform, TargetPlatform);
-	const EShaderPermutationFlags PermutationFlags = ShaderMapId.GetShaderPermutationFlags();
+	
+	// IsODSCOnly shaders aren't in shader map so don't include when gathering for IsComplete().
+	const EShaderPermutationFlags PermutationFlags = ShaderMapId.GetShaderPermutationFlags() & ~EShaderPermutationFlags::IsODSCOnly;
 
 	// traverse all global shader types
 	for (TLinkedList<FShaderType*>::TIterator ShaderTypeIt(FShaderType::GetTypeList()); ShaderTypeIt; ShaderTypeIt.Next())
diff --git a/Engine/Source/Runtime/RenderCore/Private/Shader.cpp b/Engine/Source/Runtime/RenderCore/Private/Shader.cpp
--- a/Engine/Source/Runtime/RenderCore/Private/Shader.cpp
+++ b/Engine/Source/Runtime/RenderCore/Private/Shader.cpp
@@ -2814,17 +2814,18 @@ void ShaderMapAppendKey(EShaderPlatform Platform, FShaderKeyGenerator& KeyGen)
 	}
 }
 
-EShaderPermutationFlags GetShaderPermutationFlags(const FPlatformTypeLayoutParameters& LayoutParams)
-{
-	EShaderPermutationFlags Result = EShaderPermutationFlags::None;
+static EShaderPermutationFlags GAdditionalShaderPermutationFlags = EShaderPermutationFlags::None;
 
-	static bool bProjectSupportsCookedEditor = []()
-	{
-		bool bSupportCookedEditorConfigValue = false;
-		return GConfig->GetBool(TEXT("CookedEditorSettings"), TEXT("bSupportCookedEditor"), bSupportCookedEditorConfigValue, GGameIni) && bSupportCookedEditorConfigValue;
-	}();
+void SetAdditionalShaderPermutationFlags(EShaderPermutationFlags AdditionalFlags)
+{
+	check(GAdditionalShaderPermutationFlags == EShaderPermutationFlags::None);
+	GAdditionalShaderPermutationFlags = AdditionalFlags;
+}
 
-	if (bProjectSupportsCookedEditor || LayoutParams.WithEditorOnly())
+EShaderPermutationFlags GetShaderPermutationFlags(const FPlatformTypeLayoutParameters& LayoutParams)
+{
+	EShaderPermutationFlags Result = GAdditionalShaderPermutationFlags;
+	if (LayoutParams.WithEditorOnly())
 	{
 		Result |= EShaderPermutationFlags::HasEditorOnlyData;
 	}
diff --git a/Engine/Source/Runtime/RenderCore/Private/ShaderMap.cpp b/Engine/Source/Runtime/RenderCore/Private/ShaderMap.cpp
--- a/Engine/Source/Runtime/RenderCore/Private/ShaderMap.cpp
+++ b/Engine/Source/Runtime/RenderCore/Private/ShaderMap.cpp
@@ -21,11 +21,9 @@
 
 static EShaderPermutationFlags GetCurrentShaderPermutationFlags()
 {
-	EShaderPermutationFlags Result = EShaderPermutationFlags::None;
-#if WITH_EDITORONLY_DATA 
-	Result |= EShaderPermutationFlags::HasEditorOnlyData;
-#endif
-	return Result;
+	FPlatformTypeLayoutParameters LayoutParams;
+	LayoutParams.InitializeForCurrent();
+	return GetShaderPermutationFlags(LayoutParams);
 }
 
 #ifndef ALLOW_SHADERMAP_TRACKING
diff --git a/Engine/Source/Runtime/RenderCore/Public/Shader.h b/Engine/Source/Runtime/RenderCore/Public/Shader.h
--- a/Engine/Source/Runtime/RenderCore/Public/Shader.h
+++ b/Engine/Source/Runtime/RenderCore/Public/Shader.h
@@ -819,7 +819,10 @@ enum class EShaderPermutationPrecacheRequest : uint8
 	NotUsed,			//< Permutation is not used with current cvar/configuration/feature level
 };
 
+/** Get the shader permutation flags associated with a platform layout. */
 RENDERCORE_API EShaderPermutationFlags GetShaderPermutationFlags(const FPlatformTypeLayoutParameters& LayoutParams);
+/** Set global additional flags that are always added in GetShaderPermutationFlags(). Expected to be called once at startup. */
+RENDERCORE_API void SetAdditionalShaderPermutationFlags(EShaderPermutationFlags AdditionalFlags);
 
 namespace Freeze
 {
diff --git a/Engine/Source/Runtime/RenderCore/Public/ShaderPermutation.h b/Engine/Source/Runtime/RenderCore/Public/ShaderPermutation.h
--- a/Engine/Source/Runtime/RenderCore/Public/ShaderPermutation.h
+++ b/Engine/Source/Runtime/RenderCore/Public/ShaderPermutation.h
@@ -19,6 +19,7 @@ enum class EShaderPermutationFlags : uint32
 {
 	None = 0u,
 	HasEditorOnlyData = (1u << 0),
+	IsODSCOnly = (1u << 1),
 };
 ENUM_CLASS_FLAGS(EShaderPermutationFlags);
 
