[DSTM] Implementing physics state migration for returned borrowed objects, when they are returned as part of a multi-server commit.

This CL also:

- Moves the application of the migrated state to the new PostMigrate method available in Actor Components
- Handles the case where a borrowing server destroyed the physics state
- Handles the case where a borrowing server altered collision settings before returning the objects.
- Any write operation on a borrowed object, will flag the current transaction as multi-server commit. We might want to look into reducing the scope of this in the future (if possible).
- Adds a new flag to the migration context object so we can know if we are serializing data for a multi-server commit operation or not, and act accordingly.
-Fixed an issue where the serialization implementation for migration for constraints was applying the incorrect lock type.

diff --git a/Engine/Source/Runtime/CoreUObject/Internal/UObject/UObjectMigrationContext.h b/Engine/Source/Runtime/CoreUObject/Internal/UObject/UObjectMigrationContext.h
--- a/Engine/Source/Runtime/CoreUObject/Internal/UObject/UObjectMigrationContext.h
+++ b/Engine/Source/Runtime/CoreUObject/Internal/UObject/UObjectMigrationContext.h
@@ -72,6 +72,9 @@ struct COREUOBJECT_API FUObjectMigrationContext
 	/** Which side of the migration does this context apply to? Sending or Receiving? */
 	EObjectMigrationSide MigrationSide;
 
+	/** Whether this migration is part of a multi-server commit */
+	bool bForMultiServerCommit = false;
+
 	/** Are the values of this context valid? */
 	bool IsValid() const;
 
diff --git a/Engine/Source/Runtime/CoreUObject/Private/UObject/RemoteObjectTransfer.cpp b/Engine/Source/Runtime/CoreUObject/Private/UObject/RemoteObjectTransfer.cpp
--- a/Engine/Source/Runtime/CoreUObject/Private/UObject/RemoteObjectTransfer.cpp
+++ b/Engine/Source/Runtime/CoreUObject/Private/UObject/RemoteObjectTransfer.cpp
@@ -264,12 +264,13 @@ public:
 				.OwnerServerId = OwnerServerId,
 				.PhysicsServerId = OwnerServerId,
 				.RequestId = ActiveRequest->RequestId,
-				.MigrationSide = EObjectMigrationSide::Send
+				.MigrationSide = EObjectMigrationSide::Send,
+				.bForMultiServerCommit = true
 			};
 
 			TSet<UObject*> SentObjects;
 			FRemoteObjectData ObjectData = SerializeObjectData(Object, SentObjects, MultiServerCommitReferencedObjects, &MigrationContext);
-			FMigrateSendParams SendRemoteObjectParams = { .MigrationContext = MigrationContext, .ObjectData = ObjectData, .bForMultiServerCommit = true };
+			FMigrateSendParams SendRemoteObjectParams = { .MigrationContext = MigrationContext, .ObjectData = ObjectData};
 
 			UE_AUTORTFM_OPEN
 			{
diff --git a/Engine/Source/Runtime/CoreUObject/Public/UObject/RemoteObjectTransfer.h b/Engine/Source/Runtime/CoreUObject/Public/UObject/RemoteObjectTransfer.h
--- a/Engine/Source/Runtime/CoreUObject/Public/UObject/RemoteObjectTransfer.h
+++ b/Engine/Source/Runtime/CoreUObject/Public/UObject/RemoteObjectTransfer.h
@@ -48,9 +48,6 @@ namespace UE::RemoteObject::Transfer
 
 		/** The serialized data of the object being sent */
 		FRemoteObjectData ObjectData;
-
-		/** Whether this migration is part of a multi-server commit */
-		bool bForMultiServerCommit = false;
 	};
 
 	constexpr FRemoteTransactionId PushMigrationRequestId = FRemoteTransactionId(0xFFFFFFFFu);
diff --git a/Engine/Source/Runtime/Engine/Classes/Components/PrimitiveComponent.h b/Engine/Source/Runtime/Engine/Classes/Components/PrimitiveComponent.h
--- a/Engine/Source/Runtime/Engine/Classes/Components/PrimitiveComponent.h
+++ b/Engine/Source/Runtime/Engine/Classes/Components/PrimitiveComponent.h
@@ -50,6 +50,8 @@ struct FNavigableGeometryExport;
 struct FPSOPrecacheParams;
 struct FOverlapResult;
 
+enum class EObjectMigrationSendType : uint8;
+
 namespace PrimitiveComponentCVars
 {
 	extern float HitDistanceToleranceCVar;
@@ -3161,12 +3163,17 @@ public:
 protected:
 
 	UFUNCTION()
-	void HandlePostMigrationPhysicsStateChanged(UPrimitiveComponent* ChangedComponent, EComponentPhysicsStateChange StateChange);
+	ENGINE_API virtual void HandlePostMigrationPhysicsStateChanged(UPrimitiveComponent* ChangedComponent, EComponentPhysicsStateChange StateChange);
 
 #if UE_WITH_REMOTE_OBJECT_HANDLE
+
+	ENGINE_API bool CanSendMigratedPhysicsState(const FUObjectMigrationContext* MigrationContext, EObjectMigrationSendType ObjMigrationSendType) const;
+
 	ENGINE_API virtual void MigratePhysicsState(FArchive& Ar) override;
 	
 	ENGINE_API virtual void ApplyMigratedPhysicsState(Chaos::FSerializedDataBufferPtr MigratedStateData = nullptr) override;
+
+	ENGINE_API virtual void PostMigrate(const struct FUObjectMigrationContext& MigrationContext) override;
 #endif
 
 protected:
diff --git a/Engine/Source/Runtime/Engine/Private/Components/PrimitiveComponent.cpp b/Engine/Source/Runtime/Engine/Private/Components/PrimitiveComponent.cpp
--- a/Engine/Source/Runtime/Engine/Private/Components/PrimitiveComponent.cpp
+++ b/Engine/Source/Runtime/Engine/Private/Components/PrimitiveComponent.cpp
@@ -1298,6 +1298,22 @@ void UPrimitiveComponent::Serialize(FArchive& Ar)
 }
 
 #if UE_WITH_REMOTE_OBJECT_HANDLE
+
+bool UPrimitiveComponent::CanSendMigratedPhysicsState(const FUObjectMigrationContext* MigrationContext, EObjectMigrationSendType ObjMigrationSendType) const
+{
+	if (!HasValidPhysicsState())
+	{
+		return false;
+	}
+
+	if (ObjMigrationSendType == EObjectMigrationSendType::ReturnBorrowed)
+	{
+		return MigrationContext->bForMultiServerCommit;
+	}
+
+	return true;
+}
+
 void UPrimitiveComponent::MigratePhysicsState(FArchive& Ar)
 {
 	const FUObjectMigrationContext* MigrationContext = Ar.GetMigrationContext();
@@ -1317,28 +1333,22 @@ void UPrimitiveComponent::MigratePhysicsState(FArchive& Ar)
 			return;
 		}
 
-		const bool bFinishedDestroy = HasAnyFlags(EObjectFlags::RF_MirroredGarbage);
-		if (ObjMigrationRecvType == EObjectMigrationRecvType::ReturnedLoan || bFinishedDestroy)
+		if (const bool bFinishedDestroy = HasAnyFlags(EObjectFlags::RF_MirroredGarbage))
 		{
 			return;
 		}
 
-		if (HasValidPhysicsState())
-		{
-			ApplyMigratedPhysicsState(MoveTemp(MigratedStateDataBuffer));
-		}
-		else if (FPhysScene* PhysicsScene = FPhysicsInterface::GetCurrentSceneForComponent(this))
+		FPhysScene* PhysicsScene = FPhysicsInterface::GetCurrentSceneForComponent(this);
+		if (ensure(PhysicsScene))
 		{
 			PhysicsScene->PushPendingSerializedStateForObject(this, MoveTemp(MigratedStateDataBuffer));
-			OnComponentPhysicsStateChanged.AddUniqueDynamic(this, &UPrimitiveComponent::HandlePostMigrationPhysicsStateChanged);
 		}
 	}
 	else if (Ar.IsSaving())
 	{		
 		const EObjectMigrationSendType ObjMigrationSendType = MigrationContext->GetObjectMigrationSendType(this);
-		const bool bCanSendPhysicsState = ObjMigrationSendType != EObjectMigrationSendType::ReturnBorrowed && HasValidPhysicsState();
 
-		if (bCanSendPhysicsState)
+		if (CanSendMigratedPhysicsState(MigrationContext, ObjMigrationSendType))
 		{
 			const Chaos::ESerializedDataContext PhysicsDataContext = ObjMigrationSendType == EObjectMigrationSendType::ReassignOwnership ? Chaos::ESerializedDataContext::Internal : Chaos::ESerializedDataContext::External;
 			Chaos::FSerializedMultiPhysicsState MigratedState(PhysicsDataContext);
@@ -1413,6 +1423,50 @@ void UPrimitiveComponent::ApplyMigratedPhysicsState(Chaos::FSerializedDataBuffer
 		}	
 	}
 }
+
+void UPrimitiveComponent::PostMigrate(const struct FUObjectMigrationContext& MigrationContext)
+{
+	Super::PostMigrate(MigrationContext);
+
+	if (MigrationContext.MigrationSide == EObjectMigrationSide::Receive)
+	{
+		EObjectMigrationRecvType RecvType = MigrationContext.GetObjectMigrationRecvType(this);
+
+		if (RecvType == EObjectMigrationRecvType::ReturnedLoan)
+		{
+			if (bPhysicsStateCreated)
+			{
+				// Collision settings might have been changed by the borrowing server
+				// TODO: Only do this if we know these settings were changed. Currently we don't have an easy way of knowing at this point
+				BodyInstance.UpdatePhysicsFilterData();
+			}
+			else
+			{
+				FPhysScene* PhysicsScene = FPhysicsInterface::GetCurrentSceneForComponent(this);
+				Chaos::FSerializedDataBufferPtr RawMigratedData = PhysicsScene ? PhysicsScene->PopPendingSerializedStateForObject(this) : nullptr;
+				
+				if (!ensure(!RawMigratedData))
+				{
+					UE_LOG(LogPrimitiveComponent, Error, TEXT("[%hs] Received a loaned object that destroyed our physics state, but we still have pending data to process. This should not happen..."), __func__);
+				}
+
+				// If the borrowing server destroyed the object, we will not receive a physics state to apply, but the component flag bPhysicsStateCreated will be migrated as false
+				// If so, we need to destroy our local physics state.
+				bPhysicsStateCreated = true;
+				DestroyPhysicsState();
+			}
+		}
+
+		if (HasValidPhysicsState())
+		{
+			ApplyMigratedPhysicsState();
+		}
+		else
+		{
+			OnComponentPhysicsStateChanged.AddUniqueDynamic(this, &UPrimitiveComponent::HandlePostMigrationPhysicsStateChanged);
+		}
+	}
+}
 #endif
 
 void UPrimitiveComponent::HandlePostMigrationPhysicsStateChanged(UPrimitiveComponent* ChangedComponent, EComponentPhysicsStateChange StateChange)
diff --git a/Engine/Source/Runtime/Engine/Private/Components/SkeletalMeshComponent.cpp b/Engine/Source/Runtime/Engine/Private/Components/SkeletalMeshComponent.cpp
--- a/Engine/Source/Runtime/Engine/Private/Components/SkeletalMeshComponent.cpp
+++ b/Engine/Source/Runtime/Engine/Private/Components/SkeletalMeshComponent.cpp
@@ -675,6 +675,24 @@ void USkeletalMeshComponent::PostMigrate(const struct FUObjectMigrationContext&
 	Super::PostMigrate(MigrationContext);
 
 	AnimationMode = RestoreAnimationMode;
+
+	if (!HasValidPhysicsState())
+	{
+		// If PerPolyCollision is disabled, the primitive component physics state initialization is skipped, therefore we need to listen to
+		// a skeletal mesh component specific delegate
+		if (!bEnablePerPolyCollision)
+		{
+			OnComponentPhysicsStateChanged.RemoveDynamic(this, &USkeletalMeshComponent::HandlePostMigrationPhysicsStateChanged);
+
+			RegisterOnPhysicsCreatedDelegate(FOnSkelMeshPhysicsCreated::CreateLambda([WeakThis = TWeakObjectPtr<USkeletalMeshComponent>(this)]()
+			{
+				if (USkeletalMeshComponent* SKComponent = WeakThis.Get())
+				{
+					SKComponent->HandlePostMigrationPhysicsStateChanged(SKComponent, EComponentPhysicsStateChange::Created);
+				}
+			}));
+		}
+	}
 }
 #endif
 
@@ -3860,35 +3878,23 @@ void USkeletalMeshComponent::MigratePhysicsState(FArchive& Ar)
 			UE_LOG(LogSkeletalMesh, Verbose, TEXT("Migrated an empty physics state (this can be expected) | [%s] will use its default physics state or the one already had | Migration Context [%s]."), *GetPathNameSafe(this), ToString(ObjMigrationRecvType));
 			return;
 		}
-	
-		if (ObjMigrationRecvType == EObjectMigrationRecvType::ReturnedLoan)
-		{
-			return;
-		}
 
-		if (HasValidPhysicsState())
+		if (const bool bFinishedDestroy = HasAnyFlags(EObjectFlags::RF_MirroredGarbage))
 		{
-			ApplyMigratedPhysicsState(MoveTemp(MigratedStateData));
+			return;
 		}
-		else if (FPhysScene* PhysicsScene = FPhysicsInterface::GetCurrentSceneForComponent(this))
+		
+		FPhysScene* PhysicsScene = FPhysicsInterface::GetCurrentSceneForComponent(this);
+		if (ensure(PhysicsScene))
 		{
 			PhysicsScene->PushPendingSerializedStateForObject(this, MoveTemp(MigratedStateData));
-
-			RegisterOnPhysicsCreatedDelegate(FOnSkelMeshPhysicsCreated::CreateLambda([WeakThis = TWeakObjectPtr<USkeletalMeshComponent>(this)]()
-			{
-				if (USkeletalMeshComponent* SKComponent = WeakThis.Get())
-				{
-					SKComponent->ApplyMigratedPhysicsState();
-				}
-			}));
 		}
 	}
 	else if (Ar.IsSaving())
 	{
 		const EObjectMigrationSendType ObjMigrationSendType = MigrationContext->GetObjectMigrationSendType(this);
-		const bool bCanSendPhysicsState = ObjMigrationSendType != EObjectMigrationSendType::ReturnBorrowed && HasValidPhysicsState();
 
-		if (bCanSendPhysicsState)
+		if (CanSendMigratedPhysicsState(MigrationContext, ObjMigrationSendType))
 		{
 			Chaos::FSerializedMultiPhysicsState MigratedState(ObjMigrationSendType == EObjectMigrationSendType::ReassignOwnership ? Chaos::ESerializedDataContext::Internal : Chaos::ESerializedDataContext::External);
 			MigratedState.SetNumOfBodies(Bodies.Num());
diff --git a/Engine/Source/Runtime/Engine/Private/PhysicsEngine/BodyInstance.cpp b/Engine/Source/Runtime/Engine/Private/PhysicsEngine/BodyInstance.cpp
--- a/Engine/Source/Runtime/Engine/Private/PhysicsEngine/BodyInstance.cpp
+++ b/Engine/Source/Runtime/Engine/Private/PhysicsEngine/BodyInstance.cpp
@@ -3468,7 +3468,7 @@ void FBodyInstance::ApplySerializedState(Chaos::FSerializedDataBufferPtr&& InSer
 	FPhysicsCommand::ExecuteWrite(ActorHandle, [&InSerializedState, DataContext](const FPhysicsActorHandle& Actor)
 	{
 		FChaosEngineInterface::ApplySerializedStateToParticle(Actor, MoveTemp(InSerializedState), DataContext);
-	}, LockContext);
+	}, LockContext, EPhysicsInterfaceScopedTransactionMode::Normal);
 }
 
 void FBodyInstance::SerializeState(Chaos::FSerializedDataBuffer& OutSerializedState, Chaos::ESerializedDataContext DataContext)
diff --git a/Engine/Source/Runtime/Engine/Private/PhysicsEngine/ConstraintInstance.cpp b/Engine/Source/Runtime/Engine/Private/PhysicsEngine/ConstraintInstance.cpp
--- a/Engine/Source/Runtime/Engine/Private/PhysicsEngine/ConstraintInstance.cpp
+++ b/Engine/Source/Runtime/Engine/Private/PhysicsEngine/ConstraintInstance.cpp
@@ -728,7 +728,7 @@ void FConstraintInstance::ApplySerializedState(Chaos::FSerializedDataBufferPtr&&
 	FPhysicsCommand::ExecuteWrite(ConstraintHandle, [&InSerializedState, this, DataContext](const FPhysicsConstraintHandle& Constraint)
 	{
 		FChaosEngineInterface::ApplySerializedStateToConstraint(Constraint, MoveTemp(InSerializedState), DataContext);
-	}, LockContext);
+	}, LockContext, EPhysicsInterfaceScopedTransactionMode::Normal);
 }
 
 void FConstraintInstance::CopyConstraintGeometryFrom(const FConstraintInstance* FromInstance)
diff --git a/Engine/Source/Runtime/Engine/Private/PhysicsEngine/Experimental/ChaosScopedSceneLock.cpp b/Engine/Source/Runtime/Engine/Private/PhysicsEngine/Experimental/ChaosScopedSceneLock.cpp
--- a/Engine/Source/Runtime/Engine/Private/PhysicsEngine/Experimental/ChaosScopedSceneLock.cpp
+++ b/Engine/Source/Runtime/Engine/Private/PhysicsEngine/Experimental/ChaosScopedSceneLock.cpp
@@ -184,8 +184,16 @@ void FScopedSceneLock_Chaos::LockScene()
 		SceneLock->ReadLock();
 		break;
 	case EPhysicsInterfaceScopedLockType::Write:
-		SceneLock->WriteLock();
-		break;
+		{
+#if UE_WITH_REMOTE_OBJECT_HANDLE
+			if (TransactionMode == EPhysicsInterfaceScopedTransactionMode::MultiServer)
+			{
+				UE::RemoteExecutor::TransactionRequiresMultiServerCommit(TEXT("Physics Write Lock"));
+			}
+#endif
+			SceneLock->WriteLock();
+			break;
+		}
 	}
 
 	bHasLock = true;
@@ -279,19 +287,23 @@ void FScopedSceneLock_Chaos::LockSceneForConstraint(FPhysicsConstraintHandle con
 	LockScene();
 }
 
-FScopedSceneLockWithContext_Chaos::FScopedSceneLockWithContext_Chaos(const FPhysicsActorHandle& InActorHandle, EPhysicsInterfaceScopedLockType InLockType, EPhysicsInterfaceScopedThreadContext InThreadContext)
+FScopedSceneLockWithContext_Chaos::FScopedSceneLockWithContext_Chaos(const FPhysicsActorHandle& InActorHandle, EPhysicsInterfaceScopedLockType InLockType, EPhysicsInterfaceScopedThreadContext InThreadContext, EPhysicsInterfaceScopedTransactionMode InTransactionMode)
 	: FScopedSceneLock_Chaos(InLockType)
 {
 	ThreadContext = InThreadContext;
-
+#if UE_WITH_REMOTE_OBJECT_HANDLE
+	TransactionMode = InTransactionMode;
+#endif
 	Solver = GetSolverForActor(InActorHandle);
 	LockScene();
 }
 
-FScopedSceneLockWithContext_Chaos::FScopedSceneLockWithContext_Chaos(FPhysicsConstraintHandle const* InConstraintHandle, EPhysicsInterfaceScopedLockType InLockType, EPhysicsInterfaceScopedThreadContext InThreadContext)
+FScopedSceneLockWithContext_Chaos::FScopedSceneLockWithContext_Chaos(FPhysicsConstraintHandle const* InConstraintHandle, EPhysicsInterfaceScopedLockType InLockType, EPhysicsInterfaceScopedThreadContext InThreadContext, EPhysicsInterfaceScopedTransactionMode InTransactionMode)
 	: FScopedSceneLock_Chaos(InLockType)
 {
 	ThreadContext = InThreadContext;
-
+#if UE_WITH_REMOTE_OBJECT_HANDLE
+	TransactionMode = InTransactionMode;
+#endif
 	LockSceneForConstraint(InConstraintHandle);
 }
diff --git a/Engine/Source/Runtime/Engine/Private/PhysicsEngine/Experimental/PhysInterface_Chaos.cpp b/Engine/Source/Runtime/Engine/Private/PhysicsEngine/Experimental/PhysInterface_Chaos.cpp
--- a/Engine/Source/Runtime/Engine/Private/PhysicsEngine/Experimental/PhysInterface_Chaos.cpp
+++ b/Engine/Source/Runtime/Engine/Private/PhysicsEngine/Experimental/PhysInterface_Chaos.cpp
@@ -673,11 +673,11 @@ void FPhysInterface_Chaos::ExecuteShapeWrite(FBodyInstance* InInstance, FPhysics
 }
 
 
-bool FPhysInterface_Chaos::ExecuteWrite(FPhysicsActorHandle& InActorReference, TFunctionRef<void(FPhysicsActorHandle& Actor)> InCallable, EPhysicsInterfaceScopedThreadContext ThreadContext)
+bool FPhysInterface_Chaos::ExecuteWrite(FPhysicsActorHandle& InActorReference, TFunctionRef<void(FPhysicsActorHandle& Actor)> InCallable, EPhysicsInterfaceScopedThreadContext ThreadContext, EPhysicsInterfaceScopedTransactionMode TransactionMode)
 {
 	if (InActorReference)
 	{
-		FScopedSceneLockWithContext_Chaos SceneLock(InActorReference, EPhysicsInterfaceScopedLockType::Write, ThreadContext);
+		FScopedSceneLockWithContext_Chaos SceneLock(InActorReference, EPhysicsInterfaceScopedLockType::Write, ThreadContext, TransactionMode);
 		InCallable(InActorReference);
 		return true;
 	}
@@ -690,7 +690,7 @@ bool FPhysInterface_Chaos::ExecuteRead(FPhysicsActorHandle& InActorReference, TF
 {
 	if (InActorReference)
 	{
-		FScopedSceneLockWithContext_Chaos SceneLock(InActorReference, EPhysicsInterfaceScopedLockType::Read, ThreadContext);
+		FScopedSceneLockWithContext_Chaos SceneLock(InActorReference, EPhysicsInterfaceScopedLockType::Read, ThreadContext, EPhysicsInterfaceScopedTransactionMode::Normal);
 		InCallable(InActorReference);
 		return true;
 	}
@@ -698,11 +698,11 @@ bool FPhysInterface_Chaos::ExecuteRead(FPhysicsActorHandle& InActorReference, TF
 	return false;
 }
 
-bool FPhysInterface_Chaos::ExecuteWrite(const FPhysicsConstraintHandle& InConstraintRef, TFunctionRef<void(const FPhysicsConstraintHandle& Constraint)> InCallable, EPhysicsInterfaceScopedThreadContext ThreadContext)
+bool FPhysInterface_Chaos::ExecuteWrite(const FPhysicsConstraintHandle& InConstraintRef, TFunctionRef<void(const FPhysicsConstraintHandle& Constraint)> InCallable, EPhysicsInterfaceScopedThreadContext ThreadContext, EPhysicsInterfaceScopedTransactionMode TransactionMode)
 {
 	if (InConstraintRef.IsValid())
 	{
-		FScopedSceneLockWithContext_Chaos SceneLock(&InConstraintRef, EPhysicsInterfaceScopedLockType::Write, ThreadContext);
+		FScopedSceneLockWithContext_Chaos SceneLock(&InConstraintRef, EPhysicsInterfaceScopedLockType::Write, ThreadContext, TransactionMode);
 		InCallable(InConstraintRef);
 		return true;
 	}
@@ -714,7 +714,7 @@ bool FPhysInterface_Chaos::ExecuteRead(const FPhysicsConstraintHandle& InConstra
 {
 	if (InConstraintRef.IsValid())
 	{
-		FScopedSceneLockWithContext_Chaos SceneLock(&InConstraintRef, EPhysicsInterfaceScopedLockType::Read, ThreadContext);
+		FScopedSceneLockWithContext_Chaos SceneLock(&InConstraintRef, EPhysicsInterfaceScopedLockType::Read, ThreadContext, EPhysicsInterfaceScopedTransactionMode::Normal);
 		InCallable(InConstraintRef);
 		return true;
 	}
diff --git a/Engine/Source/Runtime/Engine/Public/Physics/Experimental/ChaosScopedSceneLock.h b/Engine/Source/Runtime/Engine/Public/Physics/Experimental/ChaosScopedSceneLock.h
--- a/Engine/Source/Runtime/Engine/Public/Physics/Experimental/ChaosScopedSceneLock.h
+++ b/Engine/Source/Runtime/Engine/Public/Physics/Experimental/ChaosScopedSceneLock.h
@@ -18,6 +18,12 @@ enum class EPhysicsInterfaceScopedThreadContext : uint8
 	Internal
 };
 
+enum class EPhysicsInterfaceScopedTransactionMode : uint8
+{
+	Normal,
+	MultiServer
+};
+
 class USkeletalMeshComponent;
 class FChaosScene;
 
@@ -69,11 +75,14 @@ protected:
 
 	EPhysicsInterfaceScopedThreadContext ThreadContext = EPhysicsInterfaceScopedThreadContext::External;
 
+#if UE_WITH_REMOTE_OBJECT_HANDLE
+	/** By default, any physics write operation needs to be done as part of a multi-server commit transaction */
+	EPhysicsInterfaceScopedTransactionMode TransactionMode = EPhysicsInterfaceScopedTransactionMode::MultiServer;
+#endif
 };
 
 struct UE_INTERNAL FScopedSceneLockWithContext_Chaos : FScopedSceneLock_Chaos
 {
-	ENGINE_API FScopedSceneLockWithContext_Chaos(const FPhysicsActorHandle& InActorHandle, EPhysicsInterfaceScopedLockType InLockType, EPhysicsInterfaceScopedThreadContext InThreadContext);
-	ENGINE_API FScopedSceneLockWithContext_Chaos(FPhysicsConstraintHandle const* InConstraintHandle, EPhysicsInterfaceScopedLockType InLockType, EPhysicsInterfaceScopedThreadContext InThreadContext);
+	ENGINE_API FScopedSceneLockWithContext_Chaos(const FPhysicsActorHandle& InActorHandle, EPhysicsInterfaceScopedLockType InLockType, EPhysicsInterfaceScopedThreadContext InThreadContext, EPhysicsInterfaceScopedTransactionMode InTransactionMode);
+	ENGINE_API FScopedSceneLockWithContext_Chaos(FPhysicsConstraintHandle const* InConstraintHandle, EPhysicsInterfaceScopedLockType InLockType, EPhysicsInterfaceScopedThreadContext InThreadContext, EPhysicsInterfaceScopedTransactionMode InTransactionMode);
 };
-
diff --git a/Engine/Source/Runtime/Engine/Public/Physics/Experimental/PhysInterface_Chaos.h b/Engine/Source/Runtime/Engine/Public/Physics/Experimental/PhysInterface_Chaos.h
--- a/Engine/Source/Runtime/Engine/Public/Physics/Experimental/PhysInterface_Chaos.h
+++ b/Engine/Source/Runtime/Engine/Public/Physics/Experimental/PhysInterface_Chaos.h
@@ -15,6 +15,7 @@
 
 //NOTE: Do not include Chaos headers directly as it means recompiling all of engine. This should be reworked to avoid allocations
 
+enum class EPhysicsInterfaceScopedTransactionMode : uint8;
 enum class EPhysicsInterfaceScopedThreadContext : uint8;
 static int32 NextBodyIdValue = 0;
 static int32 NextConstraintIdValue = 0;
@@ -87,9 +88,9 @@ public:
 
     static ENGINE_API void ExecuteShapeWrite(FBodyInstance* InInstance, FPhysicsShapeHandle& InShape, TFunctionRef<void(FPhysicsShapeHandle& InShape)> InCallable);
 
-	UE_INTERNAL static ENGINE_API bool ExecuteWrite(FPhysicsActorHandle& InActorReference, TFunctionRef<void(FPhysicsActorHandle& Actor)> InCallable, EPhysicsInterfaceScopedThreadContext ThreadContext);
+	UE_INTERNAL static ENGINE_API bool ExecuteWrite(FPhysicsActorHandle& InActorReference, TFunctionRef<void(FPhysicsActorHandle& Actor)> InCallable, EPhysicsInterfaceScopedThreadContext ThreadContext, EPhysicsInterfaceScopedTransactionMode TransactionMode);
 	UE_INTERNAL static ENGINE_API bool ExecuteRead(FPhysicsActorHandle& InActorReference, TFunctionRef<void(FPhysicsActorHandle& Actor)> InCallable, EPhysicsInterfaceScopedThreadContext ThreadContext);
-	UE_INTERNAL static ENGINE_API bool ExecuteWrite(const FPhysicsConstraintHandle& InConstraintRef, TFunctionRef<void(const FPhysicsConstraintHandle& Constraint)> InCallable, EPhysicsInterfaceScopedThreadContext ThreadContext);
+	UE_INTERNAL static ENGINE_API bool ExecuteWrite(const FPhysicsConstraintHandle& InConstraintRef, TFunctionRef<void(const FPhysicsConstraintHandle& Constraint)> InCallable, EPhysicsInterfaceScopedThreadContext ThreadContext, EPhysicsInterfaceScopedTransactionMode TransactionMode);
 	UE_INTERNAL static ENGINE_API bool ExecuteRead(const FPhysicsConstraintHandle& InConstraintRef, TFunctionRef<void(const FPhysicsConstraintHandle& Constraint)> InCallable, EPhysicsInterfaceScopedThreadContext ThreadContext);
 
 	// Misc
