RemoteObjects: Players need to notify the new server of what World "ClientWorldPackageName" they were apart of.  This helps them be able to Respawn.

#rnx

diff --git a/Engine/Source/Runtime/Engine/Private/World.cpp b/Engine/Source/Runtime/Engine/Private/World.cpp
--- a/Engine/Source/Runtime/Engine/Private/World.cpp
+++ b/Engine/Source/Runtime/Engine/Private/World.cpp
@@ -7258,6 +7258,14 @@ void UWorld::NotifyControlMessage(UNetConnection* Connection, uint8 MessageType,
 					Connection->PlayerController->SetClientHandshakeId(Connection->GetClientHandshakeId());
 
 					Connection->OwningActor = Connection->PlayerController;
+
+#if WITH_EDITORONLY_DATA
+					check(CurrentLevel);
+#else
+					ULevel* CurrentLevel = PersistentLevel;
+#endif
+
+					Connection->SetClientWorldPackageName(CurrentLevel->GetOutermost()->GetFName());
 				}
 				else
 				{
@@ -7306,6 +7314,14 @@ void UWorld::NotifyControlMessage(UNetConnection* Connection, uint8 MessageType,
 						ChildConnection->PlayerController->SetClientHandshakeId(ChildConnection->GetClientHandshakeId());
 
 						ChildConnection->OwningActor = ChildConnection->PlayerController;
+
+#if WITH_EDITORONLY_DATA
+						check(CurrentLevel);
+#else
+						ULevel* CurrentLevel = PersistentLevel;
+#endif
+
+						ChildConnection->SetClientWorldPackageName(CurrentLevel->GetOutermost()->GetFName());
 					}
 				}
 				else
