SplineComponent: Preserve attribute values during BP reconstruct

diff --git a/Engine/Source/Runtime/Engine/Classes/Components/SplineComponent.h b/Engine/Source/Runtime/Engine/Classes/Components/SplineComponent.h
--- a/Engine/Source/Runtime/Engine/Classes/Components/SplineComponent.h
+++ b/Engine/Source/Runtime/Engine/Classes/Components/SplineComponent.h
@@ -393,7 +393,8 @@ public:
 	virtual bool AllowsSplinePointLeaveTangentEditing() const { return true; }
 
 	ENGINE_API void ApplyComponentInstanceData(struct FSplineInstanceData* ComponentInstanceData, const bool bPostUCS);
-
+	ENGINE_API void ApplyComponentInstanceData(struct FSplineComponentInstanceData* ComponentInstanceData, const bool bPostUCS);
+	
 	/** Reset the spline to its default shape (a spline of two points) */
 	ENGINE_API void ResetToDefault();
 	ENGINE_API bool CanResetToDefault() const;
@@ -1085,7 +1086,44 @@ public:
 	FSplineCurves SplineCurvesPreUCS;
 };
 
+/** Used to store spline data during RerunConstructionScripts */
+USTRUCT()
+struct FSplineComponentInstanceData : public FSceneComponentInstanceData
+{
+	GENERATED_BODY()
+public:
+	FSplineComponentInstanceData()
+		: bSplineHasBeenEdited(false)
+	{}
+	explicit FSplineComponentInstanceData(const USplineComponent* SourceComponent)
+		: FSceneComponentInstanceData(SourceComponent)
+		, bSplineHasBeenEdited(false)
+	{}
+	virtual ~FSplineComponentInstanceData() = default;
+
+	virtual bool ContainsData() const override
+	{
+		return Super::ContainsData() || bSplineHasBeenEdited;
+	}
+
+	virtual void ApplyToComponent(UActorComponent* Component, const ECacheApplyPhase CacheApplyPhase) override
+	{
+		Super::ApplyToComponent(Component, CacheApplyPhase);
+		CastChecked<USplineComponent>(Component)->ApplyComponentInstanceData(this, (CacheApplyPhase == ECacheApplyPhase::PostUserConstructionScript));
+	}
+
+	UPROPERTY()
+	bool bSplineHasBeenEdited;
 
+	UPROPERTY()
+	FSpline Spline;
+
+	UPROPERTY()
+	bool bClosedLoop = false;
+
+	UPROPERTY()
+	FSpline SplinePreUCS;
+};
 
 ENGINE_API EInterpCurveMode ConvertSplinePointTypeToInterpCurveMode(ESplinePointType::Type SplinePointType);
 ENGINE_API ESplinePointType::Type ConvertInterpCurveModeToSplinePointType(EInterpCurveMode InterpCurveMode);
diff --git a/Engine/Source/Runtime/Engine/Private/Components/SplineComponent.cpp b/Engine/Source/Runtime/Engine/Private/Components/SplineComponent.cpp
--- a/Engine/Source/Runtime/Engine/Private/Components/SplineComponent.cpp
+++ b/Engine/Source/Runtime/Engine/Private/Components/SplineComponent.cpp
@@ -2523,17 +2523,34 @@ bool USplineComponent::GetIgnoreBoundsForEditorFocus() const
 
 TStructOnScope<FActorComponentInstanceData> USplineComponent::GetComponentInstanceData() const
 {
-	TStructOnScope<FActorComponentInstanceData> InstanceData = MakeStructOnScope<FActorComponentInstanceData, FSplineInstanceData>(this);
-	FSplineInstanceData* SplineInstanceData = InstanceData.Cast<FSplineInstanceData>();
-	
-	if (bSplineHasBeenEdited)
+	if (UE::SplineComponent::GUseSplineCurves)
 	{
-		SplineInstanceData->SplineCurves = GetSplineCurves();
-		SplineInstanceData->bClosedLoop = bClosedLoop;
+		TStructOnScope<FActorComponentInstanceData> InstanceData = MakeStructOnScope<FActorComponentInstanceData, FSplineInstanceData>(this);
+		FSplineInstanceData* SplineInstanceData = InstanceData.Cast<FSplineInstanceData>();
+	
+		if (bSplineHasBeenEdited)
+		{
+			SplineInstanceData->SplineCurves = GetSplineCurves();
+			SplineInstanceData->bClosedLoop = bClosedLoop;
+		}
+		SplineInstanceData->bSplineHasBeenEdited = bSplineHasBeenEdited;
+		
+		return InstanceData;
 	}
-	SplineInstanceData->bSplineHasBeenEdited = bSplineHasBeenEdited;
+	else
+	{
+		TStructOnScope<FActorComponentInstanceData> InstanceData = MakeStructOnScope<FActorComponentInstanceData, FSplineComponentInstanceData>(this);
+		FSplineComponentInstanceData* SplineComponentInstanceData = InstanceData.Cast<FSplineComponentInstanceData>();
+	
+		if (bSplineHasBeenEdited)
+		{
+			SplineComponentInstanceData->Spline = Spline;
+			SplineComponentInstanceData->bClosedLoop = bClosedLoop;
+		}
+		SplineComponentInstanceData->bSplineHasBeenEdited = bSplineHasBeenEdited;
 
-	return InstanceData;
+		return InstanceData;
+	}
 }
 
 FSplineCurves USplineComponent::GetSplineCurves() const
@@ -2635,6 +2652,46 @@ void USplineComponent::ApplyComponentInstanceData(FSplineInstanceData* SplineIns
 	UpdateSpline();
 }
 
+void USplineComponent::ApplyComponentInstanceData(struct FSplineComponentInstanceData* SplineInstanceData, const bool bPostUCS)
+{
+	check(SplineInstanceData);
+
+	if (bPostUCS)
+	{
+		if (bInputSplinePointsToConstructionScript)
+		{
+			// Don't reapply the saved state after the UCS has run if we are inputting the points to it.
+			// This allows the UCS to work on the edited points and make its own changes.
+			return;
+		}
+		else
+		{
+			bModifiedByConstructionScript = SplineInstanceData->SplinePreUCS != Spline;
+
+			// If we are restoring the saved state, unmark the SplineCurves property as 'modified'.
+			// We don't want to consider that these changes have been made through the UCS.
+			TArray<FProperty*> Properties;
+			Properties.Emplace(FindFProperty<FProperty>(StaticClass(), GetSplinePropertyName()));
+			RemoveUCSModifiedProperties(Properties);
+		}
+	}
+	else
+	{
+		SplineInstanceData->SplinePreUCS = Spline;
+	}
+
+	if (SplineInstanceData->bSplineHasBeenEdited)
+	{
+		SetSpline(SplineInstanceData->Spline);
+		bClosedLoop = SplineInstanceData->bClosedLoop;
+		bModifiedByConstructionScript = false;
+	}
+
+	bSplineHasBeenEdited = SplineInstanceData->bSplineHasBeenEdited;
+
+	UpdateSpline();
+}
+
 #if WITH_EDITOR
 void USplineComponent::PostEditChangeChainProperty(FPropertyChangedChainEvent& PropertyChangedEvent)
 {
diff --git a/Engine/Source/Runtime/Engine/Private/Curves/Spline.cpp b/Engine/Source/Runtime/Engine/Private/Curves/Spline.cpp
--- a/Engine/Source/Runtime/Engine/Private/Curves/Spline.cpp
+++ b/Engine/Source/Runtime/Engine/Private/Curves/Spline.cpp
@@ -404,30 +404,17 @@ FSpline::FSpline(const FSpline& Other)
 
 FSpline& FSpline::operator=(const FSpline& Other)
 {
-	Reset();
-	
 	CurrentImplementation = Other.CurrentImplementation;
 
 	switch (CurrentImplementation)
 	{
-	default: break;		// Intentionally doing nothing
+	default:
+		break;
 	case 1:
-		if (Other.LegacyData)
-		{
-			FLegacySpline Temp = *Other.LegacyData;
-			LegacyData = MakeShared<FLegacySpline>(Temp);
-		}
-		else
-			{ LegacyData = MakeShared<FLegacySpline>(); }
+		*LegacyData = Other.LegacyData ? *Other.LegacyData : FLegacySpline();
 		break;
 	case 2:
-		if (Other.NewData)
-		{
-			FNewSpline Temp = *Other.NewData;
-			NewData = MakeShared<FNewSpline>(Temp);
-		}
-		else
-			{ NewData = MakeShared<FNewSpline>(); }
+		*NewData = Other.NewData ? *Other.NewData : FNewSpline();
 		break;
 	}
 	
