Use NeverCull for render captures initiated by custom render passes (follow-up to CL37111575)

#rnx

diff --git a/Engine/Source/Runtime/Engine/Private/Rendering/CustomRenderPass.cpp b/Engine/Source/Runtime/Engine/Private/Rendering/CustomRenderPass.cpp
--- a/Engine/Source/Runtime/Engine/Private/Rendering/CustomRenderPass.cpp
+++ b/Engine/Source/Runtime/Engine/Private/Rendering/CustomRenderPass.cpp
@@ -18,7 +18,7 @@ void FCustomRenderPassBase::BeginPass(FRDGBuilder& GraphBuilder)
 	{
 		GraphBuilder.AddPass(
 			RDG_EVENT_NAME("BeginCapture"),
-			ERDGPassFlags::None,
+			ERDGPassFlags::NeverCull,
 			[FileName = RenderCaptureFileName](FRHICommandListImmediate& RHICommandListLocal)
 		{
 			IRenderCaptureProvider::Get().BeginCapture(&RHICommandListLocal, IRenderCaptureProvider::ECaptureFlags_Launch, FileName);
@@ -47,7 +47,7 @@ void FCustomRenderPassBase::EndPass(FRDGBuilder & GraphBuilder)
 	{
 		GraphBuilder.AddPass(
 			RDG_EVENT_NAME("EndCapture"),
-			ERDGPassFlags::None,
+			ERDGPassFlags::NeverCull,
 			[](FRHICommandListImmediate& RHICommandListLocal)
 		{
 			IRenderCaptureProvider::Get().EndCapture(&RHICommandListLocal);
