Ensure that MaterialFunctionCalls preload the material functions they are dependent on to ensure the function is loaded before attempting a PostLoad call.

diff --git a/Engine/Source/Runtime/Engine/Private/Materials/MaterialExpressions.cpp b/Engine/Source/Runtime/Engine/Private/Materials/MaterialExpressions.cpp
--- a/Engine/Source/Runtime/Engine/Private/Materials/MaterialExpressions.cpp
+++ b/Engine/Source/Runtime/Engine/Private/Materials/MaterialExpressions.cpp
@@ -7,6 +7,7 @@
 #include "Field/FieldSystemTypes.h"
 #include "Misc/MessageDialog.h"
 #include "UObject/AssetRegistryTagsContext.h"
+#include "UObject/LinkerLoad.h"
 #include "UObject/UE5MainStreamObjectVersion.h"
 #include "Internationalization/LocKeyFuncs.h"
 #include "UObject/NaniteResearchStreamObjectVersion.h"
@@ -16989,6 +16990,10 @@ void UMaterialExpressionMaterialFunctionCall::PostLoad()
 {
 	if (MaterialFunction)
 	{
+		if (MaterialFunction->HasAnyFlags(RF_NeedLoad))
+		{
+			MaterialFunction->GetLinker()->Preload(MaterialFunction);
+		}
 		MaterialFunction->ConditionalPostLoad();
 	}
 
