[AsyncCompilation/StaticMesh]
- Fix warnings and error that could be reported after cancelation

diff --git a/Engine/Source/Developer/MeshBuilder/Private/StaticMeshBuilder.cpp b/Engine/Source/Developer/MeshBuilder/Private/StaticMeshBuilder.cpp
--- a/Engine/Source/Developer/MeshBuilder/Private/StaticMeshBuilder.cpp
+++ b/Engine/Source/Developer/MeshBuilder/Private/StaticMeshBuilder.cpp
@@ -645,6 +645,11 @@ static bool BuildNanite(
 			&InputAssemblyData
 	))
 	{
+		if (UE::Tasks::FCancellationTokenScope::IsCurrentWorkCanceled())
+		{
+			return false;
+		}
+
 		UE_LOGFMT_NSLOC(LogStaticMesh, Warning, "StaticMesh", "NaniteHiResBuildError", "Failed to build Nanite for HiRes static mesh. See previous line(s) for details.");
 		return false;
 	}
diff --git a/Engine/Source/Runtime/Engine/Private/StaticMesh.cpp b/Engine/Source/Runtime/Engine/Private/StaticMesh.cpp
--- a/Engine/Source/Runtime/Engine/Private/StaticMesh.cpp
+++ b/Engine/Source/Runtime/Engine/Private/StaticMesh.cpp
@@ -4073,6 +4073,11 @@ void FStaticMeshRenderData::Cache(const ITargetPlatform* TargetPlatform, UStatic
 
 				if (!MeshBuilderModule.BuildMesh(*this, FStaticMeshBuildParameters(Owner, TargetPlatform, LODGroup)))
 				{
+					if (UE::Tasks::FCancellationTokenScope::IsCurrentWorkCanceled())
+					{
+						return;
+					}
+
 					UE_LOG(LogStaticMesh, Error, TEXT("Failed to build static mesh. See previous line(s) for details."));
 					return;
 				}
