Fixed an issue with unwelding from a parent when there is a parent with disabled physics in the middle of the hierarchy. The invalid body was not skipped when searching for the weld root.

diff --git a/Engine/Source/Runtime/Engine/Private/PrimitiveComponentPhysics.cpp b/Engine/Source/Runtime/Engine/Private/PrimitiveComponentPhysics.cpp
--- a/Engine/Source/Runtime/Engine/Private/PrimitiveComponentPhysics.cpp
+++ b/Engine/Source/Runtime/Engine/Private/PrimitiveComponentPhysics.cpp
@@ -931,7 +931,7 @@ UPrimitiveComponent * GetRootWelded(const UPrimitiveComponent* PrimComponent, FN
 		PrevSocketName = RootComponent->GetAttachSocketName();
 
 		RootBI = RootComponent->GetBodyInstance(SocketName, false);
-		if (RootBI && RootBI->WeldParent != nullptr)
+		if (RootBI && (!RootBI->IsValidBodyInstance() || RootBI->WeldParent != nullptr))
 		{
 			continue;
 		}
