[HWRT] Revert part of previous change.
- need to handle building runtime BLAS even then there's no offline BLAS.

diff --git a/Engine/Source/Runtime/Engine/Private/Rendering/RayTracingGeometryManager.cpp b/Engine/Source/Runtime/Engine/Private/Rendering/RayTracingGeometryManager.cpp
--- a/Engine/Source/Runtime/Engine/Private/Rendering/RayTracingGeometryManager.cpp
+++ b/Engine/Source/Runtime/Engine/Private/Rendering/RayTracingGeometryManager.cpp
@@ -930,10 +930,15 @@ bool FRayTracingGeometryManager::RequestRayTracingGeometryStreamIn(FRHICommandLi
 		}
 	}
 
-	checkf(RegisteredGeometry.StreamableBVHSize > 0 || RegisteredGeometry.StreamableBuffersSize > 0, TEXT("FRayTracingStreamableAsset should have data to stream."));
-
 	// TODO: Support DDC streaming
 
+	if (RegisteredGeometry.StreamableBuffersSize == 0 && RegisteredGeometry.StreamableBVHSize == 0)
+	{
+		// no offline data -> build from VB/IB at runtime
+
+		RegisteredGeometry.Status = FRegisteredGeometry::FStatus::StreamedIn;
+	}
+	else
 	{
 		if (NumStreamingRequests >= GRayTracingStreamingMaxPendingRequests)
 		{
@@ -972,6 +977,27 @@ bool FRayTracingGeometryManager::RequestRayTracingGeometryStreamIn(FRHICommandLi
 		RegisteredGeometry.Status = FRegisteredGeometry::FStatus::Streaming;
 	}
 
+	if (RegisteredGeometry.Status == FRegisteredGeometry::FStatus::StreamedIn)
+	{
+		{
+			FRHIResourceReplaceBatcher Batcher(RHICmdList, 1);
+			FRayTracingGeometryInitializer IntermediateInitializer = Geometry->Initializer;
+			IntermediateInitializer.Type = ERayTracingGeometryInitializerType::StreamingSource;
+			IntermediateInitializer.OfflineData = nullptr;
+
+			FRayTracingGeometryRHIRef IntermediateRayTracingGeometry = RHICmdList.CreateRayTracingGeometry(IntermediateInitializer);
+
+			Geometry->SetRequiresBuild(IntermediateInitializer.OfflineData == nullptr || IntermediateRayTracingGeometry->IsCompressed());
+
+			Geometry->InitRHIForStreaming(IntermediateRayTracingGeometry, Batcher);
+
+			// When Batcher goes out of scope it will add commands to copy the BLAS buffers on RHI thread.
+			// We need to do it before we build the current geometry (also on RHI thread).
+		}
+
+		Geometry->RequestBuildIfNeeded(RHICmdList, ERTAccelerationStructureBuildPriority::Normal);
+	}
+
 	return true;
 }
 
