Added FNavigationSystem::SupportsDynamicChanges (returns whether Navigation System supports dynamic changes).

diff --git a/Engine/Source/Runtime/Engine/Classes/AI/NavigationSystemBase.h b/Engine/Source/Runtime/Engine/Classes/AI/NavigationSystemBase.h
--- a/Engine/Source/Runtime/Engine/Classes/AI/NavigationSystemBase.h
+++ b/Engine/Source/Runtime/Engine/Classes/AI/NavigationSystemBase.h
@@ -158,6 +158,7 @@ namespace FNavigationSystem
 	ENGINE_API void UnregisterNavRelevantObject(UObject& Object);
 	ENGINE_API void OnObjectBoundsChanged(UObject& Object, const FBox& NewBounds, TConstArrayView<FBox> DirtyAreas);
 
+	ENGINE_API bool SupportsDynamicChanges(UWorld* World);
 	ENGINE_API FNavigationElementHandle AddNavigationElement(UWorld* World, FNavigationElement&& Element);
 	ENGINE_API void UpdateNavigationElement(UWorld* World, FNavigationElementHandle ElementHandle, FNavigationElement&& Element);
 	ENGINE_API void RemoveNavigationElement(UWorld* World, FNavigationElementHandle ElementHandle);
@@ -204,6 +205,7 @@ namespace FNavigationSystem
 	ENGINE_API void StopMovement(const AController& Controller);
 	ENGINE_API IPathFollowingAgentInterface* FindPathFollowingAgentForActor(const AActor& Actor);
 
+	DECLARE_DELEGATE_RetVal_OneParam(bool, FSupportsDynamicChangesSignature, UWorld* /*World*/);
 	DECLARE_DELEGATE_RetVal_TwoParams(FNavigationElementHandle, FAddNavigationElementSignature, UWorld* /*World*/, FNavigationElement&& /*Element*/);
 	DECLARE_DELEGATE_TwoParams(FRemoveNavigationElementSignature, UWorld* /*World*/, FNavigationElementHandle /*ElementHandle*/);
 	DECLARE_DELEGATE_ThreeParams(FUpdateNavigationElementSignature, UWorld* /*World*/, FNavigationElementHandle /*ElementHandle*/, FNavigationElement&& /*Element*/)
@@ -307,6 +309,7 @@ protected:
 
 	static ENGINE_API void ResetEventDelegates();
 
+	static ENGINE_API FNavigationSystem::FSupportsDynamicChangesSignature& GetSupportsDynamicChangesDelegate();
 	static ENGINE_API FNavigationSystem::FAddNavigationElementSignature& GetAddNavigationElementDelegate();
 	static ENGINE_API FNavigationSystem::FUpdateNavigationElementSignature& GetUpdateNavigationElementDelegate();
 	static ENGINE_API FNavigationSystem::FRemoveNavigationElementSignature& GetRemoveNavigationElementDelegate();
diff --git a/Engine/Source/Runtime/Engine/Private/AI/NavigationSystemBase.cpp b/Engine/Source/Runtime/Engine/Private/AI/NavigationSystemBase.cpp
--- a/Engine/Source/Runtime/Engine/Private/AI/NavigationSystemBase.cpp
+++ b/Engine/Source/Runtime/Engine/Private/AI/NavigationSystemBase.cpp
@@ -130,6 +130,7 @@ namespace FNavigationSystem
 	class FDelegates
 	{
 	public:
+		FSupportsDynamicChangesSignature SupportsDynamicChangesDelegate;
 		FAddNavigationElementSignature AddNavigationElementDelegate;
 		FUpdateNavigationElementSignature UpdateNavigationElementDelegate;
 		FRemoveNavigationElementSignature RemoveNavigationElementDelegate;
@@ -237,6 +238,11 @@ PRAGMA_ENABLE_DEPRECATION_WARNINGS
 		new(&Delegates)FDelegates();
 	}
 
+	bool SupportsDynamicChanges(UWorld* World)
+	{
+		return Delegates.SupportsDynamicChangesDelegate.Execute(World);
+	}
+
 	FNavigationElementHandle AddNavigationElement(UWorld* World, FNavigationElement&& Element)
 	{
 		return Delegates.AddNavigationElementDelegate.Execute(World, MoveTemp(Element));
@@ -443,6 +449,11 @@ void UNavigationSystemBase::SetDefaultObstacleArea(TSubclassOf<UNavAreaBase> InA
 
 void UNavigationSystemBase::ResetEventDelegates() { FNavigationSystem::ResetDelegates(); }
 
+FNavigationSystem::FSupportsDynamicChangesSignature& UNavigationSystemBase::GetSupportsDynamicChangesDelegate()
+{
+	return FNavigationSystem::Delegates.SupportsDynamicChangesDelegate;
+}
+
 FNavigationSystem::FAddNavigationElementSignature& UNavigationSystemBase::GetAddNavigationElementDelegate()
 {
 	return FNavigationSystem::Delegates.AddNavigationElementDelegate;
diff --git a/Engine/Source/Runtime/NavigationSystem/Private/NavigationSystem.cpp b/Engine/Source/Runtime/NavigationSystem/Private/NavigationSystem.cpp
--- a/Engine/Source/Runtime/NavigationSystem/Private/NavigationSystem.cpp
+++ b/Engine/Source/Runtime/NavigationSystem/Private/NavigationSystem.cpp
@@ -760,6 +760,7 @@ UNavigationSystemV1::UNavigationSystemV1(const FObjectInitializer& ObjectInitial
 	{
 		FDelegatesInitializer()
 		{
+			UNavigationSystemBase::GetSupportsDynamicChangesDelegate().BindStatic(&UNavigationSystemV1::SupportsDynamicChanges);
 			UNavigationSystemBase::GetAddNavigationElementDelegate().BindStatic(&UNavigationSystemV1::AddNavigationElement);
 			UNavigationSystemBase::GetRemoveNavigationElementDelegate().BindStatic(&UNavigationSystemV1::RemoveNavigationElement);
 			UNavigationSystemBase::GetUpdateNavigationElementDelegate().BindStatic(&UNavigationSystemV1::OnNavigationElementUpdated);
@@ -3392,6 +3393,17 @@ void UNavigationSystemV1::RegisterComponentToNavOctree(UActorComponent* Comp)
 	}
 }
 
+bool UNavigationSystemV1::SupportsDynamicChanges(UWorld* World)
+{
+	if (IsNavigationSystemStatic())
+	{
+		return false;
+	}
+
+	UNavigationSystemV1* NavSys = FNavigationSystem::GetCurrent<UNavigationSystemV1>(World);
+	return NavSys && NavSys->RequiresNavOctree();
+}
+
 FNavigationElementHandle UNavigationSystemV1::AddNavigationElement(UWorld* World, FNavigationElement&& Element)
 {
 	SCOPE_CYCLE_COUNTER(STAT_NavOctreeBookkeeping);
diff --git a/Engine/Source/Runtime/NavigationSystem/Public/NavigationSystem.h b/Engine/Source/Runtime/NavigationSystem/Public/NavigationSystem.h
--- a/Engine/Source/Runtime/NavigationSystem/Public/NavigationSystem.h
+++ b/Engine/Source/Runtime/NavigationSystem/Public/NavigationSystem.h
@@ -896,6 +896,7 @@ public:
 	//----------------------------------------------------------------------//
 	// navigation octree related functions
 	//----------------------------------------------------------------------//
+	static NAVIGATIONSYSTEM_API bool SupportsDynamicChanges(UWorld* World);
 	static NAVIGATIONSYSTEM_API FNavigationElementHandle AddNavigationElement(UWorld* World, FNavigationElement&& Element);
 	static NAVIGATIONSYSTEM_API void RemoveNavigationElement(UWorld* World, FNavigationElementHandle ElementHandle);
 	static NAVIGATIONSYSTEM_API void OnNavigationElementUpdated(UWorld* World, FNavigationElementHandle ElementHandle, FNavigationElement&& Element);
