Make Submix & Soundclass parents Transient

diff --git a/Engine/Source/Runtime/Engine/Classes/Sound/SoundClass.h b/Engine/Source/Runtime/Engine/Classes/Sound/SoundClass.h
--- a/Engine/Source/Runtime/Engine/Classes/Sound/SoundClass.h
+++ b/Engine/Source/Runtime/Engine/Classes/Sound/SoundClass.h
@@ -221,7 +221,7 @@ public:
 	UPROPERTY(EditAnywhere, BlueprintReadOnly, Category = General)
 	TArray<FPassiveSoundMixModifier> PassiveSoundMixModifiers;
 
-	UPROPERTY(BlueprintReadOnly, Category = General)
+	UPROPERTY(Transient, BlueprintReadOnly, Category = General)
 	TObjectPtr<USoundClass> ParentClass;
 
 #if WITH_EDITORONLY_DATA
diff --git a/Engine/Source/Runtime/Engine/Classes/Sound/SoundSubmix.h b/Engine/Source/Runtime/Engine/Classes/Sound/SoundSubmix.h
--- a/Engine/Source/Runtime/Engine/Classes/Sound/SoundSubmix.h
+++ b/Engine/Source/Runtime/Engine/Classes/Sound/SoundSubmix.h
@@ -275,7 +275,7 @@ class USoundSubmixWithParentBase : public USoundSubmixBase
 public:
 
 
-	UPROPERTY(VisibleAnywhere, BlueprintReadOnly, Category = SoundSubmix, meta = (EditCondition = "!bIsDynamic"))
+	UPROPERTY(Transient, VisibleAnywhere, BlueprintReadOnly, Category = SoundSubmix, meta = (EditCondition = "!bIsDynamic"))
 	TObjectPtr<USoundSubmixBase> ParentSubmix;
 
 	/**
diff --git a/Engine/Source/Runtime/Engine/Private/SoundSubmix.cpp b/Engine/Source/Runtime/Engine/Private/SoundSubmix.cpp
--- a/Engine/Source/Runtime/Engine/Private/SoundSubmix.cpp
+++ b/Engine/Source/Runtime/Engine/Private/SoundSubmix.cpp
@@ -447,6 +447,22 @@ void USoundSubmixBase::PostLoad()
 {
 	Super::PostLoad();
 
+	for (int32 ChildIndex = 0; ChildIndex < ChildSubmixes.Num(); ChildIndex++)
+	{
+		USoundSubmixBase* ChildSubmix = ChildSubmixes[ChildIndex];
+
+		if (!ChildSubmix)
+		{
+			continue;
+		}
+
+		if (USoundSubmixWithParentBase* CastedChildSubmix = Cast<USoundSubmixWithParentBase>(ChildSubmix))
+		{
+			// Parent is transient so don't need to mark dirty
+			CastedChildSubmix->SetParentSubmix(this, false /* bModifyAssets */);
+		}
+	}
+	
 	if (ClearBrokenSubmixAssetsCVar)
 	{
 		for (int32 ChildIndex = ChildSubmixes.Num() - 1; ChildIndex >= 0; ChildIndex--)
