PR #12235: Fix ViewOriginHigh for multiview / instanced stereo in the right eye

diff --git a/Engine/Source/Runtime/Engine/Private/Materials/HLSLMaterialTranslator.cpp b/Engine/Source/Runtime/Engine/Private/Materials/HLSLMaterialTranslator.cpp
--- a/Engine/Source/Runtime/Engine/Private/Materials/HLSLMaterialTranslator.cpp
+++ b/Engine/Source/Runtime/Engine/Private/Materials/HLSLMaterialTranslator.cpp
@@ -16236,7 +16236,7 @@ int32 FHLSLMaterialTranslator::GetPrimitiveProperty(EMaterialValueType Type, con
 	return AddInlinedCodeChunkZeroDeriv(Type, TEXT("GetPrimitiveData(Parameters).%s"), HLSLName);
 }
 
-// The compiler can run in a different state and this affects caching of sub expression, Expressions are different (e.g. View.PrevWorldViewOrigin) when using previous frame's values
+// The compiler can run in a different state and this affects caching of sub expression, Expressions are different (e.g. ResolvedView.PrevWorldViewOrigin) when using previous frame's values
 bool FHLSLMaterialTranslator::IsCurrentlyCompilingForPreviousFrame() const
 {
 	return bCompilingPreviousFrame;
diff --git a/Engine/Source/Runtime/Engine/Private/Materials/HLSLMaterialTranslator.h b/Engine/Source/Runtime/Engine/Private/Materials/HLSLMaterialTranslator.h
--- a/Engine/Source/Runtime/Engine/Private/Materials/HLSLMaterialTranslator.h
+++ b/Engine/Source/Runtime/Engine/Private/Materials/HLSLMaterialTranslator.h
@@ -1457,7 +1457,7 @@ protected:
 	int32 GetPrimitiveProperty(EMaterialValueType Type, const TCHAR* ExpressionName, const TCHAR* HLSLName);
 
 	/**
-	 * The compiler can run in a different state and this affects caching of sub expression, Expressions are different(e.g.View.PrevWorldViewOrigin) when using previous frame's values.
+	 * The compiler can run in a different state and this affects caching of sub expression, Expressions are different(e.g. ResolvedView.PrevWorldViewOrigin) when using previous frame's values.
 	 */
 	virtual bool IsCurrentlyCompilingForPreviousFrame() const;
 
diff --git a/Engine/Source/Runtime/Engine/Public/MaterialCompiler.h b/Engine/Source/Runtime/Engine/Public/MaterialCompiler.h
--- a/Engine/Source/Runtime/Engine/Public/MaterialCompiler.h
+++ b/Engine/Source/Runtime/Engine/Public/MaterialCompiler.h
@@ -843,7 +843,7 @@ public:
 	virtual int32 SceneDepthWithoutWater(int32 Offset, int32 ViewportUV, bool bUseOffset, float FallbackDepth) = 0;
 
 	virtual int32 MapARPassthroughCameraUV(int32 UV) = 0;
-	// The compiler can run in a different state and this affects caching of sub expression, Expressions are different (e.g. View.PrevWorldViewOrigin) when using previous frame's values
+	// The compiler can run in a different state and this affects caching of sub expression, Expressions are different (e.g. ResolvedView.PrevWorldViewOrigin) when using previous frame's values
 	// If possible we should re-factor this to avoid having to deal with compiler state
 	virtual bool IsCurrentlyCompilingForPreviousFrame() const { return false; }
 	virtual bool IsDevelopmentFeatureEnabled(const FName& FeatureName) const { return true; }
diff --git a/Engine/Source/Runtime/Engine/Public/MaterialShared.h b/Engine/Source/Runtime/Engine/Public/MaterialShared.h
--- a/Engine/Source/Runtime/Engine/Public/MaterialShared.h
+++ b/Engine/Source/Runtime/Engine/Public/MaterialShared.h
@@ -1869,7 +1869,7 @@ public:
 	int32 OutputIndex;
 	/** Attribute currently being compiled through a MatterialAttributes connection. */
 	FGuid MaterialAttributeID;
-	// Expressions are different (e.g. View.PrevWorldViewOrigin) when using previous frame's values, value if from FHLSLMaterialTranslator::bCompilingPreviousFrame
+	// Expressions are different (e.g. ResolvedView.PrevWorldViewOrigin) when using previous frame's values, value if from FHLSLMaterialTranslator::bCompilingPreviousFrame
 	bool bCompilingPreviousFrameKey;
 
 	FMaterialExpressionKey(UMaterialExpression* InExpression, int32 InOutputIndex) :
diff --git a/Engine/Source/Runtime/Engine/Public/SceneView.h b/Engine/Source/Runtime/Engine/Public/SceneView.h
--- a/Engine/Source/Runtime/Engine/Public/SceneView.h
+++ b/Engine/Source/Runtime/Engine/Public/SceneView.h
@@ -859,7 +859,7 @@ enum ETranslucencyVolumeCascade
 	VIEW_UNIFORM_BUFFER_MEMBER_PER_VIEW(FMatrix44f, MobileMultiviewDecalTransform) \
 	VIEW_UNIFORM_BUFFER_MEMBER_PER_VIEW(FMatrix44f, FirstPersonTransform) \
 	VIEW_UNIFORM_BUFFER_MEMBER_PER_VIEW(FMatrix44f, PrevFirstPersonTransform) \
-	VIEW_UNIFORM_BUFFER_MEMBER(FVector3f, ViewOriginHigh) \
+	VIEW_UNIFORM_BUFFER_MEMBER_PER_VIEW(FVector3f, ViewOriginHigh) \
 	VIEW_UNIFORM_BUFFER_MEMBER_EX(FVector3f, ViewForward, EShaderPrecisionModifier::Half) \
 	VIEW_UNIFORM_BUFFER_MEMBER_EX(FVector3f, ViewUp, EShaderPrecisionModifier::Half) \
 	VIEW_UNIFORM_BUFFER_MEMBER_EX(FVector3f, ViewRight, EShaderPrecisionModifier::Half) \
