Moved/refactors StaticMeshComponent and InstancedStaticMeshComponent navigation code in helpers for other component types to share the implementation (will be used by FastGeo).

diff --git a/Engine/Source/Runtime/Engine/Classes/Components/StaticMeshComponent.h b/Engine/Source/Runtime/Engine/Classes/Components/StaticMeshComponent.h
--- a/Engine/Source/Runtime/Engine/Classes/Components/StaticMeshComponent.h
+++ b/Engine/Source/Runtime/Engine/Classes/Components/StaticMeshComponent.h
@@ -908,6 +908,7 @@ private:
 
 	friend class FStaticMeshComponentRecreateRenderStateContext;
 	friend class FStaticMeshComponentHelper;
+	friend class FInstancedStaticMeshComponentHelper;
 };
 
 /** Vertex data stored per-LOD */
diff --git a/Engine/Source/Runtime/Engine/Private/Components/PrimitiveComponent.cpp b/Engine/Source/Runtime/Engine/Private/Components/PrimitiveComponent.cpp
--- a/Engine/Source/Runtime/Engine/Private/Components/PrimitiveComponent.cpp
+++ b/Engine/Source/Runtime/Engine/Private/Components/PrimitiveComponent.cpp
@@ -3481,18 +3481,12 @@ bool UPrimitiveComponent::ShouldDispatchWakeEvents(FName BoneName) const
 
 void UPrimitiveComponent::GetNavigationData(FNavigationRelevantData& OutData) const
 {
-	AddNavigationModifier(OutData);
+	FPrimitiveComponentHelper::GetNavigationData(*this, OutData);
 }
 
 void UPrimitiveComponent::AddNavigationModifier(FNavigationRelevantData& OutData) const
 {
-	if (bFillCollisionUnderneathForNavmesh || bRasterizeAsFilledConvexVolume)
-	{
-		FCompositeNavModifier CompositeNavModifier;
-		CompositeNavModifier.SetFillCollisionUnderneathForNavmesh(bFillCollisionUnderneathForNavmesh);
-		CompositeNavModifier.SetRasterizeAsFilledConvexVolume(bRasterizeAsFilledConvexVolume);
-		OutData.Modifiers.Add(CompositeNavModifier);
-	}
+	FPrimitiveComponentHelper::AddNavigationModifier(*this, OutData);
 }
 
 bool UPrimitiveComponent::IsNavigationRelevant() const 
diff --git a/Engine/Source/Runtime/Engine/Private/Components/StaticMeshComponent.cpp b/Engine/Source/Runtime/Engine/Private/Components/StaticMeshComponent.cpp
--- a/Engine/Source/Runtime/Engine/Private/Components/StaticMeshComponent.cpp
+++ b/Engine/Source/Runtime/Engine/Private/Components/StaticMeshComponent.cpp
@@ -3047,34 +3047,7 @@ void UStaticMeshComponent::ComputeHLODHash(FHLODHashBuilder& HashBuilder) const
 
 bool UStaticMeshComponent::DoCustomNavigableGeometryExport(FNavigableGeometryExport& GeomExport) const
 {
-	const FVector Scale3D = GetComponentToWorld().GetScale3D();
-
-	if (!Scale3D.IsZero())
-	{
-		if (const UStaticMesh* Mesh = GetStaticMesh())
-		{
-			if (ensureMsgf(!Mesh->IsCompiling(), TEXT("%s is not considered relevant to navigation until associated mesh is compiled."), *GetFullName()))
-			{
-				if (const UNavCollisionBase* NavCollision = Mesh->GetNavCollision())
-				{
-					if (ShouldExportAsObstacle(*NavCollision))
-					{
-						// skip default export
-						return false;
-					}
-
-					const bool bHasData = NavCollision->ExportGeometry(GetComponentToWorld(), GeomExport);
-					if (bHasData)
-					{
-						// skip default export
-						return false;
-					}
-				}
-			}
-		}
-	}
-
-	return true;
+	return FStaticMeshComponentHelper::DoCustomNavigableGeometryExport(*this, GeomExport);
 }
 
 UMaterialInterface* UStaticMeshComponent::GetMaterialFromCollisionFaceIndex(int32 FaceIndex, int32& SectionIndex) const
@@ -3254,59 +3227,17 @@ bool UStaticMeshComponent::CanMeshPaintTextureColors() const
 
 bool UStaticMeshComponent::IsNavigationRelevant() const
 {
-	if (const UStaticMesh* Mesh = GetStaticMesh())
-	{
-		// Pending compilation, update to the navigation system will be done once compilation finishes.
-		return !Mesh->IsCompiling() && Mesh->IsNavigationRelevant() && Super::IsNavigationRelevant();
-	}
-
-	return false;
+	return FStaticMeshComponentHelper::IsNavigationRelevant(*this);
 }
 
 FBox UStaticMeshComponent::GetNavigationBounds() const
 {
-	if (const UStaticMesh* Mesh = GetStaticMesh())
-	{
-#if WITH_EDITOR
-		// @see GetNavigationData
-		if (!Mesh->IsCompiling())
-#endif // WITH_EDITOR
-		{
-			return Mesh->GetNavigationBounds(GetComponentTransform());
-		}
-	}
-
-	return Super::GetNavigationBounds();
+	return FStaticMeshComponentHelper::GetNavigationBounds(*this);
 }
 
 void UStaticMeshComponent::GetNavigationData(FNavigationRelevantData& Data) const
 {
-	Super::GetNavigationData(Data);
-
-	const FVector Scale3D = GetComponentToWorld().GetScale3D();
-	if (!Scale3D.IsZero())
-	{
-		if (const UStaticMesh* Mesh = GetStaticMesh())
-		{
-#if WITH_EDITOR
-			// In Editor it's possible that compilation of a StaticMesh gets triggered on a newly registered component for
-			// which a pending update is queued for the Navigation system.
-			// Then GetNavigationData is called when the pending update is processed but we don't consider the current component
-			// relevant to navigation until associated mesh is compiled.
-			// On mesh post compilation the component will reregister with the right mesh.
-			if (!Mesh->IsCompiling())
-#endif // WITH_EDITOR
-			{
-				if (UNavCollisionBase* NavCollision = Mesh->GetNavCollision())
-				{
-					if (ShouldExportAsObstacle(*NavCollision))
-					{
-						NavCollision->GetNavigationModifier(Data.Modifiers, GetComponentTransform());
-					}
-				}
-			}
-		}
-	}
+	FStaticMeshComponentHelper::GetNavigationData(*this, Data);
 }
 
 bool UStaticMeshComponent::ShouldExportAsObstacle(const UNavCollisionBase& InNavCollision) const
diff --git a/Engine/Source/Runtime/Engine/Private/InstancedStaticMesh.cpp b/Engine/Source/Runtime/Engine/Private/InstancedStaticMesh.cpp
--- a/Engine/Source/Runtime/Engine/Private/InstancedStaticMesh.cpp
+++ b/Engine/Source/Runtime/Engine/Private/InstancedStaticMesh.cpp
@@ -5200,33 +5200,7 @@ void UInstancedStaticMeshComponent::PartialNavigationUpdates(const TConstArrayVi
 
 bool UInstancedStaticMeshComponent::DoCustomNavigableGeometryExport(FNavigableGeometryExport& GeomExport) const
 {
-	if (GetStaticMesh() && GetStaticMesh()->GetNavCollision())
-	{
-		UNavCollisionBase* NavCollision = GetStaticMesh()->GetNavCollision();
-		if (ShouldExportAsObstacle(*NavCollision))
-		{
-			return false;
-		}
-		
-		if (NavCollision->HasConvexGeometry())
-		{
-			NavCollision->ExportGeometry(FTransform::Identity, GeomExport);
-		}
-		else
-		{
-			UBodySetup* BodySetup = GetStaticMesh()->GetBodySetup();
-			if (BodySetup)
-			{
-				GeomExport.ExportRigidBodySetup(*BodySetup, FTransform::Identity);
-			}
-		}
-
-		// Hook per instance transform delegate
-		GeomExport.SetNavDataPerInstanceTransformDelegate(FNavDataPerInstanceTransformDelegate::CreateUObject(this, &UInstancedStaticMeshComponent::GetNavigationPerInstanceTransforms));
-	}
-
-	// we don't want "regular" collision export for this component
-	return false;
+	return FInstancedStaticMeshComponentHelper::DoCustomNavigableGeometryExport(*this, GeomExport, FNavDataPerInstanceTransformDelegate::CreateUObject(this, &UInstancedStaticMeshComponent::GetNavigationPerInstanceTransforms));
 }
 
 extern float DebugLineLifetime;
@@ -5373,21 +5347,7 @@ bool UInstancedStaticMeshComponent::ComponentOverlapMultiImpl(TArray<struct FOve
 
 void UInstancedStaticMeshComponent::GetNavigationData(FNavigationRelevantData& Data) const
 {
-	AddNavigationModifier(Data);
-
-	// Navigation data will get refreshed once async compilation finishes
-	if (GetStaticMesh() && !GetStaticMesh()->IsCompiling() && GetStaticMesh()->GetNavCollision())
-	{
-		UNavCollisionBase* NavCollision = GetStaticMesh()->GetNavCollision();
-		if (ShouldExportAsObstacle(*NavCollision))
-		{
-			Data.Modifiers.MarkAsPerInstanceModifier();
-			NavCollision->GetNavigationModifier(Data.Modifiers, FTransform::Identity);
-
-			// Hook per instance transform delegate
-			Data.NavDataPerInstanceTransformDelegate = FNavDataPerInstanceTransformDelegate::CreateUObject(this, &UInstancedStaticMeshComponent::GetNavigationPerInstanceTransforms);
-		}
-	}
+	FInstancedStaticMeshComponentHelper::GetNavigationData(*this, Data, FNavDataPerInstanceTransformDelegate::CreateUObject(this, &UInstancedStaticMeshComponent::GetNavigationPerInstanceTransforms));
 }
 
 FBox UInstancedStaticMeshComponent::GetNavigationBounds() const
@@ -5412,35 +5372,12 @@ bool UInstancedStaticMeshComponent::ShouldSkipDirtyAreaOnAddOrRemove() const
 
 FBox UInstancedStaticMeshComponent::GetInstanceNavigationBounds() const
 {
-	if (const UStaticMesh* Mesh = GetStaticMesh())
-	{
-		const UNavCollisionBase* NavCollision = Mesh->GetNavCollision();
-		const FBox NavBounds = NavCollision ? NavCollision->GetBounds() : FBox();
-		return NavBounds.IsValid ? NavBounds : Mesh->GetBounds().GetBox();
-	}
-
-	return FBox();
+	return FInstancedStaticMeshComponentHelper::GetInstanceNavigationBounds(*this);
 }
 
 void UInstancedStaticMeshComponent::GetNavigationPerInstanceTransforms(const FBox& AreaBox, TArray<FTransform>& InstanceData) const
 {
-	const FBox InstanceBounds = GetInstanceNavigationBounds();
-	if (InstanceBounds.IsValid)
-	{
-		const FBox LocalAreaBox = AreaBox.InverseTransformBy(GetComponentTransform());
-		for (const auto& InstancedData : PerInstanceSMData)
-		{
-			const FTransform InstanceToComponent(InstancedData.Transform);
-			if (!InstanceToComponent.GetScale3D().IsZero())
-			{
-				const FBoxSphereBounds TransformedInstanceBounds = InstanceBounds.TransformBy(InstancedData.Transform);
-				if (LocalAreaBox.Intersect(TransformedInstanceBounds.GetBox()))
-				{
-					InstanceData.Add(InstanceToComponent*GetComponentTransform());
-				}
-			}
-		}
-	}
+	FInstancedStaticMeshComponentHelper::GetNavigationPerInstanceTransforms(*this, AreaBox, InstanceData);
 }
 
 void UInstancedStaticMeshComponent::GetResourceSizeEx(FResourceSizeEx& CumulativeResourceSize)
diff --git a/Engine/Source/Runtime/Engine/Public/InstancedStaticMeshComponentHelper.h b/Engine/Source/Runtime/Engine/Public/InstancedStaticMeshComponentHelper.h
--- a/Engine/Source/Runtime/Engine/Public/InstancedStaticMeshComponentHelper.h
+++ b/Engine/Source/Runtime/Engine/Public/InstancedStaticMeshComponentHelper.h
@@ -3,6 +3,10 @@
 #pragma once
 
 #include "StaticMeshComponentHelper.h"
+#include "PrimitiveComponentHelper.h"
+#include "AI/Navigation/NavigationRelevantData.h"
+#include "AI/Navigation/NavCollisionBase.h"
+#include "AI/NavigationSystemHelpers.h"
 #include "Rendering/NaniteResourcesHelper.h"
 #include "Engine/InstancedStaticMesh.h"
 #include "NaniteVertexFactory.h"
@@ -16,6 +20,18 @@ public:
 	template<class T>
 	static void CollectPSOPrecacheData(const T& Component, const FPSOPrecacheParams& BasePrecachePSOParams, FMaterialInterfacePSOPrecacheParamsList& OutParams);
 
+	template<class T>
+	static void GetNavigationData(const T& Component, FNavigationRelevantData& Data, const FNavDataPerInstanceTransformDelegate& Delegate);
+
+	template<class T>
+	static bool DoCustomNavigableGeometryExport(const T& Component, FNavigableGeometryExport& GeomExport, const FNavDataPerInstanceTransformDelegate& Delegate);
+
+	template<class T>
+	static void GetNavigationPerInstanceTransforms(const T& Component, const FBox& AreaBox, TArray<FTransform>& InstanceData);
+
+	template<class T>
+	static FBox GetInstanceNavigationBounds(const T& Component);
+
 private:
 	// Dummy instance buffer used for PSO pre-caching
 	class FDummyStaticMeshInstanceBuffer : public FStaticMeshInstanceBuffer
@@ -65,4 +81,93 @@ void FInstancedStaticMeshComponentHelper::CollectPSOPrecacheData(const T& Compon
 	{
 		FStaticMeshComponentHelper::CollectPSOPrecacheDataImpl(Component, &FInstancedStaticMeshVertexFactory::StaticType, BasePrecachePSOParams, ISMC_GetElements, OutParams);
 	}
+}
+
+template<class T>
+void FInstancedStaticMeshComponentHelper::GetNavigationData(const T& Component, FNavigationRelevantData& Data, const FNavDataPerInstanceTransformDelegate& Delegate)
+{
+	FPrimitiveComponentHelper::AddNavigationModifier(Component, Data);
+
+	// Navigation data will get refreshed once async compilation finishes
+	UStaticMesh* StaticMesh = Component.GetStaticMesh();
+	if (StaticMesh && !StaticMesh->IsCompiling() && StaticMesh->GetNavCollision())
+	{
+		UNavCollisionBase* NavCollision = StaticMesh->GetNavCollision();
+		if (Component.ShouldExportAsObstacle(*NavCollision))
+		{
+			Data.Modifiers.MarkAsPerInstanceModifier();
+			NavCollision->GetNavigationModifier(Data.Modifiers, FTransform::Identity);
+
+			// Hook per instance transform delegate
+			Data.NavDataPerInstanceTransformDelegate = Delegate;
+		}
+	}
+}
+
+template<class T>
+bool FInstancedStaticMeshComponentHelper::DoCustomNavigableGeometryExport(const T& Component, FNavigableGeometryExport& GeomExport, const FNavDataPerInstanceTransformDelegate& Delegate)
+{
+	UStaticMesh* StaticMesh = Component.GetStaticMesh();
+	if (StaticMesh && StaticMesh->GetNavCollision())
+	{
+		UNavCollisionBase* NavCollision = StaticMesh->GetNavCollision();
+		if (Component.ShouldExportAsObstacle(*NavCollision))
+		{
+			return false;
+		}
+
+		if (NavCollision->HasConvexGeometry())
+		{
+			NavCollision->ExportGeometry(FTransform::Identity, GeomExport);
+		}
+		else
+		{
+			UBodySetup* BodySetup = StaticMesh->GetBodySetup();
+			if (BodySetup)
+			{
+				GeomExport.ExportRigidBodySetup(*BodySetup, FTransform::Identity);
+			}
+		}
+
+		// Hook per instance transform delegate
+		GeomExport.SetNavDataPerInstanceTransformDelegate(Delegate);
+	}
+
+	// we don't want "regular" collision export for this component
+	return false;
+}
+
+template<class T>
+void FInstancedStaticMeshComponentHelper::GetNavigationPerInstanceTransforms(const T& Component, const FBox& AreaBox, TArray<FTransform>& InstanceData)
+{
+	const FBox InstanceBounds = FInstancedStaticMeshComponentHelper::GetInstanceNavigationBounds(Component);
+	if (InstanceBounds.IsValid)
+	{
+		const FBox LocalAreaBox = AreaBox.InverseTransformBy(Component.GetComponentTransform());
+		for (const auto& InstancedData : Component.PerInstanceSMData)
+		{
+			const FTransform InstanceToComponent(InstancedData.Transform);
+			if (!InstanceToComponent.GetScale3D().IsZero())
+			{
+				const FBoxSphereBounds TransformedInstanceBounds = InstanceBounds.TransformBy(InstancedData.Transform);
+				if (LocalAreaBox.Intersect(TransformedInstanceBounds.GetBox()))
+				{
+					InstanceData.Add(InstanceToComponent * Component.GetComponentTransform());
+				}
+			}
+		}
+	}
+}
+
+template<class T>
+FBox FInstancedStaticMeshComponentHelper::GetInstanceNavigationBounds(const T& Component)
+{
+	if (const UStaticMesh* Mesh = Component.GetStaticMesh())
+	{
+		const UNavCollisionBase* NavCollision = Mesh->GetNavCollision();
+		const FBox NavBounds = NavCollision ? NavCollision->GetBounds() : FBox();
+		return NavBounds.IsValid ? NavBounds : Mesh->GetBounds().GetBox();
+	}
+
+	return FBox();
 }
\ No newline at end of file
diff --git a/Engine/Source/Runtime/Engine/Public/PrimitiveComponentHelper.h b/Engine/Source/Runtime/Engine/Public/PrimitiveComponentHelper.h
--- a/Engine/Source/Runtime/Engine/Public/PrimitiveComponentHelper.h
+++ b/Engine/Source/Runtime/Engine/Public/PrimitiveComponentHelper.h
@@ -12,6 +12,8 @@
 #include "RenderUtils.h"
 #include "SceneTypes.h"
 #include "Engine/EngineTypes.h"
+#include "AI/Navigation/NavigationRelevantData.h"
+#include "AI/NavigationModifier.h"
 
 /** Helper class used to share implementation for different PrimitiveComponent types */
 class FPrimitiveComponentHelper
@@ -19,6 +21,12 @@ class FPrimitiveComponentHelper
 public:
 	template<class T>
 	static FPrimitiveMaterialPropertyDescriptor GetUsedMaterialPropertyDesc(const T& Component, ERHIFeatureLevel::Type FeatureLevel);
+
+	template<class T>
+	static void GetNavigationData(const T& Component, FNavigationRelevantData& OutData);
+
+	template<class T>
+	static void AddNavigationModifier(const T& Component, FNavigationRelevantData& OutData);
 };
 
 template<class T>
@@ -64,4 +72,22 @@ FPrimitiveMaterialPropertyDescriptor FPrimitiveComponentHelper::GetUsedMaterialP
 	}
 
 	return Result;
+}
+
+template<class T>
+void FPrimitiveComponentHelper::GetNavigationData(const T& Component, FNavigationRelevantData& OutData)
+{
+	FPrimitiveComponentHelper::AddNavigationModifier(Component, OutData);
+}
+
+template<class T>
+void FPrimitiveComponentHelper::AddNavigationModifier(const T& Component, FNavigationRelevantData& OutData)
+{
+	if (Component.bFillCollisionUnderneathForNavmesh || Component.bRasterizeAsFilledConvexVolume)
+	{
+		FCompositeNavModifier CompositeNavModifier;
+		CompositeNavModifier.SetFillCollisionUnderneathForNavmesh(Component.bFillCollisionUnderneathForNavmesh);
+		CompositeNavModifier.SetRasterizeAsFilledConvexVolume(Component.bRasterizeAsFilledConvexVolume);
+		OutData.Modifiers.Add(CompositeNavModifier);
+	}
 }
\ No newline at end of file
diff --git a/Engine/Source/Runtime/Engine/Public/StaticMeshComponentHelper.h b/Engine/Source/Runtime/Engine/Public/StaticMeshComponentHelper.h
--- a/Engine/Source/Runtime/Engine/Public/StaticMeshComponentHelper.h
+++ b/Engine/Source/Runtime/Engine/Public/StaticMeshComponentHelper.h
@@ -3,6 +3,9 @@
 #pragma once
 
 #include "Rendering/NaniteResourcesHelper.h"
+#include "AI/Navigation/NavigationRelevantData.h"
+#include "AI/Navigation/NavCollisionBase.h"
+#include "AI/NavigationSystemHelpers.h"
 #include "MeshComponentHelper.h"
 #include "StaticMeshResources.h"
 #include "NaniteVertexFactory.h"
@@ -27,6 +30,18 @@ public:
 
 	template<class T>
 	static void CollectPSOPrecacheData(const T& Component, const FPSOPrecacheParams& BasePrecachePSOParams, FMaterialInterfacePSOPrecacheParamsList& OutParams);
+
+	template<class T>
+	static bool IsNavigationRelevant(const T& Component);
+
+	template<class T>
+	static FBox GetNavigationBounds(const T& Component);
+
+	template<class T>
+	static void GetNavigationData(const T& Component, FNavigationRelevantData& Data);
+
+	template<class T>
+	static bool DoCustomNavigableGeometryExport(const T& Component, FNavigableGeometryExport& GeomExport);
 };
 
 template<class T>
@@ -188,4 +203,119 @@ void FStaticMeshComponentHelper::CollectPSOPrecacheData(const T& Component, cons
 	{
 		FStaticMeshComponentHelper::CollectPSOPrecacheDataImpl(Component, &FLocalVertexFactory::StaticType, BasePrecachePSOParams, SMC_GetElements, OutParams);
 	}
+}
+
+template<class T>
+bool FStaticMeshComponentHelper::IsNavigationRelevant(const T& Component)
+{
+	if (const UStaticMesh* Mesh = Component.GetStaticMesh())
+	{
+		// Pending compilation, update to the navigation system will be done once compilation finishes.
+		return !Mesh->IsCompiling() && Mesh->IsNavigationRelevant() && Component.Super::IsNavigationRelevant();
+	}
+	return false;
+}
+
+template<class T>
+FBox FStaticMeshComponentHelper::GetNavigationBounds(const T& Component)
+{
+	if (const UStaticMesh* Mesh = Component.GetStaticMesh())
+	{
+#if WITH_EDITOR
+		// @see GetNavigationData
+		if (!Mesh->IsCompiling())
+#endif // WITH_EDITOR
+		{
+			return Mesh->GetNavigationBounds(Component.GetComponentTransform());
+		}
+	}
+
+	return Component.Super::GetNavigationBounds();
+}
+
+template<class T>
+void FStaticMeshComponentHelper::GetNavigationData(const T& Component, FNavigationRelevantData& Data)
+{
+	Component.Super::GetNavigationData(Data);
+
+	const FVector Scale3D = Component.GetComponentTransform().GetScale3D();
+	if (!Scale3D.IsZero())
+	{
+		if (const UStaticMesh* Mesh = Component.GetStaticMesh())
+		{
+#if WITH_EDITOR
+			// In Editor it's possible that compilation of a StaticMesh gets triggered on a newly registered component for
+			// which a pending update is queued for the Navigation system.
+			// Then GetNavigationData is called when the pending update is processed but we don't consider the current component
+			// relevant to navigation until associated mesh is compiled.
+			// On mesh post compilation the component will reregister with the right mesh.
+			if (!Mesh->IsCompiling())
+#endif // WITH_EDITOR
+			{
+				if (UNavCollisionBase* NavCollision = Mesh->GetNavCollision())
+				{
+					if (Component.ShouldExportAsObstacle(*NavCollision))
+					{
+						NavCollision->GetNavigationModifier(Data.Modifiers, Component.GetComponentTransform());
+					}
+				}
+			}
+		}
+	}
+}
+
+/** Helper to get information from a UObject if available */
+namespace UObjectHelper
+{
+	template<typename T, typename = void>
+	struct THasGetFullName : std::false_type {};
+
+	template<typename T>
+	struct THasGetFullName<T, std::void_t<decltype(std::declval<T>().GetFullName())>> : std::true_type {};
+
+	template<typename T>
+	FString GetFullNameIfAvailable(T& Object)
+	{
+		if constexpr (THasGetFullName<T>::value)
+		{
+			return Object.GetFullName();
+		}
+		else
+		{
+			return FString("None");
+		}
+	}
+}
+
+template<class T>
+bool FStaticMeshComponentHelper::DoCustomNavigableGeometryExport(const T& Component, FNavigableGeometryExport& GeomExport)
+{
+	const FVector Scale3D = Component.GetComponentTransform().GetScale3D();
+
+	if (!Scale3D.IsZero())
+	{
+		if (const UStaticMesh* Mesh = Component.GetStaticMesh())
+		{
+			if (ensureMsgf(!Mesh->IsCompiling(), TEXT("Component %s is not considered relevant to navigation until associated mesh is compiled."), *UObjectHelper::GetFullNameIfAvailable(Component)))
+			{
+				if (const UNavCollisionBase* NavCollision = Mesh->GetNavCollision())
+				{
+					if (Component.ShouldExportAsObstacle(*NavCollision))
+					{
+						// skip default export
+						return false;
+					}
+
+					const bool bHasData = NavCollision->ExportGeometry(Component.GetComponentTransform(), GeomExport);
+					if (bHasData)
+					{
+						// skip default export
+						return false;
+					}
+				}
+			}
+		}
+	}
+
+	return true;
 }
\ No newline at end of file
