[UBA] Improve logging of re-scheduled distributed shader compile jobs.

When a distributed job result is invalid, it is re-scheduled for local compilation.
This change logs the job distribution before each job is executed locally to improve tracking down the root cause of the crash.
Here is an example what this logging can look like when a remote job failed due to a backend compiler crash:

LogUbaController: Display: Distributed job output file [d:\build\Saved/UbaControllerWorkingDir/A6384141478CA808A84F4C8A72A20327/14/14.uba.j-14.out] is invalid or does not exist
- Directory "d:\build\Saved/UbaControllerWorkingDir/A6384141478CA808A84F4C8A72A20327/14" exists and contains 0 file(s)
- Disk space: 2621422 MiB, free 496909 MiB
LogShaderCompilers: Display: Rescheduling shader compilation to run locally after distributed job failed: d:\build\Saved/UbaControllerWorkingDir/A6384141478CA808A84F4C8A72A20327/14/14.uba.j-14.out
LogShaderCompilers: Display: Executing Job [1/14]: Global/FShader_VirtualTextureCopy_PS< ERuntimeVirtualTextureMaterialType::BaseColor_Normal_Specular_YCoCg >/0:/Engine/Private/VirtualTextureCompress.usf|CopyBaseColorNormalSpecularYCoCgPS (PermutationId=0)

#rnx

diff --git a/Engine/Plugins/UbaController/Source/Private/UbaJobProcessor.cpp b/Engine/Plugins/UbaController/Source/Private/UbaJobProcessor.cpp
--- a/Engine/Plugins/UbaController/Source/Private/UbaJobProcessor.cpp
+++ b/Engine/Plugins/UbaController/Source/Private/UbaJobProcessor.cpp
@@ -652,7 +652,7 @@ bool FUbaJobProcessor::ProcessOutputFile(FDistributedBuildTask* CompileTask)
 				DiagnosticsInfo.Appendf(TEXT("\n - Disk space: %lld MiB, free %lld MiB"), (long long int)TotalNumberOfBytes >> 20, (long long int)NumberOfFreeBytes >> 20);
 			}
 
-			UE_LOG(LogUbaController, Error, TEXT("Output File [%s] is invalid or does not exist%s"), *CompileTask->CommandData.OutputFileName, *DiagnosticsInfo);
+			UE_LOG(LogUbaController, Display, TEXT("Distributed job output file [%s] is invalid or does not exist%s"), *CompileTask->CommandData.OutputFileName, *DiagnosticsInfo);
 		}
 		else
 		{
diff --git a/Engine/Source/Runtime/Engine/Private/ShaderCompiler/ShaderCompilerDistributed.cpp b/Engine/Source/Runtime/Engine/Private/ShaderCompiler/ShaderCompilerDistributed.cpp
--- a/Engine/Source/Runtime/Engine/Private/ShaderCompiler/ShaderCompilerDistributed.cpp
+++ b/Engine/Source/Runtime/Engine/Private/ShaderCompiler/ShaderCompilerDistributed.cpp
@@ -514,11 +514,19 @@ int32 FShaderCompileDistributedThreadRunnable_Interface::CompilingLoop()
 
 			if (!bCompileJobsSucceeded)
 			{
-				// Reading result from XGE job failed, so recompile shaders in current job batch locally
-				UE_LOG(LogShaderCompilers, Log, TEXT("Rescheduling shader compilation to run locally after distributed job failed: %s"), *Task->OutputFilePath);
+				// Reading result from distributed job failed, so recompile shaders in current job batch locally
+				UE_LOG(LogShaderCompilers, Display, TEXT("Rescheduling shader compilation to run locally after distributed job failed: %s"), *Task->OutputFilePath);
 
-				for (FShaderCommonCompileJobPtr Job : Task->ShaderJobs)
+				FString JobDiagnostics;
+				for (int32 JobIndex = 0; JobIndex < Task->ShaderJobs.Num(); ++JobIndex)
 				{
+					const FShaderCommonCompileJobPtr& Job = Task->ShaderJobs[JobIndex];
+
+					// Rescheduling jobs after distributed readback failed should be rare, so display all job details with default verbosity
+					JobDiagnostics.Empty();
+					Job->AppendDiagnostics(JobDiagnostics, JobIndex, Task->ShaderJobs.Num());
+					UE_LOG(LogShaderCompilers, Display, TEXT("Executing %s"), *JobDiagnostics);
+
 					FShaderCompileUtilities::ExecuteShaderCompileJob(*Job);
 				}
 			}
diff --git a/Engine/Source/Runtime/RenderCore/Private/ShaderCore.cpp b/Engine/Source/Runtime/RenderCore/Private/ShaderCore.cpp
--- a/Engine/Source/Runtime/RenderCore/Private/ShaderCore.cpp
+++ b/Engine/Source/Runtime/RenderCore/Private/ShaderCore.cpp
@@ -4419,8 +4419,12 @@ FStringView FShaderCompileJob::GetFinalSourceView() const
 
 void FShaderCompileJob::AppendDiagnostics(FString& OutDiagnostics, int32 InJobIndex, int32 InNumJobs, const TCHAR* Indentation) const
 {
-	OutDiagnostics += FString::Printf(
-		TEXT("%sJob [%d/%d]: %s (PermutationId=%d)\n"),
+	if (!OutDiagnostics.IsEmpty())
+	{
+		OutDiagnostics.AppendChar(TEXT('\n'));
+	}
+	OutDiagnostics.Appendf(
+		TEXT("%sJob [%d/%d]: %s (PermutationId=%d)"),
 		Indentation == nullptr ? TEXT("") : Indentation, InJobIndex + 1, InNumJobs, *Input.GenerateShaderName(), Key.PermutationId
 	);
 }
@@ -4501,8 +4505,12 @@ void FShaderPipelineCompileJob::AppendDebugName(FStringBuilderBase& OutName) con
 void FShaderPipelineCompileJob::AppendDiagnostics(FString& OutDiagnostics, int32 InJobIndex, int32 InNumJobs, const TCHAR* Indentation) const
 {
 	FString NewIndentation = Indentation == nullptr ? TEXT("") : Indentation;
-	OutDiagnostics += FString::Printf(
-		TEXT("%sPipeline Job [%d/%d] (PermutationId=%d):\n"),
+	if (!OutDiagnostics.IsEmpty())
+	{
+		OutDiagnostics.AppendChar(TEXT('\n'));
+	}
+	OutDiagnostics.Appendf(
+		TEXT("%sPipeline Job [%d/%d] (PermutationId=%d):"),
 		*NewIndentation, InJobIndex + 1, InNumJobs, Key.PermutationId
 	);
 	NewIndentation += TEXT("  ");
