Anim Compression: Fix deadlock with saving multiple linked anim sequences. Fix was from Jurre.

diff --git a/Engine/Source/Runtime/Engine/Private/Animation/AnimSequence.cpp b/Engine/Source/Runtime/Engine/Private/Animation/AnimSequence.cpp
--- a/Engine/Source/Runtime/Engine/Private/Animation/AnimSequence.cpp
+++ b/Engine/Source/Runtime/Engine/Private/Animation/AnimSequence.cpp
@@ -5222,27 +5222,27 @@ bool UAnimSequence::IsCompiling() const
 
 void UAnimSequence::ClearAllCompressionData()
 {
-	UE_COMPRESSED_DATA_WRITE_SCOPE(this);
 	{
 		// Delete any cache tasks first because the destructor will cancel the cache and build tasks,
 		// and drop their pointers to the data.
-		
+
 		for (auto It = CacheTasksByKeyHash.CreateIterator(); It; ++It)
-		{	
-			if(!RequiresResidency(It->Key))
+		{
+			if (!RequiresResidency(It->Key))
 			{
 				It.RemoveCurrent();
 			}
 		}
-		
+
+		UE_COMPRESSED_DATA_WRITE_SCOPE(this);
 		for (auto It = DataByPlatformKeyHash.CreateIterator(); It; ++It)
 		{
-			if(!RequiresResidency(It->Key))
+			if (!RequiresResidency(It->Key))
 			{
 				It.RemoveCurrent();
 			}
 		}
-		
+
 		CurrentPlatformData = nullptr;
 	}
 }
