Implement ColorRamp Material Expression

diff --git a/Engine/Config/BaseMaterialExpressions.ini b/Engine/Config/BaseMaterialExpressions.ini
--- a/Engine/Config/BaseMaterialExpressions.ini
+++ b/Engine/Config/BaseMaterialExpressions.ini
@@ -164,6 +164,9 @@ bShaderInputData=true
 [/Script/Engine.MaterialExpressionCollectionTransform]
 +MenuCategories=NSLOCTEXT("MaterialExpression", "Parameters", "Parameters")
 
+[/Script/Engine.MaterialExpressionColorRamp]
++MenuCategories=NSLOCTEXT("MaterialExpression", "Color", "Color")
+
 [/Script/Engine.MaterialExpressionComponentMask]
 +MenuCategories=NSLOCTEXT("MaterialExpression", "Math", "Math")
 +MenuCategories=NSLOCTEXT("MaterialExpression", "VectorOps", "VectorOps")
diff --git a/Engine/Source/Editor/MaterialEditor/Private/MaterialEditorDetailCustomization.cpp b/Engine/Source/Editor/MaterialEditor/Private/MaterialEditorDetailCustomization.cpp
--- a/Engine/Source/Editor/MaterialEditor/Private/MaterialEditorDetailCustomization.cpp
+++ b/Engine/Source/Editor/MaterialEditor/Private/MaterialEditorDetailCustomization.cpp
@@ -15,6 +15,7 @@
 #include "Materials/MaterialExpressionScalarParameter.h"
 #include "Materials/MaterialExpressionVectorParameter.h"
 #include "Materials/MaterialExpressionComposite.h"
+#include "Materials/MaterialExpressionColorRamp.h"
 #include "Materials/MaterialExpressionPinBase.h"
 #include "Materials/MaterialExpressionTextureSampleParameter.h"
 #include "MaterialLayersFunctionsCustomization.h"
@@ -47,6 +48,8 @@
 #include "Curves/CurveLinearColorAtlas.h"
 #include "RenderUtils.h"
 #include "MaterialShared.h"
+#include "SColorGradientEditor.h"
+
 
 #define LOCTEXT_NAMESPACE "MaterialEditor"
 
@@ -1089,4 +1092,49 @@ void FMaterialExpressionCompositeDetails::CustomizeDetails(IDetailLayoutBuilder&
 	}
 }
 
+TSharedRef<class IDetailCustomization> FMaterialExpressionColorRampCustomization::MakeInstance()
+{
+	return MakeShareable(new FMaterialExpressionColorRampCustomization);
+}
+
+void FMaterialExpressionColorRampCustomization::CustomizeDetails(IDetailLayoutBuilder& DetailBuilder)
+{
+	IDetailCategoryBuilder& Category = DetailBuilder.EditCategory("MaterialExpressionColorRamp");
+
+	// Hide default ColorCurve details
+	TSharedPtr<IPropertyHandle> GradientProperty = DetailBuilder.GetProperty(GET_MEMBER_NAME_CHECKED(UMaterialExpressionColorRamp, ColorCurve));
+	DetailBuilder.HideProperty(GradientProperty);
+
+	TArray<TWeakObjectPtr<UObject>> CustomizedObjects;
+	DetailBuilder.GetObjectsBeingCustomized(CustomizedObjects);
+	if (CustomizedObjects.Num() > 0)
+	{
+		UMaterialExpressionColorRamp* ColorRamp = Cast<UMaterialExpressionColorRamp>(CustomizedObjects[0].Get());
+		if (ColorRamp)
+		{
+			TSharedPtr<SColorGradientEditor> GradientEditor = SNew(SColorGradientEditor)
+				.ViewMinInput(0.0f)
+				.ViewMaxInput(1.0f)
+				.ClampStopsToViewRange(true);
+			GradientEditor->SetCurveOwner(ColorRamp->ColorCurve);
+
+			// Recompile shader when the curve is modified
+			ColorRamp->ColorCurve->OnUpdateCurve.AddUObject(ColorRamp, &UMaterialExpressionColorRamp::HandleCurvePropertyChanged);
+			
+			Category.AddCustomRow(FText::FromString("Color Gradient"))
+				.NameContent()
+				[
+					SNew(STextBlock)
+						.Text(FText::FromString("Color Gradient"))
+						.Font(FCoreStyle::GetDefaultFontStyle("Regular", 8))
+				]
+				.ValueContent()
+				.MinDesiredWidth(300.f)
+				[
+					GradientEditor.ToSharedRef()
+				];
+		}
+	}
+}
+
 #undef LOCTEXT_NAMESPACE
diff --git a/Engine/Source/Editor/MaterialEditor/Private/MaterialEditorDetailCustomization.h b/Engine/Source/Editor/MaterialEditor/Private/MaterialEditorDetailCustomization.h
--- a/Engine/Source/Editor/MaterialEditor/Private/MaterialEditorDetailCustomization.h
+++ b/Engine/Source/Editor/MaterialEditor/Private/MaterialEditorDetailCustomization.h
@@ -143,6 +143,20 @@ public:
 	virtual void CustomizeDetails(IDetailLayoutBuilder& DetailLayout) override;
 };
 
+/**
+ * Customizes the details of a Color Ramp node, adding a gradient slider to the panel
+ */
+class FMaterialExpressionColorRampCustomization : public IDetailCustomization
+{
+public:
+
+	/** Makes a new instance of this detail layout class for a specific detail view requesting it */
+	static TSharedRef<class IDetailCustomization> MakeInstance();
+
+	/** IDetailCustomization interface */
+	virtual void CustomizeDetails(IDetailLayoutBuilder& DetailLayout) override;
+};
+
 class FMaterialDetailCustomization : public IDetailCustomization
 {
 public:
diff --git a/Engine/Source/Editor/MaterialEditor/Private/MaterialEditorModule.cpp b/Engine/Source/Editor/MaterialEditor/Private/MaterialEditorModule.cpp
--- a/Engine/Source/Editor/MaterialEditor/Private/MaterialEditorModule.cpp
+++ b/Engine/Source/Editor/MaterialEditor/Private/MaterialEditorModule.cpp
@@ -24,6 +24,7 @@
 #include "Interfaces/IMainFrameModule.h"
 #include "MaterialEditorGraphPanelPinFactory.h"
 #include "MaterialEditorStyle.h"
+#include "MaterialEditorDetailCustomization.h"
 
 const FName MaterialEditorAppIdentifier = FName(TEXT("MaterialEditorApp"));
 const FName MaterialInstanceEditorAppIdentifier = FName(TEXT("MaterialInstanceEditorApp"));
@@ -75,6 +76,12 @@ public:
 				GetMutableDefault<UMaterialEditorSettings>());
 		}
 
+		FPropertyEditorModule& PropertyModule = FModuleManager::LoadModuleChecked<FPropertyEditorModule>("PropertyEditor");
+		PropertyModule.RegisterCustomClassLayout(
+			"MaterialExpressionColorRamp",
+			FOnGetDetailCustomizationInstance::CreateStatic(&FMaterialExpressionColorRampCustomization::MakeInstance)
+		);
+
 		GraphPanelPinFactory = MakeShared<FMaterialEditorGraphPanelPinFactory>();
 		FEdGraphUtilities::RegisterVisualPinFactory(GraphPanelPinFactory);
 
@@ -90,6 +97,12 @@ public:
 
 		MenuExtensibilityManager.Reset();
 		ToolBarExtensibilityManager.Reset();
+
+		if (FModuleManager::Get().IsModuleLoaded("PropertyEditor"))
+		{
+			FPropertyEditorModule& PropertyModule = FModuleManager::GetModuleChecked<FPropertyEditorModule>("PropertyEditor");
+			PropertyModule.UnregisterCustomClassLayout("MaterialExpressionColorRamp");
+		}
 		
 		FSubstrateMaterialEditorStyle::Shutdown();
 	}
diff --git a/Engine/Source/Runtime/Engine/Private/Materials/MaterialExpressions.cpp b/Engine/Source/Runtime/Engine/Private/Materials/MaterialExpressions.cpp
--- a/Engine/Source/Runtime/Engine/Private/Materials/MaterialExpressions.cpp
+++ b/Engine/Source/Runtime/Engine/Private/Materials/MaterialExpressions.cpp
@@ -70,6 +70,7 @@
 #include "Materials/MaterialExpressionClamp.h"
 #include "Materials/MaterialExpressionCollectionParameter.h"
 #include "Materials/MaterialExpressionCollectionTransform.h"
+#include "Materials/MaterialExpressionColorRamp.h"
 #include "Materials/MaterialExpressionComment.h"
 #include "Materials/MaterialExpressionComposite.h"
 #include "Materials/MaterialExpressionComponentMask.h"
@@ -4733,6 +4734,154 @@ void UMaterialExpressionLinearInterpolate::GetCaption(TArray<FString>& OutCaptio
 }
 #endif // WITH_EDITOR
 
+//
+//	UMaterialExpressionColorRamp
+//
+UMaterialExpressionColorRamp::UMaterialExpressionColorRamp(const FObjectInitializer& ObjectInitializer) 
+	: Super(ObjectInitializer)
+{
+	// Initialize ColorCurve RGB value showing a white->black gradient
+	ColorCurve = ObjectInitializer.CreateDefaultSubobject<UCurveLinearColor>(this, TEXT("ColorCurve"));
+	for (int i = 0; i < 3; i++)
+	{
+		if (FRichCurve* Curve = &ColorCurve->FloatCurves[i])
+		{
+			Curve->Reset();
+			Curve->AddKey(0.f, 1.f);
+			Curve->AddKey(1.f, 0.f);
+		}
+	}
+}
+
+#if WITH_EDITOR
+void UMaterialExpressionColorRamp::HandleCurvePropertyChanged(UCurveBase* Curve, EPropertyChangeType::Type ChangeType)
+{
+	if (Material)
+	{
+		UMaterialExpression::RefreshNode(true);
+	}
+}
+
+int32 UMaterialExpressionColorRamp::Compile(FMaterialCompiler* Compiler, int32 OutputIndex)
+{
+	if (!ColorCurve)
+	{
+		return Compiler->Errorf(TEXT("Missing ColorCurve"));
+	}
+
+	// If the input is constant, evaluate at compile time.
+	if (!Input.GetTracedInput().Expression)
+	{
+		FLinearColor ColorValue = ColorCurve->GetLinearColorValue(ConstInput);
+		return Compiler->Constant4(ColorValue.R, ColorValue.G, ColorValue.B, ColorValue.A);
+	}
+
+	// If the input is dynamic, evaluate with translator
+	int32 InputCode = Compiler->ComponentMask(Input.Compile(Compiler), true, false, false, false);
+
+	// Helper lambda to evaluate a curve using vectorized multiply–add operations
+	auto EvaluateCurve = [Compiler](const FRichCurve& Curve, int32 InValue) -> int32
+		{
+			const int32 NumKeys = Curve.Keys.Num();
+			if (NumKeys == 0) { return Compiler->Constant(0.0f); }
+
+			if (NumKeys == 1) { return Compiler->Constant(Curve.Keys[0].Value); }
+
+			if (NumKeys == 2)
+			{
+				const float StartTime = Curve.Keys[0].Time;
+				const float EndTime = Curve.Keys[1].Time;
+				const float StartValue = Curve.Keys[0].Value;
+				const float EndValue = Curve.Keys[1].Value;
+
+				int32 TimeDelta = Compiler->Constant(EndTime - StartTime);
+				int32 TimeDiff = Compiler->Sub(InValue, Compiler->Constant(StartTime));
+				int32 Fraction = Compiler->Div(TimeDiff, TimeDelta);
+
+				return Compiler->Lerp(Compiler->Constant(StartValue), Compiler->Constant(EndValue), Fraction);
+			}
+
+			// Turn input into a vector4
+			int32 InValueVec = Compiler->AppendVector(Compiler->AppendVector(InValue, InValue), Compiler->AppendVector(InValue, InValue));
+
+			int32 Result = Compiler->Constant(Curve.Keys[0].Value);
+			int32 i = 0;
+
+			// Use vector operations for segments of 4
+			for (; i < NumKeys - 4; i += 4)
+			{
+				int32 StartTimeVec = Compiler->Constant4(
+					Curve.Keys[i].Time,
+					Curve.Keys[i + 1].Time,
+					Curve.Keys[i + 2].Time,
+					Curve.Keys[i + 3].Time
+				);
+				int32 EndTimeVec = Compiler->Constant4(
+					Curve.Keys[i + 1].Time,
+					Curve.Keys[i + 2].Time,
+					Curve.Keys[i + 3].Time,
+					Curve.Keys[i + 4].Time
+				);
+				int32 StartValueVec = Compiler->Constant4(
+					Curve.Keys[i].Value,
+					Curve.Keys[i + 1].Value,
+					Curve.Keys[i + 2].Value,
+					Curve.Keys[i + 3].Value
+				);
+				int32 EndValueVec = Compiler->Constant4(
+					Curve.Keys[i + 1].Value,
+					Curve.Keys[i + 2].Value,
+					Curve.Keys[i + 3].Value,
+					Curve.Keys[i + 4].Value
+				);
+
+				int32 TimeDeltaVec = Compiler->Sub(EndTimeVec, StartTimeVec);
+				int32 ValueDeltaVec = Compiler->Sub(EndValueVec, StartValueVec);
+
+				int32 TimeDiffVec = Compiler->Sub(InValueVec, StartTimeVec);
+				int32 FractionVec = Compiler->Div(TimeDiffVec, TimeDeltaVec);
+				int32 SatFractionVec = Compiler->Saturate(FractionVec);
+
+				int32 ContributionVec = Compiler->Mul(ValueDeltaVec, SatFractionVec);
+				int32 ContributionSum = Compiler->Dot(ContributionVec, Compiler->Constant4(1.0f, 1.0f, 1.0f, 1.0f));
+				Result = Compiler->Add(Result, ContributionSum);
+			}
+
+			// Use scalar operations for the remaining keys
+			for (; i < NumKeys - 1; i++)
+			{
+				const float StartTime = Curve.Keys[i].Time;
+				const float EndTime = Curve.Keys[i + 1].Time;
+				const float StartValue = Curve.Keys[i].Value;
+				const float EndValue = Curve.Keys[i + 1].Value;
+
+				int32 TimeDelta = Compiler->Constant(EndTime - StartTime);
+				int32 ValueDelta = Compiler->Constant(EndValue - StartValue);
+
+				int32 TimeDiff = Compiler->Sub(InValue, Compiler->Constant(StartTime));
+				int32 Fraction = Compiler->Div(TimeDiff, TimeDelta);
+				int32 SatFraction = Compiler->Saturate(Fraction);
+
+				int32 Contribution = Compiler->Mul(ValueDelta, SatFraction);
+				Result = Compiler->Add(Result, Contribution);
+			}
+			return Result;
+		};
+
+	int32 Red = EvaluateCurve(ColorCurve->FloatCurves[0], InputCode);
+	int32 Green = EvaluateCurve(ColorCurve->FloatCurves[1], InputCode);
+	int32 Blue = EvaluateCurve(ColorCurve->FloatCurves[2], InputCode);
+	int32 Alpha = EvaluateCurve(ColorCurve->FloatCurves[3], InputCode);
+	int32 FinalVector = Compiler->AppendVector(Compiler->AppendVector(Red, Green), Compiler->AppendVector(Blue, Alpha));
+	return FinalVector;
+}
+
+void UMaterialExpressionColorRamp::GetCaption(TArray<FString>& OutCaptions) const
+{
+	OutCaptions.Add(TEXT("Color Ramp"));
+}
+#endif // WITH_EDITOR
+
 UMaterialExpressionGenericConstant::UMaterialExpressionGenericConstant(const FObjectInitializer& ObjectInitializer)
 	: Super(ObjectInitializer)
 {
diff --git a/Engine/Source/Runtime/Engine/Public/Materials/MaterialExpressionColorRamp.h b/Engine/Source/Runtime/Engine/Public/Materials/MaterialExpressionColorRamp.h
new file mode 100644
--- /dev/null
+++ b/Engine/Source/Runtime/Engine/Public/Materials/MaterialExpressionColorRamp.h
@@ -0,0 +1,38 @@
+// Copyright Epic Games, Inc. All Rights Reserved.
+
+
+#pragma once
+
+#include "CoreMinimal.h"
+#include "UObject/ObjectMacros.h"
+#include "MaterialExpressionIO.h"
+#include "Materials/MaterialExpression.h"
+#include "Curves/CurveLinearColor.h"
+#include "MaterialExpressionColorRamp.generated.h"
+
+UCLASS(MinimalAPI)
+class UMaterialExpressionColorRamp : public UMaterialExpression
+{
+	GENERATED_BODY()
+
+public:
+
+	UPROPERTY(meta = (RequiredInput = "false", ToolTip = "Defaults to 'ConstInput' if not specified"))
+	FExpressionInput Input;
+
+	UPROPERTY(EditAnywhere, Category = "MaterialExpressionColorRamp", meta = (OverridingInputProperty = "Input"))
+	float ConstInput = 0.5f;
+
+	UPROPERTY(EditAnywhere, Instanced, Category = "MaterialExpressionColorRamp")
+	TObjectPtr<UCurveLinearColor> ColorCurve;
+
+	UMaterialExpressionColorRamp(const FObjectInitializer& ObjectInitializer);
+
+	//~ Begin UMaterialExpression Interface
+#if WITH_EDITOR
+	virtual void HandleCurvePropertyChanged(UCurveBase* Curve, EPropertyChangeType::Type ChangeType);
+	virtual int32 Compile(class FMaterialCompiler* Compiler, int32 OutputIndex) override;
+	virtual void GetCaption(TArray<FString>& OutCaptions) const override;
+#endif // WITH_EDITOR
+	//~ End UMaterialExpression Interface
+};
\ No newline at end of file
