[Nanite] Fix an issue where Nanite fallback generation of Nanite skinned meshes can cause mesh sections to drop out, due to all triangles of that material being simplified away. Also fix crashes and validation errors that keeping a mesh section with zero triangles causes.

diff --git a/Engine/Source/Developer/MeshBuilder/Private/SkeletalMeshBuilder.cpp b/Engine/Source/Developer/MeshBuilder/Private/SkeletalMeshBuilder.cpp
--- a/Engine/Source/Developer/MeshBuilder/Private/SkeletalMeshBuilder.cpp
+++ b/Engine/Source/Developer/MeshBuilder/Private/SkeletalMeshBuilder.cpp
@@ -190,13 +190,15 @@ public:
 	int32 								LODIndex				= INDEX_NONE;
 	bool								bRegenDepLODs			= false;
 	bool								bBuildNaniteFallback	= false;
+	Nanite::FMeshDataSectionArray		NaniteFallbackMeshSections;
 
 	FContext(
 		const FSkeletalMeshBuildParameters& InBuildParameters,
 		const FMeshDescription* InSourceMeshDescription,
 		FMeshDescription* OutFallbackMeshDescription,
 		bool bInBuildNanite,
-		bool bInBuildNaniteFallback
+		bool bInBuildNaniteFallback,
+		Nanite::FMeshDataSectionArray&& InNaniteFallbackMeshSections = {}
 	)
 		: SkeletalMesh(InBuildParameters.SkeletalMesh)
 		, MeshUtilities(&FModuleManager::Get().LoadModuleChecked<IMeshUtilities>("MeshUtilities"))
@@ -206,6 +208,7 @@ public:
 		, LODIndex(InBuildParameters.LODIndex)
 		, bRegenDepLODs(InBuildParameters.bRegenDepLODs)
 		, bBuildNaniteFallback(bInBuildNaniteFallback)
+		, NaniteFallbackMeshSections(InNaniteFallbackMeshSections)
 	{
 		if (bInBuildNanite && SkeletalMesh->IsNaniteEnabled())
 		{
@@ -277,6 +280,30 @@ public:
 			InlineReductionCacheDatas[LODIndex].SetCacheGeometryInfo(LODModel);
 		}
 
+		// For Nanite fallbacks, it's possible that all triangles of a given mesh section were simplified out. Add empty mesh
+		// sections for any unrepresented materials.
+		for (const Nanite::FMeshDataSection& FallbackSection : NaniteFallbackMeshSections)
+		{
+			if (LODModel.Sections.FindByPredicate(
+				[&FallbackSection](const FSkelMeshSection& Section)
+				{
+					return Section.MaterialIndex == FallbackSection.MaterialIndex;
+				}) == nullptr)
+			{
+				// Add an empty mesh section for this material
+				FSkelMeshSection& NewSection = LODModel.Sections.Emplace_GetRef();
+				FMemory::Memzero(NewSection);
+
+				NewSection.MaterialIndex				= FallbackSection.MaterialIndex;
+				NewSection.bDisabled 					= true;
+				NewSection.bCastShadow					= EnumHasAnyFlags(FallbackSection.Flags, Nanite::EMeshDataSectionFlags::CastShadow);
+				NewSection.bVisibleInRayTracing			= EnumHasAnyFlags(FallbackSection.Flags, Nanite::EMeshDataSectionFlags::VisibleInRayTracing);
+				NewSection.OriginalDataSectionIndex		= LODModel.Sections.Num() - 1;
+				NewSection.ChunkedParentSectionIndex	= INDEX_NONE;
+				NewSection.CorrespondClothAssetIndex	= INDEX_NONE;
+			}
+		}
+
 		// Re-Apply the user section changes, the UserSectionsData is map to original section and should match the built LODModel
 		LODModel.SyncronizeUserSectionsDataArray();
 	}
@@ -313,6 +340,8 @@ public:
 		{
 			check(FallbackMeshDescription);
 			BuildNaniteFallbackMeshDescription(SkeletalMesh, FallbackMeshData, *FallbackMeshDescription);
+
+			NaniteFallbackMeshSections = MoveTemp(FallbackMeshData.Sections);
 		}
 
 		return true;
@@ -629,6 +658,7 @@ bool FSkeletalMeshBuilder::Build(FSkeletalMeshRenderData& OutRenderData, const F
 		TRACE_CPUPROFILER_EVENT_SCOPE(FSkeletalMeshBuilder::BuildNanite);
 
 		FMeshDescription FallbackMeshDescription;
+		Nanite::FMeshDataSectionArray NaniteFallbackMeshSections;
 		{
 			SkeletalMeshBuilderPrivate::FContext NaniteContext(
 				SkeletalMeshBuildParameters,
@@ -645,6 +675,7 @@ bool FSkeletalMeshBuilder::Build(FSkeletalMeshRenderData& OutRenderData, const F
 			NaniteContext.BuildLODModel();
 
 			BuildNanite(SlowTask, NaniteContext, OutRenderData);
+			NaniteFallbackMeshSections = MoveTemp(NaniteContext.NaniteFallbackMeshSections);
 		}
 
 		SkeletalMeshBuilderPrivate::FContext Context(
@@ -652,7 +683,8 @@ bool FSkeletalMeshBuilder::Build(FSkeletalMeshRenderData& OutRenderData, const F
 			&FallbackMeshDescription,
 			nullptr /* FallbackMeshDescription */,
 			false /* bBuildNanite */,
-			false /* bBuildNaniteFallback */
+			false /* bBuildNaniteFallback */,
+			MoveTemp(NaniteFallbackMeshSections)
 		);
 
 		SlowTask.EnterProgressFrame(1.0f);
diff --git a/Engine/Source/Runtime/Core/Private/UObject/UE5SpecialProjectStreamObjectVersion.cpp b/Engine/Source/Runtime/Core/Private/UObject/UE5SpecialProjectStreamObjectVersion.cpp
--- a/Engine/Source/Runtime/Core/Private/UObject/UE5SpecialProjectStreamObjectVersion.cpp
+++ b/Engine/Source/Runtime/Core/Private/UObject/UE5SpecialProjectStreamObjectVersion.cpp
@@ -10,7 +10,8 @@ TMap<FGuid, FGuid> FUE5SpecialProjectStreamObjectVersion::GetSystemGuids()
 	TMap<FGuid, FGuid> SystemGuids;
 	const FDevSystemGuids& DevGuids = FDevSystemGuids::Get();
 
-	SystemGuids.Add(DevGuids.SkeletalMeshDerivedDataVersion, FGuid("5D7013CA-0F57-48ED-A81F-CBE3038DDBBA"));
+	SystemGuids.Add(DevGuids.NANITE_DERIVEDDATA_VER, FGuid("C73C42D8-4CA0-4EAF-8968-722EB9EEF812"));
+	SystemGuids.Add(DevGuids.SkeletalMeshDerivedDataVersion, FGuid("56B88323-F44D-46B8-A3A3-B3441BC8F448"));
 
 	return SystemGuids;
 }
diff --git a/Engine/Source/Runtime/Engine/Private/GPUSkinCache.cpp b/Engine/Source/Runtime/Engine/Private/GPUSkinCache.cpp
--- a/Engine/Source/Runtime/Engine/Private/GPUSkinCache.cpp
+++ b/Engine/Source/Runtime/Engine/Private/GPUSkinCache.cpp
@@ -866,7 +866,8 @@ public:
 		OutVertexBuffers.SetNum(DispatchData.Num());
 		for (int32 SectionIdx = 0; SectionIdx < DispatchData.Num(); SectionIdx++)
 		{
-			OutVertexBuffers[SectionIdx] = DispatchData[SectionIdx].PositionBuffer->Buffer.Buffer;
+			FGPUSkinCache::FSkinCacheRWBuffer* PositionBuffer = DispatchData[SectionIdx].PositionBuffer;
+			OutVertexBuffers[SectionIdx] = PositionBuffer ? PositionBuffer->Buffer.Buffer : nullptr;
 		}
 	}
 #endif // RHI_RAYTRACING
diff --git a/Engine/Source/Runtime/Engine/Private/SkeletalMeshRenderData.cpp b/Engine/Source/Runtime/Engine/Private/SkeletalMeshRenderData.cpp
--- a/Engine/Source/Runtime/Engine/Private/SkeletalMeshRenderData.cpp
+++ b/Engine/Source/Runtime/Engine/Private/SkeletalMeshRenderData.cpp
@@ -168,26 +168,25 @@ void VerifyAllLodSkeletalMeshModelIntegrity(USkinnedAsset* Owner)
 	for (int32 LODIndex = 0; LODIndex < SkelMeshModel->LODModels.Num(); LODIndex++)
 	{
 		FSkeletalMeshLODModel* LODModel = &(SkelMeshModel->LODModels[LODIndex]);
-		int32 SectionsVerticeNum = 0;
+		int32 SectionsVertexNum = 0;
 		int32 SectionsTriangleNum = 0;
 		for (const FSkelMeshSection& Section : LODModel->Sections)
 		{
-			SectionsVerticeNum += Section.GetNumVertices();
-			SectionsTriangleNum += Section.NumTriangles;
-			int32 LastSectionIndexBuffer = Section.BaseIndex + (Section.NumTriangles * 3);
-			if (Section.NumTriangles > 0)
+			if (Section.NumTriangles == 0)
 			{
-				//Remove 1 if we have at least one triangle
-				LastSectionIndexBuffer--;
+				continue; // don't validate index/vertex buffer for empty sections
 			}
 
+			SectionsVertexNum += Section.GetNumVertices();
+			SectionsTriangleNum += Section.NumTriangles;
+			int32 LastSectionIndexBuffer = Section.BaseIndex + (Section.NumTriangles * 3) - 1;
 			if (LODModel->IndexBuffer.IsValidIndex(LastSectionIndexBuffer))
 			{
 				uint32 FirstSectionIndexBufferValue = LODModel->IndexBuffer[Section.BaseIndex];
 				uint32 LastSectionIndexBufferValue = LODModel->IndexBuffer[LastSectionIndexBuffer];
 				if (FirstSectionIndexBufferValue < Section.BaseVertexIndex || LastSectionIndexBufferValue >= Section.BaseVertexIndex + Section.GetNumVertices())
 				{
-					UE_ASSET_LOG(LogSkeletalMesh, Error, Owner, TEXT("The source model is corrupted! Section triangle refer to a vertex not in the section. LOD %d"), LODIndex);
+					UE_ASSET_LOG(LogSkeletalMesh, Error, Owner, TEXT("The source model is corrupted! Section triangle refers to a vertex not in the section. LOD %d"), LODIndex);
 				}
 			}
 			else
@@ -196,9 +195,9 @@ void VerifyAllLodSkeletalMeshModelIntegrity(USkinnedAsset* Owner)
 			}
 		}
 
-		if (LODModel->NumVertices != SectionsVerticeNum)
+		if (LODModel->NumVertices != SectionsVertexNum)
 		{
-			UE_ASSET_LOG(LogSkeletalMesh, Error, Owner, TEXT("The source model is corrupted! Total sections vertice count is different from source model vertice count. LOD %d"), LODIndex);
+			UE_ASSET_LOG(LogSkeletalMesh, Error, Owner, TEXT("The source model is corrupted! Total sections vertex count is different from source model vertex count. LOD %d"), LODIndex);
 		}
 		if ((LODModel->IndexBuffer.Num() / 3) != SectionsTriangleNum)
 		{
diff --git a/Engine/Source/Runtime/Engine/Private/SkeletalRenderGPUSkin.cpp b/Engine/Source/Runtime/Engine/Private/SkeletalRenderGPUSkin.cpp
--- a/Engine/Source/Runtime/Engine/Private/SkeletalRenderGPUSkin.cpp
+++ b/Engine/Source/Runtime/Engine/Private/SkeletalRenderGPUSkin.cpp
@@ -846,6 +846,12 @@ void FSkeletalMeshObjectGPUSkin::UpdateBufferData(FRHICommandList& RHICmdList, E
 	for (int32 SectionIdx = 0; SectionIdx < Sections.Num(); SectionIdx++)
 	{
 		const FSkelMeshRenderSection& Section = Sections[SectionIdx];
+
+		if (!Section.IsValid())
+		{
+			continue;
+		}
+
 		FGPUBaseSkinVertexFactory* VertexFactory = VertexFactories[SectionIdx].Get();
 		FGPUBaseSkinVertexFactory::FShaderDataType& ShaderData = VertexFactory->GetShaderData();
 
@@ -957,7 +963,14 @@ void FSkeletalMeshObjectGPUSkin::ProcessUpdatedDynamicData(FRHICommandList& RHIC
 	{
 		const FSkelMeshRenderSection& Section = Sections[SectionIdx];
 
+		if (!Section.IsValid())
+		{
+			continue;
+		}
+
 		FGPUBaseSkinVertexFactory* VertexFactory = VertexFactoryData.VertexFactories[SectionIdx].Get();
+		check(VertexFactory != nullptr);
+
 		FGPUBaseSkinVertexFactory::FShaderDataType& ShaderData = VertexFactory->GetShaderData();
 
 		ShaderData.SetRevisionNumbers(CurrentRevisionNumber, PreviousRevisionNumber);
@@ -1155,7 +1168,7 @@ void FSkeletalMeshObjectGPUSkin::UpdateRayTracingGeometry_Internal(
 				Segment.NumPrimitives = Section.NumTriangles;
 
 				// TODO: If we are at a dropped LOD, route material index through the LODMaterialMap in the LODInfo struct.
-				Segment.bEnabled = !MeshObject->IsMaterialHidden(LODIndex, Section.MaterialIndex) && !Section.bDisabled && Section.bVisibleInRayTracing;
+				Segment.bEnabled = !MeshObject->IsMaterialHidden(LODIndex, Section.MaterialIndex) && Section.IsValid() && Section.bVisibleInRayTracing;
 				Initializer.Segments.Add(Segment);
 			}
 
@@ -1672,7 +1685,7 @@ void FSkeletalMeshObjectGPUSkin::FSkeletalMeshObjectLOD::UpdateMorphVertexBuffer
 		for (int32 i = 0; i < LodData.RenderSections.Num(); ++i)
 		{
 			const FSkelMeshRenderSection& RenderSection = LodData.RenderSections[i];
-			if (!RenderSection.bRecomputeTangent)
+			if (RenderSection.NumTriangles > 0 && !RenderSection.bRecomputeTangent)
 			{
 				bAllSectionsDoGPURecomputeTangent = false;
 				break;
@@ -1845,7 +1858,7 @@ const FVertexFactory* FSkeletalMeshObjectGPUSkin::GetStaticSkinVertexFactory(int
 	}
 
 	const FGPUBaseSkinVertexFactory* VertexFactory = LOD.GPUSkinVertexFactories.VertexFactories[ChunkIdx].Get();
-	check(VertexFactory->IsReadyForStaticMeshCaching());
+	check(VertexFactory == nullptr || VertexFactory->IsReadyForStaticMeshCaching());
 	return VertexFactory;
 }
 
@@ -1889,7 +1902,8 @@ void FSkeletalMeshObjectGPUSkin::RefreshClothingTransforms(const FMatrix& InNewL
 
 		for(int32 SectionIndex = 0; SectionIndex < NumSections; ++SectionIndex)
 		{
-			if(LOD.GPUSkinVertexFactories.VertexFactories.IsValidIndex(SectionIndex))
+			if (LOD.GPUSkinVertexFactories.VertexFactories.IsValidIndex(SectionIndex) &&
+				LOD.GPUSkinVertexFactories.VertexFactories[SectionIndex].IsValid())
 			{
 				FGPUBaseSkinAPEXClothVertexFactory* ClothFactory = LOD.GPUSkinVertexFactories.VertexFactories[SectionIndex]->GetClothVertexFactory();
 
@@ -2289,7 +2303,11 @@ void FSkeletalMeshObjectGPUSkin::FVertexFactoryData::InitVertexFactories(
 
 	for (const FSkelMeshRenderSection& Section : Sections)
 	{
-		if (Section.HasClothingData() && bClothEnabled)
+		if (!Section.IsValid())
+		{
+			VertexFactories.Add(nullptr);
+		}
+		else if (Section.HasClothingData() && bClothEnabled)
 		{
 			constexpr int32 ClothLODBias = 0;
 			const uint32 NumClothWeights = Section.ClothMappingDataLODs.Num() ? Section.ClothMappingDataLODs[ClothLODBias].Num(): 0;
@@ -2308,9 +2326,12 @@ void FSkeletalMeshObjectGPUSkin::FVertexFactoryData::InitVertexFactories(
 
 void FSkeletalMeshObjectGPUSkin::FVertexFactoryData::ReleaseVertexFactories()
 {
-	for( int32 FactoryIdx=0; FactoryIdx < VertexFactories.Num(); FactoryIdx++)
+	for (auto& VertexFactory : VertexFactories)
 	{
-		VertexFactories[FactoryIdx]->ReleaseResource();
+		if (VertexFactory.IsValid())
+		{
+			VertexFactory->ReleaseResource();
+		}
 	}
 
 	if (PassthroughVertexFactory)
diff --git a/Engine/Source/Runtime/Engine/Public/Rendering/SkeletalMeshLODRenderData.h b/Engine/Source/Runtime/Engine/Public/Rendering/SkeletalMeshLODRenderData.h
--- a/Engine/Source/Runtime/Engine/Public/Rendering/SkeletalMeshLODRenderData.h
+++ b/Engine/Source/Runtime/Engine/Public/Rendering/SkeletalMeshLODRenderData.h
@@ -116,6 +116,11 @@ struct FSkelMeshRenderSection
 		return NumVertices;
 	}
 
+	FORCEINLINE bool IsValid() const
+	{
+		return NumTriangles > 0 && !bDisabled;
+	}
+
 	ENGINE_API friend FArchive& operator<<(FArchive& Ar, FSkelMeshRenderSection& S);
 };
 
