Lock the skinned mesh component when mutating from the task from AnimNext. This avoids a race condition when the component is destroyed or recreated while the task is active.

diff --git a/Engine/Plugins/Experimental/AnimNext/Source/AnimNext/Private/Graph/RigUnit_AnimNextWriteSkeletalMeshComponentPose.cpp b/Engine/Plugins/Experimental/AnimNext/Source/AnimNext/Private/Graph/RigUnit_AnimNextWriteSkeletalMeshComponentPose.cpp
--- a/Engine/Plugins/Experimental/AnimNext/Source/AnimNext/Private/Graph/RigUnit_AnimNextWriteSkeletalMeshComponentPose.cpp
+++ b/Engine/Plugins/Experimental/AnimNext/Source/AnimNext/Private/Graph/RigUnit_AnimNextWriteSkeletalMeshComponentPose.cpp
@@ -70,6 +70,7 @@ FRigUnit_AnimNextWriteSkeletalMeshComponentPose_Execute()
 	}
 
 	FMemMark MemMark(FMemStack::Get());
+	USkinnedMeshComponent::FRenderStateLockScope LockScope(OutputComponent);
 
 	// TODO: The LOD level we got here is not guaranteed to be stable. We should find another way to get/bind the skeletal mesh component's LOD level.
 	const int32 SkeletalMeshComponentLODLevel = OutputComponent->GetPredictedLODLevel();
diff --git a/Engine/Source/Runtime/Engine/Classes/Components/SkinnedMeshComponent.h b/Engine/Source/Runtime/Engine/Classes/Components/SkinnedMeshComponent.h
--- a/Engine/Source/Runtime/Engine/Classes/Components/SkinnedMeshComponent.h
+++ b/Engine/Source/Runtime/Engine/Classes/Components/SkinnedMeshComponent.h
@@ -17,6 +17,8 @@
 #include "BoneContainer.h"
 #include "ClothingSystemRuntimeTypes.h"
 #include "StateStream/SkinnedMeshStateStreamHandle.h"
+#include "Async/RecursiveMutex.h"
+#include "Misc/ScopeLock.h"
 
 #if WITH_EDITORONLY_DATA
 #include "Animation/MeshDeformerGeometryReadback.h"
@@ -1671,7 +1673,20 @@ public:
 	*/
 	ENGINE_API virtual bool ShouldUpdateTransform(bool bLODHasChanged) const;
 
+	class FRenderStateLockScope
+	{
+	public:
+		FRenderStateLockScope(USkinnedMeshComponent* InComponent)
+			: Lock(InComponent->RenderStateLock)
+		{}
+
+	private:
+		UE::TScopeLock<UE::FRecursiveMutex> Lock;
+	};
+
 protected:
+	friend class FRenderStateLockScope;
+	UE::FRecursiveMutex RenderStateLock;
 
 	/** Flip the editable space base buffer */
 	ENGINE_API void FlipEditableSpaceBases();
diff --git a/Engine/Source/Runtime/Engine/Private/Components/SkinnedMeshComponent.cpp b/Engine/Source/Runtime/Engine/Private/Components/SkinnedMeshComponent.cpp
--- a/Engine/Source/Runtime/Engine/Private/Components/SkinnedMeshComponent.cpp
+++ b/Engine/Source/Runtime/Engine/Private/Components/SkinnedMeshComponent.cpp
@@ -1056,6 +1056,8 @@ EPreviousBoneTransformUpdateMode USkinnedMeshComponent::GetPreviousBoneTransform
 void USkinnedMeshComponent::CreateRenderState_Concurrent(FRegisterComponentContext* Context)
 {
 	LLM_SCOPE(ELLMTag::SkeletalMesh);
+	UE::TScopeLock Lock(RenderStateLock);
+
 
 #if WITH_STATE_STREAM_ACTOR
 	if (true) // This branch is here to easily toggle between statestream and not.. will be removed once we are on par
@@ -1256,6 +1258,8 @@ void USkinnedMeshComponent::CreateRenderState_Concurrent(FRegisterComponentConte
 
 void USkinnedMeshComponent::DestroyRenderState_Concurrent()
 {
+	UE::TScopeLock Lock(RenderStateLock);
+
 #if WITH_STATE_STREAM_ACTOR
 	MeshHandle = {};
 #endif
