Reworked some high frequency CPU profiler events to reduce CPU cost.

diff --git a/Engine/Source/Runtime/Engine/Private/SkeletalRenderGPUSkin.cpp b/Engine/Source/Runtime/Engine/Private/SkeletalRenderGPUSkin.cpp
--- a/Engine/Source/Runtime/Engine/Private/SkeletalRenderGPUSkin.cpp
+++ b/Engine/Source/Runtime/Engine/Private/SkeletalRenderGPUSkin.cpp
@@ -1305,8 +1305,6 @@ static void CalculateMorphDeltaBoundsIncludingExternalMorphs(
 
 void FSkeletalMeshObjectGPUSkin::UpdateMorphVertexBuffer(FRHICommandList& RHICmdList)
 {
-	QUICK_SCOPE_CYCLE_COUNTER(STAT_FSkeletalMeshObjectGPUSkin_ProcessUpdatedDynamicData_UpdateMorphBuffer);
-
 	if (!MorphVertexBuffer)
 	{
 		return;
diff --git a/Engine/Source/Runtime/Engine/Private/Streaming/TextureInstanceState.cpp b/Engine/Source/Runtime/Engine/Private/Streaming/TextureInstanceState.cpp
--- a/Engine/Source/Runtime/Engine/Private/Streaming/TextureInstanceState.cpp
+++ b/Engine/Source/Runtime/Engine/Private/Streaming/TextureInstanceState.cpp
@@ -758,7 +758,6 @@ bool FRenderAssetInstanceState::ConditionalUpdateBounds(int32 BoundIndex)
 
 void FRenderAssetInstanceState::UpdateLastRenderTimeAndMaxDrawDistance(int32 BoundIndex)
 {
-	TRACE_CPUPROFILER_EVENT_SCOPE(UpdateLastRenderTimeAndMaxDrawDistance);
 	const UPrimitiveComponent* Component = ensure(Bounds4Components.IsValidIndex(BoundIndex)) ? Bounds4Components[BoundIndex] : nullptr;
 	if (Component)
 	{
diff --git a/Engine/Source/Runtime/Engine/Private/Streaming/TextureInstanceTask.cpp b/Engine/Source/Runtime/Engine/Private/Streaming/TextureInstanceTask.cpp
--- a/Engine/Source/Runtime/Engine/Private/Streaming/TextureInstanceTask.cpp
+++ b/Engine/Source/Runtime/Engine/Private/Streaming/TextureInstanceTask.cpp
@@ -176,12 +176,25 @@ void FDoWorkTask::ProcessTasks(TArray< TRefCountPtr<TTask> >& Tasks)
 void FDoWorkTask::DoWork()
 {
 	LLM_SCOPE(ELLMTag::Textures);
+	TRACE_CPUPROFILER_EVENT_SCOPE(RenderAssetInstanceTask);
 
 	// The order of execution is important in case other tasks need to sync.
-	ProcessTasks(RefreshFullTasks);
-	ProcessTasks(CreateViewWithUninitializedBoundsTasks);
-	ProcessTasks(RefreshVisibilityTasks);
-	ProcessTasks(NormalizeLightmapTexelFactorTasks);
+	{
+		TRACE_CPUPROFILER_EVENT_SCOPE(RefreshFull);
+		ProcessTasks(RefreshFullTasks);
+	}
+	{
+		TRACE_CPUPROFILER_EVENT_SCOPE(CreateViewWithUninitializedBounds);
+		ProcessTasks(CreateViewWithUninitializedBoundsTasks);
+	}
+	{
+		TRACE_CPUPROFILER_EVENT_SCOPE(RefreshVisibility);
+		ProcessTasks(RefreshVisibilityTasks);
+	}
+	{
+		TRACE_CPUPROFILER_EVENT_SCOPE(NormalizeLightmapTexelFactor);
+		ProcessTasks(NormalizeLightmapTexelFactorTasks);
+	}
 }
 
 } // namespace RenderAssetInstanceTask
