Port the debug view mode (non-nanite) shaders to the ODSCOnly shader permutation flag.
This maintains their functionality in editor but also allows the modes to be used in cooked builds attached to an ODSC server.
Unfortunately most of the modes require some further editor only data so are disabled for the cooked build case.
But in future we could consider storing the shader complexity data in cooked material data (at a memory cost) to enable the cooked build functionality.
Note that the debug view modes are used for in commandlets for building texture streaming data, so we can't remove the HasEditorOnlyData shader permutation flag.
See IsDebugViewShaderModeODSCOnly() for the setting and associated comments.

diff --git a/Engine/Shaders/Private/DebugViewModeCommon.ush b/Engine/Shaders/Private/DebugViewModeCommon.ush
--- a/Engine/Shaders/Private/DebugViewModeCommon.ush
+++ b/Engine/Shaders/Private/DebugViewModeCommon.ush
@@ -32,7 +32,9 @@ struct FDebugPSIn
 	float3 TangentToWorld1 : TEXCOORD4;
 	float3 TangentToWorld2 : TEXCOORD5;
 	FStereoPSInput StereoInput;
+#if (OUTPUT_QUAD_OVERDRAW)
 	uint SvPrimitiveID : SV_PrimitiveID;
+#endif
 };
 
 // This one is only compatible if no previous stage outputed SV_PrimitiveID (i.e. geometry shader)
diff --git a/Engine/Source/Runtime/Engine/Private/DebugViewModeHelpers.cpp b/Engine/Source/Runtime/Engine/Private/DebugViewModeHelpers.cpp
--- a/Engine/Source/Runtime/Engine/Private/DebugViewModeHelpers.cpp
+++ b/Engine/Source/Runtime/Engine/Private/DebugViewModeHelpers.cpp
@@ -5,28 +5,23 @@
 =============================================================================*/
 
 #include "DebugViewModeHelpers.h"
+
+#include "ActorEditorUtils.h"
+#include "Components/PrimitiveComponent.h"
+#include "DataDrivenShaderPlatformInfo.h"
 #include "DebugViewModeInterface.h"
-#include "Materials/MaterialInterface.h"
+#include "Engine/Level.h"
+#include "Engine/World.h"
 #include "MaterialDomain.h"
+#include "Materials/MaterialInterface.h"
 #include "MeshMaterialShader.h"
-#include "RenderingThread.h"
-#include "ShaderCompiler.h"
 #include "Misc/ScopedSlowTask.h"
 #include "Misc/FeedbackContext.h"
-#include "Engine/Level.h"
-#include "Engine/World.h"
-#include "Components/PrimitiveComponent.h"
-#include "ActorEditorUtils.h"
-#include "DataDrivenShaderPlatformInfo.h"
+#include "RenderingThread.h"
+#include "ShaderCompiler.h"
 
 #define LOCTEXT_NAMESPACE "LogDebugViewMode"
 
-static TAutoConsoleVariable<bool> CVarEnableDebugViewModeHelpers(
-	TEXT("DebugViewModeHelpers.Enable"),
-	true,
-	TEXT("Specifies whether to enable the debug view mode shaders. Typically only disabled for a special case editor build, if it doesn't require them"),
-	ECVF_Default);
-
 const TCHAR* DebugViewShaderModeToString(EDebugViewShaderMode InShaderMode)
 {
 	switch (InShaderMode)
@@ -66,81 +61,79 @@ const TCHAR* DebugViewShaderModeToString(EDebugViewShaderMode InShaderMode)
 
 #if WITH_DEBUG_VIEW_MODES
 
-static bool PlatformSupportsDebugViewShaders(EShaderPlatform Platform)
+bool IsDebugViewShaderModeODSCOnly()
 {
-	// List of platforms that have been tested and proved functional.
-	return FDataDrivenShaderPlatformInfo::GetSupportsDebugViewShaders(Platform);
+	// If we return true here then the debug mode shaders would _only_ compile through ODSC.
+	// This would save on always compiling the debug shader types for any loaded materials in render commandlets (where ODSC is off).
+	// BUT the DVSM_OutputMaterialTextureScales mode _is_ used in commandlets to build texture streaming data, so we can't take advantage of this (yet).
+	return false;
 }
 
-bool AllowDebugViewVSDSHS(EShaderPlatform Platform)
+bool SupportDebugViewShaderMode(EDebugViewShaderMode ShaderMode, EShaderPlatform Platform)
 {
-	if(!CVarEnableDebugViewModeHelpers.GetValueOnAnyThread())
+	// Quad modes don't compile/work on Metal or Playstation. 
+	// Needs further testing for other platforms.
+	const bool bSupportQuadComplexityModes = !IsMobilePlatform(Platform) && 
+		!FDataDrivenShaderPlatformInfo::GetIsLanguageSony(Platform) && 
+		!FDataDrivenShaderPlatformInfo::GetIsLanguageMetal(Platform);
+	
+	switch (ShaderMode)
 	{
-		return false;
-	}
-	return IsPCPlatform(Platform); 
-}
+	case DVSM_LODColoration:
+	case DVSM_VisualizeGPUSkinCache:
+	case DVSM_VirtualTexturePendingMips:
+		return true;
 
-bool AllowDebugViewShaderMode(EDebugViewShaderMode ShaderMode, EShaderPlatform Platform, ERHIFeatureLevel::Type FeatureLevel)
-{
-#if WITH_EDITOR
-	// Those options are used to test compilation on specific platforms
-	static const bool bForceQuadOverdraw = FParse::Param(FCommandLine::Get(), TEXT("quadoverdraw"));
-	static const bool bForceStreamingAccuracy = FParse::Param(FCommandLine::Get(), TEXT("streamingaccuracy"));
-	static const bool bForceTextureStreamingBuild = FParse::Param(FCommandLine::Get(), TEXT("streamingbuild"));
+	case DVSM_QuadComplexity:
+		return bSupportQuadComplexityModes;
 
-	switch (ShaderMode)
-	{
-	case DVSM_None:
-		return false;
+#if WITH_EDITORONLY_DATA
+	// All shader complexity modes requires editor data.
 	case DVSM_ShaderComplexity:
-	case DVSM_LODColoration:
-		return IsPCPlatform(Platform);
+	case DVSM_LWCComplexity:
+		return true;
 	case DVSM_ShaderComplexityContainedQuadOverhead:
 	case DVSM_ShaderComplexityBleedingQuadOverhead:
-	case DVSM_QuadComplexity:
-		return (bForceQuadOverdraw || (PlatformSupportsDebugViewShaders(Platform)));
+		return bSupportQuadComplexityModes;
+
+	// All texture streaming modes require editor data and are not supported on mobile.
 	case DVSM_PrimitiveDistanceAccuracy:
 	case DVSM_MeshUVDensityAccuracy:
-		return FeatureLevel >= ERHIFeatureLevel::SM5 && (bForceStreamingAccuracy || PlatformSupportsDebugViewShaders(Platform));
 	case DVSM_MaterialTextureScaleAccuracy:
 	case DVSM_RequiredTextureResolution:
 	case DVSM_OutputMaterialTextureScales:
-	case DVSM_VirtualTexturePendingMips:
-		return FeatureLevel >= ERHIFeatureLevel::SM5 && (bForceTextureStreamingBuild || PlatformSupportsDebugViewShaders(Platform));
-	case DVSM_VisualizeGPUSkinCache:
-	case DVSM_LWCComplexity:
-		return PlatformSupportsDebugViewShaders(Platform);
+		return !IsMobilePlatform(Platform);
+#endif
+
 	default:
 		return false;
 	}
-#else
-	return ShaderMode == DVSM_ShaderComplexity;
-#endif
 }
 
-bool ShouldCompileDebugViewModeShader(const FMeshMaterialShaderPermutationParameters& Parameters)
+bool AllowDebugViewShaderMode(EDebugViewShaderMode ShaderMode, EShaderPlatform Platform, ERHIFeatureLevel::Type FeatureLevel)
 {
-	if(!CVarEnableDebugViewModeHelpers.GetValueOnAnyThread())
+	// Only allow debug modes which are supported.
+	if (!SupportDebugViewShaderMode(ShaderMode, Platform))
 	{
 		return false;
 	}
 
-	if (!PlatformSupportsDebugViewShaders(Parameters.Platform))
-	{
-		return false;
-	}
+#if WITH_EDITORONLY_DATA
+	bool bRequireODSCShaders = IsDebugViewShaderModeODSCOnly();
+#else
+	bool bRequireODSCShaders = true;
+#endif
 
-	if (!EnumHasAllFlags(Parameters.Flags, EShaderPermutationFlags::HasEditorOnlyData))
+	// Don't allow debug modes if we require ODSC but can't compile ODSC shaders.
+	if (bRequireODSCShaders && !ShouldCompileODSCOnlyShaders())
 	{
-		// Debug view shaders only in editor
 		return false;
 	}
 
 	return true;
 }
 
-#endif // !(UE_BUILD_SHIPPING || UE_BUILD_TEST)
+#endif // WITH_DEBUG_VIEW_MODES
 
 
 int32 GetNumActorsInWorld(UWorld* InWorld)
diff --git a/Engine/Source/Runtime/Engine/Private/SceneView.cpp b/Engine/Source/Runtime/Engine/Private/SceneView.cpp
--- a/Engine/Source/Runtime/Engine/Private/SceneView.cpp
+++ b/Engine/Source/Runtime/Engine/Private/SceneView.cpp
@@ -3007,7 +3007,6 @@ FSceneViewFamily::FSceneViewFamily(const ConstructionValues& CVS)
 	{
 		DebugViewShaderMode = DVSM_None;
 	}
-	bUsedDebugViewVSDSHS = DebugViewShaderMode != DVSM_None && AllowDebugViewVSDSHS(GetShaderPlatform());
 #endif
 
 	bCurrentlyBeingEdited = false;
diff --git a/Engine/Source/Runtime/Engine/Public/DebugViewModeHelpers.h b/Engine/Source/Runtime/Engine/Public/DebugViewModeHelpers.h
--- a/Engine/Source/Runtime/Engine/Public/DebugViewModeHelpers.h
+++ b/Engine/Source/Runtime/Engine/Public/DebugViewModeHelpers.h
@@ -43,13 +43,15 @@ enum EDebugViewShaderMode
 ENGINE_API const TCHAR* DebugViewShaderModeToString(EDebugViewShaderMode InShaderMode);
 
 #if WITH_DEBUG_VIEW_MODES
-/** Returns true if the vertex shader (and potential hull and domain) should be compiled on the given platform. */
-ENGINE_API bool AllowDebugViewVSDSHS(EShaderPlatform Platform);
-/** Returns true if the shader mode can be enabled. This is only for UI elements as no shader platform is actually passed. */
+/** Returns true debug view shaders are only ODSC compiled (in editor or cooked builds). Returns false if they are can be compiled for editor only shader maps as well as ODSC. */
+ENGINE_API bool IsDebugViewShaderModeODSCOnly();
+/** Returns true if a specific debug view supported for a platform. Can be used as input for shader cooking. */
+ENGINE_API bool SupportDebugViewShaderMode(EDebugViewShaderMode ShaderMode, EShaderPlatform Platform);
+/** Returns true if a specific debug view can be displayed at runtime. */
 ENGINE_API bool AllowDebugViewShaderMode(EDebugViewShaderMode ShaderMode, EShaderPlatform Platform, ERHIFeatureLevel::Type FeatureLevel);
-ENGINE_API bool ShouldCompileDebugViewModeShader(const FMeshMaterialShaderPermutationParameters& Parameters);
 #else
-FORCEINLINE bool AllowDebugViewVSDSHS(EShaderPlatform Platform)  { return false; }
+FORCEINLINE bool IsDebugViewShaderModeODSCOnly() { return false; }
+FORCEINLINE bool SupportDebugViewShaderMode(EDebugViewShaderMode ShaderMode, EShaderPlatform Platform) { return false; }
 FORCEINLINE bool AllowDebugViewShaderMode(EDebugViewShaderMode ShaderMode, EShaderPlatform Platform, ERHIFeatureLevel::Type FeatureLevel) { return false; }
 #endif
 
diff --git a/Engine/Source/Runtime/Engine/Public/SceneView.h b/Engine/Source/Runtime/Engine/Public/SceneView.h
--- a/Engine/Source/Runtime/Engine/Public/SceneView.h
+++ b/Engine/Source/Runtime/Engine/Public/SceneView.h
@@ -2409,18 +2409,15 @@ public:
 	int32 ViewModeParam;
 	FName ViewModeParamName;
 
-	bool bUsedDebugViewVSDSHS;
 	FORCEINLINE EDebugViewShaderMode GetDebugViewShaderMode() const { return DebugViewShaderMode; }
 	FORCEINLINE int32 GetViewModeParam() const { return ViewModeParam; }
 	FORCEINLINE const FName& GetViewModeParamName() const { return ViewModeParamName; }
 	ENGINE_API EDebugViewShaderMode ChooseDebugViewShaderMode() const;
-	FORCEINLINE bool UseDebugViewVSDSHS() const { return bUsedDebugViewVSDSHS; }
 	FORCEINLINE bool UseDebugViewPS() const { return DebugViewShaderMode != DVSM_None; }
 #else
 	FORCEINLINE EDebugViewShaderMode GetDebugViewShaderMode() const { return DVSM_None; }
 	FORCEINLINE int32 GetViewModeParam() const { return -1; }
 	FORCEINLINE FName GetViewModeParamName() const { return NAME_None; }
-	FORCEINLINE bool UseDebugViewVSDSHS() const { return false; }
 	FORCEINLINE bool UseDebugViewPS() const { return false; }
 #endif
 
diff --git a/Engine/Source/Runtime/Renderer/Private/DebugViewModeRendering.cpp b/Engine/Source/Runtime/Renderer/Private/DebugViewModeRendering.cpp
--- a/Engine/Source/Runtime/Renderer/Private/DebugViewModeRendering.cpp
+++ b/Engine/Source/Runtime/Renderer/Private/DebugViewModeRendering.cpp
@@ -27,14 +27,6 @@ IMPLEMENT_STATIC_UNIFORM_BUFFER_STRUCT(FDebugViewModePassUniformParameters, "Deb
 
 #if WITH_DEBUG_VIEW_MODES
 
-int32 GCacheShaderComplexityShaders = 0;
-static FAutoConsoleVariableRef CVarCacheShaderComplexityShaders(
-	TEXT("r.ShaderComplexity.CacheShaders"),
-	GCacheShaderComplexityShaders,
-	TEXT("If non zero, store the shader complexity shaders in the material shader map, to prevent compile on-the-fly lag. (default=0)"),
-	ECVF_ReadOnly
-);
-
 int32 GShaderComplexityBaselineForwardVS = 134;
 static FAutoConsoleVariableRef CVarShaderComplexityBaselineForwardVS(
 	TEXT("r.ShaderComplexity.Baseline.Forward.VS"),
@@ -159,6 +151,12 @@ TRDGUniformBufferRef<FDebugViewModePassUniformParameters> CreateDebugViewModePas
 	return GraphBuilder.CreateUniformBuffer(UniformBufferParameters);
 }
 
+static bool ShouldCompileDebugViewModeShaders(const FMeshMaterialShaderPermutationParameters& Parameters)
+{
+	EShaderPermutationFlags Flags = IsDebugViewShaderModeODSCOnly() ? EShaderPermutationFlags::IsODSCOnly : EShaderPermutationFlags::HasEditorOnlyData | EShaderPermutationFlags::IsODSCOnly;
+	return EnumHasAnyFlags(Parameters.Flags, Flags);
+}
+
 FDebugViewModeVS::FDebugViewModeVS() = default;
 FDebugViewModeVS::FDebugViewModeVS(const FMeshMaterialShaderType::CompiledShaderInitializerType& Initializer)
 	: FMeshMaterialShader(Initializer)
@@ -166,7 +164,7 @@ FDebugViewModeVS::FDebugViewModeVS(const FMeshMaterialShaderType::CompiledShader
 
 bool FDebugViewModeVS::ShouldCompilePermutation(const FMeshMaterialShaderPermutationParameters& Parameters)
 {
-	return AllowDebugViewVSDSHS(Parameters.Platform) && EnumHasAllFlags(Parameters.Flags, EShaderPermutationFlags::HasEditorOnlyData);
+	return ShouldCompileDebugViewModeShaders(Parameters);
 }
 
 void FDebugViewModeVS::GetShaderBindings(
@@ -269,7 +267,7 @@ FDebugViewModePS::FDebugViewModePS(const FMeshMaterialShaderType::CompiledShader
 
 bool FDebugViewModePS::ShouldCompilePermutation(const FMeshMaterialShaderPermutationParameters& Parameters)
 {
-	return ShouldCompileDebugViewModeShader(Parameters);
+	return ShouldCompileDebugViewModeShaders(Parameters);
 }
 
 void FDebugViewModePS::ModifyCompilationEnvironment(const FMaterialShaderPermutationParameters& Parameters, FShaderCompilerEnvironment& OutEnvironment)
@@ -283,7 +281,7 @@ void FDebugViewModePS::ModifyCompilationEnvironment(const FMaterialShaderPermuta
 	TCHAR BufferRegister[] = { 'u', '0', 0 };
 	BufferRegister[1] += GetQuadOverdrawUAVIndex(Parameters.Platform);
 	OutEnvironment.SetDefine(TEXT("QUAD_BUFFER_REGISTER"), BufferRegister);
-	OutEnvironment.SetDefine(TEXT("OUTPUT_QUAD_OVERDRAW"), !FDataDrivenShaderPlatformInfo::GetIsLanguageMetal(Parameters.Platform));
+	OutEnvironment.SetDefine(TEXT("OUTPUT_QUAD_OVERDRAW"), SupportDebugViewShaderMode(DVSM_QuadComplexity, Parameters.Platform));
 	OutEnvironment.SetDefine(TEXT("WRITE_PENDING_VIRTUALTEXTURE_MIPS"), 1u);
 
 	for (int i = 0; i < DVSM_MAX; ++i)
diff --git a/Engine/Source/Runtime/Renderer/Private/PostProcess/PostProcessVisualizeComplexity.cpp b/Engine/Source/Runtime/Renderer/Private/PostProcess/PostProcessVisualizeComplexity.cpp
--- a/Engine/Source/Runtime/Renderer/Private/PostProcess/PostProcessVisualizeComplexity.cpp
+++ b/Engine/Source/Runtime/Renderer/Private/PostProcess/PostProcessVisualizeComplexity.cpp
@@ -46,7 +46,7 @@ public:
 		const FPermutationDomain PermutationVector(Parameters.PermutationId);
 		if (PermutationVector.Get<FVisualizeComplexityApplyPS::FQuadOverdraw>() == FVisualizeComplexityApplyPS::EQuadOverdraw::Enable)
 		{
-			return AllowDebugViewShaderMode(DVSM_QuadComplexity, Parameters.Platform, GetMaxSupportedFeatureLevel(Parameters.Platform));
+			return SupportDebugViewShaderMode(DVSM_QuadComplexity, Parameters.Platform);
 		}
 		return true;
 	}
