Fix threading issue with TaskSyncManager where it could launch the same task from two threads at once.
Improve deadlock detection code in tick so it requires multiple checks where nothing happens

diff --git a/Engine/Source/Runtime/Engine/Private/TaskSyncManager.cpp b/Engine/Source/Runtime/Engine/Private/TaskSyncManager.cpp
--- a/Engine/Source/Runtime/Engine/Private/TaskSyncManager.cpp
+++ b/Engine/Source/Runtime/Engine/Private/TaskSyncManager.cpp
@@ -426,6 +426,8 @@ FName FActiveSyncPoint::DiagnosticContext(bool bDetailed)
 
 bool FActiveSyncPoint::DispatchManually()
 {
+	UE::TScopeLock WorkScope(WorkLock);
+
 	if (CanDispatchManually())
 	{
 		ensure(FrameStatus == ESyncPointStatus::DispatchWaitingForTrigger || FrameStatus == ESyncPointStatus::DispatchWaitingForWork);
diff --git a/Engine/Source/Runtime/Engine/Private/TickTaskManager.cpp b/Engine/Source/Runtime/Engine/Private/TickTaskManager.cpp
--- a/Engine/Source/Runtime/Engine/Private/TickTaskManager.cpp
+++ b/Engine/Source/Runtime/Engine/Private/TickTaskManager.cpp
@@ -869,7 +869,8 @@ public:
 			{
 				// TODO We may want to warn if this task has prerequisites as that would lead to confusing behavior
 				
-				ensure(ManualDispatchTick->DispatchManually());
+				// This could return false if it was just dispatched on another thread
+				ManualDispatchTick->DispatchManually();
 			}
 		}
 	}
@@ -921,12 +922,16 @@ public:
 
 					if (GIdleTaskWorkMS >= 0.0)
 					{
-						bool bCanDoIdleWork = (GIdleTaskWorkMS > 0.0);
+						bool bDoIdleWork = (GIdleTaskWorkMS > 0.0);
+						bool bDoDeadlockCheck = ManualDispatchTicks[WorldTickGroup].Num() > 0;
 						uint64 EndIdleWork = 0;
 						int32 PreviousTasksRemaining = TickCompletionEvents[Block].Num();
-						auto IdleWorkUpdate = [this, &bCanDoIdleWork, &EndIdleWork, &PreviousTasksRemaining, WorldTickGroup](int32 TasksRemaining) -> FTaskGraphInterface::EProcessTasksOperation
+						int32 IdleCount = 0;
+
+						auto IdleWorkUpdate = [this, &bDoIdleWork, &bDoDeadlockCheck, &EndIdleWork, &PreviousTasksRemaining, &IdleCount, WorldTickGroup]
+							(int32 TasksRemaining) -> FTaskGraphInterface::EProcessTasksOperation
 						{
-							if (bCanDoIdleWork)
+							if (bDoIdleWork)
 							{
 								// Compute cycles for ending idle work if necessary, this is called after the first pass of processing
 								if (EndIdleWork == 0)
@@ -936,34 +941,46 @@ public:
 								}
 								else if (FPlatformTime::Cycles64() > EndIdleWork)
 								{
-									bCanDoIdleWork = false;
+									bDoIdleWork = false;
 								}
 							}
 
-							if (TasksRemaining == PreviousTasksRemaining)
+							if (bDoDeadlockCheck)
 							{
-								// Nothing changed so this could be deadlocked, make sure our manual dispatches have happened
-								this->VerifyManualDispatch(WorldTickGroup);
-
-								if (!bCanDoIdleWork)
+								if (TasksRemaining == PreviousTasksRemaining)
+								{
+									// No tasks were completed during last attempt
+									constexpr int32 IdleCountToBeDeadlocked = 10;
+									if (++IdleCount > IdleCountToBeDeadlocked)
+									{
+										// Nothing is changing so this could be deadlocked, make sure our manual dispatches have happened
+										this->VerifyManualDispatch(WorldTickGroup);
+
+										bDoDeadlockCheck = false;
+									}
+								}
+								else
 								{
-									// Out of idle work budget or never had any
-									return FTaskGraphInterface::EProcessTasksOperation::WaitUntilComplete;
+									PreviousTasksRemaining = TasksRemaining;
+									IdleCount = 0;
 								}
 							}
+							
+							if (bDoIdleWork)
+							{
+								// Run a worker thread task, this could also run other game thread tasks if we add support for that
+								return FTaskGraphInterface::EProcessTasksOperation::ProcessOneOtherTask;
+							}
+							else if (bDoDeadlockCheck)
+							{
+								// Still need to check for deadlock so process named thread tasks if there are any
+								return FTaskGraphInterface::EProcessTasksOperation::ProcessNamedThreadTasks;
+							}
 							else
 							{
-								PreviousTasksRemaining = TasksRemaining;
-
-								// This could be changed to execute other game thread work in the future
-
-								if (!bCanDoIdleWork)
-								{
-									return FTaskGraphInterface::EProcessTasksOperation::ProcessNamedThreadTasks;
-								}
+								// We don't need to do the idle or deadlock checks so wait forever
+								return FTaskGraphInterface::EProcessTasksOperation::WaitUntilComplete;
 							}
-
-							return FTaskGraphInterface::EProcessTasksOperation::ProcessOneOtherTask;
 						};
 
 						FTaskGraphInterface::Get().ProcessUntilTasksComplete(TickCompletionEvents[Block], ENamedThreads::GameThread, IdleWorkUpdate);
