Fix GlobalAnimRateScale not being applied to PreUpdateLinkedInstances

diff --git a/Engine/Source/Runtime/Engine/Private/Components/SkeletalMeshComponent.cpp b/Engine/Source/Runtime/Engine/Private/Components/SkeletalMeshComponent.cpp
--- a/Engine/Source/Runtime/Engine/Private/Components/SkeletalMeshComponent.cpp
+++ b/Engine/Source/Runtime/Engine/Private/Components/SkeletalMeshComponent.cpp
@@ -1693,7 +1693,7 @@ void USkeletalMeshComponent::TickAnimInstances(float DeltaTime, bool bNeedsValid
 	// Allow animation instance to do some processing before the linked instances update
 	if (AnimScriptInstance != nullptr)
 	{
-		AnimScriptInstance->PreUpdateLinkedInstances(DeltaTime);
+		AnimScriptInstance->PreUpdateLinkedInstances(DeltaTime * GlobalAnimRateScale);
 	}
 
 	// We update linked instances first incase we're using either root motion or non-threaded update.
