ODSC recreate render state for all primitives in the case where we are reloading a default material.
This is because default materials are special and used for depth passes, so may be referenced any primitive.

diff --git a/Engine/Source/Runtime/Engine/Private/ShaderCompiler/ShaderCompiler.cpp b/Engine/Source/Runtime/Engine/Private/ShaderCompiler/ShaderCompiler.cpp
--- a/Engine/Source/Runtime/Engine/Private/ShaderCompiler/ShaderCompiler.cpp
+++ b/Engine/Source/Runtime/Engine/Private/ShaderCompiler/ShaderCompiler.cpp
@@ -4325,8 +4325,17 @@ void ProcessCookOnTheFlyShaders(bool bReloadGlobalShaders, const TArray<uint8>&
 		// Only update materials if we need to.
 		if (LoadedMaterials.Num())
 		{
+			// need to force material update flag when reloading default material 
+			// since it may be used by any primitive in depth/shadow passes
+			const bool bIsDefaultMaterial = LoadedMaterials.ContainsByPredicate([](UMaterialInterface* MaterialInterface)
+			{
+				UMaterial* Material = Cast<UMaterial>(MaterialInterface);
+				return Material != nullptr && Material->IsDefaultMaterial();
+			});
+			const uint32 MaterialUpdateFlags = GODSCMaterialUpdateFlags | (bIsDefaultMaterial ? 1 : 0);
+
 			// this will stop the rendering thread, and reattach components, in the destructor
-			FMaterialUpdateContext UpdateContext(GODSCMaterialUpdateFlags);
+			FMaterialUpdateContext UpdateContext(MaterialUpdateFlags);
 
 			// gather the shader maps to reattach
 			for (UMaterialInterface* Material : LoadedMaterials)
