Revert 35622645 (disable stats when screen messages are disabled). Stats are not screen messages and can be enabled/disabled independently

diff --git a/Engine/Source/Runtime/Engine/Private/UnrealEngine.cpp b/Engine/Source/Runtime/Engine/Private/UnrealEngine.cpp
--- a/Engine/Source/Runtime/Engine/Private/UnrealEngine.cpp
+++ b/Engine/Source/Runtime/Engine/Private/UnrealEngine.cpp
@@ -13352,9 +13352,6 @@ void DrawStatsHUD( UWorld* World, FViewport* Viewport, FCanvas* Canvas, UCanvas*
 
 #endif // UE_BUILD_SHIPPING 
 
-#if !UE_BUILD_SHIPPING
-	if (bShowScreenMessages)
-#endif
 	{
 		int32 RightSideX = ((CanvasObject) ? CanvasObject->SizeX : TextureSize.X) / Canvas->GetDPIScale() - FPSXOffset;
 		int32 RightSideY = FMath::TruncToInt(TextureSize.Y * 0.20f) / Canvas->GetDPIScale();
