StaticMesh: Fixed missing localization entries for StaticMesh asset registry tags.

diff --git a/Engine/Source/Runtime/Engine/Private/StaticMesh.cpp b/Engine/Source/Runtime/Engine/Private/StaticMesh.cpp
--- a/Engine/Source/Runtime/Engine/Private/StaticMesh.cpp
+++ b/Engine/Source/Runtime/Engine/Private/StaticMesh.cpp
@@ -5976,11 +5976,117 @@ void UStaticMesh::GetAssetRegistryTagMetadata(TMap<FName, FAssetRegistryTagMetad
 {
 	Super::GetAssetRegistryTagMetadata(OutMetadata);
 
+	OutMetadata.Add("NaniteEnabled",
+		FAssetRegistryTagMetadata()
+			.SetDisplayName(NSLOCTEXT("UStaticMesh", "NaniteEnabled_Label", "Nanite Enabled"))
+		);
+
+	OutMetadata.Add("NaniteFallbackPercent",
+		FAssetRegistryTagMetadata()
+			.SetDisplayName(NSLOCTEXT("UStaticMesh", "NaniteFallbackPercent_Label", "Nanite Fallback Percent"))
+		);
+	
+	OutMetadata.Add("NaniteTriangles",
+		FAssetRegistryTagMetadata()
+			.SetDisplayName(NSLOCTEXT("UStaticMesh", "NaniteTriangles_Label", "Nanite Triangles"))
+		);
+
+	OutMetadata.Add("NaniteVertices",
+		FAssetRegistryTagMetadata()
+			.SetDisplayName(NSLOCTEXT("UStaticMesh", "NaniteVertices_Label", "Nanite Vertices"))
+		);
+	
+	OutMetadata.Add("Triangles",
+		FAssetRegistryTagMetadata()
+			.SetDisplayName(NSLOCTEXT("UStaticMesh", "Triangles_Label", "Triangles"))
+		);
+
+	OutMetadata.Add("Vertices",
+		FAssetRegistryTagMetadata()
+			.SetDisplayName(NSLOCTEXT("UStaticMesh", "Vertices_Label", "Vertices"))
+		);
+
+	OutMetadata.Add("UVChannels",
+		FAssetRegistryTagMetadata()
+			.SetDisplayName(NSLOCTEXT("UStaticMesh", "UVChannels_Label", "UV Channels"))
+		);
+
+	OutMetadata.Add("Materials",
+		FAssetRegistryTagMetadata()
+			.SetDisplayName(NSLOCTEXT("UStaticMesh", "Materials_Label", "Materials"))
+		);
+	
+	OutMetadata.Add(GetStaticMaterialsName(),
+		FAssetRegistryTagMetadata()
+			.SetDisplayName(NSLOCTEXT("UStaticMesh", "StaticMaterials_Label", "Static Materials"))
+		);
+
+	OutMetadata.Add("ApproxSize",
+		FAssetRegistryTagMetadata()
+			.SetDisplayName(NSLOCTEXT("UStaticMesh", "ApproxSize_Label", "Approximate Size"))
+		);
+
 	OutMetadata.Add("CollisionPrims",
 		FAssetRegistryTagMetadata()
-			.SetTooltip(NSLOCTEXT("UStaticMesh", "CollisionPrimsTooltip", "The number of collision primitives in the static mesh"))
+			.SetDisplayName(NSLOCTEXT("UStaticMesh", "CollisionPrims_Label", "Collision Primitives"))
+			.SetTooltip(NSLOCTEXT("UStaticMesh", "CollisionPrims_Tooltip", "The number of collision primitives in the static mesh"))
 			.SetImportantValue(TEXT("0"))
 		);
+
+	OutMetadata.Add("LODs",
+		FAssetRegistryTagMetadata()
+			.SetDisplayName(NSLOCTEXT("UStaticMesh", "LODs_Label", "LODs"))
+		);
+
+	OutMetadata.Add("MinLOD",
+		FAssetRegistryTagMetadata()
+			.SetDisplayName(NSLOCTEXT("UStaticMesh", "MinLOD_Label", "Min LOD"))
+		);
+
+	OutMetadata.Add("QualityLevelMinLOD",
+		FAssetRegistryTagMetadata()
+			.SetDisplayName(NSLOCTEXT("UStaticMesh", "QualityLevelMinLOD_Label", "Quality Level Min LOD"))
+		);
+
+	OutMetadata.Add("SectionsWithCollision",
+		FAssetRegistryTagMetadata()
+			.SetDisplayName(NSLOCTEXT("UStaticMesh", "SectionsWithCollision_Label", "Sections With Collision"))
+		);
+
+	OutMetadata.Add("DefaultCollision",
+		FAssetRegistryTagMetadata()
+			.SetDisplayName(NSLOCTEXT("UStaticMesh", "DefaultCollision_Label", "Default Collision"))
+		);
+
+	OutMetadata.Add("CollisionComplexity",
+		FAssetRegistryTagMetadata()
+			.SetDisplayName(NSLOCTEXT("UStaticMesh", "CollisionComplexity_Label", "Collision Complexity"))
+		);
+
+	OutMetadata.Add("DistanceFieldSize",
+		FAssetRegistryTagMetadata()
+			.SetDisplayName(NSLOCTEXT("UStaticMesh", "DistanceFieldSize_Label", "Distance Field Size"))
+		);
+
+	OutMetadata.Add("EstTotalCompressedSize",
+		FAssetRegistryTagMetadata()
+			.SetDisplayName(NSLOCTEXT("UStaticMesh", "EstTotalCompressedSize_Label", "Estimated Total Compressed Size"))
+		);
+
+	OutMetadata.Add("EstNaniteCompressedSize",
+		FAssetRegistryTagMetadata()
+			.SetDisplayName(NSLOCTEXT("UStaticMesh", "EstNaniteCompressedSize_Label", "Estimated Nanite Compressed Size"))
+		);
+
+	OutMetadata.Add("HasHiResMesh",
+		FAssetRegistryTagMetadata()
+			.SetDisplayName(NSLOCTEXT("UStaticMesh", "HasHiResMesh_Label", "Has HiRes Mesh"))
+		);
+
+	OutMetadata.Add("PhysicsSize",
+		FAssetRegistryTagMetadata()
+			.SetDisplayName(NSLOCTEXT("UStaticMesh", "PhysicsSize_Label", "Physics Size"))
+		);
 }
 #endif
 
