PR #12726: Allow Suspending RerunConstructionScripts during details customization edits.

diff --git a/Engine/Source/Runtime/Engine/Private/ActorConstruction.cpp b/Engine/Source/Runtime/Engine/Private/ActorConstruction.cpp
--- a/Engine/Source/Runtime/Engine/Private/ActorConstruction.cpp
+++ b/Engine/Source/Runtime/Engine/Private/ActorConstruction.cpp
@@ -16,6 +16,7 @@
 #include "Engine/LevelScriptActor.h"
 #include "Engine/CullDistanceVolume.h"
 #include "Engine/SimpleConstructionScript.h"
+#include "ScopedSuspendRerunConstructionScripts.h"
 
 #if WITH_EDITOR
 #include "Editor.h"
@@ -242,6 +243,13 @@ bool AActor::HasNonTrivialUserConstructionScript() const
 #if WITH_EDITOR
 void AActor::RerunConstructionScripts()
 {
+	if (FScopedSuspendRerunConstructionScripts::IsSuspended())
+	{
+		// Queue this actor for a deferred rerun
+		FScopedSuspendRerunConstructionScripts::DeferRerun(this);
+		return;
+	}
+
 	checkf(!HasAnyFlags(RF_ClassDefaultObject), TEXT("RerunConstructionScripts should never be called on a CDO as it can mutate the transient data on the CDO which then propagates to instances!"));
 
 	FEditorScriptExecutionGuard ScriptGuard;
@@ -1412,4 +1420,3 @@ void AActor::PostCreateBlueprintComponent(UActorComponent* NewActorComp)
 	}
 }
 
-
diff --git a/Engine/Source/Runtime/Engine/Private/ScopedSuspendRerunConstructionScripts.cpp b/Engine/Source/Runtime/Engine/Private/ScopedSuspendRerunConstructionScripts.cpp
new file mode 100644
--- /dev/null
+++ b/Engine/Source/Runtime/Engine/Private/ScopedSuspendRerunConstructionScripts.cpp
@@ -0,0 +1,44 @@
+// Copyright Epic Games, Inc. All Rights Reserved.
+
+#include "ScopedSuspendRerunConstructionScripts.h"
+
+#if WITH_EDITOR
+
+int32 FScopedSuspendRerunConstructionScripts::SuspensionCount = 0;
+TSet<TWeakObjectPtr<AActor>> FScopedSuspendRerunConstructionScripts::PendingActors;
+
+FScopedSuspendRerunConstructionScripts::FScopedSuspendRerunConstructionScripts()
+{
+	++SuspensionCount;
+}
+
+FScopedSuspendRerunConstructionScripts::~FScopedSuspendRerunConstructionScripts()
+{
+	--SuspensionCount;
+	// When all scopes have been exited, fire off any deferred reruns
+	if (SuspensionCount == 0 && PendingActors.Num() > 0)
+	{
+		// Make a local copy; we'll clear the set before we run scripts,
+		// so if any subsequent rerun calls happen, they get queued again.
+		TArray<TWeakObjectPtr<AActor>> ActorsToRerun = PendingActors.Array();
+		PendingActors.Empty();
+
+		for (TWeakObjectPtr<AActor> WeakActor : ActorsToRerun)
+		{
+			if (AActor* Actor = WeakActor.Get())
+			{
+				Actor->RerunConstructionScripts();
+			}
+		}
+	}
+}
+
+void FScopedSuspendRerunConstructionScripts::DeferRerun(AActor* Actor)
+{
+	if (Actor)
+	{
+		PendingActors.Add(Actor);
+	}
+}
+
+#endif
\ No newline at end of file
diff --git a/Engine/Source/Runtime/Engine/Public/ScopedSuspendRerunConstructionScripts.h b/Engine/Source/Runtime/Engine/Public/ScopedSuspendRerunConstructionScripts.h
new file mode 100644
--- /dev/null
+++ b/Engine/Source/Runtime/Engine/Public/ScopedSuspendRerunConstructionScripts.h
@@ -0,0 +1,33 @@
+// Copyright Epic Games, Inc. All Rights Reserved.
+
+#pragma once
+
+#include "CoreMinimal.h"
+#include "GameFramework/Actor.h"
+
+#if WITH_EDITOR
+
+class ENGINE_API FScopedSuspendRerunConstructionScripts
+{
+public:
+    FScopedSuspendRerunConstructionScripts();
+    ~FScopedSuspendRerunConstructionScripts();
+
+    // Called by AActor::RerunConstructionScripts() to see if it's suspended
+    static bool IsSuspended()
+    {
+        return SuspensionCount > 0;
+    }
+
+    // Queues an actor for a deferred rerun
+    static void DeferRerun(AActor* Actor);
+
+private:
+    // How many nested scopes are active
+    static int32 SuspensionCount;
+
+    // All actors that attempted to rerun while suspended
+    static TSet<TWeakObjectPtr<AActor>> PendingActors;
+};
+
+#endif
\ No newline at end of file
