[Shaders] Allow specifying compiler flags to "recompileshaders" command and filter by regex instead of fixed filename.

Previoulsy, "recompileshaders /Engine/Private/<SHADER>.usf" recompiled all shader types with the given virtual file path.
With this change, the client can specify a fixed filename or a regular expression including the special case of asterisks '*', e.g. "recompileshaders *Clear*" recompiles all global shaders whos filenames contain "Clear".

Moreover, a '+' separated list of compiler flags can be specified with either their full name such as "CFLAG_SkipOptimizations" or their abbreviated form "SkipOptimizations".
For example, "recompileshaders *Clear* SkipOptimizations+HLSL2021" recompiles all global shaders whos filenames contain "Clear" and compiles them with the additional compiler flags CFLAG_SkipOptimizations and CFLAG_HLSL2021.

#rnx

diff --git a/Engine/Source/Runtime/Engine/Private/ODSC/ODSCManager.cpp b/Engine/Source/Runtime/Engine/Private/ODSC/ODSCManager.cpp
--- a/Engine/Source/Runtime/Engine/Private/ODSC/ODSCManager.cpp
+++ b/Engine/Source/Runtime/Engine/Private/ODSC/ODSCManager.cpp
@@ -219,7 +219,8 @@ void FODSCManager::AddThreadedRequest(
 	ERHIFeatureLevel::Type FeatureLevel,
 	EMaterialQualityLevel::Type QualityLevel,
 	ODSCRecompileCommand RecompileCommandType,
-	const FString& RequestedMaterialName
+	const FString& RequestedMaterialName,
+	const FShaderCompilerFlags& ExtraCompilerFlags
 )
 {
 	if (IsHandlingRequests())
@@ -239,12 +240,12 @@ void FODSCManager::AddThreadedRequest(
 			// when we ask for "changed", we want both materials and global
 			if (RecompileCommandType == ODSCRecompileCommand::Changed)
 			{
-				Thread->AddRequest(TArray<FString>(), FString(), ShaderPlatform, FeatureLevel, QualityLevel, ODSCRecompileCommand::Changed);
+				Thread->AddRequest(TArray<FString>(), FString(), ShaderPlatform, FeatureLevel, QualityLevel, ODSCRecompileCommand::Changed, ExtraCompilerFlags);
 			}
 		}
 		else
 		{
-			Thread->AddRequest(MaterialsToCompile, ShaderTypesToLoad, ShaderPlatform, FeatureLevel, QualityLevel, RecompileCommandType);
+			Thread->AddRequest(MaterialsToCompile, ShaderTypesToLoad, ShaderPlatform, FeatureLevel, QualityLevel, RecompileCommandType, ExtraCompilerFlags);
 		}
 	}
 }
diff --git a/Engine/Source/Runtime/Engine/Private/ODSC/ODSCThread.cpp b/Engine/Source/Runtime/Engine/Private/ODSC/ODSCThread.cpp
--- a/Engine/Source/Runtime/Engine/Private/ODSC/ODSCThread.cpp
+++ b/Engine/Source/Runtime/Engine/Private/ODSC/ODSCThread.cpp
@@ -69,13 +69,21 @@ FODSCMessageHandler::FODSCMessageHandler(EShaderPlatform InShaderPlatform, ERHIF
 {
 }
 
-FODSCMessageHandler::FODSCMessageHandler(const TArray<FString>& InMaterials, const FString& ShaderTypesToLoad, EShaderPlatform InShaderPlatform, ERHIFeatureLevel::Type InFeatureLevel, EMaterialQualityLevel::Type InQualityLevel, ODSCRecompileCommand InRecompileCommandType) :
-	MaterialsToLoad(std::move(InMaterials)),
+FODSCMessageHandler::FODSCMessageHandler(
+	const TArray<FString>& InMaterials,
+	const FString& ShaderTypesToLoad,
+	EShaderPlatform InShaderPlatform,
+	ERHIFeatureLevel::Type InFeatureLevel,
+	EMaterialQualityLevel::Type InQualityLevel,
+	ODSCRecompileCommand InRecompileCommandType,
+	const FShaderCompilerFlags& InExtraCompilerFlags)
+:	MaterialsToLoad(std::move(InMaterials)),
 	ShaderTypesToLoad(ShaderTypesToLoad),
 	ShaderPlatform(InShaderPlatform),
 	FeatureLevel(InFeatureLevel),
 	QualityLevel(InQualityLevel),
-	RecompileCommandType(InRecompileCommandType)
+	RecompileCommandType(InRecompileCommandType),
+	ExtraCompilerFlags(InExtraCompilerFlags)
 {
 }
 
@@ -90,6 +98,7 @@ void FODSCMessageHandler::FillPayload(FArchive& Payload)
 
 	Payload << MaterialsToLoad;
 	Payload << ShaderTypesToLoad;
+	Payload << ExtraCompilerFlags;
 	Payload << ConvertedShaderPlatform;
 	Payload << ConvertedFeatureLevel;
 	Payload << ConvertedQualityLevel;
@@ -350,9 +359,16 @@ void FODSCThread::RetrieveMissedMaterials(TArray<FString>& OutMaterialPaths) con
 	}
 }
 
-void FODSCThread::AddRequest(const TArray<FString>& MaterialsToCompile, const FString& ShaderTypesToLoad, EShaderPlatform ShaderPlatform, ERHIFeatureLevel::Type FeatureLevel, EMaterialQualityLevel::Type QualityLevel, ODSCRecompileCommand RecompileCommandType)
+void FODSCThread::AddRequest(
+	const TArray<FString>& MaterialsToCompile,
+	const FString& ShaderTypesToLoad,
+	EShaderPlatform ShaderPlatform,
+	ERHIFeatureLevel::Type FeatureLevel,
+	EMaterialQualityLevel::Type QualityLevel,
+	ODSCRecompileCommand RecompileCommandType,
+	const FShaderCompilerFlags& ExtraCompilerFlags)
 {
-	PendingMaterialThreadedRequests.Enqueue(new FODSCMessageHandler(MaterialsToCompile, ShaderTypesToLoad, ShaderPlatform, FeatureLevel, QualityLevel, RecompileCommandType));
+	PendingMaterialThreadedRequests.Enqueue(new FODSCMessageHandler(MaterialsToCompile, ShaderTypesToLoad, ShaderPlatform, FeatureLevel, QualityLevel, RecompileCommandType, ExtraCompilerFlags));
 }
 
 void FODSCThread::AddShaderPipelineRequest(
diff --git a/Engine/Source/Runtime/Engine/Private/ODSC/ODSCThread.h b/Engine/Source/Runtime/Engine/Private/ODSC/ODSCThread.h
--- a/Engine/Source/Runtime/Engine/Private/ODSC/ODSCThread.h
+++ b/Engine/Source/Runtime/Engine/Private/ODSC/ODSCThread.h
@@ -36,7 +36,16 @@ class FODSCMessageHandler : public IPlatformFile::IFileServerMessageHandler
 {
 public:
 	FODSCMessageHandler(EShaderPlatform InShaderPlatform, ERHIFeatureLevel::Type InFeatureLevel, EMaterialQualityLevel::Type InQualityLevel, ODSCRecompileCommand InRecompileCommandType);
-	FODSCMessageHandler(const TArray<FString>& InMaterials, const FString& ShaderTypesToLoad, EShaderPlatform InShaderPlatform, ERHIFeatureLevel::Type InFeatureLevel, EMaterialQualityLevel::Type InQualityLevel, ODSCRecompileCommand InRecompileCommandType);
+	FODSCMessageHandler(
+		const TArray<FString>& InMaterials,
+		const FString& ShaderTypesToLoad,
+		EShaderPlatform InShaderPlatform,
+		ERHIFeatureLevel::Type InFeatureLevel,
+		EMaterialQualityLevel::Type InQualityLevel,
+		ODSCRecompileCommand InRecompileCommandType,
+		const FShaderCompilerFlags& InExtraCompilerFlags = FShaderCompilerFlags()
+	);
+
 	/** Subclass fills out an archive to send to the server */
 	virtual void FillPayload(FArchive& Payload) override;
 
@@ -74,6 +83,9 @@ private:
 	/** Whether or not to recompile changed shaders */
 	ODSCRecompileCommand RecompileCommandType = ODSCRecompileCommand::None;
 
+	/** Extra compiler flags sent to the shader compiler. This can be used to request individual shaders to be updated without optimizations for debugging. */
+	FShaderCompilerFlags ExtraCompilerFlags = FShaderCompilerFlags();
+
 	/** The payload for compiling a specific set of shaders. */
 	TArray<FODSCRequestPayload> RequestBatch;
 
@@ -121,7 +133,15 @@ public:
 	 *
 	 * @return false if no longer needs ticking
 	 */
-	void AddRequest(const TArray<FString>& MaterialsToCompile, const FString& ShaderTypesToLoad, EShaderPlatform ShaderPlatform, ERHIFeatureLevel::Type FeatureLevel, EMaterialQualityLevel::Type QualityLevel, ODSCRecompileCommand RecompileCommandType);
+	void AddRequest(
+		const TArray<FString>& MaterialsToCompile,
+		const FString& ShaderTypesToLoad,
+		EShaderPlatform ShaderPlatform,
+		ERHIFeatureLevel::Type FeatureLevel,
+		EMaterialQualityLevel::Type QualityLevel,
+		ODSCRecompileCommand RecompileCommandType,
+		const FShaderCompilerFlags& ExtraCompilerFlags = FShaderCompilerFlags()
+	);
 
 	/**
 	 * Add a request to compile a pipeline (VS/PS) of shaders.  The results are submitted and processed in an async manner.
diff --git a/Engine/Source/Runtime/Engine/Private/ShaderCompiler/ShaderCompiler.cpp b/Engine/Source/Runtime/Engine/Private/ShaderCompiler/ShaderCompiler.cpp
--- a/Engine/Source/Runtime/Engine/Private/ShaderCompiler/ShaderCompiler.cpp
+++ b/Engine/Source/Runtime/Engine/Private/ShaderCompiler/ShaderCompiler.cpp
@@ -4126,6 +4126,47 @@ void GlobalBeginCompileShader(
 
 #endif // WITH_EDITOR
 
+static bool ParseShaderCompilerFlags(const TCHAR* InFlagsString, FShaderCompilerFlags& OutCompilerFlags)
+{
+	if (!InFlagsString || *InFlagsString == TEXT('\0'))
+	{
+		return false;
+	}
+
+	TStringBuilder<4096> UnknownFlagNameList;
+
+	FString NextFlagArg;
+	while (FParse::Token(InFlagsString, NextFlagArg, false, TEXT('+')))
+	{
+		ECompilerFlags NextFlag = CFLAG_Max;
+		if (NextFlagArg.StartsWith(TEXT("CFLAG_")))
+		{
+			LexFromString(NextFlag, *NextFlagArg);
+		}
+		else
+		{
+			LexFromString(NextFlag, *FString::Printf(TEXT("CFLAG_%s"), *NextFlagArg));
+		}
+
+		if (NextFlag != CFLAG_Max)
+		{
+			OutCompilerFlags.Add(NextFlag);
+		}
+		else
+		{
+			UnknownFlagNameList.Appendf(TEXT("%s%s"), UnknownFlagNameList.Len() > 0 ? TEXT(", ") : TEXT(""), *NextFlagArg);
+		}
+	}
+
+	if (UnknownFlagNameList.Len() > 0)
+	{
+		UE_LOG(LogShaderCompilers, Warning, TEXT("Unknown shader compiler flags: %s"), *UnknownFlagNameList);
+		return false;
+	}
+
+	return true;
+}
+
 /** Timer class used to report information on the 'recompileshaders' console command. */
 class FRecompileShadersTimer
 {
@@ -4187,7 +4228,7 @@ namespace
 		}
 	}
 
-	ODSCRecompileCommand ParseRecompileCommandString(const TCHAR* CmdString, TArray<FString>& OutMaterialsToLoad, FString& OutShaderTypesToLoad, FString& OutRequestedMaterialName)
+	ODSCRecompileCommand ParseRecompileCommandString(const TCHAR* CmdString, TArray<FString>& OutMaterialsToLoad, FString& OutShaderTypesToLoad, FString& OutRequestedMaterialName, FShaderCompilerFlags& OutExtraCompilerFlags)
 	{
 		FString CmdName = FParse::Token(CmdString, 0);
 
@@ -4261,6 +4302,10 @@ namespace
 
 			OutShaderTypesToLoad = CmdName;
 
+			// Parse optional extra compiler flags from commandline
+			const FString FlagsStr(FParse::Token(CmdString, false));
+			ParseShaderCompilerFlags(*FlagsStr, OutExtraCompilerFlags);
+
 			// tell other side which materials to load and compile the single
 			// shader for.
 			for (TObjectIterator<UMaterialInterface> It; It; ++It)
@@ -4424,11 +4469,12 @@ bool RecompileShaders(const TCHAR* Cmd, FOutputDevice& Ar)
 		TArray<FString> MaterialsToLoad;
 		FString ShaderTypesToLoad;
 		FString RequestedMaterialName;
-		ODSCRecompileCommand CommandType = ParseRecompileCommandString(Cmd, MaterialsToLoad, ShaderTypesToLoad, RequestedMaterialName);
+		FShaderCompilerFlags ExtraCompilerFlags;
+		ODSCRecompileCommand CommandType = ParseRecompileCommandString(Cmd, MaterialsToLoad, ShaderTypesToLoad, RequestedMaterialName, ExtraCompilerFlags);
 
 		ERHIFeatureLevel::Type TargetFeatureLevel = GetMaxSupportedFeatureLevel(GMaxRHIShaderPlatform);
 		const EMaterialQualityLevel::Type ActiveQualityLevel = GetCachedScalabilityCVars().MaterialQualityLevel;
-		GODSCManager->AddThreadedRequest(MaterialsToLoad, ShaderTypesToLoad, GMaxRHIShaderPlatform, TargetFeatureLevel, ActiveQualityLevel, CommandType, RequestedMaterialName);
+		GODSCManager->AddThreadedRequest(MaterialsToLoad, ShaderTypesToLoad, GMaxRHIShaderPlatform, TargetFeatureLevel, ActiveQualityLevel, CommandType, RequestedMaterialName, ExtraCompilerFlags);
 #endif
 		return true;
 	}
@@ -4562,17 +4608,24 @@ bool RecompileShaders(const TCHAR* Cmd, FOutputDevice& Ar)
 		}
 		else
 		{
-			TArray<const FShaderType*> ShaderTypes = FShaderType::GetShaderTypesByFilename(*FlagStr);
-			TArray<const FShaderPipelineType*> ShaderPipelineTypes = FShaderPipelineType::GetShaderPipelineTypesByFilename(*FlagStr);
+			constexpr bool bSearchAsRegexFilter = true;
+			TArray<const FShaderType*> ShaderTypes = FShaderType::GetShaderTypesByFilename(*FlagStr, bSearchAsRegexFilter);
+			TArray<const FShaderPipelineType*> ShaderPipelineTypes = FShaderPipelineType::GetShaderPipelineTypesByFilename(*FlagStr, bSearchAsRegexFilter);
+
 			if (ShaderTypes.Num() > 0 || ShaderPipelineTypes.Num() > 0)
 			{
 				FRecompileShadersTimer TestTimer(TEXT("RecompileShaders SingleShader"));
 				
 				UpdateReferencedUniformBufferNames(ShaderTypes, {}, ShaderPipelineTypes);
 
+				// Parse optional extra compiler flags from commandline
+				FShaderCompilerFlags ExtraCompilerFlags;
+				const FString FlagsStr(FParse::Token(Cmd, false));
+				ParseShaderCompilerFlags(*FlagsStr, ExtraCompilerFlags);
+
 				UMaterialInterface::IterateOverActiveFeatureLevels([&](ERHIFeatureLevel::Type InFeatureLevel) {
 					auto ShaderPlatform = GShaderPlatformForFeatureLevel[InFeatureLevel];
-					BeginRecompileGlobalShaders(ShaderTypes, ShaderPipelineTypes, ShaderPlatform);
+					BeginRecompileGlobalShaders(ShaderTypes, ShaderPipelineTypes, ShaderPlatform, nullptr, ExtraCompilerFlags);
 					FinishRecompileGlobalShaders();
 				});
 			}
@@ -4590,7 +4643,7 @@ bool RecompileShaders(const TCHAR* Cmd, FOutputDevice& Ar)
 											 "    'Material [name]'     Recompile all the shaders for a single material.\n"
 											 "    'Listtypes'           List all the shader type and vertex factory type class names and their source file path.  Can be used to find shader file names to be used with `recompileshaders [shaderfilename]`.\n"
 											 "    'All'                 Recompile all materials.\n"
-											 "    [shaderfilename]      Compile the single shader associated with a specific filename.\n"
+											 "    [filename] [flags]    Compile all shaders associated with a specific filename or regular expression (including '*' for any characters). Optionally add CFLAG entries concatenated with '+'.\n"
 											 ));
 #endif // WITH_EDITOR
 
@@ -4615,7 +4668,7 @@ FOnGlobalShadersCompilation& GetOnGlobalShaderCompilation()
 *
 * @param	Platform	Platform to verify global shaders for
 */
-void VerifyGlobalShaders(EShaderPlatform Platform, const ITargetPlatform* TargetPlatform, bool bLoadedFromCacheFile, const TArray<const FShaderType*>* OutdatedShaderTypes, const TArray<const FShaderPipelineType*>* OutdatedShaderPipelineTypes)
+void VerifyGlobalShaders(EShaderPlatform Platform, const ITargetPlatform* TargetPlatform, bool bLoadedFromCacheFile, const TArray<const FShaderType*>* OutdatedShaderTypes, const TArray<const FShaderPipelineType*>* OutdatedShaderPipelineTypes, const FShaderCompilerFlags& InExtraCompilerFlags)
 {
 	SCOPED_LOADTIMER(VerifyGlobalShaders);
 
@@ -4783,6 +4836,19 @@ void VerifyGlobalShaders(EShaderPlatform Platform, const ITargetPlatform* Target
 #if WITH_EDITOR
 	if (GlobalShaderJobs.Num() > 0)
 	{
+		if (InExtraCompilerFlags.GetData() != 0)
+		{
+			for (FShaderCommonCompileJobPtr& Job : GlobalShaderJobs)
+			{
+				Job->ForEachSingleShaderJob(
+					[InExtraCompilerFlags](FShaderCompileJob& SingleJob) -> void
+					{
+						SingleJob.Input.Environment.CompilerFlags.Append(InExtraCompilerFlags);
+					}
+				);
+			}
+		}
+
 		GetOnGlobalShaderCompilation().Broadcast();
 		GShaderCompilingManager->SubmitJobs(GlobalShaderJobs, "Globals");
 
@@ -5537,7 +5603,7 @@ extern ENGINE_API const TCHAR* ODSCCmdEnumToString(ODSCRecompileCommand Cmd)
 	return TEXT("Unknown");
 }
 
-void BeginRecompileGlobalShaders(const TArray<const FShaderType*>& OutdatedShaderTypes, const TArray<const FShaderPipelineType*>& OutdatedShaderPipelineTypes, EShaderPlatform ShaderPlatform, const ITargetPlatform* TargetPlatform)
+void BeginRecompileGlobalShaders(const TArray<const FShaderType*>& OutdatedShaderTypes, const TArray<const FShaderPipelineType*>& OutdatedShaderPipelineTypes, EShaderPlatform ShaderPlatform, const ITargetPlatform* TargetPlatform, const FShaderCompilerFlags& InExtraCompilerFlags)
 {
 #if WITH_EDITOR
 	if (!FPlatformProperties::RequiresCookedData())
@@ -5552,7 +5618,7 @@ void BeginRecompileGlobalShaders(const TArray<const FShaderType*>& OutdatedShade
 		// Now check if there is any work to be done wrt outdates types
 		if (OutdatedShaderTypes.Num() > 0 || OutdatedShaderPipelineTypes.Num() > 0)
 		{
-			VerifyGlobalShaders(ShaderPlatform, TargetPlatform, false, &OutdatedShaderTypes, &OutdatedShaderPipelineTypes);
+			VerifyGlobalShaders(ShaderPlatform, TargetPlatform, false, &OutdatedShaderTypes, &OutdatedShaderPipelineTypes, InExtraCompilerFlags);
 		}
 	}
 #endif
diff --git a/Engine/Source/Runtime/Engine/Private/ShaderCompiler/ShaderCompilerEditor.cpp b/Engine/Source/Runtime/Engine/Private/ShaderCompiler/ShaderCompilerEditor.cpp
--- a/Engine/Source/Runtime/Engine/Private/ShaderCompiler/ShaderCompilerEditor.cpp
+++ b/Engine/Source/Runtime/Engine/Private/ShaderCompiler/ShaderCompilerEditor.cpp
@@ -227,12 +227,13 @@ static void CompileGlobalShaderMapForRemote(
 	const TArray<const FShaderPipelineType*>& OutdatedShaderPipelineTypes, 
 	const EShaderPlatform ShaderPlatform, 
 	const ITargetPlatform* TargetPlatform,
-	TArray<uint8>* OutArray)
+	TArray<uint8>* OutArray,
+	const FShaderCompilerFlags& InExtraCompilerFlags = {})
 {
 	UE_LOG(LogShaders, Display, TEXT("Recompiling global shaders."));
 
 	// Kick off global shader recompiles
-	BeginRecompileGlobalShaders(OutdatedShaderTypes, OutdatedShaderPipelineTypes, ShaderPlatform, TargetPlatform);
+	BeginRecompileGlobalShaders(OutdatedShaderTypes, OutdatedShaderPipelineTypes, ShaderPlatform, TargetPlatform, InExtraCompilerFlags);
 
 	// Block on global shaders
 	FinishRecompileGlobalShaders();
@@ -296,6 +297,7 @@ FArchive& operator<<(FArchive& Ar, FShaderRecompileData& RecompileData)
 
 	Ar << RecompileData.MaterialsToLoad;
 	Ar << RecompileData.ShaderTypesToLoad;
+	Ar << RecompileData.ExtraCompilerFlags;
 	Ar << iShaderPlatform;
 	Ar << iFeatureLevel;
 	Ar << iQualityLevel;
@@ -407,8 +409,9 @@ void RecompileShadersForRemote(
 				if (Args.CommandType == ODSCRecompileCommand::SingleShader &&
 					Args.ShaderTypesToLoad.Len() > 0)
 				{
-					TArray<const FShaderType*> ShaderTypes = FShaderType::GetShaderTypesByFilename(*Args.ShaderTypesToLoad);
-					TArray<const FShaderPipelineType*> ShaderPipelineTypes = FShaderPipelineType::GetShaderPipelineTypesByFilename(*Args.ShaderTypesToLoad);
+					constexpr bool bSearchAsRegexFilter = true;
+					TArray<const FShaderType*> ShaderTypes = FShaderType::GetShaderTypesByFilename(*Args.ShaderTypesToLoad, bSearchAsRegexFilter);
+					TArray<const FShaderPipelineType*> ShaderPipelineTypes = FShaderPipelineType::GetShaderPipelineTypesByFilename(*Args.ShaderTypesToLoad, bSearchAsRegexFilter);
 
 					for (const FShaderType* ShaderType : ShaderTypes)
 					{
@@ -417,7 +420,7 @@ void RecompileShadersForRemote(
 
 					UpdateReferencedUniformBufferNames(ShaderTypes, {}, ShaderPipelineTypes);
 
-					CompileGlobalShaderMapForRemote(ShaderTypes, ShaderPipelineTypes, ShaderPlatform, TargetPlatform, Args.GlobalShaderMap);
+					CompileGlobalShaderMapForRemote(ShaderTypes, ShaderPipelineTypes, ShaderPlatform, TargetPlatform, Args.GlobalShaderMap, Args.ExtraCompilerFlags);
 				}
 				else if (Args.CommandType == ODSCRecompileCommand::Global ||
 						 Args.CommandType == ODSCRecompileCommand::Changed)
@@ -438,7 +441,7 @@ void RecompileShadersForRemote(
 
 					UpdateReferencedUniformBufferNames(OutdatedShaderTypes, OutdatedFactoryTypes, OutdatedShaderPipelineTypes);
 
-					CompileGlobalShaderMapForRemote(OutdatedShaderTypes, OutdatedShaderPipelineTypes, ShaderPlatform, TargetPlatform, Args.GlobalShaderMap);
+					CompileGlobalShaderMapForRemote(OutdatedShaderTypes, OutdatedShaderPipelineTypes, ShaderPlatform, TargetPlatform, Args.GlobalShaderMap, Args.ExtraCompilerFlags);
 				}
 
 				// we only want to actually compile mesh shaders if a client directly requested it
diff --git a/Engine/Source/Runtime/Engine/Public/ODSC/ODSCManager.h b/Engine/Source/Runtime/Engine/Public/ODSC/ODSCManager.h
--- a/Engine/Source/Runtime/Engine/Public/ODSC/ODSCManager.h
+++ b/Engine/Source/Runtime/Engine/Public/ODSC/ODSCManager.h
@@ -55,7 +55,16 @@ public:
 	 *
 	 * @return false if no longer needs ticking
 	 */
-	ENGINE_API void AddThreadedRequest(const TArray<FString>& MaterialsToCompile, const FString& ShaderTypesToLoad, EShaderPlatform ShaderPlatform, ERHIFeatureLevel::Type FeatureLevel, EMaterialQualityLevel::Type QualityLevel, ODSCRecompileCommand RecompileCommandType, const FString& RequestedMaterialName = FString());
+	ENGINE_API void AddThreadedRequest(
+		const TArray<FString>& MaterialsToCompile,
+		const FString& ShaderTypesToLoad,
+		EShaderPlatform ShaderPlatform,
+		ERHIFeatureLevel::Type FeatureLevel,
+		EMaterialQualityLevel::Type QualityLevel,
+		ODSCRecompileCommand RecompileCommandType,
+		const FString& RequestedMaterialName = FString(),
+		const FShaderCompilerFlags& ExtraCompilerFlags = FShaderCompilerFlags()
+	);
 
 	/**
 	 * Add a request to compile a pipeline of shaders.  The results are submitted and processed in an async manner.
diff --git a/Engine/Source/Runtime/Engine/Public/ShaderCompiler.h b/Engine/Source/Runtime/Engine/Public/ShaderCompiler.h
--- a/Engine/Source/Runtime/Engine/Public/ShaderCompiler.h
+++ b/Engine/Source/Runtime/Engine/Public/ShaderCompiler.h
@@ -1345,7 +1345,7 @@ extern ENGINE_API FOnGlobalShadersCompilation& GetOnGlobalShaderCompilation();
 * @param	OutdatedShaderPipelineTypes		Optional list of shader pipeline types, will trigger compilation job for shader pipeline types found in this list even if the map already has the pipeline.
 */
 extern ENGINE_API void VerifyGlobalShaders(EShaderPlatform Platform, bool bLoadedFromCacheFile, const TArray<const FShaderType*>* OutdatedShaderTypes = nullptr, const TArray<const FShaderPipelineType*>* OutdatedShaderPipelineTypes = nullptr);
-extern ENGINE_API void VerifyGlobalShaders(EShaderPlatform Platform, const ITargetPlatform* TargetPlatform, bool bLoadedFromCacheFile, const TArray<const FShaderType*>* OutdatedShaderTypes = nullptr, const TArray<const FShaderPipelineType*>* OutdatedShaderPipelineTypes = nullptr);
+extern ENGINE_API void VerifyGlobalShaders(EShaderPlatform Platform, const ITargetPlatform* TargetPlatform, bool bLoadedFromCacheFile, const TArray<const FShaderType*>* OutdatedShaderTypes = nullptr, const TArray<const FShaderPipelineType*>* OutdatedShaderPipelineTypes = nullptr, const FShaderCompilerFlags& InExtraCompilerFlags = {});
 
 /** Precreates compute PSOs for global shaders. Separate from Compile/VerifyGlobalShaders to avoid blocking loads. */
 extern ENGINE_API void PrecacheComputePipelineStatesForGlobalShaders(ERHIFeatureLevel::Type FeatureLevel, const ITargetPlatform* TargetPlatform);
@@ -1367,7 +1367,7 @@ extern ENGINE_API bool RecompileChangedShadersForPlatform(const FString& Platfor
 * Begins recompiling the specified global shader types, and flushes their bound shader states.
 * FinishRecompileGlobalShaders must be called after this and before using the global shaders for anything.
 */
-extern ENGINE_API void BeginRecompileGlobalShaders(const TArray<const FShaderType*>& OutdatedShaderTypes, const TArray<const FShaderPipelineType*>& OutdatedShaderPipelineTypes, EShaderPlatform ShaderPlatform, const ITargetPlatform* TargetPlatform = nullptr);
+extern ENGINE_API void BeginRecompileGlobalShaders(const TArray<const FShaderType*>& OutdatedShaderTypes, const TArray<const FShaderPipelineType*>& OutdatedShaderPipelineTypes, EShaderPlatform ShaderPlatform, const ITargetPlatform* TargetPlatform = nullptr, const FShaderCompilerFlags& InExtraCompilerFlags = {});
 
 /** Finishes recompiling global shaders.  Must be called after BeginRecompileGlobalShaders. */
 extern ENGINE_API void FinishRecompileGlobalShaders();
@@ -1465,6 +1465,9 @@ struct FShaderRecompileData
 
 	EMaterialQualityLevel::Type QualityLevel = EMaterialQualityLevel::High;
 
+	/** Additional compiler flags for debugging, e.g. for ODSC shader requests to skip optimizations for individual shaders or enable bindless during development. */
+	FShaderCompilerFlags ExtraCompilerFlags = FShaderCompilerFlags();
+
 	/** All filenames that have been changed during the shader compilation. */
 	TArray<FString>* ModifiedFiles = nullptr;
 
diff --git a/Engine/Source/Runtime/RenderCore/Private/Shader.cpp b/Engine/Source/Runtime/RenderCore/Private/Shader.cpp
--- a/Engine/Source/Runtime/RenderCore/Private/Shader.cpp
+++ b/Engine/Source/Runtime/RenderCore/Private/Shader.cpp
@@ -12,6 +12,7 @@
 #include "Interfaces/ITargetPlatform.h"
 #include "Interfaces/ITargetPlatformManagerModule.h"
 #include "Interfaces/IShaderFormat.h"
+#include "Internationalization/Regex.h"
 #include "Serialization/ShaderKeyGenerator.h"
 #include "ShaderCodeLibrary.h"
 #include "ShaderCore.h"
@@ -356,13 +357,36 @@ FShaderType* FShaderType::GetShaderTypeByName(const TCHAR* Name)
 	return nullptr;
 }
 
-TArray<const FShaderType*> FShaderType::GetShaderTypesByFilename(const TCHAR* Filename)
+TArray<const FShaderType*> FShaderType::GetShaderTypesByFilename(const TCHAR* InFilename, bool bSearchAsRegexFilter)
+{
+	if (bSearchAsRegexFilter)
+	{
+		const FRegexPattern RegexSearch(FString(InFilename).Replace(TEXT("*"), TEXT("(.)*")), ERegexPatternFlags::CaseInsensitive);
+		return FShaderType::GetShaderTypesByFilenameFilter(
+			[&RegexSearch](const TCHAR* ShaderTypeFilename) -> bool
+			{
+				return FRegexMatcher(RegexSearch, ShaderTypeFilename).FindNext();
+			}
+		);
+	}
+	else
+	{
+		return FShaderType::GetShaderTypesByFilenameFilter(
+			[InFilename](const TCHAR* ShaderTypeFilename) -> bool
+			{
+				return FPlatformString::Strcmp(InFilename, ShaderTypeFilename) == 0;
+			}
+		);
+	}
+}
+
+TArray<const FShaderType*> FShaderType::GetShaderTypesByFilenameFilter(const TFunction<bool(const TCHAR*)>& InFilenameFilter)
 {
 	TArray<const FShaderType*> OutShaders;
-	for(TLinkedList<FShaderType*>::TIterator It(GetTypeList()); It; It.Next())
+	for (TLinkedList<FShaderType*>::TIterator It(GetTypeList()); It; It.Next())
 	{
 		const FShaderType* Type = *It;
-		if (FPlatformString::Strcmp(Filename, Type->GetShaderFilename()) == 0)
+		if (InFilenameFilter(Type->GetShaderFilename()))
 		{
 			OutShaders.Add(Type);
 		}
@@ -1093,7 +1117,30 @@ const TArray<FShaderPipelineType*>& FShaderPipelineType::GetSortedTypes(FShaderT
 	return GetSortedShaderPipelineTypes(Type);
 }
 
-TArray<const FShaderPipelineType*> FShaderPipelineType::GetShaderPipelineTypesByFilename(const TCHAR* Filename)
+TArray<const FShaderPipelineType*> FShaderPipelineType::GetShaderPipelineTypesByFilename(const TCHAR* InFilename, bool bSearchAsRegexFilter)
+{
+	if (bSearchAsRegexFilter)
+	{
+		const FRegexPattern RegexSearch(FString(InFilename).Replace(TEXT("*"), TEXT("(.)*")), ERegexPatternFlags::CaseInsensitive);
+		return FShaderPipelineType::GetShaderPipelineTypesByFilenameFilter(
+			[&RegexSearch](const TCHAR* ShaderTypeFilename) -> bool
+			{
+				return FRegexMatcher(RegexSearch, ShaderTypeFilename).FindNext();
+			}
+		);
+	}
+	else
+	{
+		return FShaderPipelineType::GetShaderPipelineTypesByFilenameFilter(
+			[InFilename](const TCHAR* ShaderTypeFilename) -> bool
+			{
+				return FPlatformString::Strcmp(InFilename, ShaderTypeFilename) == 0;
+			}
+		);
+	}
+}
+
+TArray<const FShaderPipelineType*> FShaderPipelineType::GetShaderPipelineTypesByFilenameFilter(const TFunction<bool(const TCHAR*)>& InFilenameFilter)
 {
 	TArray<const FShaderPipelineType*> PipelineTypes;
 	for (TLinkedList<FShaderPipelineType*>::TIterator It(FShaderPipelineType::GetTypeList()); It; It.Next())
@@ -1101,7 +1148,7 @@ TArray<const FShaderPipelineType*> FShaderPipelineType::GetShaderPipelineTypesBy
 		auto* PipelineType = *It;
 		for (auto* ShaderType : PipelineType->Stages)
 		{
-			if (FPlatformString::Strcmp(Filename, ShaderType->GetShaderFilename()) == 0)
+			if (InFilenameFilter(ShaderType->GetShaderFilename()))
 			{
 				PipelineTypes.AddUnique(PipelineType);
 				break;
diff --git a/Engine/Source/Runtime/RenderCore/Private/ShaderCore.cpp b/Engine/Source/Runtime/RenderCore/Private/ShaderCore.cpp
--- a/Engine/Source/Runtime/RenderCore/Private/ShaderCore.cpp
+++ b/Engine/Source/Runtime/RenderCore/Private/ShaderCore.cpp
@@ -4519,3 +4519,38 @@ void FShaderPipelineCompileJob::AppendDiagnostics(FString& OutDiagnostics, int32
 		StageJobs[StageJobIndex]->AppendDiagnostics(OutDiagnostics, StageJobIndex, StageJobs.Num(), *NewIndentation);
 	}
 }
+
+static const TMap<FString, ECompilerFlags>& GetCompilerFlagsDictionary()
+{
+	static const TMap<FString, ECompilerFlags> CompilerFlagsDictionary =
+	{
+		#define SHADER_COMPILER_FLAGS_ENTRY(NAME) { TEXT(#NAME), NAME },
+		#include "ShaderCompilerFlags.inl"
+		#undef SHADER_COMPILER_FLAGS_ENTRY
+	};
+	return CompilerFlagsDictionary;
+}
+
+void LexFromString(ECompilerFlags& OutValue, const TCHAR* InString)
+{
+	if (const ECompilerFlags* Flag = GetCompilerFlagsDictionary().Find(InString))
+	{
+		OutValue = *Flag;
+	}
+	else
+	{
+		OutValue = CFLAG_Max;
+	}
+}
+
+const TCHAR* LexToString(ECompilerFlags InValue)
+{
+	if (const FString* String = GetCompilerFlagsDictionary().FindKey(InValue))
+	{
+		return **String;
+	}
+	else
+	{
+		return nullptr;
+	}
+}
diff --git a/Engine/Source/Runtime/RenderCore/Public/Shader.h b/Engine/Source/Runtime/RenderCore/Public/Shader.h
--- a/Engine/Source/Runtime/RenderCore/Public/Shader.h
+++ b/Engine/Source/Runtime/RenderCore/Public/Shader.h
@@ -1281,7 +1281,8 @@ public:
 	static RENDERCORE_API TLinkedList<FShaderType*>*& GetTypeList();
 
 	static RENDERCORE_API FShaderType* GetShaderTypeByName(const TCHAR* Name);
-	static RENDERCORE_API TArray<const FShaderType*> GetShaderTypesByFilename(const TCHAR* Filename);
+	static RENDERCORE_API TArray<const FShaderType*> GetShaderTypesByFilename(const TCHAR* Filename, bool bSearchAsRegexFilter = false);
+	static RENDERCORE_API TArray<const FShaderType*> GetShaderTypesByFilenameFilter(const TFunction<bool(const TCHAR*)>& InFilenameFilter);
 
 	/** @return The global shader name to type map */
 	static RENDERCORE_API TMap<FHashedName, FShaderType*>& GetNameToTypeMap();
@@ -1973,7 +1974,8 @@ public:
 	/** Initialize static members, this must be called before any shader types are created. */
 	static RENDERCORE_API void Initialize();
 
-	static RENDERCORE_API TArray<const FShaderPipelineType*> GetShaderPipelineTypesByFilename(const TCHAR* Filename);
+	static RENDERCORE_API TArray<const FShaderPipelineType*> GetShaderPipelineTypesByFilename(const TCHAR* Filename, bool bSearchAsRegexFilter = false);
+	static RENDERCORE_API TArray<const FShaderPipelineType*> GetShaderPipelineTypesByFilenameFilter(const TFunction<bool(const TCHAR*)>& InFilenameFilter);
 
 	/** Serializes a shader type reference by name. */
 	RENDERCORE_API friend FArchive& operator<<(FArchive& Ar, const FShaderPipelineType*& Ref);
diff --git a/Engine/Source/Runtime/RenderCore/Public/ShaderCompilerCore.h b/Engine/Source/Runtime/RenderCore/Public/ShaderCompilerCore.h
--- a/Engine/Source/Runtime/RenderCore/Public/ShaderCompilerCore.h
+++ b/Engine/Source/Runtime/RenderCore/Public/ShaderCompilerCore.h
@@ -84,98 +84,16 @@ extern RENDERCORE_API bool AllowGlobalShaderLoad();
 
 enum ECompilerFlags
 {
-	CFLAG_PreferFlowControl = 0,
-	CFLAG_Debug,
-	CFLAG_AvoidFlowControl,
-	// Disable shader validation
-	CFLAG_SkipValidation,
-	// Only allows standard optimizations, not the longest compile times.
-	CFLAG_StandardOptimization,
-	// Always optimize, even when CFLAG_Debug is set. Required for some complex shaders and features.
-	CFLAG_ForceOptimization,
-	// Shader should generate symbols for debugging.
-	CFLAG_GenerateSymbols,
-	// Shader should insert debug/name info at the risk of generating non-deterministic libraries
-	CFLAG_ExtraShaderData,
-	// Allows the (external) symbols to be specific to each shader rather than trying to deduplicate.
-	CFLAG_AllowUniqueSymbols,
-	CFLAG_NoFastMath,
-	// Explicitly enforce zero initialization on shader platforms that may omit it.
-	CFLAG_ZeroInitialise,
-	// Explicitly enforce bounds checking on shader platforms that may omit it.
-	CFLAG_BoundsChecking,
-	// Force removing unused interpolators for platforms that can opt out
-	CFLAG_ForceRemoveUnusedInterpolators,
-	// Hint that it is a vertex to geometry shader
-	CFLAG_VertexToGeometryShader,
-	// Hint that it is a vertex to primitive shader
-	CFLAG_VertexToPrimitiveShader,
-	// Hint that a vertex shader should use automatic culling on certain platforms.
-	CFLAG_VertexUseAutoCulling,
-	// Prepare the shader for archiving in the native binary shader cache format
-	CFLAG_Archive,
-	// Shaders uses external texture so may need special runtime handling
-	CFLAG_UsesExternalTexture,
-	// Use emulated uniform buffers on supported platforms
-	CFLAG_UseEmulatedUB,
-	// Enable wave operation intrinsics (requires DX12 and DXC/DXIL on PC).
-	// Check GRHISupportsWaveOperations before using shaders compiled with this flag at runtime.
-	// [REDACTED_PROFILE_URL]/DirectXShaderCompiler/wiki/Wave-Intrinsics
-	CFLAG_WaveOperations,
-	// Use DirectX Shader Compiler (DXC) to compile all shaders, intended for compatibility testing.
-	CFLAG_ForceDXC,
-	CFLAG_SkipOptimizations,
-	// Temporarily disable optimizations with DXC compiler only, intended to workaround shader compiler bugs until they can be resolved with 1st party
-	CFLAG_SkipOptimizationsDXC,
-	// Typed UAV loads are disallowed by default, as Windows 7 D3D 11.0 does not support them; this flag allows a shader to use them.
-	CFLAG_AllowTypedUAVLoads,
-	// Prefer shader execution in wave32 mode if possible.
-	CFLAG_Wave32,
-	// Enable support of inline raytracing in compute shader.
-	CFLAG_InlineRayTracing,
-	// Enable support of C-style data types for platforms that can. Check for PLATFORM_SUPPORTS_REAL_TYPES and FDataDrivenShaderPlatformInfo::GetSupportsRealTypes()
-	CFLAG_AllowRealTypes,
-	// Precompile HLSL to optimized HLSL, then forward to FXC. Speeds up some shaders that take longer with FXC and works around crashes in FXC.
-	CFLAG_PrecompileWithDXC,
-	// Enable HLSL 2021 version. Enables templates, operator overloaing, C++ style function overloading. Contains breaking change with short-circuiting evaluation.
-	CFLAG_HLSL2021,
-	// Allow warnings to be treated as errors
-	CFLAG_WarningsAsErrors,
-	// Enabled if bindless resources are enabled for the platform
-	CFLAG_BindlessResources,
-	// Enabled if bindless samplers are enabled for the platform
-	CFLAG_BindlessSamplers,
-	// Force disable bindless resources and samplers on a shader
-	CFLAG_ForceBindful,
-	// EXPERIMENTAL: Run the shader re-writer that removes any unused functions/resources/types from source code before compilation.
-	CFLAG_RemoveDeadCode,
-	// Enable CullBeforeFetch optimization on supported platforms
-	CFLAG_CullBeforeFetch,
-	// Enable WarpCulling optimization on supported platforms
-	CFLAG_WarpCulling,
-	// Shader should generate minimal symbols info
-	CFLAG_GenerateSymbolsInfo,
-	// Enabled root constants optimization on supported platforms
-	CFLAG_RootConstants,
-	// Specifies that a shader provides derivatives, and the compiler should look in the compiled ISA for any instructions requiring
-	// auto derivatives. If none are found, the shader will be marked with EShaderResourceUsageFlags::NoDerivativeOps, meaning that
-	// calling code can safely assume only provided derivatives are used.
-	CFLAG_CheckForDerivativeOps,
-	// Shader is used with indirect draws. This flag is currently used to fix a platform specific problem with certain (rare) indirect draw setups, but it is intended to be set for all indirect draw shaders in the future.
-	// Must not be used on shaders that are used with direct draws. Doing so might cause crashes or visual corruption on certain platforms.
-	CFLAG_IndirectDraw,
-	// Shader is used with shader bundles.
-	CFLAG_ShaderBundle,
-	// Shader code should not be stripped of comments/whitespace/line directives at the end of preprocessing
-	CFLAG_DisableSourceStripping,
-	// Shader uses RHI Shader Binding Layout for global shader binding
-	CFLAG_ShaderBindingLayout,
-	// Request full shader analysis artifacts in output statistics. This may contain multiple compilation steps in full text form.
-	CFLAG_OutputAnalysisArtifacts,
-	CFLAG_Max,
+	#define SHADER_COMPILER_FLAGS_ENTRY(NAME) NAME,
+	#include "ShaderCompilerFlags.inl"
+	#undef SHADER_COMPILER_FLAGS_ENTRY
 };
+static_assert(CFLAG_PreferFlowControl == 0, "First entry in ECompilerFlags must be 'CFLAG_PreferFlowControl' and assigned to 0");
 static_assert(CFLAG_Max < 64, "Out of bitfield space! Modify FShaderCompilerFlags");
 
+extern RENDERCORE_API void LexFromString(ECompilerFlags& OutValue, const TCHAR* InString);
+extern RENDERCORE_API const TCHAR* LexToString(ECompilerFlags InValue);
+
 struct FShaderCompilerResourceTable
 {
 	/** Bits indicating which resource tables contain resources bound to this shader. */
diff --git a/Engine/Source/Runtime/RenderCore/Public/ShaderCompilerFlags.inl b/Engine/Source/Runtime/RenderCore/Public/ShaderCompilerFlags.inl
new file mode 100644
--- /dev/null
+++ b/Engine/Source/Runtime/RenderCore/Public/ShaderCompilerFlags.inl
@@ -0,0 +1,101 @@
+// Copyright Epic Games, Inc. All Rights Reserved.
+
+/*=============================================================================
+	ShaderCompilerFlags.inl: Flags used in ECompilerFlags enumeration.
+=============================================================================*/
+
+#ifndef SHADER_COMPILER_FLAGS_ENTRY
+#error SHADER_COMPILER_FLAGS_ENTRY must be defined before including ShaderCompilerFlags.inl
+#endif
+
+SHADER_COMPILER_FLAGS_ENTRY( CFLAG_PreferFlowControl )
+SHADER_COMPILER_FLAGS_ENTRY( CFLAG_Debug )
+SHADER_COMPILER_FLAGS_ENTRY( CFLAG_AvoidFlowControl )
+// Disable shader validation
+SHADER_COMPILER_FLAGS_ENTRY( CFLAG_SkipValidation )
+// Only allows standard optimizations, not the longest compile times.
+SHADER_COMPILER_FLAGS_ENTRY( CFLAG_StandardOptimization )
+// Always optimize even when CFLAG_Debug is set. Required for some complex shaders and features.
+SHADER_COMPILER_FLAGS_ENTRY( CFLAG_ForceOptimization )
+// Shader should generate symbols for debugging.
+SHADER_COMPILER_FLAGS_ENTRY( CFLAG_GenerateSymbols )
+// Shader should insert debug/name info at the risk of generating non-deterministic libraries
+SHADER_COMPILER_FLAGS_ENTRY( CFLAG_ExtraShaderData )
+// Allows the (external) symbols to be specific to each shader rather than trying to deduplicate.
+SHADER_COMPILER_FLAGS_ENTRY( CFLAG_AllowUniqueSymbols )
+SHADER_COMPILER_FLAGS_ENTRY( CFLAG_NoFastMath )
+// Explicitly enforce zero initialization on shader platforms that may omit it.
+SHADER_COMPILER_FLAGS_ENTRY( CFLAG_ZeroInitialise )
+// Explicitly enforce bounds checking on shader platforms that may omit it.
+SHADER_COMPILER_FLAGS_ENTRY( CFLAG_BoundsChecking )
+// Force removing unused interpolators for platforms that can opt out
+SHADER_COMPILER_FLAGS_ENTRY( CFLAG_ForceRemoveUnusedInterpolators )
+// Hint that it is a vertex to geometry shader
+SHADER_COMPILER_FLAGS_ENTRY( CFLAG_VertexToGeometryShader )
+// Hint that it is a vertex to primitive shader
+SHADER_COMPILER_FLAGS_ENTRY( CFLAG_VertexToPrimitiveShader )
+// Hint that a vertex shader should use automatic culling on certain platforms.
+SHADER_COMPILER_FLAGS_ENTRY( CFLAG_VertexUseAutoCulling )
+// Prepare the shader for archiving in the native binary shader cache format
+SHADER_COMPILER_FLAGS_ENTRY( CFLAG_Archive )
+// Shaders uses external texture so may need special runtime handling
+SHADER_COMPILER_FLAGS_ENTRY( CFLAG_UsesExternalTexture )
+// Use emulated uniform buffers on supported platforms
+SHADER_COMPILER_FLAGS_ENTRY( CFLAG_UseEmulatedUB )
+// Enable wave operation intrinsics (requires DX12 and DXC/DXIL on PC).
+// Check GRHISupportsWaveOperations before using shaders compiled with this flag at runtime.
+// [REDACTED_PROFILE_URL]/DirectXShaderCompiler/wiki/Wave-Intrinsics
+SHADER_COMPILER_FLAGS_ENTRY( CFLAG_WaveOperations )
+// Use DirectX Shader Compiler (DXC) to compile all shaders ) intended for compatibility testing.
+SHADER_COMPILER_FLAGS_ENTRY( CFLAG_ForceDXC )
+SHADER_COMPILER_FLAGS_ENTRY( CFLAG_SkipOptimizations )
+// Temporarily disable optimizations with DXC compiler only ) intended to workaround shader compiler bugs until they can be resolved with 1st party
+SHADER_COMPILER_FLAGS_ENTRY( CFLAG_SkipOptimizationsDXC )
+// Typed UAV loads are disallowed by default ) as Windows 7 D3D 11.0 does not support them; this flag allows a shader to use them.
+SHADER_COMPILER_FLAGS_ENTRY( CFLAG_AllowTypedUAVLoads )
+// Prefer shader execution in wave32 mode if possible.
+SHADER_COMPILER_FLAGS_ENTRY( CFLAG_Wave32 )
+// Enable support of inline raytracing in compute shader.
+SHADER_COMPILER_FLAGS_ENTRY( CFLAG_InlineRayTracing )
+// Enable support of C-style data types for platforms that can. Check for PLATFORM_SUPPORTS_REAL_TYPES and FDataDrivenShaderPlatformInfo::GetSupportsRealTypes()
+SHADER_COMPILER_FLAGS_ENTRY( CFLAG_AllowRealTypes )
+// Precompile HLSL to optimized HLSL ) then forward to FXC. Speeds up some shaders that take longer with FXC and works around crashes in FXC.
+SHADER_COMPILER_FLAGS_ENTRY( CFLAG_PrecompileWithDXC )
+// Enable HLSL 2021 version. Enables templates ) operator overloaing ) C++ style function overloading. Contains breaking change with short-circuiting evaluation.
+SHADER_COMPILER_FLAGS_ENTRY( CFLAG_HLSL2021 )
+// Allow warnings to be treated as errors
+SHADER_COMPILER_FLAGS_ENTRY( CFLAG_WarningsAsErrors )
+// Enabled if bindless resources are enabled for the platform
+SHADER_COMPILER_FLAGS_ENTRY( CFLAG_BindlessResources )
+// Enabled if bindless samplers are enabled for the platform
+SHADER_COMPILER_FLAGS_ENTRY( CFLAG_BindlessSamplers )
+// Force disable bindless resources and samplers on a shader
+SHADER_COMPILER_FLAGS_ENTRY( CFLAG_ForceBindful )
+// EXPERIMENTAL: Run the shader re-writer that removes any unused functions/resources/types from source code before compilation.
+SHADER_COMPILER_FLAGS_ENTRY( CFLAG_RemoveDeadCode )
+// Enable CullBeforeFetch optimization on supported platforms
+SHADER_COMPILER_FLAGS_ENTRY( CFLAG_CullBeforeFetch )
+// Enable WarpCulling optimization on supported platforms
+SHADER_COMPILER_FLAGS_ENTRY( CFLAG_WarpCulling )
+// Shader should generate minimal symbols info
+SHADER_COMPILER_FLAGS_ENTRY( CFLAG_GenerateSymbolsInfo )
+// Enabled root constants optimization on supported platforms
+SHADER_COMPILER_FLAGS_ENTRY( CFLAG_RootConstants )
+// Specifies that a shader provides derivatives ) and the compiler should look in the compiled ISA for any instructions requiring
+// auto derivatives. If none are found ) the shader will be marked with EShaderResourceUsageFlags::NoDerivativeOps ) meaning that
+// calling code can safely assume only provided derivatives are used.
+SHADER_COMPILER_FLAGS_ENTRY( CFLAG_CheckForDerivativeOps )
+// Shader is used with indirect draws. This flag is currently used to fix a platform specific problem with certain (rare) indirect draw setups ) but it is intended to be set for all indirect draw shaders in the future.
+// Must not be used on shaders that are used with direct draws. Doing so might cause crashes or visual corruption on certain platforms.
+SHADER_COMPILER_FLAGS_ENTRY( CFLAG_IndirectDraw )
+// Shader is used with shader bundles.
+SHADER_COMPILER_FLAGS_ENTRY( CFLAG_ShaderBundle )
+// Shader code should not be stripped of comments/whitespace/line directives at the end of preprocessing
+SHADER_COMPILER_FLAGS_ENTRY( CFLAG_DisableSourceStripping )
+// Shader uses RHI Shader Binding Layout for global shader binding
+SHADER_COMPILER_FLAGS_ENTRY( CFLAG_ShaderBindingLayout )
+// Request full shader analysis artifacts in output statistics. This may contain multiple compilation steps in full text form.
+SHADER_COMPILER_FLAGS_ENTRY( CFLAG_OutputAnalysisArtifacts )
+SHADER_COMPILER_FLAGS_ENTRY( CFLAG_Max )
+
+
