Fixed FResourcePool::Drain method skipping certain frames and not reclaiming memory. The heartbeat tickable thing is based off of a fixed delta time and can end up skipping frames. For example, frame 1 and frame 3 were being drained but 0 and 2 were always being skipped, which leaked the buffers.

diff --git a/Engine/Source/Runtime/Engine/Public/ResourcePool.h b/Engine/Source/Runtime/Engine/Public/ResourcePool.h
--- a/Engine/Source/Runtime/Engine/Public/ResourcePool.h
+++ b/Engine/Source/Runtime/Engine/Public/ResourcePool.h
@@ -97,14 +97,32 @@ public:
 
 		if(!bForceDrainAll)
 		{
-			// Index of the bucket that is now old enough to be reused
-			const int32 SafeFrameIndex = (GFrameNumberRenderThread + 1) % ResourcePoolPolicy::NumSafeFrames;
+			// DrainPool won't necessarily be called at the same frequency as the render thread number increments.
+			// Therefore, we track the frame number from the previous drain that occurred and drain all frames between
+			// the last and the current frame number (capping out at draining all frames).
 
-			for (FPooledResource& PoolEntry : SafeResourceBuckets[SafeFrameIndex])
+			uint32 SafeFrameNumber = GFrameNumberRenderThread + 1;
+
+			if (LastSafeFrameNumber == INDEX_NONE)
 			{
-				ResourceBuckets[PoolEntry.BucketIndex].Emplace(MoveTemp(PoolEntry));
+				LastSafeFrameNumber = SafeFrameNumber;
 			}
-			SafeResourceBuckets[SafeFrameIndex].Reset();
+
+			uint32 NumSafeFramesToDrain = FMath::Clamp(SafeFrameNumber - LastSafeFrameNumber, 1, ResourcePoolPolicy::NumSafeFrames);
+
+			for (uint32 Index = 0; Index < NumSafeFramesToDrain; ++Index)
+			{
+				// Index of the bucket that is now old enough to be reused
+				const int32 SafeFrameIndex = (LastSafeFrameNumber + Index) % ResourcePoolPolicy::NumSafeFrames;
+
+				for (FPooledResource& PoolEntry : SafeResourceBuckets[SafeFrameIndex])
+				{
+					ResourceBuckets[PoolEntry.BucketIndex].Emplace(MoveTemp(PoolEntry));
+				}
+				SafeResourceBuckets[SafeFrameIndex].Reset();
+			}
+
+			LastSafeFrameNumber = SafeFrameNumber;
 		}
 		else
 		{
@@ -185,6 +203,7 @@ private:
 		uint32 BucketIndex;
 	};
 
+	uint32 LastSafeFrameNumber = INDEX_NONE;
 	UE::FMutex CS;
 	bool bLocked = false;
 
@@ -290,3 +309,5 @@ public: // From FRenderResource
 		TResourcePool<ResourceType, ResourcePoolPolicy, ResourceCreationArguments>::DrainPool(true);
 	}
 };
+
+UE_ENABLE_OPTIMIZATION
\ No newline at end of file
