Removed extraneous UE_ENABLE_OPTIMIZATION.

diff --git a/Engine/Source/Runtime/Engine/Public/ResourcePool.h b/Engine/Source/Runtime/Engine/Public/ResourcePool.h
--- a/Engine/Source/Runtime/Engine/Public/ResourcePool.h
+++ b/Engine/Source/Runtime/Engine/Public/ResourcePool.h
@@ -308,6 +308,4 @@ public: // From FRenderResource
 		FTickableObjectRenderThread::Unregister();
 		TResourcePool<ResourceType, ResourcePoolPolicy, ResourceCreationArguments>::DrainPool(true);
 	}
-};
-
-UE_ENABLE_OPTIMIZATION
\ No newline at end of file
+};
\ No newline at end of file
