Fixed build error due to precision warning.

diff --git a/Engine/Source/Runtime/Engine/Public/ResourcePool.h b/Engine/Source/Runtime/Engine/Public/ResourcePool.h
--- a/Engine/Source/Runtime/Engine/Public/ResourcePool.h
+++ b/Engine/Source/Runtime/Engine/Public/ResourcePool.h
@@ -108,7 +108,7 @@ public:
 				LastSafeFrameNumber = SafeFrameNumber;
 			}
 
-			uint32 NumSafeFramesToDrain = FMath::Clamp(SafeFrameNumber - LastSafeFrameNumber, 1, ResourcePoolPolicy::NumSafeFrames);
+			uint32 NumSafeFramesToDrain = FMath::Clamp<uint32>(SafeFrameNumber - LastSafeFrameNumber, 1, ResourcePoolPolicy::NumSafeFrames);
 
 			for (uint32 Index = 0; Index < NumSafeFramesToDrain; ++Index)
 			{
