Add lcs diff to string library (editor only because this is expensive for large strings)

[RN] DiffString has been added as an editor only method for diffing two strings from blueprints, python, or C++

diff --git a/Engine/Source/Runtime/Engine/Classes/Kismet/KismetStringLibrary.h b/Engine/Source/Runtime/Engine/Classes/Kismet/KismetStringLibrary.h
--- a/Engine/Source/Runtime/Engine/Classes/Kismet/KismetStringLibrary.h
+++ b/Engine/Source/Runtime/Engine/Classes/Kismet/KismetStringLibrary.h
@@ -562,6 +562,31 @@ class UKismetStringLibrary : public UBlueprintFunctionLibrary
 	UFUNCTION(BlueprintPure, Category = "Utilities|String")
 	static ENGINE_API FString Mid(const FString& SourceString, int32 Start, int32 Count);
 
+#if WITH_EDITOR
+	/** 
+	 * Returns a string describing the LCS difference between two strings, First and Second.
+	 * This function is editor only because it is expensive for large strings. Examples:
+	 * 
+	 * ("A", "") -> "- [ 0001 |      ] A"
+	 * ("", "") -> ""
+	 * ("The quick brown fox", "The quick red fox") -> "~ [ 0001 | 0001 ]
+     *                   The quick brown fox
+     *                             ~~~~~    
+	 *                   The quick red   fox"
+	 * 
+	 * The numbers between the [], separated by |, are the line numbers in the First and Second string at
+	 * which the difference occurs. Preceding the line number information is a symbol identifying the type 
+	 * of difference: A tilde (~) indicates a modification, a minus (-) indicates a removal from First, and
+	 * a plus (+) indicates an addition to second.
+	 * 
+	 * @param First - The string to treat as basis
+	 * @param Second - The string to treat as change
+     * @return the LCS difference between the First and Second string
+	 */
+	UFUNCTION(BlueprintCallable, Category = "Utilities|String")
+	static ENGINE_API FString DiffString(const FString& First, const FString& Second);
+#endif	// WITH_EDITOR
+
 	/**
 	 * Convert a number of seconds into minutes:seconds.milliseconds format string (including leading zeroes)
 	 *
diff --git a/Engine/Source/Runtime/Engine/Private/KismetStringLibrary.cpp b/Engine/Source/Runtime/Engine/Private/KismetStringLibrary.cpp
--- a/Engine/Source/Runtime/Engine/Private/KismetStringLibrary.cpp
+++ b/Engine/Source/Runtime/Engine/Private/KismetStringLibrary.cpp
@@ -3,6 +3,13 @@
 #include "Kismet/KismetStringLibrary.h"
 #include "Math/Box.h"
 
+#if WITH_EDITOR
+#include "AsyncTreeDifferences.h"
+#include "Containers/StringView.h"
+#include "DiffUtils.h"
+#include "Internationalization/Regex.h"
+#endif
+
 #include UE_INLINE_GENERATED_CPP_BY_NAME(KismetStringLibrary)
 
 //////////////////////////////////////////////////////////////////////////
@@ -555,6 +562,210 @@ FString UKismetStringLibrary::Mid(const FString& SourceString, int32 Start, int3
 	return (Count >= 0 ? SourceString.Mid(Start, Count) : FString());
 }
 
+#if WITH_EDITOR
+namespace UE::Private
+{
+	struct FTokenizedStringNode
+	{
+		FStringView StringView;
+		int32 LineIndex = INDEX_NONE; // line number in the full string (INDEX_NONE if this is a head node)
+		int32 TokenIndex = INDEX_NONE; // index of token at current line (INDEX_NONE if this is a head or line node)
+		
+		bool operator==(const FTokenizedStringNode&) const = default;
+	};
+}
+
+// methods that make FTokenizedStringNode diffable - specializations must be done at global scope:
+template<>
+class TTreeDiffSpecification<UE::Private::FTokenizedStringNode>
+{
+public:
+	bool AreValuesEqual(const UE::Private::FTokenizedStringNode& TreeNodeA, const UE::Private::FTokenizedStringNode& TreeNodeB) const
+	{
+		return TreeNodeA.StringView.TrimStartAndEnd() == TreeNodeB.StringView.TrimStartAndEnd();
+	}
+
+	bool AreMatching(const UE::Private::FTokenizedStringNode& TreeNodeA, const UE::Private::FTokenizedStringNode& TreeNodeB) const
+	{
+		return TreeNodeA.LineIndex == TreeNodeB.LineIndex &&  TreeNodeA.TokenIndex == TreeNodeB.TokenIndex;
+	}
+
+	void GetChildren(const UE::Private::FTokenizedStringNode& InParent, TArray<UE::Private::FTokenizedStringNode>& OutChildren) const
+	{
+		if (InParent.LineIndex == INDEX_NONE)
+		{
+			// parent is the head. Each child should be a line
+			int32 LineIndex = 0;
+			FStringView RemainingView(InParent.StringView);
+			while (!RemainingView.IsEmpty())
+			{
+				int32 ChildSize;
+				if (!RemainingView.FindChar('\n', ChildSize))
+				{
+					ChildSize = RemainingView.Len();
+				}
+			
+				UE::Private::FTokenizedStringNode NextChild(InParent);
+				NextChild.StringView = FStringView(RemainingView.begin(), ChildSize);
+				NextChild.LineIndex = LineIndex++;
+				OutChildren.Add(NextChild);
+			
+				RemainingView = FStringView(RemainingView.begin() + (ChildSize + 1), FMath::Max(0,RemainingView.Len() - (ChildSize + 1)));
+			}
+		}
+		else if (InParent.TokenIndex == INDEX_NONE)
+		{
+			// parent is a line. Each child should be a token
+			int32 TokenIndex = 0;
+		
+			FRegexPattern TokenPattern(TEXT(R"(\w+|[^\w\s]+|[\s]+)"));
+			// TODO: Constructing a string here is slow! FRegexMatcher should be using FStringView!
+			FRegexMatcher TokenMatcher(TokenPattern, FString(InParent.StringView));
+
+			while (TokenMatcher.FindNext())
+			{
+				const TCHAR* Begin = InParent.StringView.begin() + TokenMatcher.GetMatchBeginning();
+				const TCHAR* End = InParent.StringView.begin() + TokenMatcher.GetMatchEnding();
+			
+				UE::Private::FTokenizedStringNode NextChild(InParent);
+				NextChild.StringView = FStringView(Begin, End-Begin);
+				NextChild.TokenIndex = TokenIndex++;
+				OutChildren.Add(NextChild);
+			}
+		}
+	}
+
+	bool ShouldMatchByValue(const UE::Private::FTokenizedStringNode& TreeNodeA) const
+	{
+		// prioritize matching equivalent substrings
+		return true;
+	}
+
+	bool ShouldInheritEqualFromChildren(const UE::Private::FTokenizedStringNode& TreeNodeA, const UE::Private::FTokenizedStringNode& TreeNodeB) const
+	{
+		return true;
+	}
+};
+
+FString UKismetStringLibrary::DiffString(const FString& First, const FString& Second)
+{
+	using namespace UE::Private;
+
+	class FStringDiffTree : public TAsyncTreeDifferences<FTokenizedStringNode>
+	{
+	public:
+		FStringDiffTree(const FString& StringA, const FString& StringB)
+			: TAsyncTreeDifferences<FTokenizedStringNode>(GetRootAttribute(StringA), GetRootAttribute(StringB))
+		{}
+
+		// generates a text based explanation of the diff
+		// and calls PrintLine for each line of the result
+		FString CollectDifferences() const
+		{
+			FString Result;
+			ForEach (
+				ETreeTraverseOrder::PreOrder,
+				[&Result](const TUniquePtr<FStringDiffTree::DiffNodeType>& DiffNode)->ETreeTraverseControl
+				{
+					// the tree has multiple levels of granularity diffing the entire text, then lines, then tokens,
+					// then individual characters.
+					// Skip to the first diff of a line
+					if (DiffNode->ValueA.LineIndex == INDEX_NONE && DiffNode->ValueB.LineIndex == INDEX_NONE)
+					{
+						return ETreeTraverseControl::Continue;
+					}
+
+					if (DiffNode->ValueA.TokenIndex == INDEX_NONE && DiffNode->ValueB.TokenIndex == INDEX_NONE)
+					{
+						// so that distances are consistent, replace tabs with spaces
+						FString LeftString = FString(DiffNode->ValueA.StringView).Replace(TEXT("\t"), TEXT("  "));
+						FString RightString = FString(DiffNode->ValueB.StringView).Replace(TEXT("\t"), TEXT("  "));
+				
+						switch (DiffNode->DiffResult)
+						{
+						case ETreeDiffResult::Invalid:
+							break;
+						case ETreeDiffResult::Identical:
+							// empty string for no differences
+							break;
+						case ETreeDiffResult::MissingFromTree1:
+							Result += FString::Printf(TEXT("+ [      | %04d ] %s\n"), DiffNode->ValueB.LineIndex + 1, *RightString);
+							break;
+						case ETreeDiffResult::MissingFromTree2:
+							Result += FString::Printf(TEXT("- [ %04d |      ] %s\n"), DiffNode->ValueA.LineIndex + 1, *LeftString);
+							break;
+						case ETreeDiffResult::DifferentValues:
+							{
+								Result += FString::Printf(TEXT("~ [ %04d | %04d ]\n"), DiffNode->ValueA.LineIndex + 1, DiffNode->ValueB.LineIndex + 1);
+
+								auto GrowString = [](FString& String, TCHAR Char, int32 NewLen)
+								{
+									while (String.Len() < NewLen) String.AppendChar(Char);
+								};
+							
+								FString LeftPadded  = "                  ";
+								FString Annotations = "                  ";
+								FString RightPadded = "                  ";
+								// the children of this node contain token diff info. Use it to generate two parallel diff strings that show what's changed
+								for (TUniquePtr<FStringDiffTree::DiffNodeType>& TokenDiffNode : DiffNode->Children)
+								{
+									// so that distances are consistent, replace tabs with spaces
+									LeftPadded += FString(TokenDiffNode->ValueA.StringView).Replace(TEXT("\t"), TEXT("  "));
+									RightPadded += FString(TokenDiffNode->ValueB.StringView).Replace(TEXT("\t"), TEXT("  "));
+							
+									int32 NewLength = FMath::Max(RightPadded.Len(), LeftPadded.Len());
+									GrowString(LeftPadded, ' ', NewLength);
+									GrowString(RightPadded, ' ', NewLength);
+
+									switch (TokenDiffNode->DiffResult)
+									{
+									case ETreeDiffResult::Invalid:
+										check(false);
+									case ETreeDiffResult::MissingFromTree1:
+										GrowString(Annotations, 'v', NewLength);
+										break;
+									case ETreeDiffResult::MissingFromTree2:
+										GrowString(Annotations, '^', NewLength);
+										break;
+									case ETreeDiffResult::DifferentValues:
+										GrowString(Annotations, '~', NewLength);
+										break;
+									case ETreeDiffResult::Identical:
+										GrowString(Annotations, ' ', NewLength);
+										break;
+									}
+								}
+						
+								Result += LeftPadded + TEXT("\n");
+								Result += Annotations + TEXT("\n");
+								Result += RightPadded + TEXT("\n");
+						
+								break;
+							}
+						}
+					}
+					return ETreeTraverseControl::SkipChildren;
+				}
+			);
+			return Result;
+		}
+
+	private:
+		static TAttribute<TArray<FTokenizedStringNode>> GetRootAttribute(const FString& String)
+		{
+			return TAttribute<TArray<FTokenizedStringNode>>::CreateLambda([String]()->TArray<FTokenizedStringNode>
+			{
+				return {FTokenizedStringNode{String}};
+			});
+		}
+	};
+
+	FStringDiffTree Difference(First, Second);
+	Difference.FlushQueue();
+	return Difference.CollectDifferences();
+}
+#endif// WITH_EDITOR
+
 FString UKismetStringLibrary::TimeSecondsToString(float InSeconds)
 {
 	// Determine whether to display this number as a negative
