[WorldPartition] Change assertion that could happen in case of mismatch in SubObjectsToExternalStreamingObjectsRemapping map, into a log to get more information and avoid blocking automated tests.

diff --git a/Engine/Source/Runtime/Engine/Private/WorldPartition/WorldPartitionLevelStreamingPolicy.cpp b/Engine/Source/Runtime/Engine/Private/WorldPartition/WorldPartitionLevelStreamingPolicy.cpp
--- a/Engine/Source/Runtime/Engine/Private/WorldPartition/WorldPartitionLevelStreamingPolicy.cpp
+++ b/Engine/Source/Runtime/Engine/Private/WorldPartition/WorldPartitionLevelStreamingPolicy.cpp
@@ -324,7 +324,11 @@ bool UWorldPartitionLevelStreamingPolicy::RemoveExternalStreamingObject(URuntime
 
 	for (const TPair<FName, FName> Pair : ExternalStreamingObject->SubObjectsToCellRemapping)
 	{
-		verify(SubObjectsToExternalStreamingObjectsRemapping.FindAndRemoveChecked(Pair.Key) == ExternalStreamingObject);
+		TWeakObjectPtr<URuntimeHashExternalStreamingObjectBase> RemovedExternalStreamingObject = SubObjectsToExternalStreamingObjectsRemapping.FindAndRemoveChecked(Pair.Key);
+		UE_CLOG(RemovedExternalStreamingObject != ExternalStreamingObject, LogWorldPartition, Log, TEXT("Mismatch in SubObjectsToExternalStreamingObjectsRemapping map. Found: %s. Expected: %s"),
+			(RemovedExternalStreamingObject.IsValid() && RemovedExternalStreamingObject->GetPackage()) ? *RemovedExternalStreamingObject->GetPackage()->GetName() : TEXT("<Invalid>"),
+			ExternalStreamingObject->GetPackage() ? *ExternalStreamingObject->GetPackage()->GetName() : TEXT("<Invalid>"));
+			
 	}
 
 	ExternalStreamingObjects.RemoveSwap(ExternalStreamingObject);
