GetDebugStringForWorld now properly identifies editor worlds as such

diff --git a/Engine/Source/Runtime/Engine/Private/UnrealEngine.cpp b/Engine/Source/Runtime/Engine/Private/UnrealEngine.cpp
--- a/Engine/Source/Runtime/Engine/Private/UnrealEngine.cpp
+++ b/Engine/Source/Runtime/Engine/Private/UnrealEngine.cpp
@@ -651,7 +651,11 @@ ENGINE_API FString GetDebugStringForWorld(const UWorld* World)
 	{
 		WorldName = NSLOCTEXT("Engine", "PlayWorldBeingCreated", "(World Being Created)").ToString();
 	}
-	else
+	else if (World->WorldType == EWorldType::Editor)
+	{
+		WorldName = NSLOCTEXT("Engine", "PlayWorldIsEditor", "Editor").ToString();
+	}
+	else if (World->IsGameWorld())
 	{
 		FWorldContext* WorldContext = GEngine->GetWorldContextFromWorld(World);
 
