Streamable Manager Error Improvements
- Add target path to error details
- Make sure error details are set for all handles waiting on a streamable

diff --git a/Engine/Source/Runtime/Engine/Classes/Engine/Experimental/StreamableManagerError.h b/Engine/Source/Runtime/Engine/Classes/Engine/Experimental/StreamableManagerError.h
--- a/Engine/Source/Runtime/Engine/Classes/Engine/Experimental/StreamableManagerError.h
+++ b/Engine/Source/Runtime/Engine/Classes/Engine/Experimental/StreamableManagerError.h
@@ -25,7 +25,7 @@ DECLARE_ERROR_ONEPARAM(AsyncLoadUnknownError, 10, StreamableManager, NSLOCTEXT("
 
 
 DECLARE_ERROR(UnknownError, 11, StreamableManager, NSLOCTEXT("StreamableManager", "UnknownError", "Unknown error occured while streaming asset"));
-
+DECLARE_ERROR(AsyncLoadNotInstalled, 12, StreamableManager, NSLOCTEXT("StreamableManager", "AsyncLoadNotInstalled", "Async load failed because the package is not installed."));
 
 namespace UE::UnifiedError::StreamableManager
 {
@@ -34,18 +34,33 @@ namespace UE::UnifiedError::StreamableManager
 		DECLARE_FERROR_DETAILS(StreamableManager, FStreamableManagerErrorDetails);
 	private:
 		FString SoftObjectPath;
+		FText ErrorFormatString;
 		FStreamableManagerErrorDetails() = default;
 	public:
 
-		FStreamableManagerErrorDetails(const FString& InSoftObjectPath) : SoftObjectPath(InSoftObjectPath)
-		{}
+		FStreamableManagerErrorDetails(const FString& InSoftObjectPath) 
+			: SoftObjectPath(InSoftObjectPath)
+			, ErrorFormatString(NSLOCTEXT("StreamableManager", "StreamableManagerErrorDetails", "{InnerErrorMessage} Target Path: {SoftObjectPath}"))
+		{
+		}
 
 		virtual void GetErrorProperties(const FError& Error, IErrorPropertyExtractor& OutProperties) const override
 		{
 			OutProperties.AddProperty(UTF8TEXTVIEW("SoftObjectPath"), SoftObjectPath);
+			FFormatNamedArguments ErrorMessageArgs;
+			UE::UnifiedError::FTextFormatArgsPropertyExtractor ExtractToFTextArgs(ErrorMessageArgs);
+			GetInnerErrorDetails()->GetErrorProperties(Error, ExtractToFTextArgs);
+			const FText& FormatString = GetInnerErrorDetails()->GetErrorFormatString(Error);
+			FText Result =FText::Format(FormatString, ErrorMessageArgs);
+			OutProperties.AddProperty(UTF8TEXTVIEW("InnerErrorMessage"), Result);
 			FDynamicErrorDetails::GetErrorProperties(Error, OutProperties);
 		}
 
+		virtual const FText& GetErrorFormatString(const FError& Error) const override
+		{
+			return ErrorFormatString;
+		}
+
 		const FString& GetSoftObjectPath() const { return SoftObjectPath; }
 	};
 
diff --git a/Engine/Source/Runtime/Engine/Classes/Engine/StreamableManager.h b/Engine/Source/Runtime/Engine/Classes/Engine/StreamableManager.h
--- a/Engine/Source/Runtime/Engine/Classes/Engine/StreamableManager.h
+++ b/Engine/Source/Runtime/Engine/Classes/Engine/StreamableManager.h
@@ -569,7 +569,7 @@ private:
 	void UnbindDelegates();
 
 	/** Record any errors that occured with this handle */
-	void SetError(UE::UnifiedError::FError&& Error);
+	void SetError(const UE::UnifiedError::FError& Error);
 	void SetErrorOnParents(const UE::UnifiedError::FError& Error);
 
 	/** Delegate to call when streaming is completed */
@@ -895,7 +895,7 @@ private:
 	void CheckCompletedRequests(const FSoftObjectPath& Target, struct FStreamable* Existing);
 
 	void OnPreGarbageCollect();
-	void AsyncLoadCallback(FSoftObjectPath Request, UPackage* Package, EAsyncLoadingResult::Type Result);
+	void AsyncLoadCallback(FSoftObjectPath Request, UPackage* Package, const TOptional<UE::UnifiedError::FError> AsyncLoadError);
 
 	ENGINE_API static bool ShouldStripDebugName();
 	ENGINE_API TSharedPtr<FStreamableHandle> RequestSyncLoadInternal(TArray<FSoftObjectPath>&& TargetsToStream, bool bManageActiveHandle, FString&& DebugName, UE::FSourceLocation&& Location);
diff --git a/Engine/Source/Runtime/Engine/Private/Experimental/StreamableManagerError.cpp b/Engine/Source/Runtime/Engine/Private/Experimental/StreamableManagerError.cpp
--- a/Engine/Source/Runtime/Engine/Private/Experimental/StreamableManagerError.cpp
+++ b/Engine/Source/Runtime/Engine/Private/Experimental/StreamableManagerError.cpp
@@ -15,6 +15,8 @@ UE::UnifiedError::FError GetStreamableError(EAsyncLoadingResult::Type Result)
 		// We could supply an error string if async loading bubbled one up...
 		// Possibly GetExplanationForUnavailablePackage?
 		return AsyncLoadFailed::MakeError();
+	/*case EAsyncLoadingResult::FailedNotInstalled:
+		return AsyncLoadNotInstalled::MakeError();*/
 	case EAsyncLoadingResult::Canceled:
 		return AsyncLoadCancelled::MakeError();
 	}
diff --git a/Engine/Source/Runtime/Engine/Private/StreamableManager.cpp b/Engine/Source/Runtime/Engine/Private/StreamableManager.cpp
--- a/Engine/Source/Runtime/Engine/Private/StreamableManager.cpp
+++ b/Engine/Source/Runtime/Engine/Private/StreamableManager.cpp
@@ -23,8 +23,6 @@
 DEFINE_LOG_CATEGORY_STATIC(LogStreamableManager, Log, All);
 
 
-// DEFINE_ERROR_DETAILS(StreamableManager, FStreamableManagerErrorDetails);
-
 // Default to 1 frame, this will cause the delegates to go off on the next tick to avoid recursion issues. Set higher to fake disk lag
 static int32 GStreamableDelegateDelayFrames = 1;
 static FAutoConsoleVariableRef CVarStreamableDelegateDelayFrames(
@@ -1191,13 +1189,12 @@ void FStreamableHandle::UnbindDelegates()
 	CompleteDelegate.Unbind();
 }
 
-void FStreamableHandle::SetError(UE::UnifiedError::FError&& InError)
+void FStreamableHandle::SetError(const UE::UnifiedError::FError& InError)
 {
 	if (!Error)
 	{
 		SetErrorOnParents(InError);
-		Error = MoveTemp(InError);
-		
+		Error = InError;
 	}
 }
 
@@ -1210,7 +1207,7 @@ void FStreamableHandle::SetErrorOnParents(const UE::UnifiedError::FError& InErro
 
 		if (Handle.IsValid())
 		{
-			Handle->SetError(CopyTemp(InError));
+			Handle->SetError(InError);
 		}
 	}
 }
@@ -1222,17 +1219,18 @@ void FStreamableHandle::AsyncLoadCallbackWrapper(const FName& PackageName, UPack
 	// Needed so we can bind with a shared pointer for safety
 	if (OwningManager)
 	{
+		TOptional<UE::UnifiedError::FError> MaybeError;
+
 		if (Result != EAsyncLoadingResult::Succeeded)
 		{
-			UE::UnifiedError::FError NewError = UE::UnifiedError::StreamableManager::GetStreamableError(Result);
+			MaybeError = UE::UnifiedError::StreamableManager::GetStreamableError(Result);
+			TRefCountPtr<UE::UnifiedError::IErrorDetails> ErrorDetails(new UE::UnifiedError::StreamableManager::FStreamableManagerErrorDetails(TargetName.ToString()));
+			MaybeError->PushErrorDetails(ErrorDetails);
 
-			TRefCountPtr<UE::UnifiedError::IErrorDetails> NewErrorDetails( new UE::UnifiedError::StreamableManager::FStreamableManagerErrorDetails(TargetName.ToString()));
-
-			NewError.PushErrorDetails(NewErrorDetails);
-			SetError(MoveTemp(NewError));
+			SetError(*MaybeError);
 		}
 
-		OwningManager->AsyncLoadCallback(MoveTemp(TargetName), Package, Result);
+		OwningManager->AsyncLoadCallback(MoveTemp(TargetName), Package, MaybeError);
 
 		if (!HasLoadCompleted())
 		{
@@ -2080,7 +2078,7 @@ void FStreamableManager::FindInMemory(FSoftObjectPath& InOutTargetName, struct F
 	}
 }
 
-void FStreamableManager::AsyncLoadCallback(FSoftObjectPath TargetName, UPackage* Package, EAsyncLoadingResult::Type Result)
+void FStreamableManager::AsyncLoadCallback(FSoftObjectPath TargetName, UPackage* Package, const TOptional<UE::UnifiedError::FError> AsyncLoadError)
 {
 	check(IsInGameThread());
 
@@ -2115,12 +2113,11 @@ void FStreamableManager::AsyncLoadCallback(FSoftObjectPath TargetName, UPackage*
 			UE_LOG(LogStreamableManager, Verbose, TEXT("    Failed async load for %s"), *TargetName.ToString());
 		}
 
-		if (Result != EAsyncLoadingResult::Succeeded)
+		if (AsyncLoadError)
 		{
 			for (const TSharedRef<FStreamableHandle>& Handle : Existing->LoadingHandles)
 			{
-				// @Per @Justin In my testing this has already been set by the AsyncLoadCallbackWrapper.  does this need to happen twice?
-				Handle->SetError(UE::UnifiedError::StreamableManager::GetStreamableError(Result));
+				Handle->SetError(*AsyncLoadError);
 			}
 		}
 		else if (!Existing->Target)
