Fix spline serialization issues

diff --git a/Engine/Source/Runtime/Engine/Private/Curves/Spline.cpp b/Engine/Source/Runtime/Engine/Private/Curves/Spline.cpp
--- a/Engine/Source/Runtime/Engine/Private/Curves/Spline.cpp
+++ b/Engine/Source/Runtime/Engine/Private/Curves/Spline.cpp
@@ -1251,6 +1251,7 @@ FNewSpline::FNewSpline(const FLegacySpline& Other)
 		AddAttributeValue(RotationAttrName, Other.GetRotation(Idx), AttributeParameter);
 		AddAttributeValue(ScaleAttrName, Other.GetScale(Idx), AttributeParameter);
 	}
+	
     UpdateSpline();
 }
 
@@ -1261,8 +1262,7 @@ FNewSpline::FNewSpline(const FSplineCurves& Other)
 	
     for (int32 Idx = 0; Idx < Other.Position.Points.Num(); Idx++)
     {
-		const auto& Position = Other.Position.Points[Idx];
-    	
+	    const FInterpCurvePoint<FVector>& Position = Other.Position.Points[Idx];
     	
     	FTangentBezierControlPoint& Point = Points.AddDefaulted_GetRef();
 
@@ -1283,8 +1283,8 @@ FNewSpline::FNewSpline(const FSplineCurves& Other)
     // setup rotation/scale attributes
     for (int32 Idx = 0; Idx < GetNumPoints(); Idx++)
     {
-    	const auto& Rotation = Other.Rotation.Points[Idx];
-        const auto& Scale = Other.Scale.Points[Idx];
+    	const FQuat& Rotation = Other.Rotation.Points[Idx].OutVal;
+        const FVector& Scale = Other.Scale.Points[Idx].OutVal;
     	
     	const float AttributeParameter = NormalizeAttributeParameter(Idx);
     	AddAttributeValue(RotationAttrName, Rotation, AttributeParameter);
@@ -1416,6 +1416,10 @@ void FNewSpline::ResetToDefault()
 {
     SetControlPoints( { FTangentBezierControlPoint(FVector(0.0, 0.0, 0.0)),
     					FTangentBezierControlPoint(FVector(100.0, 0.0, 0.0)) } );
+
+	ResetRotation();
+	ResetScale();
+	
     UpdateSpline();
 }
 
diff --git a/Engine/Source/Runtime/GeometryCore/Public/Splines/Spline.h b/Engine/Source/Runtime/GeometryCore/Public/Splines/Spline.h
--- a/Engine/Source/Runtime/GeometryCore/Public/Splines/Spline.h
+++ b/Engine/Source/Runtime/GeometryCore/Public/Splines/Spline.h
@@ -211,22 +211,7 @@ public:
 		int32 NumChannels = Spline.AttributeChannels.Num();
 		Ar << NumChannels;
         
-		if (Ar.IsSaving())
-		{
-			for (auto& Pair : Spline.AttributeChannels)
-			{
-				// Save channel name and data
-				Ar << Pair.Key;
-				
-				// Save type name explicitly - critical for correct deserialization
-				FName ChannelTypeName = Pair.Value->GetTypeName();
-				Ar << ChannelTypeName;
-            
-				// Now serialize the channel data
-				Pair.Value->Serialize(Ar);
-			}
-		}
-		else
+		if (Ar.IsLoading())
 		{
 			Spline.AttributeChannels.Empty();
             
@@ -263,10 +248,25 @@ public:
 				}
 			}
 			// After loading all channels, update their parent evaluator references
-            for (auto& Pair : Spline.AttributeChannels)
-            {
-                Pair.Value->UpdateEvaluatorParent(&Spline.GetEvaluator());
-            }
+			for (auto& Pair : Spline.AttributeChannels)
+			{
+				Pair.Value->UpdateEvaluatorParent(&Spline.GetEvaluator());
+			}
+		}
+		else
+		{
+			for (auto& Pair : Spline.AttributeChannels)
+			{
+				// Save channel name and data
+				Ar << Pair.Key;
+				
+				// Save type name explicitly - critical for correct deserialization
+				FName ChannelTypeName = Pair.Value->GetTypeName();
+				Ar << ChannelTypeName;
+            
+				// Now serialize the channel data
+				Pair.Value->Serialize(Ar);
+			}
 		}
 		        
 		return Ar;
