[AsyncCompilation/StaticMesh]
- Make static mesh async compilation cancelable by default after more than 2h of unittest soaking repeateadly triggering cancelation at random to validate all cancelation points stability

diff --git a/Engine/Source/Runtime/Engine/Private/StaticMeshCompiler.cpp b/Engine/Source/Runtime/Engine/Private/StaticMeshCompiler.cpp
--- a/Engine/Source/Runtime/Engine/Private/StaticMeshCompiler.cpp
+++ b/Engine/Source/Runtime/Engine/Private/StaticMeshCompiler.cpp
@@ -60,7 +60,7 @@ static TAutoConsoleVariable<bool> CVarAsyncStaticMeshDebugDraw(
 	TEXT("Any static meshes that were waited on due to being too close to the player will have their bounding box drawn in red for a couple of seconds."),
 	ECVF_Default);
 
-static bool GAsyncStaticMeshCompilationCancelable = false;
+static bool GAsyncStaticMeshCompilationCancelable = true;
 static FAutoConsoleVariableRef CVarAsyncStaticMeshCompilationCancelable(
 	TEXT("Editor.AsyncStaticMeshCompilationCancelable"),
 	GAsyncStaticMeshCompilationCancelable,
