Added bShouldMatchSyncPhases to UAnimBlendspace.

Enabling this flag ensures that all the followers should match the sync phases that the leader samples is at, meaning if their leader is on the third sync phase then the follows must ensure they are on their third sync phase as well.

#rnx

diff --git a/Engine/Source/Runtime/Engine/Classes/Animation/BlendSpace.h b/Engine/Source/Runtime/Engine/Classes/Animation/BlendSpace.h
--- a/Engine/Source/Runtime/Engine/Classes/Animation/BlendSpace.h
+++ b/Engine/Source/Runtime/Engine/Classes/Animation/BlendSpace.h
@@ -812,6 +812,10 @@ public:
 	UPROPERTY(EditAnywhere, Category = Animation)
 	bool bAllowMarkerBasedSync = true;
 
+	/** If true, all follower samples will pass the same amount of markers the leader sample has passed to match its sync phase. Otherwise, followers samples will only match their next valid sync position. */
+	UPROPERTY(EditAnywhere, Category = Animation)
+	bool bShouldMatchSyncPhases = false;
+	
 #if WITH_EDITORONLY_DATA
 	/** Preview Base pose for additive BlendSpace **/
 	UPROPERTY(EditAnywhere, Category = AdditiveSettings)
diff --git a/Engine/Source/Runtime/Engine/Private/Animation/BlendSpace.cpp b/Engine/Source/Runtime/Engine/Private/Animation/BlendSpace.cpp
--- a/Engine/Source/Runtime/Engine/Private/Animation/BlendSpace.cpp
+++ b/Engine/Source/Runtime/Engine/Private/Animation/BlendSpace.cpp
@@ -589,57 +589,107 @@ void UBlendSpace::TickAssetPlayer(FAnimTickRecord& Instance, struct FAnimNotifyQ
 			// Tick as leader using marked based syncing.
 			if (bCanDoMarkerSync)
 			{
-				FBlendSampleData& LeaderSampleData = SampleDataList[HighestMarkerSyncWeightIndex];
-				const FBlendSample& LeaderSample = SampleData[LeaderSampleData.SampleDataIndex];
-
-				if (LeaderSample.Animation)
+				if (bShouldMatchSyncPhases)
 				{
-					bool bResetMarkerDataOnFollowers = false;
-
-					// Invalidate sample followers' tick records if instance doesn't have any valid sync marker data. 
-					if (!Instance.MarkerTickRecord->IsValid(Instance.bLooping))
+					FBlendSampleData& LeaderSampleData = SampleDataList[HighestMarkerSyncWeightIndex];
+					const FBlendSample& LeaderSample = SampleData[LeaderSampleData.SampleDataIndex];
+					
+					if (LeaderSample.Animation)
 					{
-						LeaderSampleData.MarkerTickRecord.Reset();
-						LeaderSampleData.Time = NormalizedCurrentTime * LeaderSample.Animation->GetPlayLength();
-						bResetMarkerDataOnFollowers = true;
+						// Update leader delta.
+						float NewDelta = Instance.DeltaTimeRecord->Delta * LeaderSample.RateScale * LeaderSample.Animation->RateScale;
+						Context.SetLeaderDelta(NewDelta);
+						
+						// Advance time as usual.
+						{
+							float CurrentTime = NormalizedCurrentTime * NewAnimLength;
+                        	LeaderSampleData.PreviousTime = CurrentTime;
+                        	
+                        	FAnimationRuntime::AdvanceTime(Instance.bLooping, NewDelta, /*inout*/ CurrentTime, NewAnimLength);
+                        	
+                        	LeaderSampleData.Time = CurrentTime;
+                        	NormalizedCurrentTime = NewAnimLength ? (CurrentTime / NewAnimLength) : 0.0f;
+						}
+
+						// Sync position BEFORE advancing time.
+						LeaderSample.Animation->GetMarkerIndicesForTime(LeaderSampleData.PreviousTime, Instance.bLooping, Context.MarkerTickContext.GetValidMarkerNames(), LeaderSampleData.MarkerTickRecord.PreviousMarker, LeaderSampleData.MarkerTickRecord.NextMarker);
+						const FMarkerSyncAnimPosition StartSyncPosition = LeaderSample.Animation->GetMarkerSyncPositionFromMarkerIndicies(LeaderSampleData.MarkerTickRecord.PreviousMarker.MarkerIndex, LeaderSampleData.MarkerTickRecord.NextMarker.MarkerIndex, LeaderSampleData.PreviousTime, Instance.MirrorDataTable);
+						Context.MarkerTickContext.SetMarkerSyncStartPosition(StartSyncPosition);
+
+						// @todo: See if we can have reset blend samples directly update the context and/or tick record (no need to compute info twice, if possible). 
+						// Rest all samples to match current time.
+						ResetBlendSamples(SampleDataList, NormalizedCurrentTime, Instance.bLooping, true);
+
+						// Sync position AFTER advancing time.
+						LeaderSample.Animation->GetMarkerIndicesForTime(LeaderSampleData.Time, Instance.bLooping, Context.MarkerTickContext.GetValidMarkerNames(), LeaderSampleData.MarkerTickRecord.PreviousMarker, LeaderSampleData.MarkerTickRecord.NextMarker);
+						*Instance.MarkerTickRecord = LeaderSampleData.MarkerTickRecord;
+						const FMarkerSyncAnimPosition EndSyncPosition = LeaderSample.Animation->GetMarkerSyncPositionFromMarkerIndicies(LeaderSampleData.MarkerTickRecord.PreviousMarker.MarkerIndex, LeaderSampleData.MarkerTickRecord.NextMarker.MarkerIndex, LeaderSampleData.Time, Instance.MirrorDataTable);
+						Context.MarkerTickContext.SetMarkerSyncEndPosition(EndSyncPosition);
+
+						// Update context's position ratios
+						Context.SetAnimationPositionRatio(NormalizedCurrentTime);
+						Context.SetPreviousAnimationPositionRatio(NormalizedPreviousTime);
+						
+						// Keep track of the markers we passed.
+						FMarkerSyncData SyncData;
+						SyncData.AuthoredSyncMarkers = LeaderSample.Animation->AuthoredSyncMarkers;
+						SyncData.CollectMarkersInRange(LeaderSampleData.PreviousTime, LeaderSampleData.Time, Context.MarkerTickContext.MarkersPassedThisTick, NewDelta);
 					}
-					// Re-compute marker indices since the leader sample's tick record is invalid. Get previous and next markers.
-					else if (!LeaderSampleData.MarkerTickRecord.IsValid(Instance.bLooping) && Context.MarkerTickContext.GetMarkerSyncStartPosition().IsValid())
+				}
+				else
+				{
+					FBlendSampleData& LeaderSampleData = SampleDataList[HighestMarkerSyncWeightIndex];
+					const FBlendSample& LeaderSample = SampleData[LeaderSampleData.SampleDataIndex];
+
+					if (LeaderSample.Animation)
 					{
-						// TODO: Look into the reason for not passing bLooping variable and just forcing the vale to be true. 
-						LeaderSample.Animation->GetMarkerIndicesForPosition(Context.MarkerTickContext.GetMarkerSyncStartPosition(), true, LeaderSampleData.MarkerTickRecord.PreviousMarker, LeaderSampleData.MarkerTickRecord.NextMarker, LeaderSampleData.Time, Instance.MirrorDataTable);
-					}
+						bool bResetMarkerDataOnFollowers = false;
 
-					// Only tick samples if leader sample has any delta time to consume.
-					const float NewDeltaTime = Instance.DeltaTimeRecord->Delta * LeaderSample.RateScale * LeaderSample.Animation->RateScale;
-					Context.SetLeaderDelta(NewDeltaTime);
+						// Invalidate sample followers' tick records if instance doesn't have any valid sync marker data. 
+						if (!Instance.MarkerTickRecord->IsValid(Instance.bLooping))
+						{
+							LeaderSampleData.MarkerTickRecord.Reset();
+							LeaderSampleData.Time = NormalizedCurrentTime * LeaderSample.Animation->GetPlayLength();
+							bResetMarkerDataOnFollowers = true;
+						}
+						// Re-compute marker indices since the leader sample's tick record is invalid. Get previous and next markers.
+						else if (!LeaderSampleData.MarkerTickRecord.IsValid(Instance.bLooping) && Context.MarkerTickContext.GetMarkerSyncStartPosition().IsValid())
+						{
+							// TODO: Look into the reason for not passing bLooping variable and just forcing the vale to be true. 
+							LeaderSample.Animation->GetMarkerIndicesForPosition(Context.MarkerTickContext.GetMarkerSyncStartPosition(), true, LeaderSampleData.MarkerTickRecord.PreviousMarker, LeaderSampleData.MarkerTickRecord.NextMarker, LeaderSampleData.Time, Instance.MirrorDataTable);
+						}
 
-					if (!FMath::IsNearlyZero(NewDeltaTime))
-					{
-						// Tick leader sample
-						LeaderSample.Animation->TickByMarkerAsLeader(LeaderSampleData.MarkerTickRecord, Context.MarkerTickContext, LeaderSampleData.Time, LeaderSampleData.PreviousTime, NewDeltaTime, Instance.bLooping, Instance.MirrorDataTable);
+						// Only tick samples if leader sample has any delta time to consume.
+						const float NewDeltaTime = Instance.DeltaTimeRecord->Delta * LeaderSample.RateScale * LeaderSample.Animation->RateScale;
+						Context.SetLeaderDelta(NewDeltaTime);
+
+						if (!FMath::IsNearlyZero(NewDeltaTime))
+						{
+							// Tick leader sample
+							LeaderSample.Animation->TickByMarkerAsLeader(LeaderSampleData.MarkerTickRecord, Context.MarkerTickContext, LeaderSampleData.Time, LeaderSampleData.PreviousTime, NewDeltaTime, Instance.bLooping, Instance.MirrorDataTable);
 
-						check(!Instance.bLooping || Context.MarkerTickContext.IsMarkerSyncStartValid());
+							check(!Instance.bLooping || Context.MarkerTickContext.IsMarkerSyncStartValid());
 							
-						// Tick all the follower samples
-						TickFollowerSamples(SampleDataList, HighestMarkerSyncWeightIndex, Context, bResetMarkerDataOnFollowers, Instance.bLooping, Instance.MirrorDataTable);
-					}
-					else if (!Instance.MarkerTickRecord->IsValid(Instance.bLooping))
-					{
-						// Re-compute marker indices for leader sample's tick record. Get previous and next markers.
-						LeaderSample.Animation->GetMarkerIndicesForTime(LeaderSampleData.Time, Instance.bLooping, Context.MarkerTickContext.GetValidMarkerNames(), LeaderSampleData.MarkerTickRecord.PreviousMarker, LeaderSampleData.MarkerTickRecord.NextMarker);
+							// Tick all the follower samples
+							TickFollowerSamples(SampleDataList, HighestMarkerSyncWeightIndex, Context, bResetMarkerDataOnFollowers, Instance.bLooping, Instance.MirrorDataTable);
+						}
+						else if (!Instance.MarkerTickRecord->IsValid(Instance.bLooping))
+						{
+							// Re-compute marker indices for leader sample's tick record. Get previous and next markers.
+							LeaderSample.Animation->GetMarkerIndicesForTime(LeaderSampleData.Time, Instance.bLooping, Context.MarkerTickContext.GetValidMarkerNames(), LeaderSampleData.MarkerTickRecord.PreviousMarker, LeaderSampleData.MarkerTickRecord.NextMarker);
 
-						// Get sync position for followers to sync up to.
-						const FMarkerSyncAnimPosition SyncPosition = LeaderSample.Animation->GetMarkerSyncPositionFromMarkerIndicies(LeaderSampleData.MarkerTickRecord.PreviousMarker.MarkerIndex, LeaderSampleData.MarkerTickRecord.NextMarker.MarkerIndex, LeaderSampleData.Time, Instance.MirrorDataTable);
-						Context.MarkerTickContext.SetMarkerSyncStartPosition(SyncPosition);
-						Context.MarkerTickContext.SetMarkerSyncEndPosition(SyncPosition);
+							// Get sync position for followers to sync up to.
+							const FMarkerSyncAnimPosition SyncPosition = LeaderSample.Animation->GetMarkerSyncPositionFromMarkerIndicies(LeaderSampleData.MarkerTickRecord.PreviousMarker.MarkerIndex, LeaderSampleData.MarkerTickRecord.NextMarker.MarkerIndex, LeaderSampleData.Time, Instance.MirrorDataTable);
+							Context.MarkerTickContext.SetMarkerSyncStartPosition(SyncPosition);
+							Context.MarkerTickContext.SetMarkerSyncEndPosition(SyncPosition);
 							
-						// Make all follower samples match next sync position to equal that of the leader.
-						TickFollowerSamples(SampleDataList, HighestMarkerSyncWeightIndex, Context, true, Instance.bLooping, Instance.MirrorDataTable);
-					}
+							// Make all follower samples match next sync position to equal that of the leader.
+							TickFollowerSamples(SampleDataList, HighestMarkerSyncWeightIndex, Context, true, Instance.bLooping, Instance.MirrorDataTable);
+						}
 						
-					NormalizedCurrentTime = LeaderSampleData.Time / LeaderSample.Animation->GetPlayLength();
-					*Instance.MarkerTickRecord = LeaderSampleData.MarkerTickRecord;
+						NormalizedCurrentTime = LeaderSampleData.Time / LeaderSample.Animation->GetPlayLength();
+						*Instance.MarkerTickRecord = LeaderSampleData.MarkerTickRecord;
+					}
 				}
 			}
 			else
