Fixing issue with FScopedMovementUpdate::RevertMove not updating the cached rotations of the Scenecomponent which could lead to some bad rotation applied to the component after reverting movement with rotation and then applying a different move

diff --git a/Engine/Source/Runtime/Engine/Private/Engine/ScopedMovementUpdate.cpp b/Engine/Source/Runtime/Engine/Private/Engine/ScopedMovementUpdate.cpp
--- a/Engine/Source/Runtime/Engine/Private/Engine/ScopedMovementUpdate.cpp
+++ b/Engine/Source/Runtime/Engine/Private/Engine/ScopedMovementUpdate.cpp
@@ -89,7 +89,9 @@ void FScopedMovementUpdate::RevertMove()
 			Component->RelativeLocation = InitialRelativeLocation;
 			Component->RelativeRotation = InitialRelativeRotation;
 			Component->RelativeScale3D = InitialRelativeScale;
-
+			Component->RelativeRotationCache.RotatorToQuat(InitialRelativeRotation);
+			Component->WorldRotationCache.RotatorToQuat(InitialTransform.Rotator());
+			
 			if (!IsDeferringUpdates())
 			{
 				Component->PropagateTransformUpdate(true);
