[UENum Markup for Audio Bus]

adding 5.1 and 7.1 markup for audio busses to reflect common string representation of those channel formats elsewhere in the engine

[RN] minor audio

diff --git a/Engine/Source/Runtime/Engine/Classes/Sound/AudioBus.h b/Engine/Source/Runtime/Engine/Classes/Sound/AudioBus.h
--- a/Engine/Source/Runtime/Engine/Classes/Sound/AudioBus.h
+++ b/Engine/Source/Runtime/Engine/Classes/Sound/AudioBus.h
@@ -15,8 +15,8 @@ enum class EAudioBusChannels : uint8
 	Mono = 0,
 	Stereo = 1,
 	Quad = 3,
-	FivePointOne = 5,
-	SevenPointOne = 7,
+	FivePointOne = 5 UMETA(DisplayName = "5.1"),
+	SevenPointOne = 7 UMETA(DisplayName = "7.1"),
 	MaxChannelCount = 8
 };
 
