Determinism fix supplied by partners at The Coalition:
Fixing non-determinism in MaterialInstanceConstant assets. Sorting of TextureStreamingData was non-deterministic after multi-platform cooks when there were multiple entries with the same texture name and index but with different sampling scales, as each PreSave was arbitrary swapping the elements. Switched to StableSort.

diff --git a/Engine/Source/Runtime/Engine/Private/Materials/MaterialInterface.cpp b/Engine/Source/Runtime/Engine/Private/Materials/MaterialInterface.cpp
--- a/Engine/Source/Runtime/Engine/Private/Materials/MaterialInterface.cpp
+++ b/Engine/Source/Runtime/Engine/Private/Materials/MaterialInterface.cpp
@@ -1802,7 +1802,8 @@ void UMaterialInterface::SortTextureStreamingData(bool bForceSort, bool bFinalSo
 		}
 
 		// Sort by name to be compatible with FindTextureStreamingDataIndexRange
-		TextureStreamingData.Sort([](const FMaterialTextureInfo& Lhs, const FMaterialTextureInfo& Rhs) 
+		// Use stable sort to preserve order and avoid non-determinism from multi-platform cooks.
+		TextureStreamingData.StableSort([](const FMaterialTextureInfo& Lhs, const FMaterialTextureInfo& Rhs) 
 		{ 
 			// Sort by register indices when the name are the same, as when initially added in the streaming data.
 			if (Lhs.TextureName == Rhs.TextureName)
