Disable config tracker in UDeviceProfileManager::LoadProfiles.
LoadProfiles is called only once when the singleton for UDeviceProfileManager is created. The singleton is created when class default objects are created so there is no need to track config access at this moment.

#rnx
#incrementalcook

diff --git a/Engine/Source/Runtime/Engine/Private/DeviceProfiles/DeviceProfileManager.cpp b/Engine/Source/Runtime/Engine/Private/DeviceProfiles/DeviceProfileManager.cpp
--- a/Engine/Source/Runtime/Engine/Private/DeviceProfiles/DeviceProfileManager.cpp
+++ b/Engine/Source/Runtime/Engine/Private/DeviceProfiles/DeviceProfileManager.cpp
@@ -1200,6 +1200,8 @@ void UDeviceProfileManager::LoadProfiles()
 {
 	if( !HasAnyFlags( RF_ClassDefaultObject ) )
 	{
+		UE::ConfigAccessTracking::FIgnoreScope IgnoreConfigAccessTracking = UE::ConfigAccessTracking::FIgnoreScope();
+
 		TMap<FString, FString> DeviceProfileToPlatformConfigMap;
 		DeviceProfileToPlatformConfigMap.Add(TEXT("GlobalDefaults,None"), FPlatformProperties::IniPlatformName());
 		
