Skeletal Mesh: Ensure that vertex colors get enabled if they got set on the incoming mesh description

#rnx

diff --git a/Engine/Source/Developer/SkeletalMeshUtilitiesCommon/Private/StaticToSkeletalMeshConverter.cpp b/Engine/Source/Developer/SkeletalMeshUtilitiesCommon/Private/StaticToSkeletalMeshConverter.cpp
--- a/Engine/Source/Developer/SkeletalMeshUtilitiesCommon/Private/StaticToSkeletalMeshConverter.cpp
+++ b/Engine/Source/Developer/SkeletalMeshUtilitiesCommon/Private/StaticToSkeletalMeshConverter.cpp
@@ -631,12 +631,6 @@ bool FStaticToSkeletalMeshConverter::InitializeSkeletalMeshFromMeshDescriptions(
 		InSkeletalMesh->SetMaterials(NewMaterials);
 	}
 
-	if (HasVertexColors(InSkeletalMesh))
-	{
-		InSkeletalMesh->SetHasVertexColors(true);	
-		InSkeletalMesh->SetVertexColorGuid(FGuid::NewGuid());
-	}
-	
 	InSkeletalMesh->SetImportedBounds( FBox3d{BoundingBox} );
 
 	return true;
diff --git a/Engine/Source/Runtime/Engine/Classes/Engine/SkeletalMesh.h b/Engine/Source/Runtime/Engine/Classes/Engine/SkeletalMesh.h
--- a/Engine/Source/Runtime/Engine/Classes/Engine/SkeletalMesh.h
+++ b/Engine/Source/Runtime/Engine/Classes/Engine/SkeletalMesh.h
@@ -605,6 +605,9 @@ public:
 		 */
 		bool bUpdateVertexAttributes = true;
 		
+		/** Update the vertex color enabled state/GUID if the mesh description contains non-default vertex colors */
+		bool bUpdateVertexColors = true;
+		
 		/** Mark the package as dirty. If calling CommitMeshDescription from a non-game thread,
 		 *  this value should be set to \c false.
 		 */
diff --git a/Engine/Source/Runtime/Engine/Private/SkeletalMesh.cpp b/Engine/Source/Runtime/Engine/Private/SkeletalMesh.cpp
--- a/Engine/Source/Runtime/Engine/Private/SkeletalMesh.cpp
+++ b/Engine/Source/Runtime/Engine/Private/SkeletalMesh.cpp
@@ -2735,9 +2735,9 @@ bool USkeletalMesh::CommitMeshDescription(
 			// Get all profiles from the models on all LODs, since we may have some that aren't
 			// defined on the skeletal mesh's list of profiles.
 			TSet<FName> ValidProfileNames;
-			for (const FSkeletalMeshSourceModel& OtherSourceModels: GetAllSourceModels())
+			for (const FSkeletalMeshSourceModel& OtherSourceModel: GetAllSourceModels())
 			{
-				ValidProfileNames.Append(OtherSourceModels.GetSkinWeightProfileNames());
+				ValidProfileNames.Append(OtherSourceModel.GetSkinWeightProfileNames());
 			}
 			for (const FName& ProfileName: ValidProfileNames)
 			{
@@ -2812,6 +2812,68 @@ bool USkeletalMesh::CommitMeshDescription(
 					}, TStatId(), nullptr, ENamedThreads::GameThread);	
 #endif			
 			}
+			
+			if (InParams.bUpdateVertexColors)
+			{
+				struct FContext
+				{
+					bool bHasVertexColors = false;
+					FXxHash128 Hash;
+				};
+				TArray<FContext> VertexColorContexts;
+				TConstArrayView<FSkeletalMeshSourceModel> AllSourceModels = GetAllSourceModels();
+				ParallelForWithTaskContext(VertexColorContexts, AllSourceModels.Num(), [Models=AllSourceModels](FContext &InOutContext, int32 InIndex)
+				{
+					const FSkeletalMeshSourceModel& SourceModel = Models[InIndex];
+					if (!SourceModel.HasMeshDescription())
+					{
+						return;
+					}
+							
+					const FMeshDescription& MeshDescription = *SourceModel.GetMeshDescription();
+					FSkeletalMeshConstAttributes Attributes(MeshDescription);
+
+					TVertexInstanceAttributesConstRef<FVector4f> ColorAttrib = Attributes.GetVertexInstanceColors();
+					FXxHash128Builder HashBuilder;
+					for (FVertexInstanceID VertexInstanceID : MeshDescription.VertexInstances().GetElementIDs())
+					{
+						const FVector4f VertexColor = ColorAttrib.Get(VertexInstanceID);
+						if (!VertexColor.IsNearlyZero())
+						{
+							InOutContext.bHasVertexColors = true;
+						}
+						
+						HashBuilder.Update(&VertexColor, sizeof(FVector4f));
+					}
+
+					InOutContext.Hash = HashBuilder.Finalize();
+				});
+
+				bool bSourceHasVertexColors = false;
+				FXxHash128Builder HashBuilder;
+				
+				for (const FContext& Context: VertexColorContexts)
+				{
+					if (Context.bHasVertexColors)
+					{
+						HashBuilder.Update(&Context.Hash, sizeof(FXxHash128));
+						bSourceHasVertexColors = true;
+					}
+				}
+				
+				if (bSourceHasVertexColors)
+				{
+					FXxHash128 Hash = HashBuilder.Finalize();
+					FGuid SourceVertexColorGuid(
+						static_cast<uint32>(Hash.HashHigh >> 32),
+						static_cast<uint32>(Hash.HashHigh),
+						static_cast<uint32>(Hash.HashLow >> 32),
+						static_cast<uint32>(Hash.HashLow)
+						);
+					SetVertexColorGuid(SourceVertexColorGuid);
+				}
+				SetHasVertexColors(bSourceHasVertexColors);
+			}
 		}
 	}
 
