Fix Bug where Modulation does not properly inherit settings from Sound Class

diff --git a/Engine/Source/Runtime/AudioMixer/Private/AudioMixerSource.cpp b/Engine/Source/Runtime/AudioMixer/Private/AudioMixerSource.cpp
--- a/Engine/Source/Runtime/AudioMixer/Private/AudioMixerSource.cpp
+++ b/Engine/Source/Runtime/AudioMixer/Private/AudioMixerSource.cpp
@@ -88,17 +88,6 @@ namespace Audio
 					return EMixerSourceSubmixSendStage::PostDistanceAttenuation;
 			}
 		}
-
-		const USoundClass* GetFallbackSoundClass(const FActiveSound& InActiveSound, const FWaveInstance& InWaveInstance)
-		{
-			const USoundClass* SoundClass = InActiveSound.GetSoundClass();
-			if (InWaveInstance.SoundClass)
-			{
-				SoundClass = InWaveInstance.SoundClass;
-			}
-
-			return SoundClass;
-		}
 		
 	} // namespace MixerSourcePrivate
 
diff --git a/Engine/Source/Runtime/Engine/Private/Audio.cpp b/Engine/Source/Runtime/Engine/Private/Audio.cpp
--- a/Engine/Source/Runtime/Engine/Private/Audio.cpp
+++ b/Engine/Source/Runtime/Engine/Private/Audio.cpp
@@ -245,8 +245,12 @@ void FWaveInstance::CalculateInheritedModulationSettingsAndValue(const EModulati
 	const FSoundModulationDestinationSettings& InWaveSettings = WaveData->ModulationSettings.GetSettingsForDestination(InDestination);
 	EModulationRouting InWaveRouting = WaveData->ModulationSettings.GetRoutingForDestination(InDestination);
 
-	const USoundClass* InSoundClass = InActiveSound.GetSoundClass();
-
+	const USoundClass* EffectiveSoundClass = InActiveSound.GetSoundClass();
+	if (SoundClass)
+	{
+		EffectiveSoundClass = SoundClass;
+	}
+	
 	auto MixInRoutedValue = [&InParam](float& InOutValueA, float InValueB)
 	{
 		if (InParam.bRequiresConversion)
@@ -263,9 +267,9 @@ void FWaveInstance::CalculateInheritedModulationSettingsAndValue(const EModulati
 	
 	auto UnionSoundClassSettings = [&](FSoundModulationDestinationSettings& Settings)
 		{
-			if (InSoundClass)
+			if (EffectiveSoundClass)
 			{
-				const FSoundModulationDestinationSettings& ClassSettings = InSoundClass->Properties.ModulationSettings.GetSettingsForDestination(InDestination);
+				const FSoundModulationDestinationSettings& ClassSettings = EffectiveSoundClass->Properties.ModulationSettings.GetSettingsForDestination(InDestination);
 				OutValue = Settings.Value;
 				MixInRoutedValue(OutValue, ClassSettings.Value);
 				Settings.Modulators = Settings.Modulators.Union(ClassSettings.Modulators);
@@ -326,9 +330,9 @@ void FWaveInstance::CalculateInheritedModulationSettingsAndValue(const EModulati
 
 		case EModulationRouting::Inherit:
 		{
-			if (InSoundClass)
+			if (EffectiveSoundClass)
 			{
-				FSoundModulationDestinationSettings InheritSettings = InSoundClass->Properties.ModulationSettings.GetSettingsForDestination(InDestination);
+				FSoundModulationDestinationSettings InheritSettings = EffectiveSoundClass->Properties.ModulationSettings.GetSettingsForDestination(InDestination);
 				InOutSettings = InheritSettings;
 				OutValue = InheritSettings.Value;
 			}
