FNewSpline: Fix parameter space mapping issues

diff --git a/Engine/Source/Runtime/Engine/Private/Curves/Spline.cpp b/Engine/Source/Runtime/Engine/Private/Curves/Spline.cpp
--- a/Engine/Source/Runtime/Engine/Private/Curves/Spline.cpp
+++ b/Engine/Source/Runtime/Engine/Private/Curves/Spline.cpp
@@ -185,6 +185,7 @@ public:
 
 	template <typename AttrType> TArray<FName> GetAttributeChannelNames() const;
 	template <typename AttrType> AttrType EvaluateAttribute(FName Name, float Parameter) const;
+	template <typename AttrType> int32 AddAttributeValue(FName Name, const AttrType& Value, float Parameter);
 	int32 SetAttributeParameter(const FName& Name, int32 Index, float ParentSpaceParameter);
 	float GetAttributeParameter(const FName& Name, int32 Index) const;
 	
@@ -217,15 +218,25 @@ private:
 
 private:
 
+	/** Convert from [0, 1] to [0, NumSegments] */
+	float FromNormalizedSpace(float Parameter) const;
+	
+	/** Convert from [0, NumSegments] to [0, 1] */
+	float ToNormalizedSpace(float Parameter) const;
+
+	/** Convert from internal spline space to [0, NumSegments] */
+	float FromInternalSplineSpace(float Parameter) const;
+	
+	/** Convert from [0, NumSegments] to internal spline space */
+	float ToInternalSplineSpace(float Parameter) const;
+	
 	/** Clears and repopulates reparameterization attribute channels. */
 	void UpdateReparamTable();
 
 	/** Clears and repopulates PositionCurve, RotationCurve, and ScaleCurve */
 	void RebuildLegacyCurves();
 
-	/**
-     * Converts FSplinePoint to FTangentBezierControlPoint
-     */
+	/** Converts FSplinePoint to FTangentBezierControlPoint */
     FTangentBezierControlPoint ConvertToTangentBezierControlPoint(const FSplinePoint& Point) const;
 
     /**
@@ -236,12 +247,12 @@ private:
     void UpdatePointAttributes(const FSplinePoint& Point, int32 PointIndex);
 	
 	/**
-	 * Converts an index (possibly fractional) to a parameter value
+	 * Converts an index (possibly fractional) to a parameter value in internal spline space.
 	 */
 	float ConvertIndexToParameter(int32 Index, float Fraction = 0.0f) const;
 
     /**
-	 * Converts a parameter value to the closest index
+	 * Converts an internal spline space parameter value to the closest index
 	 */
 	int32 ConvertParameterToIndex(float Parameter) const;
 };
@@ -265,25 +276,27 @@ TArray<FName> FNewSpline::GetAttributeChannelNames() const
 template <typename AttrType>
 AttrType FNewSpline::EvaluateAttribute(FName Name, float Parameter) const
 {
-	Parameter = GetParameterRange().Interpolate(Parameter / GetNumSegments());
+	Parameter = ToInternalSplineSpace(Parameter);
 	return InternalSpline.EvaluateAttribute<AttrType>(Name, Parameter);
 }
 
+template <typename AttrType>
+int32 FNewSpline::AddAttributeValue(FName Name, const AttrType& Value, float Parameter)
+{
+	Parameter = ToInternalSplineSpace(Parameter);
+	return InternalSpline.AddAttributeValue<AttrType>(Name, Value, Parameter);
+}
+
 int32 FNewSpline::SetAttributeParameter(const FName& Name, int32 Index, float Parameter)
 {
-  	Parameter = GetParameterRange().Interpolate(Parameter / GetNumSegments());
-	return InternalSpline.SetAttributeParameter(Name, Index, Parameter);
+	return InternalSpline.SetAttributeParameter(Name, Index, ToInternalSplineSpace(Parameter));
 }
 
 float FNewSpline::GetAttributeParameter(const FName& Name, int32 Index) const
 {
-	float SplineSpaceParam = InternalSpline.GetAttributeParameter(Name, Index);
-	const float NormalizedParam = this->GetParameterRange().GetT(SplineSpaceParam);
-	const float SegmentSpaceParam = NormalizedParam * static_cast<float>(GetNumSegments());
-	return SegmentSpaceParam;
+	return FromInternalSplineSpace(InternalSpline.GetAttributeParameter(Name, Index));
 }
 
-
 /** FLegacySpline Definition */
 
 class FLegacySpline
@@ -1216,16 +1229,18 @@ FNewSpline::FNewSpline()
 }
 
 FNewSpline::FNewSpline(const FNewSpline& Other)
-    : FNewSpline()
 {
     *this = Other;
 }
 
 FNewSpline::FNewSpline(const FLegacySpline& Other)
-	: FNewSpline()
 {
-	TArray<FTangentBezierControlPoint> Points;
+	CreateAttributeChannel<FQuat>(RotationAttrName);
+	CreateAttributeChannel<FVector>(ScaleAttrName);
+	CreateAttributeChannel<float>(ParamToLengthAttrName);
+	CreateAttributeChannel<float>(LengthToParamAttrName);
 	
+	TArray<FTangentBezierControlPoint> Points;
     for (int32 Idx = 0; Idx < Other.GetNumControlPoints(); Idx++)
     {
 	    FTangentBezierControlPoint& Point = Points.AddDefaulted_GetRef();
@@ -1239,27 +1254,27 @@ FNewSpline::FNewSpline(const FLegacySpline& Other)
 	// initialize spline
 	SetControlPoints(Points);
 	
-	auto NormalizeAttributeParameter = [MaxParameter = static_cast<float>(GetNumPoints() - 1)](int32 Index)
-	{
-		return Index / MaxParameter;
-	};
-	
 	// setup rotation/scale attributes
 	for (int32 Idx = 0; Idx < GetNumPoints(); Idx++)
 	{    		
-		const float AttributeParameter = NormalizeAttributeParameter(Idx);
-		AddAttributeValue(RotationAttrName, Other.GetRotation(Idx), AttributeParameter);
-		AddAttributeValue(ScaleAttrName, Other.GetScale(Idx), AttributeParameter);
+		AddAttributeValue(RotationAttrName, Other.GetRotation(Idx), static_cast<float>(Idx));
+		AddAttributeValue(ScaleAttrName, Other.GetScale(Idx), static_cast<float>(Idx));
 	}
+
+	FSpline::FUpdateSplineParams UpdateParams;
+	UpdateParams.bClosedLoop = Other.GetSplinePointsPosition().bIsLooped;
 	
-    UpdateSpline();
+    UpdateSpline(UpdateParams);
 }
 
 FNewSpline::FNewSpline(const FSplineCurves& Other)
-	: FNewSpline()
 {
-	TArray<FTangentBezierControlPoint> Points;
+	CreateAttributeChannel<FQuat>(RotationAttrName);
+	CreateAttributeChannel<FVector>(ScaleAttrName);
+	CreateAttributeChannel<float>(ParamToLengthAttrName);
+	CreateAttributeChannel<float>(LengthToParamAttrName);
 	
+	TArray<FTangentBezierControlPoint> Points;
     for (int32 Idx = 0; Idx < Other.Position.Points.Num(); Idx++)
     {
 	    const FInterpCurvePoint<FVector>& Position = Other.Position.Points[Idx];
@@ -1275,23 +1290,20 @@ FNewSpline::FNewSpline(const FSplineCurves& Other)
 	// initialize spline
 	SetControlPoints(Points);
 	
-	auto NormalizeAttributeParameter = [MaxParameter = static_cast<float>(GetNumPoints() - 1)](int32 Index)
-    {
-    	return Index / MaxParameter;
-    };
-	
     // setup rotation/scale attributes
     for (int32 Idx = 0; Idx < GetNumPoints(); Idx++)
     {
     	const FQuat& Rotation = Other.Rotation.Points[Idx].OutVal;
         const FVector& Scale = Other.Scale.Points[Idx].OutVal;
     	
-    	const float AttributeParameter = NormalizeAttributeParameter(Idx);
-    	AddAttributeValue(RotationAttrName, Rotation, AttributeParameter);
-    	AddAttributeValue(ScaleAttrName, Scale, AttributeParameter);
+    	AddAttributeValue(RotationAttrName, Rotation, static_cast<float>(Idx));
+    	AddAttributeValue(ScaleAttrName, Scale, static_cast<float>(Idx));
     }
 
-    UpdateSpline();
+	FSpline::FUpdateSplineParams UpdateParams;
+	UpdateParams.bClosedLoop = Other.Position.bIsLooped;
+	
+    UpdateSpline(UpdateParams);
 }
 
 /** Copy assignment */
@@ -1307,13 +1319,13 @@ FNewSpline& FNewSpline::operator=(const FNewSpline& Other)
 
     	UpdateSpline();
     }
+	
     return *this;
 }
 
 bool FNewSpline::operator==(const FNewSpline& Other) const
 {
     return InternalSpline == Other.InternalSpline && Tension == Other.Tension;
-    //return Points == Other.Points && Tension == Other.Tension && InternalSpline == Other.InternalSpline;
 }
 
 bool FNewSpline::operator!=(const FNewSpline& Other) const
@@ -1326,12 +1338,12 @@ bool FNewSpline::Serialize(FArchive& Ar)
     FTangentBezierSpline::Serialize(Ar);
 	
     Ar << ReparamStepsPerSegment;
-	Ar << bStationaryEndpoint;
 	
 	if (Ar.IsLoading())
 	{
 		UpdateSpline();
 	}
+	
     return true;
 }
 
@@ -1397,10 +1409,10 @@ float FNewSpline::GetSegmentLength(const int32 Index, const float Param, const F
 
 float FNewSpline::GetSplineLength() const
 {
-    return EvaluateAttribute<float>(ParamToLengthAttrName, 1.f);
+	// Evaluate the ParamToLength channel at the very end of the spline
+    return GetDistanceAtParameter(FromNormalizedSpace(1.0f));
 }
 
-
 void FNewSpline::SetClosedLoop(bool bInClosedLoop)
 {
     FTangentBezierSpline::SetClosedLoop(bInClosedLoop);
@@ -1409,7 +1421,11 @@ void FNewSpline::SetClosedLoop(bool bInClosedLoop)
 
 void FNewSpline::Reset()
 {
-    InternalSpline.Clear();
+	InternalSpline.Clear();
+	TangentModes.Empty();
+	PositionCurve.Reset();
+	RotationCurve.Reset();
+	ScaleCurve.Reset();
 }
 
 void FNewSpline::ResetToDefault()
@@ -1427,48 +1443,34 @@ void FNewSpline::ResetRotation()
 {
     ClearAttributeChannel(RotationAttrName);
 
-    auto NormalizeAttributeParameter = [MaxParameter = static_cast<float>(GetNumPoints() - 1)](int32 Index)
-    {
-    	return Index / MaxParameter;
-    };
-    
     for (int32 Idx = 0; Idx < GetNumPoints(); Idx++)
     {    		
-    	const float AttributeParameter = NormalizeAttributeParameter(Idx);
-    	//Points[Idx].Rotation = FQuat::Identity;
-    	AddAttributeValue(RotationAttrName, FQuat::Identity, AttributeParameter);
+    	AddAttributeValue(RotationAttrName, FQuat::Identity, static_cast<float>(Idx));
     }
 }
 
 void FNewSpline::ResetScale()
 {
     ClearAttributeChannel(ScaleAttrName);
-
-    auto NormalizeAttributeParameter = [MaxParameter = static_cast<float>(GetNumPoints() - 1)](int32 Index)
-    {
-    	return Index / MaxParameter;
-    };
-    
+	
     for (int32 Idx = 0; Idx < GetNumPoints(); Idx++)
     {    		
-        const float AttributeParameter = NormalizeAttributeParameter(Idx);
-        //Points[Idx].Scale = FVector::OneVector;
-        AddAttributeValue(ScaleAttrName, FVector::OneVector, AttributeParameter);
+        AddAttributeValue(ScaleAttrName, FVector::OneVector, static_cast<float>(Idx));
     }
 }
 
 void FNewSpline::AddPoint(const FSplinePoint& Point)
 {
-    FTangentBezierControlPoint ControlPoint = ConvertToTangentBezierControlPoint(Point);
-    const int32 NumPoints = GetNumPoints();
+	FTangentBezierControlPoint ControlPoint = ConvertToTangentBezierControlPoint(Point);
+	const int32 NumPoints = GetNumPoints();
     
-    // Special case for empty spline
-    if (NumPoints == 0)
-    {
-        AppendPoint(ControlPoint);
-        UpdatePointAttributes(Point, 0);
-        return;
-    }
+	// Special case for empty spline
+	if (NumPoints == 0)
+	{
+		AppendPoint(ControlPoint);
+		UpdatePointAttributes(Point, 0);
+		return;
+	}
 	// Convert InputKey to parameter value
 	int32 BaseIndex = FMath::FloorToInt(Point.InputKey);
 	float Fraction = Point.InputKey - BaseIndex;
@@ -1482,7 +1484,6 @@ void FNewSpline::AddPoint(const FSplinePoint& Point)
 
 void FNewSpline::InsertPoint(const FSplinePoint& Point, int32 Index)
 {
-
 	FTangentBezierControlPoint ControlPoint = ConvertToTangentBezierControlPoint(Point);
 	const int32 NumPoints = GetNumPoints();
     
@@ -1500,11 +1501,11 @@ void FNewSpline::InsertPoint(const FSplinePoint& Point, int32 Index)
 		return;
 	}
 	
-    Index = FMath::Clamp(Index, 0, GetNumPoints());
+	Index = FMath::Clamp(Index, 0, GetNumPoints());
 	float SqDistance;
 	float SplineParam = FTangentBezierSpline::FindNearestOnSegment(Point.Position, Index - 1, SqDistance);
 		
-   	int32 InsertedSegmentIndex = FTangentBezierSpline::InsertPoint(SplineParam, ControlPoint);
+	int32 InsertedSegmentIndex = FTangentBezierSpline::InsertPoint(SplineParam, ControlPoint);
 	
 	UpdatePointAttributes(Point, InsertedSegmentIndex);
 }
@@ -1701,13 +1702,16 @@ float FNewSpline::GetParameterAtIndex(int32 Index) const
 float FNewSpline::GetParameterAtDistance(float Distance) const
 {
 	const float NormalizedDistance = FMath::Clamp(Distance / FMath::Max(GetSplineLength(), UE_SMALL_NUMBER), 0.f, 1.f);
-    return EvaluateAttribute<float>(LengthToParamAttrName, NormalizedDistance);
+	return FTangentBezierSpline::EvaluateAttribute<float>(LengthToParamAttrName, NormalizedDistance);
 }
 
 float FNewSpline::GetDistanceAtParameter(float Parameter) const
 {
-    Parameter = Parameter / GetNumSegments();
-    return EvaluateAttribute<float>(ParamToLengthAttrName, Parameter);
+	// a bit of a hack, we convert to normalized but then treat
+	// the resulting value as if it were in internal parent space.
+	// see the todo in UpdateReparamTable
+	Parameter = ToNormalizedSpace(Parameter);
+	return FTangentBezierSpline::EvaluateAttribute<float>(ParamToLengthAttrName, Parameter);
 }
 
 float FNewSpline::FindNearest(const FVector& InLocation) const
@@ -1720,7 +1724,7 @@ float FNewSpline::FindNearest(const FVector& InLocation) const
     }
     else
     {
-	    return FTangentBezierSpline::FindNearest(InLocation, Dummy);
+	    return FromInternalSplineSpace(FTangentBezierSpline::FindNearest(InLocation, Dummy));
     }
 }
 
@@ -1739,41 +1743,35 @@ float FNewSpline::FindNearestOnSegment(const FVector& InLocation, int32 SegmentI
     }
     else
     {
-	   return FTangentBezierSpline::FindNearestOnSegment(InLocation, SegmentIndex, Dummy);
+	   return FromInternalSplineSpace(FTangentBezierSpline::FindNearestOnSegment(InLocation, SegmentIndex, Dummy));
     }
 }
 
 FVector FNewSpline::Evaluate(float Parameter) const
 {
-    Parameter = GetParameterRange().Interpolate(Parameter / GetNumSegments());
-	
+	Parameter = ToInternalSplineSpace(Parameter);
     return FVector(GetPosition(Parameter));
 }
 
 FVector FNewSpline::EvaluateDerivative(float Parameter) const
 {
-    const int32 NumEffectivePoints = GetNumPoints();
-	Parameter = GetParameterRange().Interpolate(Parameter / GetNumSegments());
+	Parameter = ToInternalSplineSpace(Parameter);
     return FVector(GetTangent(Parameter));
 }
 
 FQuat FNewSpline::EvaluateRotation(float Parameter) const
 {
-    const int32 NumEffectivePoints = GetNumPoints();
-	Parameter = GetParameterRange().Interpolate(Parameter / GetNumSegments());
     return EvaluateAttribute<FQuat>(RotationAttrName, Parameter);
 }
 
 FVector FNewSpline::EvaluateScale(float Parameter) const
 {
-    const int32 NumEffectivePoints = GetNumPoints();
-	Parameter = GetParameterRange().Interpolate(Parameter / GetNumSegments());
     return EvaluateAttribute<FVector>(ScaleAttrName, Parameter);
 }
 
 void FNewSpline::UpdateSpline()
 {
-    //UpdateTangents();			// updates the tangents on our version of the points
+    UpdateTangents();			// updates the tangents on our version of the points
     UpdateReparamTable();		// updates both reparam tables
     RebuildLegacyCurves();		// updates legacy curves based on points, never evaluates internal spline.
 }
@@ -1789,44 +1787,71 @@ void FNewSpline::UpdateSpline(const FSpline::FUpdateSplineParams& Params)
     UpdateSpline();
 }
 
+float FNewSpline::FromNormalizedSpace(float Parameter) const
+{
+	return GetNumSegments() * Parameter;
+}
+
+float FNewSpline::ToNormalizedSpace(float Parameter) const
+{
+	return Parameter / GetNumSegments();
+}
+
+float FNewSpline::FromInternalSplineSpace(float Parameter) const
+{
+	const float NormalizedParam = this->GetParameterRange().GetT(Parameter);
+	const float SegmentSpaceParam = NormalizedParam * GetNumSegments();
+	return SegmentSpaceParam;
+}
+
+float FNewSpline::ToInternalSplineSpace(float Parameter) const
+{
+	const float NormalizedParam = Parameter / GetNumSegments();
+	const float InternalSpaceParam = GetParameterRange().Interpolate(NormalizedParam);
+	return InternalSpaceParam;
+}
+
 void FNewSpline::UpdateReparamTable()
 {
     const int32 NumPoints = GetNumPoints();
   	const int32 NumSegments = NumPoints - 1;
     const int32 TotalSteps = NumSegments * ReparamStepsPerSegment + 1;
-
+	
     // Start by clearing it
     ClearAttributeChannel(ParamToLengthAttrName);
     ClearAttributeChannel(LengthToParamAttrName);
-    
+
+	// Todo: convert normalized parameters below into [0, NumSegments] space.
+	// This does not work when last tested, but should be corrected.
+	
     // Here we get Param -> Length mapping.
     float AccumulatedLength = 0.0f;
     for (int32 SegmentIndex = 0; SegmentIndex < NumSegments; ++SegmentIndex)
     {
     	for (int32 Step = 0; Step < ReparamStepsPerSegment; ++Step)
     	{
-    		// These are with respect to current segment (so called 'segment space').
+    		// These are with respect to current segment.
     		const float SegmentParam = static_cast<float>(Step) / ReparamStepsPerSegment;
     		const float OffsetLength = GetSegmentLength(SegmentIndex, SegmentParam);
 
-    		// Map the above values to 'spline space' (as opposed to segment space).
-    		const float SplineParam = static_cast<float>(SegmentIndex * ReparamStepsPerSegment + Step) / TotalSteps;
+    		// Map the above values to 'spline space' (i.e. [0, NumSegments]) by normalizing and converting.
+    		const float SplineParam = (static_cast<float>(SegmentIndex * ReparamStepsPerSegment + Step) / TotalSteps);
     		const float CurrentLength = OffsetLength + AccumulatedLength;
 
-    		AddAttributeValue(ParamToLengthAttrName, CurrentLength, SplineParam);
+    		FTangentBezierSpline::AddAttributeValue(ParamToLengthAttrName, CurrentLength, SplineParam);
     	}
 
     	AccumulatedLength += GetSegmentLength(SegmentIndex, 1.0f);
     }
-    AddAttributeValue(ParamToLengthAttrName, AccumulatedLength, 1.f);
+    FTangentBezierSpline::AddAttributeValue(ParamToLengthAttrName, AccumulatedLength, 1.0f);
 
-    // Here we get Normalized Length -> Param mapping
+    // Here we get Length -> Param mapping
     for (int32 Step = 0; Step <= NumSegments * ReparamStepsPerSegment + 1; ++Step)
     {
     	const float SplineParam = FMath::Clamp(static_cast<float>(Step) / TotalSteps, 0.0f, 1.0f);
-    	const float CurrentLength = EvaluateAttribute<float>(ParamToLengthAttrName, SplineParam);
-    
-    	AddAttributeValue(LengthToParamAttrName, SplineParam * NumSegments, CurrentLength / AccumulatedLength);	// must be normalized param
+    	const float CurrentLength = FTangentBezierSpline::EvaluateAttribute<float>(ParamToLengthAttrName, SplineParam);
+    	
+    	FTangentBezierSpline::AddAttributeValue(LengthToParamAttrName, FromNormalizedSpace(SplineParam), CurrentLength / AccumulatedLength);
     }
 }
 
@@ -1869,14 +1894,12 @@ FTangentBezierControlPoint FNewSpline::ConvertToTangentBezierControlPoint(const
 void FNewSpline::UpdatePointAttributes(const FSplinePoint& Point, int32 PointIndex)
 {
 	if (PointIndex < 0)
+	{
 		return;
-            
-	// Calculate normalized parameter
-	float NormalizedParam = static_cast<float>(PointIndex) / 
-		FMath::Max(1.0f, static_cast<float>(GetNumPoints() - 1));
-        
-	AddAttributeValue(RotationAttrName, Point.Rotation.Quaternion(), NormalizedParam);
-	AddAttributeValue(ScaleAttrName, Point.Scale, NormalizedParam);
+	}
+	
+	SetRotation(PointIndex, Point.Rotation.Quaternion());
+	SetScale(PointIndex, Point.Scale);
         
 	RebuildLegacyCurves();
 }
