[World Partition] Add accessors that allow retrieving the URuntimePartition from a UWorldPartitionRuntimeHashSet, and URuntimePartitionLHGrid::CellSize

diff --git a/Engine/Source/Runtime/Engine/Public/WorldPartition/RuntimeHashSet/RuntimePartitionLHGrid.h b/Engine/Source/Runtime/Engine/Public/WorldPartition/RuntimeHashSet/RuntimePartitionLHGrid.h
--- a/Engine/Source/Runtime/Engine/Public/WorldPartition/RuntimeHashSet/RuntimePartitionLHGrid.h
+++ b/Engine/Source/Runtime/Engine/Public/WorldPartition/RuntimeHashSet/RuntimePartitionLHGrid.h
@@ -39,6 +39,11 @@ public:
 	virtual bool GenerateStreaming(const FGenerateStreamingParams& InParams, FGenerateStreamingResult& OutResult) override;
 	virtual FArchive& AppendCellGuid(FArchive& InAr) override;
 	//~ End URuntimePartition interface
+
+	uint32 GetCellSize() const
+	{
+		return CellSize;
+	}
 #endif
 
 private:
diff --git a/Engine/Source/Runtime/Engine/Public/WorldPartition/RuntimeHashSet/WorldPartitionRuntimeHashSet.h b/Engine/Source/Runtime/Engine/Public/WorldPartition/RuntimeHashSet/WorldPartitionRuntimeHashSet.h
--- a/Engine/Source/Runtime/Engine/Public/WorldPartition/RuntimeHashSet/WorldPartitionRuntimeHashSet.h
+++ b/Engine/Source/Runtime/Engine/Public/WorldPartition/RuntimeHashSet/WorldPartitionRuntimeHashSet.h
@@ -222,7 +222,11 @@ private:
 	ENGINE_API void ForEachStreamingData(TFunctionRef<bool(const FRuntimePartitionStreamingData&)> Func) const;
 
 	ENGINE_API void UpdateRuntimeDataGridMap();
-	const URuntimePartition* ResolveRuntimePartition(FName GridName, bool bMainPartitionLayer = false) const;
+
+public:
+	ENGINE_API const URuntimePartition* ResolveRuntimePartition(FName GridName, bool bMainPartitionLayer = false) const;
+
+private:
 	const URuntimePartition* ResolveRuntimePartitionForHLODLayer(FName GridName, const FSoftObjectPath& HLODLayerPath) const;
 
 	void FixupHLODSetup(FRuntimePartitionDesc& RuntimePartition);
