[WorldPartitionLevelStreamingPolicy] Handle subobjects missing from SubObjectsToExternalStreamingObjectsRemapping to avoid blocking automated test

diff --git a/Engine/Source/Runtime/Engine/Private/WorldPartition/WorldPartitionLevelStreamingPolicy.cpp b/Engine/Source/Runtime/Engine/Private/WorldPartition/WorldPartitionLevelStreamingPolicy.cpp
--- a/Engine/Source/Runtime/Engine/Private/WorldPartition/WorldPartitionLevelStreamingPolicy.cpp
+++ b/Engine/Source/Runtime/Engine/Private/WorldPartition/WorldPartitionLevelStreamingPolicy.cpp
@@ -324,11 +324,18 @@ bool UWorldPartitionLevelStreamingPolicy::RemoveExternalStreamingObject(URuntime
 
 	for (const TPair<FName, FName> Pair : ExternalStreamingObject->SubObjectsToCellRemapping)
 	{
-		TWeakObjectPtr<URuntimeHashExternalStreamingObjectBase> RemovedExternalStreamingObject = SubObjectsToExternalStreamingObjectsRemapping.FindAndRemoveChecked(Pair.Key);
-		UE_CLOG(RemovedExternalStreamingObject != ExternalStreamingObject, LogWorldPartition, Log, TEXT("Mismatch in SubObjectsToExternalStreamingObjectsRemapping map. Found: %s. Expected: %s"),
-			(RemovedExternalStreamingObject.IsValid() && RemovedExternalStreamingObject->GetPackage()) ? *RemovedExternalStreamingObject->GetPackage()->GetName() : TEXT("<Invalid>"),
-			ExternalStreamingObject->GetPackage() ? *ExternalStreamingObject->GetPackage()->GetName() : TEXT("<Invalid>"));
-			
+		if (TWeakObjectPtr<URuntimeHashExternalStreamingObjectBase>* RemovedExternalStreamingObjectPtr = SubObjectsToExternalStreamingObjectsRemapping.Find(Pair.Key))
+		{
+			TWeakObjectPtr<URuntimeHashExternalStreamingObjectBase> RemovedExternalStreamingObject = *RemovedExternalStreamingObjectPtr;
+			UE_CLOG(RemovedExternalStreamingObject != ExternalStreamingObject, LogWorldPartition, Log, TEXT("Mismatch in SubObjectsToExternalStreamingObjectsRemapping map. Found: %s. Expected: %s"),
+				(RemovedExternalStreamingObject.IsValid() && RemovedExternalStreamingObject->GetPackage()) ? *RemovedExternalStreamingObject->GetPackage()->GetName() : TEXT("<Invalid>"),
+				ExternalStreamingObject->GetPackage() ? *ExternalStreamingObject->GetPackage()->GetName() : TEXT("<Invalid>"));
+			SubObjectsToExternalStreamingObjectsRemapping.Remove(Pair.Key);
+		}
+		else
+		{
+			UE_LOG(LogWorldPartition, Log, TEXT("Mismatch in SubObjectsToExternalStreamingObjectsRemapping map. Subobject %s not found"), *Pair.Key.ToString());
+		}
 	}
 
 	ExternalStreamingObjects.RemoveSwap(ExternalStreamingObject);
