Allow the IoStoreOnDemand module to be compiled in the editor if the target requests bCompileIoStoreOnDemand.

#rnx

diff --git a/Engine/Source/Runtime/Engine/Engine.Build.cs b/Engine/Source/Runtime/Engine/Engine.Build.cs
--- a/Engine/Source/Runtime/Engine/Engine.Build.cs
+++ b/Engine/Source/Runtime/Engine/Engine.Build.cs
@@ -124,7 +124,7 @@ public class Engine : ModuleRules
 			PublicDependencyModuleNames.Add("ApplicationCore");
 		}
 
-		if (Target.bCompileIoStoreOnDemand && !Target.bBuildEditor)
+		if (Target.bCompileIoStoreOnDemand)
 		{
 			PrivateDependencyModuleNames.Add("IoStoreOnDemand");
 			PublicDefinitions.Add("UE_WITH_IOSTOREONDEMAND=1");
