Fix case where soft object remapping could map to an incorrect cell for an actor that has an editor only reference to it.

diff --git a/Engine/Source/Runtime/Engine/Private/WorldPartition/WorldPartitionLevelStreamingPolicy.cpp b/Engine/Source/Runtime/Engine/Private/WorldPartition/WorldPartitionLevelStreamingPolicy.cpp
--- a/Engine/Source/Runtime/Engine/Private/WorldPartition/WorldPartitionLevelStreamingPolicy.cpp
+++ b/Engine/Source/Runtime/Engine/Private/WorldPartition/WorldPartitionLevelStreamingPolicy.cpp
@@ -91,30 +91,37 @@ void UWorldPartitionLevelStreamingPolicy::PrepareActorToCellRemapping()
 	SourceWorldAssetPath = FTopLevelAssetPath(SourceWorldPath);
 	
 	// Build Actor-to-Cell remapping
+	// The use cases for remapping are the following:
+	//	- Spatially loaded or Datalayer Actors from the main World Partition map that get moved into
+	//		a Streaming Cell. In this case, an actor path like '/Game/SomePath/WorldName.WorldName:PersistentLevel.ActorA'
+	//		would be mapped to a cell name like 'WorldName_MainGrid_L0_X5_Y-4'.
+	//	- Always loaded Actors from the main World:
+	//		- In PIE they get remapped to the top level Cell 'WorldName_MainGrid_L{MAX}_X0_Y0'.
+	//		- In Cook they don't need remapping as the top level Cell is
+	//			the PersistentLevel (Cell->NeedsActorToCellRemapping() returns false).
 	WorldPartition->RuntimeHash->ForEachStreamingCells([this, &SourceWorldPath](const UWorldPartitionRuntimeCell* Cell)
 	{
 		const UWorldPartitionRuntimeLevelStreamingCell* StreamingCell = Cast<const UWorldPartitionRuntimeLevelStreamingCell>(Cell);
 		check(StreamingCell);
-		for (const FWorldPartitionRuntimeCellObjectMapping& CellObjectMap : StreamingCell->GetPackages())
+
+		if (Cell->NeedsActorToCellRemapping())
 		{
-			// The use cases for remapping are the following:
-			//
-			// - Spatially loaded or Datalayer Actors from the main World Partition map that get moved into a Streaming Cell. In thise case an actor path like:
-			//		- '/Game/SomePath/WorldName.WorldName:PersistentLevel.ActorA' would be mapped to a cell name ex: 'WorldName_MainGrid_L0_X5_Y-4'
-			// - Always loaded Actors from the main World:
-			//		- In PIE they get remapped to the top level Cell 'WorldName_MainGrid_L{MAX}_X0_Y0'
-			//		- In Cook they don't need remapping as the top level Cell is the PersistentLevel (Cell->NeedsActorToCellRemapping() returns false)
-			if (Cell->NeedsActorToCellRemapping())
+			const FName CellName = StreamingCell->GetFName();
+
+			for (const FWorldPartitionRuntimeCellObjectMapping& CellObjectMap : StreamingCell->GetPackages())
 			{
-				const FSoftObjectPath CellActorPath = FWorldPartitionLevelHelper::RemapActorPath(CellObjectMap.ContainerID, SourceWorldPath, FSoftObjectPath(CellObjectMap.Path.ToString()));
-				
-				const FString ActorPath = CellActorPath.ToString();
-				const FName CellName = StreamingCell->GetFName();
-				const int32 LastDotPos = ActorPath.Find(TEXT("."), ESearchCase::IgnoreCase, ESearchDir::FromEnd);
-				check(LastDotPos != INDEX_NONE);
-				SubObjectsToCellRemapping.Add(FName(*ActorPath.Mid(LastDotPos + 1)), CellName);
+				// Avoid remapping from any editor only references of actors that the cells may have. 
+				if (!CellObjectMap.bIsEditorOnly)
+				{
+					const FSoftObjectPath CellActorPath = FWorldPartitionLevelHelper::RemapActorPath(CellObjectMap.ContainerID, SourceWorldPath, FSoftObjectPath(CellObjectMap.Path.ToString()));
+					const FString ActorPath = CellActorPath.ToString();
+					const int32 LastDotPos = ActorPath.Find(TEXT("."), ESearchCase::IgnoreCase, ESearchDir::FromEnd);
+					check(LastDotPos != INDEX_NONE);
+					SubObjectsToCellRemapping.Add(FName(*ActorPath.Mid(LastDotPos + 1)), CellName);
+				}
 			}
 		}
+
 		return true;
 	});
 }
