Export symbols for FGPUSkinPassThroughFactoryLooseParameters so that plugins can reference these when creating custom vertex factories.

diff --git a/Engine/Source/Runtime/Engine/Public/GPUSkinVertexFactory.h b/Engine/Source/Runtime/Engine/Public/GPUSkinVertexFactory.h
--- a/Engine/Source/Runtime/Engine/Public/GPUSkinVertexFactory.h
+++ b/Engine/Source/Runtime/Engine/Public/GPUSkinVertexFactory.h
@@ -729,7 +729,7 @@ protected:
 	FGPUSkinAPEXClothDataType* ClothDataPtr = nullptr;
 };
 
-BEGIN_GLOBAL_SHADER_PARAMETER_STRUCT(FGPUSkinPassThroughFactoryLooseParameters,)
+BEGIN_GLOBAL_SHADER_PARAMETER_STRUCT(FGPUSkinPassThroughFactoryLooseParameters, ENGINE_API)
 	SHADER_PARAMETER(uint32, FrameNumber)
 	SHADER_PARAMETER_SRV(Buffer<float>, PositionBuffer)
 	SHADER_PARAMETER_SRV(Buffer<float>, PreviousPositionBuffer)
@@ -829,4 +829,4 @@ private:
 	EPixelFormat TangentFormat = PF_Unknown;
 };
 
-ENUM_CLASS_FLAGS(FGPUSkinPassthroughVertexFactory::EVertexAttributeFlags)
\ No newline at end of file
+ENUM_CLASS_FLAGS(FGPUSkinPassthroughVertexFactory::EVertexAttributeFlags)
