Add a new "Record Texture Streaming Info" material node.
This adds functionality to record the material UV scales for use by the automatic texture streaming system.
If any texture sampling is done inside Custom HLSL nodes then this node can be added to give the texture streaming system the ability to record the UV scales of the texture being sampled.

diff --git a/Engine/Config/BaseMaterialExpressions.ini b/Engine/Config/BaseMaterialExpressions.ini
--- a/Engine/Config/BaseMaterialExpressions.ini
+++ b/Engine/Config/BaseMaterialExpressions.ini
@@ -987,6 +987,9 @@ bShaderInputData=true
 +MenuCategories=NSLOCTEXT("MaterialExpression", "Texture", "Texture")
 +MenuCategories=NSLOCTEXT("MaterialExpression", "Parameters", "Parameters")
 
+[/Script/Engine.MaterialExpressionRecordTextureStreamingInfo]
++MenuCategories=NSLOCTEXT("MaterialExpression", "Texture", "Texture")
+
 [/Script/Engine.MaterialExternalCodeCollection]
 +ExternalCodeDeclarations=(ReturnType=Float3,			Name="AtmosphericLightVector",			Definition="MaterialExpressionAtmosphericLightVector(Parameters)")
 +ExternalCodeDeclarations=(ReturnType=Float3,			Name="AtmosphericLightColor",			Definition="MaterialExpressionAtmosphericLightColor(Parameters)")
diff --git a/Engine/Source/Runtime/Engine/Private/Materials/HLSLMaterialTranslator.cpp b/Engine/Source/Runtime/Engine/Private/Materials/HLSLMaterialTranslator.cpp
--- a/Engine/Source/Runtime/Engine/Private/Materials/HLSLMaterialTranslator.cpp
+++ b/Engine/Source/Runtime/Engine/Private/Materials/HLSLMaterialTranslator.cpp
@@ -8220,6 +8220,19 @@ int32 FHLSLMaterialTranslator::TextureFromCollection(int32 TextureCollectionCode
 	);
 }
 
+int32 FHLSLMaterialTranslator::TextureStreamingInfo(int32 TextureReferenceIndex, int32 TextureIndex, int32 CoordinateIndex)
+{
+	EMaterialValueType TextureType = GetParameterType(TextureIndex);
+
+	if (!(TextureType & MCT_Texture))
+	{
+		Errorf(TEXT("Sampling unknown texture type: %s"), DescribeType(TextureType));
+		return INDEX_NONE;
+	}
+
+	return AddCodeChunk(MCT_Float, TEXT("MaterialStoreTexCoordScale(Parameters, %s, %d)"), *CoerceParameter(CoordinateIndex, MCT_Float2), TextureReferenceIndex);
+}
+
 int32 FHLSLMaterialTranslator::SampleMaterialCache(int32 TexCoordIndexLiteral)
 {
 	int32 TexCoord = TextureCoordinate(TexCoordIndexLiteral, false, false);
diff --git a/Engine/Source/Runtime/Engine/Private/Materials/HLSLMaterialTranslator.h b/Engine/Source/Runtime/Engine/Private/Materials/HLSLMaterialTranslator.h
--- a/Engine/Source/Runtime/Engine/Private/Materials/HLSLMaterialTranslator.h
+++ b/Engine/Source/Runtime/Engine/Private/Materials/HLSLMaterialTranslator.h
@@ -1057,6 +1057,8 @@ protected:
 	virtual int32 TextureProperty(int32 TextureIndex, EMaterialExposedTextureProperty Property) override;
 	virtual int32 TextureFromCollection(int32 TextureCollectionCodeIndex, int32 IndexIntoCollection, EMaterialValueType ResultTextureType) override;
 
+	virtual int32 TextureStreamingInfo(int32 TextureReferenceIndex, int32 TextureIndex, int32 CoordinateIndex) override;
+
 	virtual int32 SampleMaterialCache(int32 TexCoordIndexLiteral) override;
 	virtual int32 GetMaterialCacheAttribute(int32 MaterialCacheIndex, const FGuid& Attribute) override;
 	
diff --git a/Engine/Source/Runtime/Engine/Private/Materials/MaterialExpressions.cpp b/Engine/Source/Runtime/Engine/Private/Materials/MaterialExpressions.cpp
--- a/Engine/Source/Runtime/Engine/Private/Materials/MaterialExpressions.cpp
+++ b/Engine/Source/Runtime/Engine/Private/Materials/MaterialExpressions.cpp
@@ -300,6 +300,7 @@
 #include "Materials/MaterialExpressionSparseVolumeTextureSample.h"
 #include "Materials/MaterialExpressionIsFirstPerson.h"
 #include "Materials/MaterialExternalCodeRegistry.h"
+#include "Materials/MaterialExpressionRecordTextureStreamingInfo.h"
 #include "EditorSupportDelegates.h"
 #if WITH_EDITOR
 #include "MaterialGraph/MaterialGraphNode_Comment.h"
@@ -23843,4 +23844,91 @@ bool UMaterialExpressionConvert::IsValidComponentIndex(int32 InComponentIndex, E
 }
 #endif // WITH_EDITOR
 
+
+UMaterialExpressionRecordTextureStreamingInfo::UMaterialExpressionRecordTextureStreamingInfo(const FObjectInitializer& ObjectInitializer)
+	: Super(ObjectInitializer)
+{
+#if WITH_EDITOR
+	bShowOutputNameOnPin = true;
+
+	Outputs.Reset();
+	Outputs.Add(FExpressionOutput(TEXT("Tex")));
+#endif
+}
+
+#if WITH_EDITOR
+
+FExpressionInput* UMaterialExpressionRecordTextureStreamingInfo::GetInput(int32 InputIndex)
+{
+	switch (InputIndex)
+	{
+	case 0: return &TextureObject;
+	case 1: return &Coordinates;
+	}
+
+	return nullptr;
+}
+
+EMaterialValueType UMaterialExpressionRecordTextureStreamingInfo::GetInputValueType(int32 InputIndex)
+{
+	switch (InputIndex)
+	{
+	case 0: return MCT_Texture2D;
+	case 1: return MCT_Float2;
+	}
+
+	return MCT_Unknown;
+}
+
+EMaterialValueType UMaterialExpressionRecordTextureStreamingInfo::GetOutputValueType(int32 OutputIndex)
+{
+	if (TextureObject.GetTracedInput().Expression)
+	{
+		return TextureObject.GetTracedInput().Expression->GetOutputValueType(0);
+	}
+	return MCT_Texture2D;
+}
+
+void UMaterialExpressionRecordTextureStreamingInfo::GetCaption(TArray<FString>& OutCaptions) const
+{
+	OutCaptions.Add(TEXT("Record Texture Streaming"));
+}
+
+void UMaterialExpressionRecordTextureStreamingInfo::GetExpressionToolTip(TArray<FString>& OutToolTip)
+{
+	ConvertToMultilineToolTip(
+		TEXT("Adds functionality to record the material UV scales for use by the automatic texture streaming system. If this material uses Texture ")
+		TEXT("Sample nodes then this node isn't required. But if any texture sampling is done inside Custom HLSL nodes then this node should be added ")
+		TEXT("so that the texture streaming system can record the UV scales of the texture being sampled."), 40, OutToolTip);
+}
+
+int32 UMaterialExpressionRecordTextureStreamingInfo::Compile(class FMaterialCompiler* Compiler, int32 OutputIndex)
+{
+	if (!TextureObject.GetTracedInput().Expression)	
+	{
+		return Compiler->Errorf(TEXT("Missing input TextureObject"));
+	}
+	const int32 TextureCodeIndex = TextureObject.Compile(Compiler);
+	
+	if (!Coordinates.GetTracedInput().Expression)
+	{
+		return Compiler->Errorf(TEXT("Missing input UV Coordinates"));
+	}
+	const int32 CoordinatesCodeIndex = Coordinates.Compile(Compiler);
+
+	int32 TextureReferenceIndex = INDEX_NONE;
+	EMaterialSamplerType EffectiveSamplerType = SAMPLERTYPE_Color;
+	TOptional<FName> EffectiveParameterName;
+	if (!Compiler->GetTextureForExpression(TextureCodeIndex, TextureReferenceIndex, EffectiveSamplerType, EffectiveParameterName))
+	{
+		return CompilerError(Compiler, TEXT("TextureObject input requires a texture value"));
+	}
+
+	Compiler->TextureStreamingInfo(TextureReferenceIndex, TextureCodeIndex, CoordinatesCodeIndex);
+
+	return TextureCodeIndex;
+}
+
+#endif // WITH_EDITOR
+
 #undef LOCTEXT_NAMESPACE
diff --git a/Engine/Source/Runtime/Engine/Public/MaterialCompiler.h b/Engine/Source/Runtime/Engine/Public/MaterialCompiler.h
--- a/Engine/Source/Runtime/Engine/Public/MaterialCompiler.h
+++ b/Engine/Source/Runtime/Engine/Public/MaterialCompiler.h
@@ -328,6 +328,8 @@ public:
 	virtual int32 TextureCollectionCount(int32 InTextureCollectionCodeIndex) = 0;
 	virtual int32 TextureFromCollection(int32 TextureCollectionCodeIndex, int32 IndexIntoCollectionCodeIndex, EMaterialValueType ResultTextureType) = 0;
 
+	virtual int32 TextureStreamingInfo(int32 TextureReferenceIndex, int32 TextureIndex, int32 CoordinateIndex) = 0;
+
 	virtual int32 SampleMaterialCache(int32 TexCoordIndexLiteral) = 0;
 	virtual int32 GetMaterialCacheAttribute(int32 MaterialCacheIndex, const FGuid& Attribute) = 0;
 	
@@ -1033,6 +1035,11 @@ public:
 		return Compiler->TextureFromCollection(TextureCollectionCodeIndex, IndexIntoCollectionCodeIndex, ResultTextureType);
 	}
 
+	virtual int32 TextureStreamingInfo(int32 TextureReferenceIndex, int32 TextureIndex, int32 CoordinateIndex) override
+	{
+		return Compiler->TextureStreamingInfo(TextureReferenceIndex, TextureIndex, CoordinateIndex);
+	}
+
 	virtual int32 SampleMaterialCache(int32 TexCoordIndexLiteral)
 	{
 		return Compiler->SampleMaterialCache(TexCoordIndexLiteral);
diff --git a/Engine/Source/Runtime/Engine/Public/Materials/MaterialExpressionRecordTextureStreamingInfo.h b/Engine/Source/Runtime/Engine/Public/Materials/MaterialExpressionRecordTextureStreamingInfo.h
new file mode 100644
--- /dev/null
+++ b/Engine/Source/Runtime/Engine/Public/Materials/MaterialExpressionRecordTextureStreamingInfo.h
@@ -0,0 +1,30 @@
+// Copyright Epic Games, Inc. All Rights Reserved.
+
+#pragma once
+
+#include "MaterialExpression.h"
+#include "MaterialExpressionRecordTextureStreamingInfo.generated.h"
+
+/** Adds functionality to record the material UV scales for use by the automatic texture streaming system. */
+UCLASS(collapsecategories, hidecategories=Object, MinimalAPI)
+class UMaterialExpressionRecordTextureStreamingInfo : public UMaterialExpression
+{
+	GENERATED_UCLASS_BODY()
+
+	/** A texture object that will be sampled. */
+	UPROPERTY(meta = (RequiredInput = "true"))
+	FExpressionInput TextureObject;
+
+	/** The texture coordinates that we expect to use when sampling the texture object. */
+	UPROPERTY(meta = (RequiredInput = "true"))
+	FExpressionInput Coordinates;
+
+#if WITH_EDITOR
+	virtual FExpressionInput* GetInput(int32 InputIndex) override;
+	virtual EMaterialValueType GetInputValueType(int32 InputIndex) override;
+	virtual EMaterialValueType GetOutputValueType(int32 OutputIndex) override;
+	virtual void GetCaption(TArray<FString>& OutCaptions) const override;
+	virtual void GetExpressionToolTip(TArray<FString>& OutToolTip) override;
+	virtual int32 Compile(class FMaterialCompiler* Compiler, int32 OutputIndex) override;
+#endif // WITH_EDITOR
+};
