Skeletal Mesh: Cache information on vertex colors

#rnx

diff --git a/Engine/Source/Developer/SkeletalMeshUtilitiesCommon/Private/StaticToSkeletalMeshConverter.cpp b/Engine/Source/Developer/SkeletalMeshUtilitiesCommon/Private/StaticToSkeletalMeshConverter.cpp
--- a/Engine/Source/Developer/SkeletalMeshUtilitiesCommon/Private/StaticToSkeletalMeshConverter.cpp
+++ b/Engine/Source/Developer/SkeletalMeshUtilitiesCommon/Private/StaticToSkeletalMeshConverter.cpp
@@ -631,12 +631,6 @@ bool FStaticToSkeletalMeshConverter::InitializeSkeletalMeshFromMeshDescriptions(
 		InSkeletalMesh->SetMaterials(NewMaterials);
 	}
 
-	if (HasVertexColors(InSkeletalMesh))
-	{
-		InSkeletalMesh->SetHasVertexColors(true);	
-		InSkeletalMesh->SetVertexColorGuid(FGuid::NewGuid());
-	}
-	
 	InSkeletalMesh->SetImportedBounds( FBox3d{BoundingBox} );
 
 	return true;
diff --git a/Engine/Source/Runtime/Engine/Classes/Engine/SkeletalMesh.h b/Engine/Source/Runtime/Engine/Classes/Engine/SkeletalMesh.h
--- a/Engine/Source/Runtime/Engine/Classes/Engine/SkeletalMesh.h
+++ b/Engine/Source/Runtime/Engine/Classes/Engine/SkeletalMesh.h
@@ -605,6 +605,9 @@ public:
 		 */
 		bool bUpdateVertexAttributes = true;
 		
+		/** Update the vertex color enabled state/GUID if the mesh description contains non-default vertex colors */
+		bool bUpdateVertexColors = true;
+		
 		/** Mark the package as dirty. If calling CommitMeshDescription from a non-game thread,
 		 *  this value should be set to \c false.
 		 */
diff --git a/Engine/Source/Runtime/Engine/Classes/Engine/SkeletalMeshSourceModel.h b/Engine/Source/Runtime/Engine/Classes/Engine/SkeletalMeshSourceModel.h
--- a/Engine/Source/Runtime/Engine/Classes/Engine/SkeletalMeshSourceModel.h
+++ b/Engine/Source/Runtime/Engine/Classes/Engine/SkeletalMeshSourceModel.h
@@ -133,6 +133,11 @@ struct FSkeletalMeshSourceModel
 	 */
 	TConstArrayView<FName> GetMorphTargetNames() const;
 
+	/** Get the vertex color GUID as an optional value. If the mesh has no vertex colors defined, then
+	 *  an empty optional is returned.
+	 */
+	TOptional<FGuid> GetVertexColorGuid() const;
+	
 private:
 #if WITH_EDITOR
 	USkeletalMesh* GetOwner() const;
@@ -203,6 +208,12 @@ private:
 	UPROPERTY()
 	TArray<FName> CachedMorphTargetNames{ NAME_None };
 	
+	UPROPERTY()
+	TOptional<bool> CachedHasVertexColors;
+
+	UPROPERTY()
+	FGuid CachedVertexColorGuid;
+	
 #if WITH_EDITORONLY_DATA
 	/**
 	 * Bulk data containing mesh description from imported or modeled geometry.
diff --git a/Engine/Source/Runtime/Engine/Private/SkeletalMesh.cpp b/Engine/Source/Runtime/Engine/Private/SkeletalMesh.cpp
--- a/Engine/Source/Runtime/Engine/Private/SkeletalMesh.cpp
+++ b/Engine/Source/Runtime/Engine/Private/SkeletalMesh.cpp
@@ -2735,9 +2735,9 @@ bool USkeletalMesh::CommitMeshDescription(
 			// Get all profiles from the models on all LODs, since we may have some that aren't
 			// defined on the skeletal mesh's list of profiles.
 			TSet<FName> ValidProfileNames;
-			for (const FSkeletalMeshSourceModel& OtherSourceModels: GetAllSourceModels())
+			for (const FSkeletalMeshSourceModel& OtherSourceModel: GetAllSourceModels())
 			{
-				ValidProfileNames.Append(OtherSourceModels.GetSkinWeightProfileNames());
+				ValidProfileNames.Append(OtherSourceModel.GetSkinWeightProfileNames());
 			}
 			for (const FName& ProfileName: ValidProfileNames)
 			{
@@ -2812,6 +2812,27 @@ bool USkeletalMesh::CommitMeshDescription(
 					}, TStatId(), nullptr, ENamedThreads::GameThread);	
 #endif			
 			}
+			
+			if (InParams.bUpdateVertexColors)
+			{
+				FGuid SourceVertexColorGuid{};	// Zero-initialized
+				bool bSourceHasVertexColors = false;
+				for (const FSkeletalMeshSourceModel& OtherSourceModel: GetAllSourceModels())
+				{
+					TOptional<FGuid> SourceModelVertexColorGuid = OtherSourceModel.GetVertexColorGuid();
+					if (SourceModelVertexColorGuid.IsSet())
+					{
+						bSourceHasVertexColors = true;
+						SourceVertexColorGuid = FGuid::Combine(SourceVertexColorGuid, *SourceModelVertexColorGuid);
+					}
+				}
+				
+				SetHasVertexColors(bSourceHasVertexColors);
+				if (bSourceHasVertexColors)
+				{
+					SetVertexColorGuid(SourceVertexColorGuid);
+				}
+			}
 		}
 	}
 
diff --git a/Engine/Source/Runtime/Engine/Private/SkeletalMeshSourceModel.cpp b/Engine/Source/Runtime/Engine/Private/SkeletalMeshSourceModel.cpp
--- a/Engine/Source/Runtime/Engine/Private/SkeletalMeshSourceModel.cpp
+++ b/Engine/Source/Runtime/Engine/Private/SkeletalMeshSourceModel.cpp
@@ -2,6 +2,7 @@
 
 #include "Engine/SkeletalMeshSourceModel.h"
 
+#include "Async/ParallelFor.h"
 #include "Rendering/SkeletalMeshModel.h"
 
 #if WITH_EDITOR
@@ -247,6 +248,21 @@ TConstArrayView<FName> FSkeletalMeshSourceModel::GetMorphTargetNames() const
 	return CachedMorphTargetNames;
 }
 
+TOptional<FGuid> FSkeletalMeshSourceModel::GetVertexColorGuid() const
+{
+	if (!CachedHasVertexColors.IsSet())
+	{
+		const_cast<FSkeletalMeshSourceModel*>(this)->UpdateCachedMeshStatisticsFromBulkIfNeeded();
+	}
+
+	if (CachedHasVertexColors.Get(false))
+	{
+		return CachedVertexColorGuid;
+	}
+	
+	return {};
+}
+
 USkeletalMesh* FSkeletalMeshSourceModel::GetOwner() const
 {
 	if (!ensure(MeshDescriptionBulkData))
@@ -385,6 +401,55 @@ void FSkeletalMeshSourceModel::UpdateCachedMeshStatistics(const FMeshDescription
 		constexpr bool bInUserDefinedOnly = true;
 		CachedSkinWeightProfileNames = Attributes.GetSkinWeightProfileNames(bInUserDefinedOnly);
 		CachedMorphTargetNames = Attributes.GetMorphTargetNames();
+
+		// Test whether there are non-zero vertex colors, and compute a pseudo-hash for them to ensure that
+		// each time the vertex colors change, we have a new GUID to ensure the vertex data gets updated. 
+		struct FContext
+		{
+			uint32 Parts[4] = {0};
+			bool bHasVertexColors = false;
+		};
+		TArray<FContext> VertexColorContexts;
+		
+		TVertexInstanceAttributesConstRef<FVector4f> ColorAttrib = Attributes.GetVertexInstanceColors();
+		const FVertexInstanceArray& VertexInstances = InMeshDescription->VertexInstances();
+		ParallelForWithTaskContext(VertexColorContexts, VertexInstances.GetArraySize(), [&ColorAttrib, VertexInstances](FContext &InOutContext, int32 InIndex)
+		{
+			const FVertexInstanceID VertexInstanceID(InIndex);
+			if (VertexInstances.IsValid(VertexInstanceID))
+			{
+				const FVector4f VertexColor = ColorAttrib.Get(VertexInstanceID);
+
+				if (!VertexColor.IsNearlyZero())
+				{
+					InOutContext.bHasVertexColors = true;
+				}
+
+				// Build a pseudo-hash for the vertex colors, in a thread-safe manner by just adding up the float bits. This will be consistent
+				// no matter how the vertex instances are partitioned across threads. 
+				InOutContext.Parts[0] += *reinterpret_cast<const uint32*>(&VertexColor.X) + static_cast<uint32>(InIndex);
+				InOutContext.Parts[1] += *reinterpret_cast<const uint32*>(&VertexColor.Y) + static_cast<uint32>(InIndex);
+				InOutContext.Parts[2] += *reinterpret_cast<const uint32*>(&VertexColor.Z) + static_cast<uint32>(InIndex);
+				InOutContext.Parts[3] += *reinterpret_cast<const uint32*>(&VertexColor.W) + static_cast<uint32>(InIndex);
+			}
+		});
+
+		uint32 Parts[4] = {0};
+		bool bHasVertexColors = false;
+		for (const FContext& Context: VertexColorContexts)
+		{
+			if (Context.bHasVertexColors)
+			{
+				bHasVertexColors = true;
+			}
+			Parts[0] += Context.Parts[0];
+			Parts[1] += Context.Parts[1];
+			Parts[2] += Context.Parts[2];
+			Parts[3] += Context.Parts[3];
+		}
+		
+		CachedHasVertexColors = bHasVertexColors;
+		CachedVertexColorGuid = FGuid(Parts[0], Parts[1], Parts[2], Parts[3]);
 	}
 	else
 	{
@@ -392,6 +457,8 @@ void FSkeletalMeshSourceModel::UpdateCachedMeshStatistics(const FMeshDescription
 		Bounds = FBoxSphereBounds{ForceInitToZero};
 		CachedSkinWeightProfileNames.Reset();
 		CachedMorphTargetNames.Reset();
+		CachedHasVertexColors.Reset();
+		CachedVertexColorGuid = {};
 	}
 }
 
