[AsyncCompilation/Texture]
- Fix crash in landscape while it's being investigated

diff --git a/Engine/Source/Runtime/Engine/Private/Texture.cpp b/Engine/Source/Runtime/Engine/Private/Texture.cpp
--- a/Engine/Source/Runtime/Engine/Private/Texture.cpp
+++ b/Engine/Source/Runtime/Engine/Private/Texture.cpp
@@ -535,7 +535,7 @@ bool UTexture::Modify(bool bAlwaysMarkDirty)
 	if (!IsAsyncCacheComplete())
 	{
 		// try to throw away the data since we're going to invalidate it anyway
-		FAssetCompilingManager::Get().MarkCompilationAsCanceled({ this });
+		//FAssetCompilingManager::Get().MarkCompilationAsCanceled({ this });
 	}
 
 	// Before applying any modification to the texture
