WorldPartition ISM cell transformer (UWorldPartitionRuntimeCellTransformerISM) now creates a AWorldPartitionAutoInstancedActor to allow for other transformers to easily include/exclude by actor type.

diff --git a/Engine/Source/Runtime/Engine/Private/WorldPartition/WorldPartitionRuntimeCellTransformerISM.cpp b/Engine/Source/Runtime/Engine/Private/WorldPartition/WorldPartitionRuntimeCellTransformerISM.cpp
--- a/Engine/Source/Runtime/Engine/Private/WorldPartition/WorldPartitionRuntimeCellTransformerISM.cpp
+++ b/Engine/Source/Runtime/Engine/Private/WorldPartition/WorldPartitionRuntimeCellTransformerISM.cpp
@@ -20,7 +20,7 @@
 
 #define LOCTEXT_NAMESPACE "WorldPartition"
 
-static const FName NAME_AutoInstancingActor(TEXT("AutoInstancingActor"));
+static const FName NAME_AutoInstancedActor(TEXT("AutoInstancedActor"));
 
 UWorldPartitionRuntimeCellTransformerISM::UWorldPartitionRuntimeCellTransformerISM(const FObjectInitializer& ObjectInitializer)
 	: Super(ObjectInitializer)
@@ -38,7 +38,7 @@ UWorldPartitionRuntimeCellTransformerISM::UWorldPartitionRuntimeCellTransformerI
 		{
 			if (AActor* Actor = InPrimitiveComponent ? InPrimitiveComponent->GetOwner() : nullptr)
 			{
-				if (Actor->GetFName() == NAME_AutoInstancingActor)
+				if (Actor->IsA<AWorldPartitionAutoInstancedActor>())
 				{
 					return FLinearColor::MakeRandomSeededColor(GetTypeHash(InPrimitiveComponent->GetFName()));
 				}
@@ -111,7 +111,7 @@ void UWorldPartitionRuntimeCellTransformerISM::Transform(ULevel* InLevel)
 
 	if (NumInstancedComponents)
 	{
-		AActor* PackedActor = NewObject<AActor>(InLevel, NAME_AutoInstancingActor);
+		AActor* PackedActor = NewObject<AWorldPartitionAutoInstancedActor>(InLevel, NAME_AutoInstancedActor);
 
 		for (auto& [ISMComponentDescriptor, ActorComponentBatcherDescriptor] : ISMComponentBatchers)
 		{
@@ -212,4 +212,8 @@ bool UWorldPartitionRuntimeCellTransformerISM::CanRemoveActor(AActor* InActor) c
 }
 #endif
 
+AWorldPartitionAutoInstancedActor::AWorldPartitionAutoInstancedActor(const FObjectInitializer& ObjectInitializer)
+	: Super(ObjectInitializer)
+{}
+
 #undef LOCTEXT_NAMESPACE
\ No newline at end of file
diff --git a/Engine/Source/Runtime/Engine/Public/WorldPartition/WorldPartitionRuntimeCellTransformerISM.h b/Engine/Source/Runtime/Engine/Public/WorldPartition/WorldPartitionRuntimeCellTransformerISM.h
--- a/Engine/Source/Runtime/Engine/Public/WorldPartition/WorldPartitionRuntimeCellTransformerISM.h
+++ b/Engine/Source/Runtime/Engine/Public/WorldPartition/WorldPartitionRuntimeCellTransformerISM.h
@@ -3,6 +3,7 @@
 
 #include "WorldPartition/WorldPartitionRuntimeCellTransformer.h"
 #include "Templates/SubclassOf.h"
+#include "GameFramework/Actor.h"
 
 #include "WorldPartitionRuntimeCellTransformerISM.generated.h"
 
@@ -34,4 +35,11 @@ public:
 	UPROPERTY(EditAnywhere, Category = ISM)
 	uint32 MinNumInstances;
 #endif
+};
+
+/** Actor class used by UWorldPartitionRuntimeCellTransformerISM to save transformed result */
+UCLASS(NotPlaceable, NotBlueprintable, NotBlueprintType, MinimalAPI)
+class AWorldPartitionAutoInstancedActor : public AActor
+{
+	GENERATED_UCLASS_BODY()
 };
\ No newline at end of file
