Fix layer tracking in rewind debug to prevent deleting main instance and to support shared layers

diff --git a/Engine/Source/Runtime/Engine/Private/Animation/AnimSubsystem_SharedLinkedAnimLayers.cpp b/Engine/Source/Runtime/Engine/Private/Animation/AnimSubsystem_SharedLinkedAnimLayers.cpp
--- a/Engine/Source/Runtime/Engine/Private/Animation/AnimSubsystem_SharedLinkedAnimLayers.cpp
+++ b/Engine/Source/Runtime/Engine/Private/Animation/AnimSubsystem_SharedLinkedAnimLayers.cpp
@@ -3,6 +3,7 @@
 #include "Animation/AnimSubsystem_SharedLinkedAnimLayers.h"
 #include "Animation/AnimInstance.h"
 #include "Components/SkeletalMeshComponent.h"
+#include "ObjectTrace.h"
 
 #include UE_INLINE_GENERATED_CPP_BY_NAME(AnimSubsystem_SharedLinkedAnimLayers)
 
@@ -173,6 +174,7 @@ void FLinkedAnimLayerClassData::RemoveInstance(UAnimInstance* AnimInstance)
 			InstancesData.RemoveAt(i);
 
 			AnimInstance->UninitializeAnimation();
+			TRACE_OBJECT_LIFETIME_END(AnimInstance);
 			if (SharedLinkedAnimLayersConsoleCommands::MarkLayerAsGarbageOnUninitialize)
 			{
 				AnimInstance->MarkAsGarbage();
