Fixing compile error when building with remote object handles and FUObjectItem packing  and probably deprecation warning without FUObjectItem packing

diff --git a/Engine/Source/Runtime/Engine/Private/Collision/Collision.cpp b/Engine/Source/Runtime/Engine/Private/Collision/Collision.cpp
--- a/Engine/Source/Runtime/Engine/Private/Collision/Collision.cpp
+++ b/Engine/Source/Runtime/Engine/Private/Collision/Collision.cpp
@@ -177,7 +177,7 @@ void FCollisionQueryParams::AddIgnoredSourceObject(const uint32 InIgnoreActorID)
 	IgnoreSourceObjects.Add(InIgnoreActorID);
 
 #if UE_WITH_REMOTE_OBJECT_HANDLE
-	UObject* Obj = (UObject*)GUObjectArray.IndexToObject(InIgnoreActorID)->Object;
+	UObject* Obj = (UObject*)GUObjectArray.IndexToObject(InIgnoreActorID)->GetObject();
 	IgnoreSourceObjectReferences.Add(FRemoteObjectReference(Obj));
 #endif
 }
