Fix typo of 16 bit float's max finite value in the FPScale of Orthographic autoplane scaling.

diff --git a/Engine/Source/Runtime/Engine/Private/Camera/CameraStackTypes.cpp b/Engine/Source/Runtime/Engine/Private/Camera/CameraStackTypes.cpp
--- a/Engine/Source/Runtime/Engine/Private/Camera/CameraStackTypes.cpp
+++ b/Engine/Source/Runtime/Engine/Private/Camera/CameraStackTypes.cpp
@@ -391,7 +391,7 @@ bool FMinimalViewInfo::AutoCalculateOrthoPlanes(FSceneViewProjectionData& InOutP
 		const bool bUse16bitDepth = CVarOrthoClampToMaxFPBuffer.GetValueOnAnyThread() == 1;
 		const bool bScaleIncrementingUnits = CVarOrthoScaleIncrementingUnits.GetValueOnAnyThread() && bUse16bitDepth;
 		const float MaxFPValue = bScaleIncrementingUnits ? UE_LARGE_WORLD_MAX : UE_OLD_WORLD_MAX;
-		float FPScale = bUse16bitDepth ? 66504.0f : UE_OLD_WORLD_MAX;
+		float FPScale = bUse16bitDepth ? 65504.0f : UE_OLD_WORLD_MAX;
 
 		const float AutoDepthScale = CVarOrthoAutoDepthScale.GetValueOnAnyThread();
 		if (AutoDepthScale > 0.0f)
