WorldPartition Removed hardcoded name when UWorldPartitionRuntimeCellTransformerISM generates its AWorldPartitionAutoInstancedActor as this is not required anymore. This could also lead to a crash if 2 instances of UWorldPartitionRuntimeCellTransformerISM (with different settings) were put in the cell transformer stack.

diff --git a/Engine/Source/Runtime/Engine/Private/WorldPartition/WorldPartitionRuntimeCellTransformerISM.cpp b/Engine/Source/Runtime/Engine/Private/WorldPartition/WorldPartitionRuntimeCellTransformerISM.cpp
--- a/Engine/Source/Runtime/Engine/Private/WorldPartition/WorldPartitionRuntimeCellTransformerISM.cpp
+++ b/Engine/Source/Runtime/Engine/Private/WorldPartition/WorldPartitionRuntimeCellTransformerISM.cpp
@@ -20,8 +20,6 @@
 
 #define LOCTEXT_NAMESPACE "WorldPartition"
 
-static const FName NAME_AutoInstancedActor(TEXT("AutoInstancedActor"));
-
 UWorldPartitionRuntimeCellTransformerISM::UWorldPartitionRuntimeCellTransformerISM(const FObjectInitializer& ObjectInitializer)
 	: Super(ObjectInitializer)
 {
@@ -111,7 +109,7 @@ void UWorldPartitionRuntimeCellTransformerISM::Transform(ULevel* InLevel)
 
 	if (NumInstancedComponents)
 	{
-		AActor* PackedActor = NewObject<AWorldPartitionAutoInstancedActor>(InLevel, NAME_AutoInstancedActor);
+		AActor* PackedActor = NewObject<AWorldPartitionAutoInstancedActor>(InLevel);
 
 		for (auto& [ISMComponentDescriptor, ActorComponentBatcherDescriptor] : ISMComponentBatchers)
 		{
