Lights: Allow EV units to become negative

Some of the validation on light intensity did not consider the EV unit case, leading lights with 0 EV or less to be missing from the scene.

diff --git a/Engine/Source/Runtime/Engine/Private/Components/LightComponent.cpp b/Engine/Source/Runtime/Engine/Private/Components/LightComponent.cpp
--- a/Engine/Source/Runtime/Engine/Private/Components/LightComponent.cpp
+++ b/Engine/Source/Runtime/Engine/Private/Components/LightComponent.cpp
@@ -610,7 +610,11 @@ void ULightComponent::PostLoad()
 		PreviewShadowMapChannel = INDEX_NONE;
 	}
 
-	Intensity = FMath::Max(0.0f, Intensity);
+	// EV is a logarithmic unit, so negative values are fine
+	if (GetLightUnits() != ELightUnits::EV)
+	{
+		Intensity = FMath::Max(0.0f, Intensity);
+	}
 
 	if (GetLinkerUEVersion() < VER_UE4_LIGHTCOMPONENT_USE_IES_TEXTURE_MULTIPLIER_ON_NON_IES_BRIGHTNESS)
 	{
@@ -896,7 +900,8 @@ void ULightComponent::CreateRenderState_Concurrent(FRegisterComponentContext* Co
 	if (bAffectsWorld)
 	{
 		UWorld* World = GetWorld();
-		const bool bHidden = !ShouldComponentAddToScene() || !ShouldRender() || Intensity <= 0.f;
+		const bool bValidIntensity = Intensity > 0.f || GetLightUnits() == ELightUnits::EV;
+		const bool bHidden = !ShouldComponentAddToScene() || !ShouldRender() || !bValidIntensity;
 		if (!bHidden)
 		{
 			InitializeStaticShadowDepthMap();
@@ -1311,15 +1316,16 @@ void ULightComponent::UpdateColorAndBrightness()
 	UWorld* World = GetWorld();
 	if( World && World->Scene )
 	{
-		const bool bNeedsToBeAddedToScene = (!bAddedToSceneVisible && Intensity > 0.f);
-		const bool bNeedsToBeRemovedFromScene = (bAddedToSceneVisible && Intensity <= 0.f);
+		const bool bValidIntensity = Intensity > 0.f || GetLightUnits() == ELightUnits::EV;
+		const bool bNeedsToBeAddedToScene = (!bAddedToSceneVisible && bValidIntensity);
+		const bool bNeedsToBeRemovedFromScene = (bAddedToSceneVisible && !bValidIntensity);
 		if (bNeedsToBeAddedToScene || bNeedsToBeRemovedFromScene)
 		{
 			// We may have just been set to 0 intensity or we were previously 0 intensity.
 			// Mark the render state dirty to add or remove this light from the scene as necessary.
 			MarkRenderStateDirty();
 		}
-		else if (bAddedToSceneVisible && Intensity > 0.f)
+		else if (bAddedToSceneVisible && bValidIntensity)
 		{
 			// We are already in the scene. Just update with this fast path command
 			World->Scene->UpdateLightColorAndBrightness(this);
