Added FTestWorlds::OnClientConnected delegate to run generic code whenever a client finished connected to the server

#rnx

diff --git a/Engine/Source/Runtime/Engine/Private/Tests/NetTestHelpers.cpp b/Engine/Source/Runtime/Engine/Private/Tests/NetTestHelpers.cpp
--- a/Engine/Source/Runtime/Engine/Private/Tests/NetTestHelpers.cpp
+++ b/Engine/Source/Runtime/Engine/Private/Tests/NetTestHelpers.cpp
@@ -336,7 +336,7 @@ bool FTestWorlds::CreateAndConnectClient()
 
 bool FTestWorlds::WaitForClientConnect(FTestWorldInstance& Client)
 {
-	return TickAllUntil([&Client]()
+	bool bConnected = TickAllUntil([&Client]()
 	{
 		if (UWorld* World = Client.GetWorld())
 		{
@@ -345,6 +345,13 @@ bool FTestWorlds::WaitForClientConnect(FTestWorldInstance& Client)
 		}
 		return false;
 	});
+
+	if (bConnected)
+	{
+		OnClientConnected.Broadcast(Client);
+	}
+
+	return bConnected;
 }
 
 void FTestWorlds::TickAll(int32 NumTicks)
diff --git a/Engine/Source/Runtime/Engine/Public/Tests/NetTestHelpers.h b/Engine/Source/Runtime/Engine/Public/Tests/NetTestHelpers.h
--- a/Engine/Source/Runtime/Engine/Public/Tests/NetTestHelpers.h
+++ b/Engine/Source/Runtime/Engine/Public/Tests/NetTestHelpers.h
@@ -195,10 +195,6 @@ struct FTestWorlds
 	template<class PredicateT>
 	bool TickAllUntil(const PredicateT& Predicate, float DeltaSeconds = 0.0166f, int32 MaxTicks = 60);
 
-	DECLARE_MULTICAST_DELEGATE(FPreTickAll);
-	/** Delegate called before a TickAll is executed */
-	FPreTickAll PreTickAllDelegate;
-
 	/** Ticks all server & client worlds until the passed in client world has a valid client PlayerController. */
 	ENGINE_API bool WaitForClientConnect(FTestWorldInstance& Client);
 
@@ -232,6 +228,16 @@ struct FTestWorlds
 	 */
 	ENGINE_API bool IsServerObjectReplicated(const UObject* ServerObject) const;
 
+public:
+
+	/** Delegate called before a TickAll is executed */
+	DECLARE_MULTICAST_DELEGATE(FPreTickAll);
+	FPreTickAll PreTickAllDelegate;
+
+	/** Delegate called after a client finished connecting to the server */
+	DECLARE_MULTICAST_DELEGATE_OneParam(FOnClientConnected, const FTestWorldInstance& Client);
+	FOnClientConnected OnClientConnected;
+
 public:
 
 	/** Server and Client Worlds */
