New material translator:
- Reworked the FEmitter API to be more consistent, less verbose and informative. Introduce a FValueRef concept to keep track of the input that originated the value, so that in case of erro the name of the expression input can be reported with the error message.
- In the process introduced new diagnostic APIs.

#rnx

diff --git a/Engine/Source/Runtime/Engine/Private/Materials/MaterialExpressionsIR.cpp b/Engine/Source/Runtime/Engine/Private/Materials/MaterialExpressionsIR.cpp
--- a/Engine/Source/Runtime/Engine/Private/Materials/MaterialExpressionsIR.cpp
+++ b/Engine/Source/Runtime/Engine/Private/Materials/MaterialExpressionsIR.cpp
@@ -243,43 +243,45 @@
 
 #include "Materials/MaterialIREmitter.h"
 
+using FValueRef = MIR::FValueRef;
+
 /* Constants */
 
-void UMaterialExpression::Build(MIR::FEmitter& Emitter)
+void UMaterialExpression::Build(MIR::FEmitter& Em)
 {
-	Emitter.Error(TEXT("Unsupported material expression."));
+	Em.Error(TEXT("Unsupported material expression."));
 } 
 
-void UMaterialExpressionFunctionInput::Build(MIR::FEmitter& Emitter)
+void UMaterialExpressionFunctionInput::Build(MIR::FEmitter& Em)
 {
-	MIR::FValue* OutputValue = Emitter.TryInput(&Preview);
+	FValueRef OutputValue = Em.TryInput(&Preview);
 	if (OutputValue)
 	{
-		Emitter.Output(0, OutputValue);
+		Em.Output(0, OutputValue);
 		return;
 	}
 
 	switch (InputType)
 	{
 		case FunctionInput_Scalar:
-			OutputValue = Emitter.EmitConstantFloatScalar(PreviewValue.X);
+			OutputValue = Em.ConstantFloat(PreviewValue.X);
 			break;
 
 		case FunctionInput_Vector2:
-			OutputValue = Emitter.EmitConstantFloat2({ PreviewValue.X, PreviewValue.Y });
+			OutputValue = Em.ConstantFloat2({ PreviewValue.X, PreviewValue.Y });
 			break;
 
 		case FunctionInput_Vector3:
-			OutputValue = Emitter.EmitConstantFloat3({ PreviewValue.X, PreviewValue.Y, PreviewValue.Z });
+			OutputValue = Em.ConstantFloat3({ PreviewValue.X, PreviewValue.Y, PreviewValue.Z });
 			break;
 
 		case FunctionInput_Vector4:
-			OutputValue = Emitter.EmitConstantFloat4(PreviewValue);
+			OutputValue = Em.ConstantFloat4(PreviewValue);
 			break;
 
 		case FunctionInput_Bool:
 		case FunctionInput_StaticBool:
-			OutputValue = Emitter.EmitConstantBoolScalar(PreviewValue.X != 0.0f);
+			OutputValue = Em.ConstantBool(PreviewValue.X != 0.0f);
 
 		case FunctionInput_Texture2D:
 		case FunctionInput_TextureCube:
@@ -288,168 +290,139 @@ void UMaterialExpressionFunctionInput::Build(MIR::FEmitter& Emitter)
 		case FunctionInput_MaterialAttributes:
 		case FunctionInput_TextureExternal:
 		case FunctionInput_Substrate:
-			Emitter.Error(TEXT("Function input of object type requires preview input to be provided."));
+			Em.Error(TEXT("Function input of object type requires preview input to be provided."));
 			return;
 
 		default:
 			UE_MIR_UNREACHABLE();
 	}
 
-	Emitter.Output(0, OutputValue);
+	Em.Output(0, OutputValue);
 }
 
-void UMaterialExpressionFunctionOutput::Build(MIR::FEmitter& Emitter)
+void UMaterialExpressionFunctionOutput::Build(MIR::FEmitter& Em)
 {
-	Emitter.Output(0, Emitter.Input(&A));
+	Em.Output(0, Em.Input(&A));
 }
 
-void UMaterialExpressionConstant::Build(MIR::FEmitter& Emitter)
+void UMaterialExpressionConstant::Build(MIR::FEmitter& Em)
 {
-	MIR::FValue* Value = Emitter.EmitConstantFloatScalar(R);
-	Emitter.Output(GetOutput(0), Value);
+	FValueRef Value = Em.ConstantFloat(R);
+	Em.Output(0, Value);
 }
 
-void UMaterialExpressionConstant2Vector::Build(MIR::FEmitter& Emitter)
+void UMaterialExpressionConstant2Vector::Build(MIR::FEmitter& Em)
 {
-	MIR::FValue* Value = Emitter.EmitConstantFloat2({ R, G });
-	Emitter.Output(GetOutput(0), Value);
+	FValueRef Value = Em.ConstantFloat2({ R, G });
+	Em.Output(0, Value);
 	for (int i = 0; i < 2; ++i)
 	{
-		Emitter.Output(GetOutput(i + 1), Emitter.EmitSubscript(Value, i));
+		Em.Output(i + 1, Em.Subscript(Value, i));
 	}
 }
 
-void UMaterialExpressionConstant3Vector::Build(MIR::FEmitter& Emitter)
+void UMaterialExpressionConstant3Vector::Build(MIR::FEmitter& Em)
 {
-	MIR::FValue* Value = Emitter.EmitConstantFloat3({ Constant.R, Constant.G, Constant.B });
-	Emitter.Output(GetOutput(0), Value);
+	FValueRef Value = Em.ConstantFloat3({ Constant.R, Constant.G, Constant.B });
+	Em.Output(0, Value);
 	for (int i = 0; i < 3; ++i)
 	{
-		Emitter.Output(GetOutput(i + 1), Emitter.EmitSubscript(Value, i));
+		Em.Output(i + 1, Em.Subscript(Value, i));
 	}
 }
 
-void UMaterialExpressionConstant4Vector::Build(MIR::FEmitter& Emitter)
+void UMaterialExpressionConstant4Vector::Build(MIR::FEmitter& Em)
 {
-	MIR::FValue* Value = Emitter.EmitConstantFloat4(Constant);
-	Emitter.Output(GetOutput(0), Value);
+	FValueRef Value = Em.ConstantFloat4(Constant);
+	Em.Output(0, Value);
 	for (int i = 0; i < 4; ++i)
 	{
-		Emitter.Output(GetOutput(i + 1), Emitter.EmitSubscript(Value, i));
+		Em.Output(i + 1, Em.Subscript(Value, i));
 	}
 }
 
-void UMaterialExpressionStaticBool::Build(MIR::FEmitter& Emitter)
+void UMaterialExpressionStaticBool::Build(MIR::FEmitter& Em)
 {
-	Emitter.Output(0, Emitter.EmitConstantBoolScalar(Value));
+	Em.Output(0, Em.ConstantBool(Value));
 }
 
-static MIR::FValue* BuildMaterialExpressionParameter(UMaterialExpressionParameter* ParameterExpr, MIR::FEmitter& Emitter)
+static FValueRef BuildMaterialExpressionParameter(MIR::FEmitter& Em, UMaterialExpressionParameter* ParameterExpr)
 {
 	FMaterialParameterMetadata Metadata;
 	if (!ParameterExpr->GetParameterValue(Metadata))
 	{
-		Emitter.Error(TEXT("Could not get parameter value."));
-		return MIR::FPoison::Get();
+		Em.Error(TEXT("Could not get parameter value."));
+		return Em.Poison();
 	}
 
-	return Emitter.EmitParameter(ParameterExpr->GetParameterName(), Metadata);
+	return Em.Parameter(ParameterExpr->GetParameterName(), Metadata);
 }
 
-void UMaterialExpressionParameter::Build(MIR::FEmitter& Emitter)
+void UMaterialExpressionParameter::Build(MIR::FEmitter& Em)
 {
-	Emitter.Output(0, BuildMaterialExpressionParameter(this, Emitter));
+	Em.Output(0, BuildMaterialExpressionParameter(Em, this));
 }
 
-void UMaterialExpressionVectorParameter::Build(MIR::FEmitter& Emitter)
+void UMaterialExpressionVectorParameter::Build(MIR::FEmitter& Em)
 {
-	MIR::FValue* Value = BuildMaterialExpressionParameter(this, Emitter);
-	Emitter.Output(0, Value);
-	Emitter.Output(1, Emitter.EmitSubscript(Value, 0));
-	Emitter.Output(2, Emitter.EmitSubscript(Value, 1));
-	Emitter.Output(3, Emitter.EmitSubscript(Value, 2));
-	Emitter.Output(4, Emitter.EmitSubscript(Value, 3));
+	FValueRef Value = BuildMaterialExpressionParameter(Em, this);
+	Em.Output(0, Value);
+	Em.Output(1, Em.Subscript(Value, 0));
+	Em.Output(2, Em.Subscript(Value, 1));
+	Em.Output(3, Em.Subscript(Value, 2));
+	Em.Output(4, Em.Subscript(Value, 3));
 }
 
-void UMaterialExpressionChannelMaskParameter::Build(MIR::FEmitter& Emitter)
+void UMaterialExpressionChannelMaskParameter::Build(MIR::FEmitter& Em)
 {
-	MIR::FValue* InputValue = Emitter.Input(&Input);
-	InputValue = Emitter.EmitConstruct(MIR::FPrimitiveType::GetVector(MIR::ScalarKind_Float, 4), InputValue);
-	MIR::FValue* ParameterValue = BuildMaterialExpressionParameter(this, Emitter);
-	MIR::FValue* DotResult = Emitter.EmitOperator(MIR::BO_Dot, InputValue, ParameterValue);
+	FValueRef DotResult = Em.Dot(
+		Em.CastToFloat(Em.Input(&Input), 4),
+		BuildMaterialExpressionParameter(Em, this));
 
-	Emitter.Output(0, DotResult);
-	Emitter.Output(1, Emitter.EmitSubscript(DotResult, 0));
-	Emitter.Output(2, Emitter.EmitSubscript(DotResult, 1));
-	Emitter.Output(3, Emitter.EmitSubscript(DotResult, 2));
-	Emitter.Output(4, Emitter.EmitSubscript(DotResult, 3));
+	Em.Output(0, DotResult);
+	Em.Output(1, Em.Subscript(DotResult, 1));
+	Em.Output(2, Em.Subscript(DotResult, 2));
+	Em.Output(3, Em.Subscript(DotResult, 3));
+	Em.Output(4, Em.Subscript(DotResult, 4));
 }
 
-void UMaterialExpressionStaticBoolParameter::Build(MIR::FEmitter& Emitter)
+void UMaterialExpressionStaticBoolParameter::Build(MIR::FEmitter& Em)
 {
-	MIR::FValue* Value = BuildMaterialExpressionParameter(this, Emitter);
-	
-	// Check value is a static bool scalar.
-	Emitter.ToConstantBoolScalar(Value);
-
-	Emitter.Output(0, Value);
+	FValueRef Value = BuildMaterialExpressionParameter(Em, this);
+	Em.ToConstantBool(Value); // Check that it is a constant boolean
+	Em.Output(0, Value);
 }
 
-void UMaterialExpressionStaticSwitch::Build(MIR::FEmitter& Emitter)
+void UMaterialExpressionStaticSwitch::Build(MIR::FEmitter& Em)
 {
-	Emitter.DefaultInputToBoolScalar(&Value, (bool)DefaultValue);
-	bool bCondition = Emitter.InputConstantBoolScalar(&Value);
-
-	UE_MIR_CHECKPOINT();
+	bool bCondition = Em.ConstantBool(Em.InputDefaultBool(&Value, DefaultValue));
+	UE_MIR_CHECKPOINT(Em); // Make sure that evaluating bCondition didn't raise an error
 
-	Emitter.Output(0, Emitter.Input(bCondition ? &A : &B));
+	Em.Output(0, Em.Input(bCondition ? &A : &B));
 }
 
-void UMaterialExpressionStaticSwitchParameter::Build(MIR::FEmitter& Emitter)
+void UMaterialExpressionStaticSwitchParameter::Build(MIR::FEmitter& Em)
 {
-	MIR::FValue* Value = BuildMaterialExpressionParameter(this, Emitter);
+	bool bCondition = Em.ToConstantBool(BuildMaterialExpressionParameter(Em, this));
+	UE_MIR_CHECKPOINT(Em); // Make sure that evaluating bCondition didn't raise an error
 
-	// Expect it to be a constant bool scalar, get its value.
-	bool bCondition = Emitter.ToConstantBoolScalar(Value);
-
-	UE_MIR_CHECKPOINT();
-
-	// Select active input
-	FExpressionInput* ActiveInput = bCondition ? &A : &B;
-
-	// And bypass its value to our output
-	Emitter.Output(0, Emitter.Input(ActiveInput));
+	Em.Output(0, Em.Input(bCondition ? &A : &B));
 }
 
-void UMaterialExpressionAppendVector::Build(MIR::FEmitter& Emitter)
+void UMaterialExpressionAppendVector::Build(MIR::FEmitter& Em)
 {
-	MIR::FValue* AValue = Emitter.InputPrimitive(&A);
-	MIR::FValue* BValue = Emitter.TryInput(&B) ? Emitter.InputPrimitive(&B) : nullptr;
-
-	UE_MIR_CHECKPOINT();
+	FValueRef AVal = Em.CheckIsScalarOrVector(Em.Input(&A));
+	FValueRef BVal = Em.CheckIsScalarOrVector(Em.TryInput(&B));
 
-	const MIR::FPrimitiveType* AType = AValue->Type->AsPrimitive();
-	if (AType->IsMatrix())
-	{
-		Emitter.Error(TEXT("Input A does not have scalar or vector type."));
-		return;
-	}
+	UE_MIR_CHECKPOINT(Em);
 
-	const MIR::FPrimitiveType* BType = nullptr;
-	if (BValue)
-	{
-		BType = BValue->Type->AsPrimitive();
-		if (BType->IsMatrix()) 
-		{
-			Emitter.Error(TEXT("Input B does not have scalar or vector type."));
-			return;
-		}
-	}
+	const MIR::FPrimitiveType* AType = AVal->Type->AsPrimitive();
+	const MIR::FPrimitiveType* BType = BVal ? BVal->Type->AsPrimitive() : nullptr;
 
 	int Dimensions = AType->NumRows + (BType ? BType->NumRows : 0);
 	if (Dimensions > 4)
 	{
-		Emitter.Errorf(TEXT("The resulting vector would have %d component (it can have at most 4)."), Dimensions);
+		Em.Errorf(TEXT("The resulting vector would have %d component (it can have at most 4)."), Dimensions);
 		return;
 	}
 
@@ -459,198 +432,119 @@ void UMaterialExpressionAppendVector::Build(MIR::FEmitter& Emitter)
 	const MIR::FPrimitiveType* ResultType = MIR::FPrimitiveType::GetVector(MIR::ScalarKind_Float, Dimensions);
 
 	// Set up each output vector component 
-	MIR::FValue* Components[4] = { nullptr, nullptr, nullptr, nullptr };
+	FValueRef Components[4] = { nullptr, nullptr, nullptr, nullptr };
 	int ComponentIndex = 0;
 	for (int i = 0; i < AType->NumRows; ++i, ++ComponentIndex)
 	{
-		Components[ComponentIndex] = Emitter.EmitSubscript(AValue, i);
+		Components[ComponentIndex] = Em.Subscript(AVal, i);
 	}
 
 	if (BType)
 	{
 		for (int i = 0; i < BType->NumRows; ++i, ++ComponentIndex)
 		{
-			Components[ComponentIndex] = Emitter.EmitSubscript(BValue, i);
+			Components[ComponentIndex] = Em.Subscript(BVal, i);
 		}
 	}
 
 	// Create the vector value and output it.
-	MIR::FValue* Output{};
+	FValueRef Output{};
 	if (Dimensions == 2)
 	{
-		Output = Emitter.EmitVector2(Components[0], Components[1]);
+		Output = Em.Vector2(Components[0], Components[1]);
 	}
 	else if (Dimensions == 3)
 	{
-		Output = Emitter.EmitVector3(Components[0], Components[1], Components[2]);
+		Output = Em.Vector3(Components[0], Components[1], Components[2]);
 	}
 	else
 	{
-		Output = Emitter.EmitVector4(Components[0], Components[1], Components[2], Components[3]);
+		Output = Em.Vector4(Components[0], Components[1], Components[2], Components[3]);
 	}
 
-	Emitter.Output(0, Output);
+	Em.Output(0, Output);
 }
 
 /* Unary Operators */
-static void BuildUnaryArithmeticOperatorFromValue(
-	MIR::FEmitter& Emitter,
-	MIR::EOperator Op,
-	MIR::FValue* Value)
-{
-	// Emit the unary operator.
-	Value = Emitter.EmitOperator(Op, Value);
-
-	// And flow it out of the expression's only output.
-	Emitter.Output(0, Value);
-}
-
-static void BuildUnaryArithmeticOperatorFromInputWithDefault(
-	MIR::FEmitter& Emitter,
-	MIR::EOperator Op,
-	FExpressionInput* Input,
-	float Const)
-{
-	// Default input to its relative constants if disconnected, then get it after checking it has primitive type.
-	MIR::FValue* Value = Emitter.DefaultInputToFloatScalar(Input, Const).Input(Input);
-	BuildUnaryArithmeticOperatorFromValue(Emitter, Op, Value);
-}
-
-// Special case function which doesn't try to emit a default value and will error if you don't have an Input
-static void BuildUnaryArithmeticOperatorFromInput(
-	MIR::FEmitter& Emitter,
-	MIR::EOperator Op,
-	FExpressionInput* Input)
-{
-	// Get input after checking it has primitive type.
-	MIR::FValue* Value = Emitter.Input(Input);
-	BuildUnaryArithmeticOperatorFromValue(Emitter, Op, Value);
-}
-
-static void BuildPeriodicTrigonometricOperatorFromInput(
-	MIR::FEmitter& Emitter,
-	MIR::EOperator Op,
-	FExpressionInput* Input,
-	float Period)
-{
-	// Get input after checking it has primitive type.
-	MIR::FValue* Value = Emitter.Input(Input);
-
-	if (Period > 0.0f)
-	{
-		Value = Emitter.EmitOperator(MIR::BO_Multiply, Value, Emitter.EmitConstantFloatScalar(2.0f * (float)UE_PI / Period));
-	}
-
-	BuildUnaryArithmeticOperatorFromValue(Emitter, Op, Value);
-}
-
-
-// TODO: Add the missing expressions
-
-void UMaterialExpressionAbs::Build(MIR::FEmitter& Emitter)
-{
-	BuildUnaryArithmeticOperatorFromInput(Emitter, MIR::UO_Abs, &Input);
-}
-
-void UMaterialExpressionCeil::Build(MIR::FEmitter& Emitter)
-{
-	BuildUnaryArithmeticOperatorFromInput(Emitter, MIR::UO_Ceil, &Input);
-}
-
-void UMaterialExpressionFloor::Build(MIR::FEmitter& Emitter)
-{
-	BuildUnaryArithmeticOperatorFromInput(Emitter, MIR::UO_Floor, &Input);
-}
-
-void UMaterialExpressionFrac::Build(MIR::FEmitter& Emitter)
-{
-	BuildUnaryArithmeticOperatorFromInput(Emitter, MIR::UO_Frac, &Input);
-}
-
-void UMaterialExpressionLength::Build(MIR::FEmitter& Emitter)
-{
-	MIR::FValue* Value = Emitter.Input(&Input);
-	BuildUnaryArithmeticOperatorFromValue(Emitter, MIR::UO_Length, Value);
-}
 
-void UMaterialExpressionRound::Build(MIR::FEmitter& Emitter)
+void UMaterialExpressionAbs::Build(MIR::FEmitter& Em)
 {
-	BuildUnaryArithmeticOperatorFromInput(Emitter, MIR::UO_Round, &Input);
+	Em.Output(0, Em.Abs(Em.Input(&Input)));
 }
 
-void UMaterialExpressionExponential::Build(MIR::FEmitter& Emitter)
+void UMaterialExpressionCeil::Build(MIR::FEmitter& Em)
 {
-	BuildUnaryArithmeticOperatorFromInput(Emitter, MIR::UO_Exponential, &Input);
+	Em.Output(0, Em.Ceil(Em.Input(&Input)));
 }
 
-void UMaterialExpressionExponential2::Build(MIR::FEmitter& Emitter)
+void UMaterialExpressionFloor::Build(MIR::FEmitter& Em)
 {
-	BuildUnaryArithmeticOperatorFromInput(Emitter, MIR::UO_Exponential2, &Input);
+	Em.Output(0, Em.Floor(Em.Input(&Input)));
 }
 
-void UMaterialExpressionLogarithm::Build(MIR::FEmitter& Emitter)
+void UMaterialExpressionFrac::Build(MIR::FEmitter& Em)
 {
-	BuildUnaryArithmeticOperatorFromInput(Emitter, MIR::UO_Logarithm, &Input);
+	Em.Output(0, Em.Frac(Em.Input(&Input)));
 }
 
-void UMaterialExpressionLogarithm2::Build(MIR::FEmitter& Emitter)
+void UMaterialExpressionLength::Build(MIR::FEmitter& Em)
 {
-	BuildUnaryArithmeticOperatorFromInput(Emitter, MIR::UO_Logarithm2, &X);
+	Em.Output(0, Em.Length(Em.Input(&Input)));
 }
 
-void UMaterialExpressionLogarithm10::Build(MIR::FEmitter& Emitter)
+void UMaterialExpressionRound::Build(MIR::FEmitter& Em)
 {
-	BuildUnaryArithmeticOperatorFromInput(Emitter, MIR::UO_Logarithm10, &X);
+	Em.Output(0, Em.Round(Em.Input(&Input)));
 }
 
-void UMaterialExpressionTruncate::Build(MIR::FEmitter& Emitter)
+void UMaterialExpressionExponential::Build(MIR::FEmitter& Em)
 {
-	BuildUnaryArithmeticOperatorFromInput(Emitter, MIR::UO_Truncate, &Input);
+	Em.Output(0, Em.Exponential(Em.Input(&Input)));
 }
 
-void UMaterialExpressionArccosine::Build(MIR::FEmitter& Emitter)
+void UMaterialExpressionExponential2::Build(MIR::FEmitter& Em)
 {
-	BuildUnaryArithmeticOperatorFromInput(Emitter, MIR::UO_ACos, &Input);
+	Em.Output(0, Em.Exponential2(Em.Input(&Input)));
 }
 
-void UMaterialExpressionArcsine::Build(MIR::FEmitter& Emitter)
+void UMaterialExpressionLogarithm::Build(MIR::FEmitter& Em)
 {
-	BuildUnaryArithmeticOperatorFromInput(Emitter, MIR::UO_ASin, &Input);
+	Em.Output(0, Em.Logarithm(Em.Input(&Input)));
 }
 
-void UMaterialExpressionArctangent::Build(MIR::FEmitter& Emitter)
+void UMaterialExpressionLogarithm2::Build(MIR::FEmitter& Em)
 {
-	BuildUnaryArithmeticOperatorFromInput(Emitter, MIR::UO_ATan, &Input);
+	Em.Output(0, Em.Logarithm2(Em.Input(&X)));
 }
 
-void UMaterialExpressionCosine::Build(MIR::FEmitter& Emitter)
+void UMaterialExpressionLogarithm10::Build(MIR::FEmitter& Em)
 {
-	BuildPeriodicTrigonometricOperatorFromInput(Emitter, MIR::UO_Cos, &Input, Period);
+	Em.Output(0, Em.Logarithm10(Em.Input(&X)));
 }
 
-void UMaterialExpressionSine::Build(MIR::FEmitter& Emitter)
+void UMaterialExpressionTruncate::Build(MIR::FEmitter& Em)
 {
-	BuildPeriodicTrigonometricOperatorFromInput(Emitter, MIR::UO_Sin, &Input, Period);
+	Em.Output(0, Em.Truncate(Em.Input(&Input)));
 }
 
-void UMaterialExpressionSaturate::Build(MIR::FEmitter& Emitter)
+void UMaterialExpressionArccosine::Build(MIR::FEmitter& Em)
 {
-	BuildUnaryArithmeticOperatorFromInput(Emitter, MIR::UO_Saturate, &Input);
+	Em.Output(0, Em.ACos(Em.Input(&Input)));
 }
 
-void UMaterialExpressionSquareRoot::Build(MIR::FEmitter& Emitter)
+void UMaterialExpressionArcsine::Build(MIR::FEmitter& Em)
 {
-	BuildUnaryArithmeticOperatorFromInput(Emitter, MIR::UO_Sqrt, &Input);
+	Em.Output(0, Em.ASin(Em.Input(&Input)));
 }
 
-void UMaterialExpressionTangent::Build(MIR::FEmitter& Emitter)
+void UMaterialExpressionArctangent::Build(MIR::FEmitter& Em)
 {
-	BuildPeriodicTrigonometricOperatorFromInput(Emitter, MIR::UO_Tan, &Input, Period);
+	Em.Output(0, Em.ATan(Em.Input(&Input)));
 }
 
 void UMaterialExpressionComponentMask::Build(MIR::FEmitter& Emitter)
 {
-	MIR::FValue* Value = Emitter.Input(&Input);
+	FValueRef Value = Emitter.Input(&Input);
 
 	MIR::FSwizzleMask Mask;
 	if (R)
@@ -670,291 +564,233 @@ void UMaterialExpressionComponentMask::Build(MIR::FEmitter& Emitter)
 		Mask.Components[Mask.NumComponents++] = MIR::EVectorComponent::W;
 	}
 
-	Emitter.Output(0, Emitter.EmitSwizzle(Value, Mask));
+	Emitter.Output(0, Emitter.Swizzle(Value, Mask));
 }
 
-static MIR::FValue* EmitInlineHLSL(const UMaterialExpressionExternalCodeBase& InExternalCodeExpression, MIR::FEmitter& Emitter, int32 InExternalCodeIdentifierIndex, TConstArrayView<MIR::FValue*> InArguments)
-{
-	checkf(
-		InExternalCodeIdentifierIndex >= 0 && InExternalCodeIdentifierIndex < InExternalCodeExpression.ExternalCodeIdentifiers.Num(),
-		TEXT("External code identifier index (%d) out of bounds; Upper bound is %d"), InExternalCodeIdentifierIndex, InExternalCodeExpression.ExternalCodeIdentifiers.Num()
-	);
 
-	const FName& ExternalCodeIdentifier = InExternalCodeExpression.ExternalCodeIdentifiers[InExternalCodeIdentifierIndex];
-	const FMaterialExternalCodeDeclaration* ExternalCodeDeclaration = MaterialExternalCodeRegistry::Get().FindExternalCode(ExternalCodeIdentifier);
-	if (!ExternalCodeDeclaration)
+static FValueRef GetTrigonometricInputWithPeriod(MIR::FEmitter& Em, const FExpressionInput* Input, float Period)
+{
+	// Get input after checking it has primitive type.
+	FValueRef Value = Em.CheckIsArithmetic(Em.Input(Input));
+	if (Period > 0.0f)
 	{
-		Emitter.Errorf(TEXT("Missing external code declaration for '%s' [Index=%d]"), *ExternalCodeIdentifier.ToString(), InExternalCodeIdentifierIndex);
-		return nullptr;
+		Value = Em.Multiply(Value, Em.ConstantFloat(2.0f * (float)UE_PI / Period));
 	}
-
-	return Emitter.EmitInlineHLSL(ExternalCodeDeclaration, InArguments);
+	return Value;
 }
 
-static void BuildInlineHLSLOutput(const UMaterialExpressionExternalCodeBase& InExternalCodeExpression, MIR::FEmitter& Emitter, TConstArrayView<MIR::FValue*> InArguments)
+void UMaterialExpressionCosine::Build(MIR::FEmitter& Em)
 {
-	for (int32 OutputIndex = 0; OutputIndex < InExternalCodeExpression.ExternalCodeIdentifiers.Num(); ++OutputIndex)
-	{
-		Emitter.Output(OutputIndex, EmitInlineHLSL(InExternalCodeExpression, Emitter, OutputIndex, InArguments));
-	}
+	Em.Output(0, Em.Cos(GetTrigonometricInputWithPeriod(Em, &Input, Period)));
 }
 
-void UMaterialExpressionExternalCodeBase::Build(MIR::FEmitter& Emitter)
+void UMaterialExpressionSine::Build(MIR::FEmitter& Em)
 {
-	BuildInlineHLSLOutput(*this, Emitter, {});
+	Em.Output(0, Em.Sin(GetTrigonometricInputWithPeriod(Em, &Input, Period)));
 }
 
-/* Binary Operators */
-
-static void BuildBinaryOperatorFromValue(
-	MIR::FEmitter& Emitter,
-	MIR::EOperator Op,
-	MIR::FValue* LhsValue,
-	MIR::FValue* RhsValue)
+void UMaterialExpressionTangent::Build(MIR::FEmitter& Em)
 {
-	// Finally emit the binary operator.
-	MIR::FValue* Value = Emitter.EmitOperator(Op, LhsValue, RhsValue);
-
-	// And flow it out of the expression's only output.
-	Emitter.Output(0, Value);
+	Em.Output(0, Em.Tan(GetTrigonometricInputWithPeriod(Em, &Input, Period)));
 }
 
-static void BuildBinaryOperatorFromInputWithDefault(
-	MIR::FEmitter& Emitter,
-	MIR::EOperator Op,
-	FExpressionInput* LhsInput,
-	float LhsConst,
-	FExpressionInput* RhsInput,
-	float RhsConst)
+void UMaterialExpressionSaturate::Build(MIR::FEmitter& Em)
 {
-	// Default inputs to their relative constants if disconnected, then get each input after checking it has primitive type.
-	MIR::FValue* LhsValue = Emitter.DefaultInputToFloatScalar(LhsInput, LhsConst).InputPrimitive(LhsInput);
-	MIR::FValue* RhsValue = Emitter.DefaultInputToFloatScalar(RhsInput, RhsConst).InputPrimitive(RhsInput);
-
-	UE_MIR_CHECKPOINT();
+	Em.Output(0, Em.Saturate(Em.Input(&Input)));
+}
 
-	BuildBinaryOperatorFromValue(Emitter, Op, LhsValue, RhsValue);
+void UMaterialExpressionSquareRoot::Build(MIR::FEmitter& Em)
+{
+	Em.Output(0, Em.Sqrt(Em.Input(&Input)));
 }
 
-static void BuildBinaryOperatorFromInput(
-	MIR::FEmitter& Emitter,
-	MIR::EOperator Op,
-	FExpressionInput* LhsInput,
-	FExpressionInput* RhsInput)
+static FValueRef EmitInlineHLSL(MIR::FEmitter& Em, const UMaterialExpressionExternalCodeBase& InExternalCodeExpression, int32 InExternalCodeIdentifierIndex, TConstArrayView<FValueRef> InArguments)
 {
-	// Default inputs to their relative constants if disconnected, then get each input after checking it has primitive type.
-	MIR::FValue* LhsValue = Emitter.Input(LhsInput);
-	MIR::FValue* RhsValue = Emitter.Input(RhsInput);
+	checkf(InExternalCodeIdentifierIndex >= 0 && InExternalCodeIdentifierIndex < InExternalCodeExpression.ExternalCodeIdentifiers.Num(),
+		TEXT("External code identifier index (%d) out of bounds; Upper bound is %d"), InExternalCodeIdentifierIndex, InExternalCodeExpression.ExternalCodeIdentifiers.Num());
 
-	UE_MIR_CHECKPOINT();
+	const FName& ExternalCodeIdentifier = InExternalCodeExpression.ExternalCodeIdentifiers[InExternalCodeIdentifierIndex];
+	const FMaterialExternalCodeDeclaration* ExternalCodeDeclaration = MaterialExternalCodeRegistry::Get().FindExternalCode(ExternalCodeIdentifier);
+	if (!ExternalCodeDeclaration)
+	{
+		Em.Errorf(TEXT("Missing external code declaration for '%s' [Index=%d]"), *ExternalCodeIdentifier.ToString(), InExternalCodeIdentifierIndex);
+		return {};
+	}
 
-	BuildBinaryOperatorFromValue(Emitter, Op, LhsValue, RhsValue);
+	return Em.InlineHLSL(ExternalCodeDeclaration, InArguments);
 }
 
-void UMaterialExpressionAdd::Build(MIR::FEmitter& Emitter)
-{ 
-	BuildBinaryOperatorFromInputWithDefault(Emitter, MIR::BO_Add, &A, ConstA, &B, ConstB);
+static void BuildInlineHLSLOutput(MIR::FEmitter& Em, const UMaterialExpressionExternalCodeBase& InExternalCodeExpression, TConstArrayView<FValueRef> InArguments)
+{
+	for (int32 OutputIndex = 0; OutputIndex < InExternalCodeExpression.ExternalCodeIdentifiers.Num(); ++OutputIndex)
+	{
+		Em.Output(OutputIndex, EmitInlineHLSL(Em, InExternalCodeExpression, OutputIndex, InArguments));
+	}
 }
 
-void UMaterialExpressionDesaturation::Build(MIR::FEmitter& Emitter)
+void UMaterialExpressionExternalCodeBase::Build(MIR::FEmitter& Em)
 {
-	MIR::FValue* ColorValue = Emitter.Input(&Input);
-	UE_MIR_CHECKPOINT();
-
-	ColorValue = Emitter.EmitConstruct(MIR::FPrimitiveType::GetVector(MIR::ScalarKind_Float, 3), ColorValue);
-	UE_MIR_CHECKPOINT();
+	BuildInlineHLSLOutput(Em, *this, {});
+}
 
-	MIR::FValue* GreyOrLerpValue = Emitter.EmitOperator(MIR::BO_Dot, ColorValue, Emitter.EmitConstantFloat3(FVector3f(LuminanceFactors)));
-	UE_MIR_CHECKPOINT();
+/* Binary Operators */
 
-	MIR::FValue* FractionValue = Emitter.TryInput(&Fraction);
+void UMaterialExpressionDesaturation::Build(MIR::FEmitter& Em)
+{
+	FValueRef ColorValue = Em.CastToFloat(Em.Input(&Input), 3);
+	FValueRef GreyOrLerpValue = Em.Dot(ColorValue, Em.ConstantFloat3(FVector3f(LuminanceFactors))); // todo: check
+	FValueRef FractionValue = Em.TryInput(&Fraction);
 	if (FractionValue)
 	{
-		// Cast the Grey and Fraction value to the color type value so the lerp op works below.
-		GreyOrLerpValue = Emitter.EmitConstruct(ColorValue->Type, GreyOrLerpValue);
-		FractionValue = Emitter.EmitConstruct(ColorValue->Type, FractionValue);
-
-		UE_MIR_CHECKPOINT();
-
-		GreyOrLerpValue = Emitter.EmitOperator(MIR::TO_Lerp, ColorValue, GreyOrLerpValue, FractionValue);
+		GreyOrLerpValue = Em.Lerp(ColorValue, GreyOrLerpValue, FractionValue);
 	}
-
-	UE_MIR_CHECKPOINT();
-
-	Emitter.Output(0, GreyOrLerpValue);
+	Em.Output(0, GreyOrLerpValue);
 }
 
-void UMaterialExpressionDistance::Build(MIR::FEmitter& Emitter)
+void UMaterialExpressionDistance::Build(MIR::FEmitter& Em)
 {
-	MIR::FValue* LhsValue = Emitter.Input(&A);
-	MIR::FValue* RhsValue = Emitter.Input(&B);
-	UE_MIR_CHECKPOINT();
-
-	MIR::FValue* DeltaValue = Emitter.EmitOperator(MIR::BO_Subtract, LhsValue, RhsValue);
-	UE_MIR_CHECKPOINT();
+	Em.Output(0, Em.Length(Em.Subtract(Em.Input(&A), Em.Input(&B))));
+}
 
-	MIR::FValue* ResultValue = Emitter.EmitOperator(MIR::UO_Length, DeltaValue);
-	UE_MIR_CHECKPOINT();
+void UMaterialExpressionFmod::Build(MIR::FEmitter& Em)
+{ 
+	Em.Output(0, Em.Fmod(Em.Input(&A), Em.Input(&B)));
+}
 
-	Emitter.Output(0, ResultValue);
+static void BuildBinaryOperatorWithDefaults(MIR::FEmitter& Em, MIR::EOperator Op, const FExpressionInput* A, float ConstA, const FExpressionInput* B, float ConstB)
+{
+	FValueRef AVal = Em.InputDefaultFloat(A, ConstA);
+	FValueRef BVal = Em.InputDefaultFloat(B, ConstB);
+	Em.Output(0, Em.Operator(Op, AVal, BVal));
 }
 
-void UMaterialExpressionFmod::Build(MIR::FEmitter& Emitter)
+void UMaterialExpressionAdd::Build(MIR::FEmitter& Em)
 { 
-	BuildBinaryOperatorFromInput(Emitter, MIR::BO_Fmod, &A, &B);
+	BuildBinaryOperatorWithDefaults(Em, MIR::BO_Add, &A, ConstA, &B, ConstB);
 }
 
-void UMaterialExpressionSubtract::Build(MIR::FEmitter& Emitter)
-{ 
-	BuildBinaryOperatorFromInputWithDefault(Emitter, MIR::BO_Subtract, &A, ConstA, &B, ConstB);
+void UMaterialExpressionSubtract::Build(MIR::FEmitter& Em)
+{
+	BuildBinaryOperatorWithDefaults(Em, MIR::BO_Subtract, &A, ConstA, &B, ConstB);
 }
 
-void UMaterialExpressionMultiply::Build(MIR::FEmitter& Emitter)
-{ 
-	BuildBinaryOperatorFromInputWithDefault(Emitter, MIR::BO_Multiply, &A, ConstA, &B, ConstB);
+void UMaterialExpressionMultiply::Build(MIR::FEmitter& Em)
+{
+	BuildBinaryOperatorWithDefaults(Em, MIR::BO_Multiply, &A, ConstA, &B, ConstB);
 }
 
-void UMaterialExpressionDivide::Build(MIR::FEmitter& Emitter)
-{ 
-	BuildBinaryOperatorFromInputWithDefault(Emitter, MIR::BO_Divide, &A, ConstA, &B, ConstB);
+void UMaterialExpressionDivide::Build(MIR::FEmitter& Em)
+{
+	BuildBinaryOperatorWithDefaults(Em, MIR::BO_Divide, &A, ConstA, &B, ConstB);
 }
 
-void UMaterialExpressionMin::Build(MIR::FEmitter& Emitter)
+void UMaterialExpressionMin::Build(MIR::FEmitter& Em)
 { 
-	BuildBinaryOperatorFromInputWithDefault(Emitter, MIR::BO_Min, &A, ConstA, &B, ConstB);
+	BuildBinaryOperatorWithDefaults(Em, MIR::BO_Min, &A, ConstA, &B, ConstB);
 }
 
-void UMaterialExpressionDotProduct::Build(MIR::FEmitter& Emitter)
+void UMaterialExpressionDotProduct::Build(MIR::FEmitter& Em)
 {
-	BuildBinaryOperatorFromInput(Emitter, MIR::BO_Dot, &A, &B);
+	Em.Output(0, Em.Dot(Em.Input(&A), Em.Input(&B)));
 }
 
-void UMaterialExpressionCrossProduct::Build(MIR::FEmitter& Emitter)
+void UMaterialExpressionCrossProduct::Build(MIR::FEmitter& Em)
 {
-	BuildBinaryOperatorFromInput(Emitter, MIR::BO_Cross, &A, &B);
+	Em.Output(0, Em.Cross(Em.Input(&A), Em.Input(&B)));
 }
 
-void UMaterialExpressionEyeAdaptationInverse::Build(MIR::FEmitter& Emitter)
+void UMaterialExpressionEyeAdaptationInverse::Build(MIR::FEmitter& Em)
 {
-	MIR::FValue* LightValue = Emitter
-		.DefaultInputToFloat3(&LightValueInput, { 1.0f, 1.0f, 1.0f })
-		.InputFloatVector(&LightValueInput, 3);
-
-	MIR::FValue* AlphaValue = Emitter
-		.DefaultInputToFloatScalar(&AlphaInput, 1.0f)
-		.InputFloatScalar(&AlphaInput);
-
-	UE_MIR_CHECKPOINT();
-
 	check(ExternalCodeIdentifiers.Num() == 1);
-	MIR::FValue* MultiplierValue = EmitInlineHLSL(*this, Emitter, 0, { AlphaValue });
-
-	Emitter.Output(0, Emitter.EmitOperator(MIR::BO_Multiply, LightValue, MultiplierValue));
+	FValueRef LightValue = Em.CastToFloat(Em.InputDefaultFloat(&LightValueInput, 1.0f), 3);
+	FValueRef AlphaValue = Em.CastToFloat(Em.InputDefaultFloat(&AlphaInput, 1.0f), 1);
+	FValueRef MultiplierValue = EmitInlineHLSL(Em, *this, 0, { AlphaValue });
+	Em.Output(0, Em.Multiply(LightValue, MultiplierValue));
 }
 
-void UMaterialExpressionOneMinus::Build(MIR::FEmitter& Emitter)
+void UMaterialExpressionOneMinus::Build(MIR::FEmitter& Em)
 {
 	// Default input to zero if not connected, then get it as a primitive.
-	MIR::FValue* Value = Emitter.DefaultInputToFloatScalar(&Input, 0.0f).InputPrimitive(&Input);
-	
-	UE_MIR_CHECKPOINT();
+	FValueRef Value = Em.InputDefaultFloat(&Input, 0.0f);
 
 	// Make a "One" value of the same type and dimension as input's.
-	MIR::FValue* OneValue = Emitter.EmitConstruct(Value->Type, Emitter.EmitConstantFloatScalar(1.0f));
+	FValueRef One = Em.ConstantOne(Value->Type->AsPrimitive()->ScalarKind);
 
-	UE_MIR_CHECKPOINT();
-
-	// Finally emit the binary operator.
-	MIR::FValue* Output = Emitter.EmitOperator(MIR::BO_Subtract, OneValue, Value);
-
-	// And flow it out of the expression's only output.
-	Emitter.Output(0, Output);
+	// And flow the subtraction out of the expression's only output.
+	Em.Output(0, Em.Subtract(One, Value));
 }
 
-void UMaterialExpressionIfThenElse::Build(MIR::FEmitter& Emitter)
+void UMaterialExpressionIfThenElse::Build(MIR::FEmitter& Em)
 {
-	// If Condition isn't connected, default its value to false
-	Emitter.DefaultInputToBoolScalar(&Condition, false);
-
 	// Get the condition value checking it is a bool scalar
-	MIR::FValue* ConditionValue = Emitter.InputBoolScalar(&Condition);
+	FValueRef ConditionValue = Em.CastToBool(Em.InputDefaultBool(&Condition, false), 1);
 
-	UE_MIR_CHECKPOINT();
+	UE_MIR_CHECKPOINT(Em); // Make sure the condition value is valid
 	
 	// If condition boolean is constant, select which input is active and simply
 	// bypass its value to our output.
 	if (MIR::FConstant* Constant = ConditionValue->As<MIR::FConstant>())
 	{
 		FExpressionInput* ActiveInput = Constant->Boolean ? &True : &False;
-		Emitter.Output(0, Emitter.Input(ActiveInput));
+		Em.Output(0, Em.Input(ActiveInput));
 		return;
 	}
 
 	// The condition isn't static; Get the true and false values.
 	// If any is disconnected, the emitter will report an error.
-	MIR::FValue* ThenValue = Emitter.Input(&True);
-	MIR::FValue* ElseValue = Emitter.Input(&False);
+	FValueRef ThenValue = Em.Input(&True);
+	FValueRef ElseValue = Em.Input(&False);
 	
-	UE_MIR_CHECKPOINT();
+	const MIR::FType* CommonType = Em.GetCommonType(ThenValue->Type, ElseValue->Type);
+
+	UE_MIR_CHECKPOINT(Em); // Make sure the common type is valid
 
-	// Convert the false type to the true's
-	ElseValue = Emitter.EmitConstruct(ThenValue->Type, ElseValue);
+	// Cast the "then" and "else" values to the common type.
+	ThenValue = Em.Cast(ThenValue, CommonType);
+	ElseValue = Em.Cast(ElseValue, CommonType);
 
 	// Emit the branch instruction
-	MIR::FValue* OutputValue = Emitter.EmitBranch(ConditionValue, ThenValue, ElseValue);
+	FValueRef OutputValue = Em.Branch(ConditionValue, ThenValue, ElseValue);
 
-	Emitter.Output(0, OutputValue);
+	Em.Output(0, OutputValue);
 }
 
-static MIR::FValue* EmitAlmostEquals(MIR::FEmitter& Emitter, MIR::FValue* A, MIR::FValue* B, float Threshold)
+static FValueRef EmitAlmostEquals(MIR::FEmitter& Em, FValueRef A, FValueRef B, float Threshold)
 {
-	return Emitter.EmitOperator(
-		MIR::BO_LessThanOrEquals, 
-		Emitter.EmitOperator(MIR::UO_Abs, Emitter.EmitOperator(MIR::BO_Subtract, A, B)),
-		Emitter.EmitConstantFloatScalar(Threshold));
+	// abs(A - B) <= Threshold
+	return Em.LessThanOrEquals(Em.Abs(Em.Subtract(A, B)), Em.ConstantFloat(Threshold));
 }
 
-void UMaterialExpressionIf::Build(MIR::FEmitter& Emitter)
+void UMaterialExpressionIf::Build(MIR::FEmitter& Em)
 {
-	// Create default values flowing into disconnected inputs
-	Emitter.DefaultInputToFloatScalarZero(&A);
-	Emitter.DefaultInputToFloatScalar(&B, ConstB);
-	Emitter.DefaultInputToFloatScalarZero(&AGreaterThanB);
-	Emitter.DefaultInputToFloatScalarZero(&AEqualsB);
-	Emitter.DefaultInputToFloatScalarZero(&ALessThanB);
-
 	// Get input values and check their types are what we expect.
-	MIR::FValue* AValue = Emitter.InputScalar(&A);
-	MIR::FValue* BValue = Emitter.InputScalar(&B);
-	MIR::FValue* AGreaterThanBValue = Emitter.InputPrimitive(&AGreaterThanB);
-	MIR::FValue* AEqualsBValue = Emitter.InputPrimitive(&AEqualsB);
-	MIR::FValue* ALessThanBValue = Emitter.InputPrimitive(&ALessThanB);
+	FValueRef AValue = Em.CheckIsScalar(Em.InputDefaultFloat(&A, 0.f));
+	FValueRef BValue = Em.CheckIsScalar(Em.InputDefaultFloat(&B, ConstB));
+	FValueRef AGreaterThanBValue = Em.InputDefaultFloat(&AGreaterThanB, 0.f);
+	FValueRef AEqualsBValue = Em.InputDefaultFloat(&AEqualsB, 0.f);
+	FValueRef ALessThanBValue = Em.InputDefaultFloat(&ALessThanB, 0.f);
 
 	// Emit the comparison expressions.
-	MIR::FValue* ALessThanBConditionValue = Emitter.EmitOperator(MIR::BO_LessThan, AValue, BValue);
-	MIR::FValue* AEqualsBConditionValue = EmitAlmostEquals(Emitter, AValue, BValue, EqualsThreshold);
+	FValueRef ALessThanBConditionValue = Em.LessThan(AValue, BValue);
+	FValueRef AEqualsBConditionValue = EmitAlmostEquals(Em, AValue, BValue, EqualsThreshold);
 
 	// And finally emit the full conditional expression.
-	MIR::FValue* OutputValue = Emitter.EmitBranch(AEqualsBConditionValue, AEqualsBValue, AGreaterThanBValue);
-	OutputValue = Emitter.EmitBranch(ALessThanBConditionValue, ALessThanBValue, OutputValue);
+	FValueRef OutputValue = Em.Branch(AEqualsBConditionValue, AEqualsBValue, AGreaterThanBValue);
+	OutputValue = Em.Branch(ALessThanBConditionValue, ALessThanBValue, OutputValue);
 
-	Emitter.Output(0, OutputValue);
+	Em.Output(0, OutputValue);
 }
 
-void UMaterialExpressionSphericalParticleOpacity::Build(MIR::FEmitter& Emitter)
+void UMaterialExpressionSphericalParticleOpacity::Build(MIR::FEmitter& Em)
 {
-	MIR::FValue* DensityValue = Emitter
-		.DefaultInputToFloatScalar(&Density, ConstantDensity)
-		.Input(&Density);
-		
-	UE_MIR_CHECKPOINT();
-
-	BuildInlineHLSLOutput(*this, Emitter, { DensityValue });
+	FValueRef DensityValue = Em.InputDefaultFloat(&Density, ConstantDensity);
+	UE_MIR_CHECKPOINT(Em); // Early out in case of errors
+	BuildInlineHLSLOutput(Em, *this, { DensityValue });
 }
 
-void UMaterialExpressionTextureObject::Build(MIR::FEmitter& Emitter)
+void UMaterialExpressionTextureObject::Build(MIR::FEmitter& Em)
 {
-	Emitter.Output(0, Emitter.EmitTextureObject(Texture, SamplerType));
+	Em.Output(0, Em.TextureObject(Texture, SamplerType));
 }
 
 static MIR::ETextureReadMode TextureGatherModeToMIR(::ETextureGatherMode Mode)
@@ -969,73 +805,71 @@ static MIR::ETextureReadMode TextureGatherModeToMIR(::ETextureGatherMode Mode)
 	}
 }
 
-static void BuildTextureSampleExpression(MIR::FValue* Texture, UMaterialExpressionTextureSample* Expr, MIR::FEmitter& Emitter)
+static void BuildTextureSampleExpression(MIR::FEmitter& Em, UMaterialExpressionTextureSample* Expr, FValueRef Texture)
 {
-	MIR::FValue* TexCoords = Emitter.TryInput(&Expr->Coordinates);
+	FValueRef TexCoords = Em.TryInput(&Expr->Coordinates);
 	if (!TexCoords)
 	{
-		TexCoords = Emitter.EmitExternalInput(MIR::TexCoordIndexToExternalInput(Expr->ConstCoordinate));
+		TexCoords = Em.ExternalInput(MIR::TexCoordIndexToExternalInput(Expr->ConstCoordinate));
 	}
 
-	MIR::FValue* TextureRead{};
+	FValueRef TextureRead{};
 	if (Expr->GatherMode != TGM_None)
 	{
 		if (Expr->MipValueMode != TMVM_None)
 		{
-			Emitter.Errorf(TEXT("Texture gather does not support mipmap overrides (it implicitly accesses a specific mip)."));
+			Em.Errorf(TEXT("Texture gather does not support mipmap overrides (it implicitly accesses a specific mip)."));
 			return;
 		}
 
-		TextureRead = Emitter.EmitTextureGather(Texture, TexCoords, TextureGatherModeToMIR(Expr->GatherMode), Expr->SamplerSource, Expr->SamplerType);
+		TextureRead = Em.TextureGather(Texture, TexCoords, TextureGatherModeToMIR(Expr->GatherMode), Expr->SamplerSource, Expr->SamplerType);
 	} 
 	else
 	{
 		// Determine if automatic view mip bias should be used, by trying ot acquire a its input as a static boolean.
-		bool bAutomaticViewMipBias = Emitter
-			.DefaultInputToBoolScalar(&Expr->AutomaticViewMipBiasValue, Expr->AutomaticViewMipBias)
-			.InputConstantBoolScalar(&Expr->AutomaticViewMipBiasValue);
+		bool bAutomaticViewMipBias = Em.ToConstantBool(Em.InputDefaultBool(&Expr->AutomaticViewMipBiasValue, Expr->AutomaticViewMipBias));
 	
-		UE_MIR_CHECKPOINT();
+		UE_MIR_CHECKPOINT(Em);
 
 		// Get the mip value level (either through the expression input or using the given constant if disconnected).
-		MIR::FValue* MipValue = (Expr->MipValueMode == TMVM_MipLevel || Expr->MipValueMode == TMVM_MipBias)
-			? Emitter.DefaultInputToFloatScalar(&Expr->MipValue, Expr->ConstMipValue).InputScalar(&Expr->MipValue)
-			: MIR::FPoison::Get();
+		FValueRef MipValue = (Expr->MipValueMode == TMVM_MipLevel || Expr->MipValueMode == TMVM_MipBias)
+			? Em.CheckIsScalar(Em.InputDefaultInt(&Expr->MipValue, Expr->ConstMipValue))
+			: Em.Poison();
 
 		switch (Expr->MipValueMode)
 		{
 			case TMVM_None:
-				TextureRead = Emitter.EmitTextureSample(Texture, TexCoords, bAutomaticViewMipBias, Expr->SamplerSource, Expr->SamplerType);
+				TextureRead = Em.TextureSample(Texture, TexCoords, bAutomaticViewMipBias, Expr->SamplerSource, Expr->SamplerType);
 				break;
 		
 			case TMVM_MipBias:
-				TextureRead = Emitter.EmitTextureSampleBias(Texture, TexCoords, MipValue, bAutomaticViewMipBias, Expr->SamplerSource, Expr->SamplerType);
+				TextureRead = Em.TextureSampleBias(Texture, TexCoords, MipValue, bAutomaticViewMipBias, Expr->SamplerSource, Expr->SamplerType);
 				break;
 		
 			case TMVM_MipLevel:
-				TextureRead = Emitter.EmitTextureSampleLevel(Texture, TexCoords, MipValue, bAutomaticViewMipBias, Expr->SamplerSource, Expr->SamplerType);
+				TextureRead = Em.TextureSampleLevel(Texture, TexCoords, MipValue, bAutomaticViewMipBias, Expr->SamplerSource, Expr->SamplerType);
 				break;
 
 			case TMVM_Derivative:
 			{
-				MIR::FValue* TexCoordsDdx = Emitter.EmitConstruct(TexCoords->Type, Emitter.Input(&Expr->CoordinatesDX));
-				MIR::FValue* TexCoordsDdy = Emitter.EmitConstruct(TexCoords->Type, Emitter.Input(&Expr->CoordinatesDY));
-				TextureRead = Emitter.EmitTextureSampleGrad(Texture, TexCoords, TexCoordsDdx, TexCoordsDdy, bAutomaticViewMipBias, Expr->SamplerSource, Expr->SamplerType);
+				FValueRef TexCoordsDdx = Em.Cast(Em.Input(&Expr->CoordinatesDX), TexCoords->Type);
+				FValueRef TexCoordsDdy = Em.Cast(Em.Input(&Expr->CoordinatesDY), TexCoords->Type);
+				TextureRead = Em.TextureSampleGrad(Texture, TexCoords, TexCoordsDdx, TexCoordsDdy, bAutomaticViewMipBias, Expr->SamplerSource, Expr->SamplerType);
 				break;
 			}
 		}
 	}
 
-	Emitter.Output(0, Emitter.EmitSwizzle(TextureRead, MIR::FSwizzleMask::XYZ()));
-	Emitter.Output(1, Emitter.EmitSubscript(TextureRead, 0));
-	Emitter.Output(2, Emitter.EmitSubscript(TextureRead, 1)); 
-	Emitter.Output(3, Emitter.EmitSubscript(TextureRead, 2));
-	Emitter.Output(4, Emitter.EmitSubscript(TextureRead, 3));
-	Emitter.Output(5, TextureRead); 
+	Em.Output(0, Em.Swizzle(TextureRead, MIR::FSwizzleMask::XYZ()));
+	Em.Output(1, Em.Subscript(TextureRead, 0));
+	Em.Output(2, Em.Subscript(TextureRead, 1));
+	Em.Output(3, Em.Subscript(TextureRead, 2));
+	Em.Output(4, Em.Subscript(TextureRead, 3));
+	Em.Output(5, TextureRead); 
 }
 
 // Returns this value's textures if it has one (nullptr otherwise).
-static EMaterialSamplerType GetValueMaterialSamplerType(const MIR::FValue* Value)
+static EMaterialSamplerType GetValueMaterialSamplerType(FValueRef Value)
 {
 	if (auto TextureObject = MIR::As<MIR::FTextureObject>(Value))
 	{
@@ -1048,188 +882,164 @@ static EMaterialSamplerType GetValueMaterialSamplerType(const MIR::FValue* Value
 	return SAMPLERTYPE_MAX;
 }
 
-void UMaterialExpressionTextureSample::Build(MIR::FEmitter& Emitter)
+void UMaterialExpressionTextureSample::Build(MIR::FEmitter& Em)
 {
-	MIR::FValue* TextureValue = Emitter.TryInput(&TextureObject);
-	if (TextureValue)
-	{
-		TextureValue = Emitter.InputTexture(&TextureObject);
-	}
-	else
+	FValueRef TextureValue = Em.TryInput(&TextureObject);
+	if (!TextureValue)
 	{
 		if (!Texture)
 		{
-			Emitter.Error(TEXT("No texture specified for this expression."));
+			Em.Error(TEXT("No texture specified for this expression."));
 			return;
 		}
 
-		TextureValue = Emitter.EmitTextureObject(Texture.Get(), SamplerType);
+		TextureValue = Em.TextureObject(Texture.Get(), SamplerType);
 	}
 
-	UE_MIR_CHECKPOINT();
+	UE_MIR_CHECKPOINT(Em);
 
 	if (TextureValue && GetValueMaterialSamplerType(TextureValue) != SAMPLERTYPE_Color)
 	{
-		Emitter.Error(TEXT("Input texture sampler type must be color."));
+		Em.Error(TEXT("Input texture sampler type must be color."));
 		return;
 	}
 
-	BuildTextureSampleExpression(TextureValue, this, Emitter);
+	BuildTextureSampleExpression(Em, this, TextureValue);
 }
 
-static MIR::FValue* BuildTextureObjectParameter(MIR::FEmitter& Emitter, UMaterialExpressionTextureSampleParameter* Expr)
+static FValueRef BuildTextureObjectParameter(MIR::FEmitter& Em, UMaterialExpressionTextureSampleParameter* Expr)
 {
 	FMaterialParameterMetadata Param;
 	if (!Expr->GetParameterValue(Param))
 	{
-		Emitter.Error(TEXT("Failed to get parameter value"));
+		Em.Error(TEXT("Failed to get parameter value"));
 		return nullptr;
 	}
 
-	MIR::FValue* ParameterValue = Emitter.EmitParameter(Expr->GetParameterName(), Param, Expr->SamplerType);
+	FValueRef ParameterValue = Em.Parameter(Expr->GetParameterName(), Param, Expr->SamplerType);
 	if (!ParameterValue->Type->IsTexture())
 	{
-		Emitter.Error(TEXT("Parameter is not a texture"));
+		Em.Error(TEXT("Parameter is not a texture"));
 		return nullptr;
 	}
 
 	return ParameterValue;
 }
 
-void UMaterialExpressionTextureSampleParameter::Build(MIR::FEmitter& Emitter)
+void UMaterialExpressionTextureSampleParameter::Build(MIR::FEmitter& Em)
 {
-	MIR::FValue* ParameterValue = BuildTextureObjectParameter(Emitter, this);
-	UE_MIR_CHECKPOINT();
-	BuildTextureSampleExpression(ParameterValue, this, Emitter);
+	FValueRef ParameterValue = BuildTextureObjectParameter(Em, this);
+	UE_MIR_CHECKPOINT(Em);
+	BuildTextureSampleExpression(Em, this, ParameterValue);
 }
 
-void UMaterialExpressionTextureSampleParameterCube::Build(MIR::FEmitter& Emitter)
+void UMaterialExpressionTextureSampleParameterCube::Build(MIR::FEmitter& Em)
 {
-	if (!Coordinates.GetTracedInput().Expression)
-	{
-		Emitter.Error(TEXT("Cube sample needs UV input"));
-		return;
-	}
-
-	UMaterialExpressionTextureSampleParameter::Build(Emitter);
+	Em.Input(&Coordinates); // Cubemap sampling requires coordinates input specified
+	UE_MIR_CHECKPOINT(Em);
+	UMaterialExpressionTextureSampleParameter::Build(Em);
 }
 
-void UMaterialExpressionTextureObjectParameter::Build(MIR::FEmitter& Emitter)
+void UMaterialExpressionTextureObjectParameter::Build(MIR::FEmitter& Em)
 {
-	MIR::FValue* ParameterValue = BuildTextureObjectParameter(Emitter, this);
-	UE_MIR_CHECKPOINT();
-	Emitter.Output(0, ParameterValue);
+	Em.Output(0, BuildTextureObjectParameter(Em, this));
 }
 
-void UMaterialExpressionTextureCoordinate::Build(MIR::FEmitter& Emitter)
+void UMaterialExpressionTextureCoordinate::Build(MIR::FEmitter& Em)
 {
 	if (UnMirrorU || UnMirrorV)
 	{
-		Emitter.Error(TEXT("Unmirroring unsupported"));
+		Em.Error(TEXT("Unmirroring unsupported"));
 		return;
 	}
 
-	MIR::FValue* OutputValue = Emitter.EmitExternalInput(MIR::TexCoordIndexToExternalInput(CoordinateIndex));
+	FValueRef OutputValue = Em.ExternalInput(MIR::TexCoordIndexToExternalInput(CoordinateIndex));
 	
 	// Multiply the UV input by the UV tiling constants
-	OutputValue = Emitter.EmitOperator(MIR::BO_Multiply, OutputValue, Emitter.EmitConstantFloat2({ UTiling, VTiling }));
+	OutputValue = Em.Multiply(OutputValue, Em.ConstantFloat2({ UTiling, VTiling }));
 	
-	Emitter.Output(GetOutput(0), OutputValue);
+	Em.Output(GetOutput(0), OutputValue);
 }
 
-void UMaterialExpressionTime::Build(MIR::FEmitter& Emitter)
+void UMaterialExpressionTime::Build(MIR::FEmitter& Em)
 {
 	const MIR::FType* ScalarFloatType = MIR::FPrimitiveType::GetScalar(MIR::ScalarKind_Float);
 
 	// When pausing the game is ignored for this time expression, use real-time instead of game-time.
 	const TCHAR* TimeFieldName = bIgnorePause ? TEXT("View.<PREV>RealTime") : TEXT("View.<PREV>GameTime");
-
 	if (!bOverride_Period)
 	{
-		Emitter.Output(0, Emitter.EmitInlineHLSL(ScalarFloatType, TimeFieldName, {}, MIR::EValueFlags::SubstituteTagsInInlineHLSL));
+		Em.Output(0, Em.InlineHLSL(ScalarFloatType, TimeFieldName, {}, MIR::EValueFlags::SubstituteTagsInInlineHLSL));
 	}
 	else if (Period == 0.0f)
 	{
-		Emitter.Output(0, Emitter.EmitConstantFloatScalar(0.0f));
+		Em.Output(0, Em.ConstantFloat(0.0f));
 	}
 	else
 	{
 		// Note: Don't use IR intrinsic for Fmod() here to avoid conversion to fp16 on mobile.
 		// We want full 32 bit float precision until the fmod when using a period.
-		MIR::FValue* PeriodValue = Emitter.EmitConstantFloatScalar(Period);
-		Emitter.Output(0, Emitter.EmitInlineHLSL(ScalarFloatType, FString::Printf(TEXT("fmod(%s,$0)"), TimeFieldName), { PeriodValue }, MIR::EValueFlags::SubstituteTagsInInlineHLSL));
+		FValueRef PeriodValue = Em.ConstantFloat(Period);
+		Em.Output(0, Em.InlineHLSL(ScalarFloatType, FString::Printf(TEXT("fmod(%s,$0)"), TimeFieldName), { PeriodValue }, MIR::EValueFlags::SubstituteTagsInInlineHLSL));
 	}
 }
 
-void UMaterialExpressionReroute::Build(MIR::FEmitter& Emitter)
-{
-	MIR::FValue* Value = Emitter.Input(&Input);
-	Emitter.Output(0, Value);
-}
-
-void BuildTernaryArithmeticOperator(
-	MIR::FEmitter& Emitter,
-	MIR::EOperator Op,
-	FExpressionInput* A,
-	float ConstA,
-	FExpressionInput* B,
-	float ConstB,
-	FExpressionInput* C,
-	float ConstC)
+void UMaterialExpressionReroute::Build(MIR::FEmitter& Em)
 {
-	Emitter.DefaultInputToFloatScalar(A, ConstA);
-	Emitter.DefaultInputToFloatScalar(B, ConstB);
-	Emitter.DefaultInputToFloatScalar(C, ConstC);
-
-	MIR::FValue* ValueA = Emitter.InputPrimitive(A);
-	MIR::FValue* ValueB = Emitter.Input(B);
-	MIR::FValue* ValueC = Emitter.Input(C);
-
-	Emitter.Output(0, Emitter.EmitOperator(Op, ValueA, ValueB, ValueC));
+	FValueRef Value = Em.Input(&Input);
+	Em.Output(0, Value);
 }
 
-void UMaterialExpressionClamp::Build(MIR::FEmitter& Emitter)
+void UMaterialExpressionClamp::Build(MIR::FEmitter& Em)
 {
-	MIR::FValue* InputValue = Emitter.Input(&Input);
-	MIR::FValue* MinValue = Emitter.DefaultInputToFloatScalar(&Min, MinDefault).Input(&Min);
-	MIR::FValue* MaxValue = Emitter.DefaultInputToFloatScalar(&Max, MaxDefault).Input(&Max);
+	FValueRef InputValue = Em.Input(&Input);
+	FValueRef MinValue = Em.InputDefaultFloat(&Min, MinDefault);
+	FValueRef MaxValue = Em.InputDefaultFloat(&Max, MaxDefault);
 
-	MIR::FValue* OutputValue = nullptr;
+	FValueRef OutputValue = nullptr;
 	if (ClampMode == CMODE_Clamp)
 	{
-		OutputValue = Emitter.EmitOperator(MIR::TO_Clamp, InputValue, MinValue, MaxValue);
+		OutputValue = Em.Clamp(InputValue, MinValue, MaxValue);
 	}
 	else if (ClampMode == CMODE_ClampMin)
 	{
-		OutputValue = Emitter.EmitOperator(MIR::BO_Max, InputValue, MinValue);
+		OutputValue = Em.Max(InputValue, MinValue);
 	}
 	else if (ClampMode == CMODE_ClampMax)
 	{
-		OutputValue = Emitter.EmitOperator(MIR::BO_Min, InputValue, MaxValue);
+		OutputValue = Em.Min(InputValue, MaxValue);
 	}
 
-	Emitter.Output(0, OutputValue);
+	Em.Output(0, OutputValue);
 }
 
-void UMaterialExpressionLinearInterpolate::Build(MIR::FEmitter& Emitter)
+void BuildTernaryArithmeticOperator(MIR::FEmitter& Em, MIR::EOperator Op, FExpressionInput* A, float ConstA, FExpressionInput* B, float ConstB, FExpressionInput* C, float ConstC)
 {
-	BuildTernaryArithmeticOperator(Emitter, MIR::TO_Lerp, &A, ConstA, &B, ConstB, &Alpha, ConstAlpha);
+	FValueRef ValueA = Em.InputDefaultFloat(A, ConstA);
+	FValueRef ValueB = Em.InputDefaultFloat(B, ConstB);
+	FValueRef ValueC = Em.InputDefaultFloat(C, ConstC);
+	Em.Output(0, Em.Operator(Op, ValueA, ValueB, ValueC));
 }
 
-void UMaterialExpressionSmoothStep::Build(MIR::FEmitter& Emitter)
+void UMaterialExpressionLinearInterpolate::Build(MIR::FEmitter& Em)
 {
-	BuildTernaryArithmeticOperator(Emitter, MIR::TO_Smoothstep, &Min, ConstMin, &Max, ConstMax, &Value, ConstValue);
+	BuildTernaryArithmeticOperator(Em, MIR::TO_Lerp, &A, ConstA, &B, ConstB, &Alpha, ConstAlpha);
 }
 
-void UMaterialExpressionConvert::Build(MIR::FEmitter& Emitter)
+void UMaterialExpressionSmoothStep::Build(MIR::FEmitter& Em)
 {
-	TArray<MIR::FValue*, TInlineAllocator<8>> InputValues;
+	BuildTernaryArithmeticOperator(Em, MIR::TO_Smoothstep, &Min, ConstMin, &Max, ConstMax, &Value, ConstValue);
+}
+
+void UMaterialExpressionConvert::Build(MIR::FEmitter& Em)
+{
+	TArray<FValueRef, TInlineAllocator<8>> InputValues;
 	InputValues.Init(nullptr, ConvertInputs.Num());
 
 	for (int32 OutputIndex = 0; OutputIndex < ConvertOutputs.Num(); ++OutputIndex)
 	{
 		const FMaterialExpressionConvertOutput& ConvertOutput = ConvertOutputs[OutputIndex];
-		MIR::FValue* OutComponents[4] = { nullptr, nullptr, nullptr, nullptr };
+		FValueRef OutComponents[4] = { nullptr, nullptr, nullptr, nullptr };
 
 		for (const FMaterialExpressionConvertMapping& Mapping : ConvertMappings)
 		{
@@ -1242,14 +1052,14 @@ void UMaterialExpressionConvert::Build(MIR::FEmitter& Emitter)
 			const int32 OutputComponentIndex = Mapping.OutputComponentIndex;
 			if (!IsValidComponentIndex(OutputComponentIndex, ConvertOutput.Type))
 			{
-				Emitter.Errorf(TEXT("Convert mapping's output component `%d` is invalid."), OutputComponentIndex);
+				Em.Errorf(TEXT("Convert mapping's output component `%d` is invalid."), OutputComponentIndex);
 				continue;
 			}
 
 			const int32 InputIndex = Mapping.InputIndex;
 			if (!ConvertInputs.IsValidIndex(InputIndex))
 			{
-				Emitter.Errorf(TEXT("Convert mapping's input `%d` is invalid."), InputIndex);
+				Em.Errorf(TEXT("Convert mapping's input `%d` is invalid."), InputIndex);
 				continue;
 			}
 
@@ -1257,7 +1067,7 @@ void UMaterialExpressionConvert::Build(MIR::FEmitter& Emitter)
 			const int32 InputComponentIndex = Mapping.InputComponentIndex;
 			if (!IsValidComponentIndex(InputComponentIndex, ConvertInput.Type))
 			{
-				Emitter.Errorf(TEXT("Convert mapping's input component `%d` is invalid."), InputComponentIndex);
+				Em.Errorf(TEXT("Convert mapping's input component `%d` is invalid."), InputComponentIndex);
 				continue;
 			}
 
@@ -1267,13 +1077,13 @@ void UMaterialExpressionConvert::Build(MIR::FEmitter& Emitter)
 			if (!InputValues[InputIndex])
 			{
 				// Read the input's value (or read float zero if disconnected).
-				InputValues[InputIndex] = Emitter.DefaultInputToFloat4(&ConvertInput.ExpressionInput, ConvertInput.DefaultValue).Input(&ConvertInput.ExpressionInput);
+				InputValues[InputIndex] = Em.InputDefaultFloat4(&ConvertInput.ExpressionInput, ConvertInput.DefaultValue);
 
 				// Expect type to be primitive.
 				const MIR::FPrimitiveType* InputPrimitiveType = InputValues[InputIndex]->Type->AsPrimitive();
 				if (!InputPrimitiveType)
 				{
-					Emitter.Errorf(TEXT("Input `%d` of type `%s` is not primitive."), InputComponentIndex, InputValues[InputIndex]->Type->GetSpelling().GetData());
+					Em.Errorf(TEXT("Input `%d` of type `%s` is not primitive."), InputComponentIndex, InputValues[InputIndex]->Type->GetSpelling().GetData());
 					continue;
 				}
 
@@ -1281,11 +1091,11 @@ void UMaterialExpressionConvert::Build(MIR::FEmitter& Emitter)
 				const MIR::FType* InputType = MIR::FPrimitiveType::GetVector(InputPrimitiveType->ScalarKind, MaterialExpressionConvertType::GetComponentCount(ConvertInput.Type));
 
 				// Cast the input value to the target type.
-				InputValues[InputIndex] = Emitter.EmitConstruct(InputType, InputValues[InputIndex]);
+				InputValues[InputIndex] = Em.Cast(InputValues[InputIndex], InputType);
 			}
 
 			// Subscript the input value to the specified component index.
-			OutComponents[OutputComponentIndex] = Emitter.EmitSubscript(InputValues[InputIndex], InputComponentIndex);
+			OutComponents[OutputComponentIndex] = Em.Subscript(InputValues[InputIndex], InputComponentIndex);
 		}
 
 		const int32 OutputNumComponents = MaterialExpressionConvertType::GetComponentCount(ConvertOutput.Type);
@@ -1296,35 +1106,34 @@ void UMaterialExpressionConvert::Build(MIR::FEmitter& Emitter)
 			// If we don't have a compile result here, default it to a that component's default value
 			if (!OutComponents[OutputComponentIndex])
 			{
-				OutComponents[OutputComponentIndex] = Emitter.EmitConstantFloatScalar(ConvertOutput.DefaultValue.Component(OutputComponentIndex));
+				OutComponents[OutputComponentIndex] = Em.ConstantFloat(ConvertOutput.DefaultValue.Component(OutputComponentIndex));
 			}
 		}
 
 		// Finally create the output dimensional value by combining the output components.
-		MIR::FValue* OutValue = nullptr;
+		FValueRef OutValue;
 		switch (OutputNumComponents)
 		{
 			case 1: OutValue = OutComponents[0]; break;
-			case 2: OutValue = Emitter.EmitVector2(OutComponents[0], OutComponents[1]); break;
-			case 3: OutValue = Emitter.EmitVector3(OutComponents[0], OutComponents[1], OutComponents[2]); break;
-			case 4: OutValue = Emitter.EmitVector4(OutComponents[0], OutComponents[1], OutComponents[2], OutComponents[3]); break;
+			case 2: OutValue = Em.Vector2(OutComponents[0], OutComponents[1]); break;
+			case 3: OutValue = Em.Vector3(OutComponents[0], OutComponents[1], OutComponents[2]); break;
+			case 4: OutValue = Em.Vector4(OutComponents[0], OutComponents[1], OutComponents[2], OutComponents[3]); break;
 			
 			default:
-				OutValue = MIR::FPoison::Get();
-				Emitter.Errorf(TEXT("Convert node has an invalid component count of %d"), OutputNumComponents);
+				OutValue = Em.Poison();
+				Em.Errorf(TEXT("Convert node has an invalid component count of %d"), OutputNumComponents);
 				break;
 		}
 
-		Emitter.Output(OutputIndex, OutValue);
+		Em.Output(OutputIndex, OutValue);
 	}
 }
 
-static MIR::FValue* BuildViewProperty(MIR::FEmitter& Emitter, EMaterialExposedViewProperty InProperty, bool bInvProperty = false)
+static FValueRef BuildViewProperty(MIR::FEmitter& Em, EMaterialExposedViewProperty InProperty, bool bInvProperty = false)
 {
 	check(InProperty < MEVP_MAX);
 
 	const FMaterialExposedViewPropertyMeta& PropertyMeta = MaterialExternalCodeRegistry::Get().GetExternalViewPropertyCode(InProperty);
-
 	const bool bHasCustomInverseCode = PropertyMeta.InvPropertyCode != nullptr;
 
 	FString HLSLCode = bInvProperty && bHasCustomInverseCode ? PropertyMeta.InvPropertyCode : PropertyMeta.PropertyCode;
@@ -1333,50 +1142,48 @@ static MIR::FValue* BuildViewProperty(MIR::FEmitter& Emitter, EMaterialExposedVi
 	if (bInvProperty && !bHasCustomInverseCode)
 	{
 		// Fall back to compute the property's inverse from PropertyCode
-		return Emitter.EmitOperator(MIR::BO_Divide,
-			Emitter.EmitConstantFloatScalar(1.0f),
-			Emitter.EmitInlineHLSL(HLSLCodeType, MoveTemp(HLSLCode), {}, MIR::EValueFlags::SubstituteTagsInInlineHLSL));
+		return Em.Divide(Em.ConstantFloat(1.0f), Em.InlineHLSL(HLSLCodeType, MoveTemp(HLSLCode), {}, MIR::EValueFlags::SubstituteTagsInInlineHLSL));
 	}
 	else
 	{
 		// @todo-laura.hermanns - old translator used CastToNonLWCIfDisabled(), but LWC not supported in MIR yet
-		return Emitter.EmitInlineHLSL(HLSLCodeType, MoveTemp(HLSLCode), {}, MIR::EValueFlags::SubstituteTagsInInlineHLSL);
+		return Em.InlineHLSL(HLSLCodeType, MoveTemp(HLSLCode), {}, MIR::EValueFlags::SubstituteTagsInInlineHLSL);
 	}
 }
 
-void UMaterialExpressionViewProperty::Build(MIR::FEmitter& Emitter)
+void UMaterialExpressionViewProperty::Build(MIR::FEmitter& Em)
 {
 	for (int32 OutputIndex = 0; OutputIndex < 2; ++OutputIndex)
 	{
 		const bool bInvProperty = OutputIndex == 1;
-		Emitter.Output(OutputIndex, BuildViewProperty(Emitter, Property, bInvProperty));
+		Em.Output(OutputIndex, BuildViewProperty(Em, Property, bInvProperty));
 	}
 }
 
-void UMaterialExpressionViewSize::Build(MIR::FEmitter& Emitter)
+void UMaterialExpressionViewSize::Build(MIR::FEmitter& Em)
 {
-	Emitter.Output(0, BuildViewProperty(Emitter, MEVP_ViewSize));
+	Em.Output(0, BuildViewProperty(Em, MEVP_ViewSize));
 }
 
-void UMaterialExpressionCameraPositionWS::Build(MIR::FEmitter& Emitter)
+void UMaterialExpressionCameraPositionWS::Build(MIR::FEmitter& Em)
 {
-	Emitter.Output(0, BuildViewProperty(Emitter, MEVP_WorldSpaceCameraPosition));
+	Em.Output(0, BuildViewProperty(Em, MEVP_WorldSpaceCameraPosition));
 }
 
-void UMaterialExpressionPixelNormalWS::Build(MIR::FEmitter& Emitter)
+void UMaterialExpressionPixelNormalWS::Build(MIR::FEmitter& Em)
 {
-	MIR::FValue* Output = Emitter.EmitInlineHLSL(MIR::FPrimitiveType::GetFloat3(), "Parameters.WorldNormal", {}, MIR::EValueFlags::None, MIR::EGraphProperties::ReadsPixelNormal);
-	Emitter.Output(0, Output);
+	FValueRef Output = Em.InlineHLSL(MIR::FPrimitiveType::GetFloat3(), "Parameters.WorldNormal", {}, MIR::EValueFlags::None, MIR::EGraphProperties::ReadsPixelNormal);
+	Em.Output(0, Output);
 }
 
-void UMaterialExpressionDDX::Build(MIR::FEmitter& Emitter)
+void UMaterialExpressionDDX::Build(MIR::FEmitter& Em)
 {
-	Emitter.Output(0, Emitter.EmitPartialDerivative(Emitter.Input(&Value), MIR::EDerivativeAxis::X));
+	Em.Output(0, Em.PartialDerivative(Em.Input(&Value), MIR::EDerivativeAxis::X));
 }
 
-void UMaterialExpressionDDY::Build(MIR::FEmitter& Emitter)
+void UMaterialExpressionDDY::Build(MIR::FEmitter& Em)
 {
-	Emitter.Output(0, Emitter.EmitPartialDerivative(Emitter.Input(&Value), MIR::EDerivativeAxis::Y));
+	Em.Output(0, Em.PartialDerivative(Em.Input(&Value), MIR::EDerivativeAxis::Y));
 }
 
 #endif // WITH_EDITOR
diff --git a/Engine/Source/Runtime/Engine/Private/Materials/MaterialIR.cpp b/Engine/Source/Runtime/Engine/Private/Materials/MaterialIR.cpp
--- a/Engine/Source/Runtime/Engine/Private/Materials/MaterialIR.cpp
+++ b/Engine/Source/Runtime/Engine/Private/Materials/MaterialIR.cpp
@@ -390,12 +390,12 @@
 
 bool FConstant::IsInteger() const
 {
-	return Type == FPrimitiveType::GetInt1();
+	return Type == FPrimitiveType::GetInt();
 }
 
 bool FConstant::IsFloat() const
 {
-	return Type == FPrimitiveType::GetFloat1();
+	return Type == FPrimitiveType::GetFloat();
 }
 
 const TCHAR* LexToString(EExternalInput Input)
@@ -441,8 +441,8 @@
 	
 	switch (Id)
 	{
-		case EExternalInput::ViewMaterialTextureMipBias: return FPrimitiveType::GetFloat1();
-		case EExternalInput::ViewMaterialTextureDerivativeMultiply: return FPrimitiveType::GetFloat1();
+		case EExternalInput::ViewMaterialTextureMipBias: return FPrimitiveType::GetFloat();
+		case EExternalInput::ViewMaterialTextureDerivativeMultiply: return FPrimitiveType::GetFloat();
 		default: UE_MIR_UNREACHABLE();
 	}
 }
@@ -584,30 +584,10 @@
 
 const TCHAR* LexToString(EOperator Op)
 {
+	// Note: sorted alphabetically
 	switch (Op)
 	{
-		// Note: sorted alphabetically
-		case BO_Add: return TEXT("Add");
-		case BO_And: return TEXT("And");
-		case BO_Cross: return TEXT("Cross");
-		case BO_Divide: return TEXT("Divide");
-		case BO_Dot: return TEXT("Dot");
-		case BO_Equals: return TEXT("Equals");
-		case BO_Fmod: return TEXT("Fmod");
-		case BO_GreaterThan: return TEXT("GreaterThan");
-		case BO_GreaterThanOrEquals: return TEXT("GreaterThanOrEquals");
-		case BO_LessThan: return TEXT("LessThan");
-		case BO_LessThanOrEquals: return TEXT("LessThanOrEquals");
-		case BO_Max: return TEXT("Max");
-		case BO_Min: return TEXT("Min");
-		case BO_Multiply: return TEXT("Multiply");
-		case BO_NotEquals: return TEXT("NotEquals");
-		case BO_Or: return TEXT("Or");
-		case BO_Subtract: return TEXT("Subtract");
-		case TO_Clamp: return TEXT("Clamp");
-		case TO_Lerp: return TEXT("Lerp");
-		case TO_Select: return TEXT("Select");
-		case TO_Smoothstep: return TEXT("Smoothstep");
+		/* Unary operators */
 		case UO_Abs: return TEXT("Abs");
 		case UO_ACos: return TEXT("ACos");
 		case UO_ASin: return TEXT("ASin");
@@ -631,6 +611,31 @@
 		case UO_Tan: return TEXT("Tan");
 		case UO_Truncate: return TEXT("Truncate");
 
+		/* Binary operators */
+		case BO_Add: return TEXT("Add");
+		case BO_And: return TEXT("And");
+		case BO_Cross: return TEXT("Cross");
+		case BO_Divide: return TEXT("Divide");
+		case BO_Dot: return TEXT("Dot");
+		case BO_Equals: return TEXT("Equals");
+		case BO_Fmod: return TEXT("Fmod");
+		case BO_GreaterThan: return TEXT("GreaterThan");
+		case BO_GreaterThanOrEquals: return TEXT("GreaterThanOrEquals");
+		case BO_LessThan: return TEXT("LessThan");
+		case BO_LessThanOrEquals: return TEXT("LessThanOrEquals");
+		case BO_Max: return TEXT("Max");
+		case BO_Min: return TEXT("Min");
+		case BO_Multiply: return TEXT("Multiply");
+		case BO_NotEquals: return TEXT("NotEquals");
+		case BO_Or: return TEXT("Or");
+		case BO_Subtract: return TEXT("Subtract");
+
+		/* Ternary operators */
+		case TO_Clamp: return TEXT("Clamp");
+		case TO_Lerp: return TEXT("Lerp");
+		case TO_Select: return TEXT("Select");
+		case TO_Smoothstep: return TEXT("Smoothstep");
+		
 		case O_Invalid: return TEXT("Invalid");
 		case OperatorCount: UE_MIR_UNREACHABLE();
 	}
diff --git a/Engine/Source/Runtime/Engine/Private/Materials/MaterialIREmitter.cpp b/Engine/Source/Runtime/Engine/Private/Materials/MaterialIREmitter.cpp
--- a/Engine/Source/Runtime/Engine/Private/Materials/MaterialIREmitter.cpp
+++ b/Engine/Source/Runtime/Engine/Private/Materials/MaterialIREmitter.cpp
@@ -67,30 +67,30 @@ void FSwizzleMask::Append(EVectorComponent Component)
 struct FEmitter::FPrivate
 {
 	// Looks for an existing value in the module that matches `Prototype` and returns it if found.
-	static FValue* FindValue(FEmitter* Emitter, const FValue* Prototype)
+	static FValue* FindValue(FEmitter& Emitter, const FValue* Prototype)
 	{
-		FValue** Value = Emitter->ValueSet.Find(Prototype);
+		FValue** Value = Emitter.ValueSet.Find(Prototype);
 		return Value ? *Value : nullptr;
 	}
 
 	// Allocates a block of memory using the emitter's allocator.
-	static uint8* Allocate(FEmitter* Emitter, int32 Size, int32 Alignment)
+	static uint8* Allocate(FEmitter& Emitter, int32 Size, int32 Alignment)
 	{
-		uint8* Bytes = (uint8*) FMemory::Malloc(Size, Alignment); // Emitter->Module->Allocator.PushBytes(Size, Alignment);
+		uint8* Bytes = (uint8*) FMemory::Malloc(Size, Alignment); // Emitter.Module->Allocator.PushBytes(Size, Alignment);
 		FMemory::Memzero(Bytes, Size);
 		return Bytes;
 	}
 
 	// Pushes a new value to the list of values.
-	static void PushNewValue(FEmitter* Emitter, FValue* Value)
+	static void PushNewValue(FEmitter& Emitter, FValue* Value)
 	{
-		Emitter->Module->Values.Add(Value);
-		Emitter->ValueSet.Add(Value);
+		Emitter.Module->Values.Add(Value);
+		Emitter.ValueSet.Add(Value);
 	}
 };
 
 // Creates a new `FDimensional` value of specified `Type` and returns it.
-static FDimensional* NewDimensionalValue(FEmitter* Emitter, const FPrimitiveType* Type)
+static FDimensional* NewDimensionalValue(FEmitter& Emitter, const FPrimitiveType* Type)
 {
 	check(!Type->IsScalar());
 
@@ -108,7 +108,7 @@ static FDimensional* NewDimensionalValue(FEmitter* Emitter, const FPrimitiveType
 
 // Emits specified newly created `Value`. If the exact value already exists,
 // specified one is *destroyed* and existing one is returned instead.
-static FValue* EmitNew(FEmitter* Emitter, FValue* Value)
+static FValueRef EmitNew(FEmitter& Emitter, FValue* Value)
 {
 	if (FValue* Existing = FEmitter::FPrivate::FindValue(Emitter, Value))
 	{
@@ -137,7 +137,7 @@ static T MakePrototype(const FType* InType)
 // If none found, it creates a new value as a copy of the prototype, adds it to
 // the module then returns it.
 template <typename TValueType>
-static FValue* EmitPrototype(FEmitter* Emitter, const TValueType& Prototype)
+static FValueRef EmitPrototype(FEmitter& Emitter, const TValueType& Prototype)
 {
 	if (FValue* Existing = FEmitter::FPrivate::FindValue(Emitter, &Prototype))
 	{
@@ -173,607 +173,419 @@ static FName SlowFindInputName(UMaterialExpression* Expression, const FExpressio
 	return Expression->GetInputName(InputIndex);
 }
 
-// FEmitter API
+/*------------------------------------- FValueRef --------------------------------------*/
 
-FEmitter& FEmitter::DefaultInputToFloatScalarZero(const FExpressionInput* InInput)
+static FValueRef With(FValueRef ValueRef, FValue* Value)
 {
-	return DefaultInputToFloatScalar(InInput, 0.0f);
+	FValueRef Copy = ValueRef;
+	Copy.Value = Value;
+	return Copy;
 }
 
-FEmitter& FEmitter::DefaultInputToBoolScalar(const FExpressionInput* InInput, bool Value)
+static inline bool IsAnyNotValid()
 {
-	if (!InInput->IsConnected())
-	{
-		Internal::BindValueToExpressionInput(BuilderImpl, InInput, EmitConstantBoolScalar(Value));
-	}
-
-	return *this;
+	return false;
 }
 
-FEmitter& FEmitter::DefaultInputToIntScalar(const FExpressionInput* InInput, TInteger Value)
+// Returns whether any of the values is invalid (null or poison).
+template <typename... TTail>
+static bool IsAnyNotValid(FValueRef Head, TTail... Tail)
 {
-	if (!InInput->IsConnected())
-	{
-		Internal::BindValueToExpressionInput(BuilderImpl, InInput, EmitConstantIntScalar(Value));
-	}
-
-	return *this;
+	return !Head.IsValid() || IsAnyNotValid(Tail...);
 }
 
-FEmitter& FEmitter::DefaultInputToFloatScalar(const FExpressionInput* InInput, TFloat Value)
+// Returns whether any of the values is invalid (null or poison).
+static inline bool IsAnyNotValid(TConstArrayView<FValueRef> Values)
 {
-	if (!InInput->IsConnected())
+	for (FValueRef Value : Values)
 	{
-		Internal::BindValueToExpressionInput(BuilderImpl, InInput, EmitConstantFloatScalar(Value));
+		if (!Value.IsValid())
+		{
+			return true;
+		}
 	}
-
-	return *this;
+	return false;
 }
 
-FEmitter& FEmitter::DefaultInputToFloat3(const FExpressionInput* InInput, FVector3f Value)
+bool FValueRef::IsValid() const
 {
-	if (!InInput->IsConnected())
-	{
-		Internal::BindValueToExpressionInput(BuilderImpl, InInput, EmitConstantFloat3(Value));
-	}
-
-	return *this;
+	return Value && !Value->IsPoison();
 }
 
-FEmitter& FEmitter::DefaultInputToFloat4(const FExpressionInput* InInput, FVector4f Value)
+FValueRef FValueRef::To(FValue* InValue) const
 {
-	if (!InInput->IsConnected())
-	{
-		Internal::BindValueToExpressionInput(BuilderImpl, InInput, EmitConstantFloat4(Value));
-	}
+	return { InValue, Input };
+}
 
-	return *this;
+FValueRef FValueRef::ToPoison() const
+{
+	return To(FPoison::Get());
 }
 
-FValue* FEmitter::TryInput(const FExpressionInput* InInput)
+/*----------------------------------- Error handling -----------------------------------*/
+
+void FEmitter::Error(FValueRef Source, FStringView Message)
 {
-	FValue* Value = MIR::Internal::FetchValueFromExpressionInput(BuilderImpl, InInput);
-	return Value ? Value : nullptr;
+	Source.Input
+		? Error(FString::Printf(TEXT("From expression input '%s': %s"), *SlowFindInputName(Expression, Source.Input).ToString(),Message.GetData()))
+		: Error(Message);
 }
 
-FValue* FEmitter::Input(const FExpressionInput* InInput)
+void FEmitter::Error(FStringView Message)
 {
-	if (FValue* Value = TryInput(InInput))
-	{
-		return Value;
-	}
+	FMaterialIRModule::FError Error;
+	Error.Expression = Expression;
 
-	Errorf(TEXT("Missing '%s' input value."), *SlowFindInputName(Expression, InInput).ToString());
-	return MIR::FPoison::Get();
+	// Add the node type to the error message
+	const int32 ChopCount = FCString::Strlen(TEXT("MaterialExpression"));
+	const FString& ErrorClassName = Expression->GetClass()->GetName();
+	
+	Error.Message.Appendf(TEXT("(Node %s) %s"), *ErrorClassName + ChopCount, Message.GetData());
+	
+	Module->Errors.Push(Error);
+	bCurrentExpressionHasErrors = true;
 }
 
-FValue* FEmitter::InputOfType(const FExpressionInput* InInput, ETypeKind Kind)
+/*--------------------------------- Type handling ----------------------------------*/
+
+const FType* FEmitter::GetCommonType(const FType* A, const FType* B)
 {
-	FValue* Value = Input(InInput);
-	if (IsAnyNonNullAPoison(Value))
+	if (const FType* CommonType = TryGetCommonType(A, B))
 	{
-		return Value;
+		return CommonType;
 	}
 
-	CheckInputTypeIs(InInput, Value, Kind);
-	return Value;
+	Errorf(TEXT("No common type between '%s' and '%s'."), A->GetSpelling().GetData(), B->GetSpelling().GetData());
+	return nullptr;
 }
 
-bool FEmitter::InputConstantBoolScalar(const FExpressionInput* InInput)
+/*-------------------------------- Input management --------------------------------*/
+
+FValueRef FEmitter::TryInput(const FExpressionInput* InInput)
 {
-	FValue* Value = InputBoolScalar(InInput);
-	if (IsAnyNonNullAPoison(Value))
-	{
-		return false;
-	}
+	return FValueRef{ MIR::Internal::FetchValueFromExpressionInput(BuilderImpl, InInput), InInput };
+}
 
-	const FConstant* Constant = MIR::As<FConstant>(Value);
-	if (!Constant)
+FValueRef FEmitter::Input(const FExpressionInput* InInput)
+{
+	FValueRef Value = TryInput(InInput);
+	if (!Value)
 	{
-		Errorf(TEXT("Expected constant bool value through '%s' input."), *SlowFindInputName(Expression, InInput).ToString());
-		return false;
+		Errorf(TEXT("Missing '%s' input value."), *SlowFindInputName(Expression, InInput).ToString());
+		return Value.ToPoison();
 	}
-
-	return Constant->Boolean;
+	return Value;
 }
 
-FValue* FEmitter::InputBoolScalar(const FExpressionInput* InInput)
+FValueRef FEmitter::InputDefaultBool(const FExpressionInput* Input, bool Default)
 {
-	FValue* Value = Input(InInput);
-	CheckInputIsScalar(InInput, Value, ScalarKind_Bool);
-	return Value;
+	FValueRef Value = TryInput(Input);
+	return Value ? Value : Value.To(ConstantBool(Default));
 }
 
-FValue* FEmitter::InputFloatScalar(const FExpressionInput* InInput)
+FValueRef FEmitter::InputDefaultInt(const FExpressionInput* Input, TInteger Default)
 {
-	FValue* Value = Input(InInput);
-	CheckInputIsScalar(InInput, Value, ScalarKind_Float);
-	return Value;
+	FValueRef Value = TryInput(Input);
+	return Value ? Value : Value.To(ConstantInt(Default));
 }
 
-FValue* FEmitter::InputScalar(const FExpressionInput* InInput)
+FValueRef FEmitter::InputDefaultInt2(const FExpressionInput* Input, UE::Math::TIntVector2<TInteger> Default)
 {
-	FValue* Value = InputPrimitive(InInput);
-	if (IsAnyNonNullAPoison(Value))
-	{
-		return Value;
-	}
-
-	const FPrimitiveType* ScalarType = Value->Type->AsPrimitive()->ToScalar();
-	return EmitConstruct(ScalarType, Value);
+	FValueRef Value = TryInput(Input);
+	return Value ? Value : Value.To(ConstantInt2(Default));
 }
 
-FValue* FEmitter::InputPrimitive(const FExpressionInput* InInput)
+FValueRef FEmitter::InputDefaultInt3(const FExpressionInput* Input, UE::Math::TIntVector3<TInteger> Default)
 {
-	FValue* Value = Input(InInput);
-	if (IsAnyNonNullAPoison(Value))
-	{
-		return Value;
-	}
-
-	CheckInputTypeIs(InInput, Value, TypeKind_Primitive);
-	return Value;
+	FValueRef Value = TryInput(Input);
+	return Value ? Value : Value.To(ConstantInt3(Default));
 }
 
-FValue* FEmitter::InputFloatVector(const FExpressionInput* InInput, int32 NumComponents)
+FValueRef FEmitter::InputDefaultInt4(const FExpressionInput* Input, UE::Math::TIntVector4<TInteger> Default)
 {
-	FValue* Value = Input(InInput);
-	CheckInputIsVector(InInput, Value, NumComponents, ScalarKind_Float);
-	return Value;
+	FValueRef Value = TryInput(Input);
+	return Value ? Value : Value.To(ConstantInt4(Default));
 }
 
-FValue* FEmitter::InputTexture(const FExpressionInput* InInput)
+FValueRef FEmitter::InputDefaultFloat(const FExpressionInput* Input, TFloat Default)
 {
-	FValue* Value = Input(InInput);
-	if (IsAnyNonNullAPoison(Value))
-	{
-		return Value;
-	}
-
-	CheckValueIsTexture(Value);
-	return Value;
+	FValueRef Value = TryInput(Input);
+	return Value ? Value : Value.To(ConstantFloat(Default));
 }
 
-FEmitter& FEmitter::Output(int32 OutputIndex, FValue* Value)
+FValueRef FEmitter::InputDefaultFloat2(const FExpressionInput* Input, UE::Math::TVector2<TFloat> Default)
 {
-	Output(Expression->GetOutput(OutputIndex), Value);
-	return *this;
+	FValueRef Value = TryInput(Input);
+	return Value ? Value : Value.To(ConstantFloat2(Default));
 }
 
-FEmitter& FEmitter::Output(const FExpressionOutput* Output, FValue* Value)
+FValueRef FEmitter::InputDefaultFloat3(const FExpressionInput* Input, UE::Math::TVector<TFloat> Default)
 {
-	MIR::Internal::BindValueToExpressionOutput(BuilderImpl, Output, Value);
-	return *this;
+	FValueRef Value = TryInput(Input);
+	return Value ? Value : Value.To(ConstantFloat3(Default));
 }
 
-bool FEmitter::ToConstantBoolScalar(MIR::FValue* Value)
+FValueRef FEmitter::InputDefaultFloat4(const FExpressionInput* Input, UE::Math::TVector4<TFloat> Default)
 {
-	CheckValueIsScalar(Value, ScalarKind_Bool);
-
-	if (CurrentExpressionHasErrors())
-	{
-		return false;
-	}
-
-	if (MIR::FConstant* Constant = As<MIR::FConstant>(Value))
-	{
-		return Constant->Boolean;
-	}
-
-	Errorf(TEXT("Value is not a constant bool."));
-	return false;
+	FValueRef Value = TryInput(Input);
+	return Value ? Value : Value.To(ConstantFloat4(Default));
 }
 
-void FEmitter::CheckInputIsScalar(const FExpressionInput* InInput, FValue* InputValue)
+FValueRef FEmitter::CheckTypeIsKind(FValueRef Value, ETypeKind Kind)
 {
-	if (IsAnyNonNullAPoison(InputValue))
+	if (Value.IsValid() && Value->Type->Kind != Kind)
 	{
-		return;
-	}
-
-	const FPrimitiveType* PrimitiveType = InputValue->Type->AsPrimitive();
-	if (!PrimitiveType || !PrimitiveType->IsScalar())
-	{
-		Errorf(TEXT("Input '%s' expected to be a scalar. It is a %s instead."), *SlowFindInputName(Expression, InInput).ToString(), InputValue->Type->GetSpelling().GetData());
+		Errorf(Value, TEXT("Expected a '%s' value, got a '%s' instead."), TypeKindToString(Kind), Value->Type->GetSpelling().GetData());
+		return Value.ToPoison();
 	}
+	return Value;
 }
 
-void FEmitter::CheckInputIsScalar(const FExpressionInput* InInput, FValue* InputValue, EScalarKind Kind)
+FValueRef FEmitter::CheckIsPrimitive(FValueRef Value)
 {
-	if (IsAnyNonNullAPoison(InputValue))
-	{
-		return;
-	}
-
-	const FPrimitiveType* PrimitiveType = InputValue->Type->AsPrimitive();
-	if (!PrimitiveType || !PrimitiveType->IsScalar() || PrimitiveType->ScalarKind != Kind)
+	if (Value.IsValid() && !Value->Type->AsPrimitive())
 	{
-		Errorf(TEXT("Input '%s' expected to be a %s scalar. It is %s instead."), *SlowFindInputName(Expression, InInput).ToString(), ScalarKindToString(Kind), InputValue->Type->GetSpelling().GetData());
+		Errorf(Value, TEXT("Expected a primitive value, got a '%s' instead."), Value->Type->GetSpelling().GetData());
+		return Value.ToPoison();
 	}
+	return Value;
 }
 
-void FEmitter::CheckInputIsVector(const FExpressionInput* InInput, FValue* InputValue, int32 NumComponents, EScalarKind Kind)
+FValueRef FEmitter::CheckIsArithmetic(FValueRef Value)
 {
-	if (IsAnyNonNullAPoison(InputValue))
+	if (Value.IsValid() && !Value->Type->AsArithmetic())
 	{
-		return;
+		Errorf(Value, TEXT("Expected an arithmetic value, got a '%s' instead."), Value->Type->GetSpelling().GetData());
+		return Value.ToPoison();
 	}
+	return Value;
+}
 
-	const FPrimitiveType* PrimitiveType = InputValue->Type->AsPrimitive();
-	if (!PrimitiveType || !PrimitiveType->IsVector() || PrimitiveType->NumRows != NumComponents || PrimitiveType->ScalarKind != Kind)
+FValueRef FEmitter::CheckIsScalar(FValueRef Value)
+{
+	if (Value.IsValid() && (!Value->Type->AsPrimitive() || !Value->Type->AsPrimitive()->IsScalar()))
 	{
-		Errorf(TEXT("Input '%s' expected to be a %s %d-component vector. It is %s instead."), *SlowFindInputName(Expression, InInput).ToString(), ScalarKindToString(Kind), NumComponents, InputValue->Type->GetSpelling().GetData());
+		Errorf(Value, TEXT("Expected a scalar value, got a '%s' instead."), Value->Type->GetSpelling().GetData());
+		return Value.ToPoison();
 	}
+	return Value;
 }
 
-void FEmitter::CheckInputTypeIs(const FExpressionInput* InInput, FValue* InputValue, ETypeKind Kind)
+FValueRef FEmitter::CheckIsScalarOrVector(FValueRef Value)
 {
-	if (IsAnyNonNullAPoison(InputValue))
+	if (Value.IsValid() && (!Value->Type->AsPrimitive() || Value->Type->AsPrimitive()->IsMatrix()))
 	{
-		return;
+		Errorf(Value, TEXT("Expected a scalar or vector value, got a '%s' instead."), Value->Type->GetSpelling().GetData());
+		return Value.ToPoison();
 	}
+	return Value;
+}
 
-	if (InputValue && InputValue->Type->Kind != Kind)
+FValueRef FEmitter::CheckIsTexture(FValueRef Value)
+{
+	if (Value.IsValid() && !Value->Type->IsTexture())
 	{
-		Errorf(TEXT("Input '%s' expected to be have type %s. It is %s instead."), *SlowFindInputName(Expression, InInput).ToString(), TypeKindToString(Kind), InputValue->Type->GetSpelling().GetData());
+		Errorf(Value, TEXT("Expected a texture value, got a '%s' instead."), Value->Type->GetSpelling().GetData());
+		return Value.ToPoison();
 	}
+	return Value;
 }
 
-void FEmitter::CheckValueIsScalar(const FValue* Value, EScalarKind Kind)
+bool FEmitter::ToConstantBool(FValueRef Value)
 {
-	if (IsAnyNonNullAPoison(Value))
+	if (!Value.IsValid())
 	{
-		return;
+		return false;
 	}
+	const MIR::FConstant* Constant = Value->As<MIR::FConstant>();
+    if (!Constant)
+    {
+        Errorf(Value, TEXT("Expected a constant bool value, got a non-constant value instead."));
+        return false;
+    }
+	if (Constant->Type != FPrimitiveType::GetBool())
+    {
+        Errorf(Value, TEXT("Expected a constant bool value, got a '%s' instead."), Constant->Type->GetSpelling().GetData());
+        return false;
+    }
+	return Constant->Boolean;
+ }
 
-	const FPrimitiveType* PrimitiveType = Value->Type->AsPrimitive();
-	if (!PrimitiveType || !PrimitiveType->IsScalar() || PrimitiveType->ScalarKind != Kind)
-	{
-		Errorf(TEXT("Value expected to be a %s scalar. It is %s instead."), ScalarKindToString(Kind), Value->Type->GetSpelling().GetData());
-		return;
-	}
-}
+/*-------------------------------- Output management -------------------------------*/
 
-void FEmitter::CheckValueIsTexture(const FValue* Value)
+FEmitter& FEmitter::Output(int32 OutputIndex, FValueRef Value)
 {
-	if (IsAnyNonNullAPoison(Value))
-	{
-		return;
-	}
+	Output(Expression->GetOutput(OutputIndex), Value);
+	return *this;
+}
 
-	if (!Value->Type->IsTexture())
-	{
-		Errorf(TEXT("Value expected to be a texture. It is %s instead."), Value->Type->GetSpelling().GetData());
-		return;
-	}
+FEmitter& FEmitter::Output(const FExpressionOutput* Output, FValueRef Value)
+{
+	MIR::Internal::BindValueToExpressionOutput(BuilderImpl, Output, Value);
+	return *this;
 }
 
-FValue* FEmitter::EmitConstantFromShaderValue(const UE::Shader::FValue& InValue)
+/*------------------------------- Constants emission -------------------------------*/
+
+FValueRef FEmitter::ConstantFromShaderValue(const UE::Shader::FValue& InValue)
 {
 	using namespace UE::Shader;
 
 	switch (InValue.Type.ValueType)
 	{
-		case UE::Shader::EValueType::Float1: return EmitConstantFloatScalar(InValue.AsFloatScalar());
-		case UE::Shader::EValueType::Float2: return EmitConstantFloat2(FVector2f{ InValue.Component[0].Float, InValue.Component[1].Float });
-		case UE::Shader::EValueType::Float3: return EmitConstantFloat3(FVector3f{ InValue.Component[0].Float, InValue.Component[1].Float, InValue.Component[2].Float });
-		case UE::Shader::EValueType::Float4: return EmitConstantFloat4(FVector4f{ InValue.Component[0].Float, InValue.Component[1].Float, InValue.Component[2].Float, InValue.Component[3].Float });
+		case UE::Shader::EValueType::Float1: return ConstantFloat(InValue.AsFloatScalar());
+		case UE::Shader::EValueType::Float2: return ConstantFloat2(UE::Math::TVector2<TFloat>{ InValue.Component[0].Float, InValue.Component[1].Float });
+		case UE::Shader::EValueType::Float3: return ConstantFloat3(UE::Math::TVector<TFloat>{ InValue.Component[0].Float, InValue.Component[1].Float, InValue.Component[2].Float });
+		case UE::Shader::EValueType::Float4: return ConstantFloat4(UE::Math::TVector4<TFloat>{ InValue.Component[0].Float, InValue.Component[1].Float, InValue.Component[2].Float, InValue.Component[3].Float });
 
-		case UE::Shader::EValueType::Int1: return EmitConstantIntScalar(InValue.AsFloatScalar());
-		case UE::Shader::EValueType::Int2: return EmitConstantInt2(FIntVector2{ InValue.Component[0].Int, InValue.Component[1].Int });
-		case UE::Shader::EValueType::Int3: return EmitConstantInt3(FIntVector3{ InValue.Component[0].Int, InValue.Component[1].Int, InValue.Component[2].Int });
-		case UE::Shader::EValueType::Int4: return EmitConstantInt4(FIntVector4{ InValue.Component[0].Int, InValue.Component[1].Int, InValue.Component[2].Int, InValue.Component[3].Int });
+		case UE::Shader::EValueType::Int1: return ConstantInt(InValue.AsFloatScalar());
+		case UE::Shader::EValueType::Int2: return ConstantInt2({ InValue.Component[0].Int, InValue.Component[1].Int });
+		case UE::Shader::EValueType::Int3: return ConstantInt3({ InValue.Component[0].Int, InValue.Component[1].Int, InValue.Component[2].Int });
+		case UE::Shader::EValueType::Int4: return ConstantInt4({ InValue.Component[0].Int, InValue.Component[1].Int, InValue.Component[2].Int, InValue.Component[3].Int });
 	}
 
 	UE_MIR_UNREACHABLE();
-} 
+}
 
-FValue* FEmitter::EmitConstantScalarZero(EScalarKind Kind)
+FValueRef FEmitter::ConstantZero(EScalarKind Kind)
 {
 	switch (Kind)
 	{
-		case ScalarKind_Bool: return EmitConstantFalse();
-		case ScalarKind_Int: return EmitConstantIntScalar(0);
-		case ScalarKind_Float: return EmitConstantFloatScalar(0.0f);
+		case ScalarKind_Bool: return ConstantFalse();
+		case ScalarKind_Int: return ConstantInt(0);
+		case ScalarKind_Float: return ConstantFloat(0.0f);
 		default: UE_MIR_UNREACHABLE();
 	}
 }
 
-FValue* FEmitter::EmitConstantScalarOne(EScalarKind Kind)
+FValueRef FEmitter::ConstantOne(EScalarKind Kind)
 {
 	switch (Kind)
 	{
-		case ScalarKind_Bool: return EmitConstantTrue();
-		case ScalarKind_Int: return EmitConstantIntScalar(1);
-		case ScalarKind_Float: return EmitConstantFloatScalar(1.0f);
+		case ScalarKind_Bool: return ConstantTrue();
+		case ScalarKind_Int: return ConstantInt(1);
+		case ScalarKind_Float: return ConstantFloat(1.0f);
 		default: UE_MIR_UNREACHABLE();
 	}
 }
 
-FValue* FEmitter::EmitConstantTrue()
+FValueRef FEmitter::ConstantScalar(EScalarKind Kind, TFloat FromFloat)
 {
-	return ConstantTrue;
+	switch (Kind)
+	{
+		case ScalarKind_Bool: return ConstantBool(FromFloat != 0.0f);
+		case ScalarKind_Int: return ConstantInt(TInteger(FromFloat));
+		case ScalarKind_Float: return ConstantFloat(FromFloat);
+		default: UE_MIR_UNREACHABLE();
+	}
 }
 
-FValue* FEmitter::EmitConstantFalse()
+FValueRef FEmitter::ConstantTrue()
 {
-	return ConstantFalse;
+	return TrueConstant;
 }
 
-FValue* FEmitter::EmitConstantBoolScalar(bool InX)
+FValueRef FEmitter::ConstantFalse()
 {
-	return InX ? EmitConstantTrue() : EmitConstantFalse();
+	return FalseConstant;
 }
 
-FValue* FEmitter::EmitConstantFloatScalar(TFloat InX)
+FValueRef FEmitter::ConstantBool(bool InX)
 {
-	FConstant Scalar = MakePrototype<FConstant>(FPrimitiveType::GetScalar(ScalarKind_Float));
-	Scalar.Float = InX;
-	return EmitPrototype(this, Scalar);
+	return InX ? ConstantTrue() : ConstantFalse();
 }
 
-FValue* FEmitter::EmitConstantIntScalar(TInteger InX)
+FValueRef FEmitter::ConstantInt(TInteger InX)
 {
 	FConstant Scalar = MakePrototype<FConstant>(FPrimitiveType::GetScalar(ScalarKind_Int));
 	Scalar.Integer = InX;
-	return EmitPrototype(this, Scalar);
+	return EmitPrototype(*this, Scalar);
 }
 
-FValue* FEmitter::EmitConstantScalar(EScalarKind Kind, TFloat FromFloat)
+FValueRef FEmitter::ConstantInt2(UE::Math::TIntVector2<TInteger> InValue)
 {
-	switch (Kind)
-	{
-		case ScalarKind_Bool: return EmitConstantBoolScalar(FromFloat != 0.0f);
-		case ScalarKind_Int: return EmitConstantIntScalar(TInteger(FromFloat));
-		case ScalarKind_Float: return EmitConstantFloatScalar(FromFloat);
-		default: UE_MIR_UNREACHABLE();
-	}
+	FValueRef X = ConstantInt(InValue.X);
+	FValueRef Y = ConstantInt(InValue.Y);
+	return Vector2(X, Y);
 }
 
-FValue* FEmitter::EmitConstantFloat2(const FVector2f& InValue)
+FValueRef FEmitter::ConstantInt3(UE::Math::TIntVector3<TInteger> InValue)
 {
-	FValue* X = EmitConstantFloatScalar(InValue.X);
-	FValue* Y = EmitConstantFloatScalar(InValue.Y);
-	return EmitVector2(X, Y);
+	FValueRef X = ConstantInt(InValue.X);
+	FValueRef Y = ConstantInt(InValue.Y);
+	FValueRef Z = ConstantInt(InValue.Z);
+	return Vector3(X, Y, Z);
 }
 
-FValue* FEmitter::EmitConstantFloat3(const FVector3f& InValue)
+FValueRef FEmitter::ConstantInt4(UE::Math::TIntVector4<TInteger> InValue)
 {
-	FValue* X = EmitConstantFloatScalar(InValue.X);
-	FValue* Y = EmitConstantFloatScalar(InValue.Y);
-	FValue* Z = EmitConstantFloatScalar(InValue.Z);
-	return EmitVector3(X, Y, Z);
+	FValueRef X = ConstantInt(InValue.X);
+	FValueRef Y = ConstantInt(InValue.Y);
+	FValueRef Z = ConstantInt(InValue.Z);
+	FValueRef W = ConstantInt(InValue.W);
+	return Vector4(X, Y, Z, W);
 }
 
-FValue* FEmitter::EmitConstantFloat4(const FVector4f& InValue)
+FValueRef FEmitter::ConstantFloat(TFloat InX)
 {
-	FValue* X = EmitConstantFloatScalar(InValue.X);
-	FValue* Y = EmitConstantFloatScalar(InValue.Y);
-	FValue* Z = EmitConstantFloatScalar(InValue.Z);
-	FValue* W = EmitConstantFloatScalar(InValue.W);
-	return EmitVector4(X, Y, Z, W);
+	FConstant Scalar = MakePrototype<FConstant>(FPrimitiveType::GetScalar(ScalarKind_Float));
+	Scalar.Float = InX;
+	return EmitPrototype(*this, Scalar);
 }
 
-FValue* FEmitter::EmitConstantInt2(const FIntVector2& InValue)
+FValueRef FEmitter::ConstantFloat2(UE::Math::TVector2<TFloat> InValue)
 {
-	FValue* X = EmitConstantIntScalar(InValue.X);
-	FValue* Y = EmitConstantIntScalar(InValue.Y);
-	return EmitVector2(X, Y);
+	FValueRef X = ConstantFloat(InValue.X);
+	FValueRef Y = ConstantFloat(InValue.Y);
+	return Vector2(X, Y);
 }
 
-FValue* FEmitter::EmitConstantInt3(const FIntVector3& InValue)
+FValueRef FEmitter::ConstantFloat3(UE::Math::TVector<TFloat> InValue)
 {
-	FValue* X = EmitConstantIntScalar(InValue.X);
-	FValue* Y = EmitConstantIntScalar(InValue.Y);
-	FValue* Z = EmitConstantIntScalar(InValue.Z);
-	return EmitVector3(X, Y, Z);
+	FValueRef X = ConstantFloat(InValue.X);
+	FValueRef Y = ConstantFloat(InValue.Y);
+	FValueRef Z = ConstantFloat(InValue.Z);
+	return Vector3(X, Y, Z);
 }
 
-FValue* FEmitter::EmitConstantInt4(const FIntVector4& InValue)
+FValueRef FEmitter::ConstantFloat4(UE::Math::TVector4<TFloat> InValue)
 {
-	FValue* X = EmitConstantIntScalar(InValue.X);
-	FValue* Y = EmitConstantIntScalar(InValue.Y);
-	FValue* Z = EmitConstantIntScalar(InValue.Z);
-	FValue* W = EmitConstantIntScalar(InValue.W);
-	return EmitVector4(X, Y, Z, W);
+	FValueRef X = ConstantFloat(InValue.X);
+	FValueRef Y = ConstantFloat(InValue.Y);
+	FValueRef Z = ConstantFloat(InValue.Z);
+	FValueRef W = ConstantFloat(InValue.W);
+	return Vector4(X, Y, Z, W);
 }
 
-FValue* FEmitter::EmitVector2(FValue* InX, FValue* InY)
-{
-	if (IsAnyNonNullAPoison(InX, InY))
-	{
-		return FPoison::Get();
-	}
-
-	check(InX->Type->AsScalar());
-	check(InX->Type == InY->Type);
+/*--------------------- Other non-instruction values emission ---------------------*/
 
-	TDimensional<2> Vector = MakePrototype<TDimensional<2>>(FPrimitiveType::GetVector(InX->Type->AsPrimitive()->ScalarKind, 2));
-	TArrayView<FValue*> Components = Vector.GetMutableComponents();
-	Components[0] = InX;
-	Components[1] = InY;
-
-	return EmitPrototype(this, Vector);
-}
-
-FValue* FEmitter::EmitVector3(FValue* InX, FValue* InY, FValue* InZ)
+FValueRef FEmitter::Poison()
 {
-	if (IsAnyNonNullAPoison(InX, InY, InZ))
-	{
-		return FPoison::Get();
-	}
-
-	check(InX->Type->AsScalar());
-	check(InX->Type == InY->Type);
-	check(InY->Type == InZ->Type);
-
-	TDimensional<3> Vector = MakePrototype<TDimensional<3>>(FPrimitiveType::GetVector(InX->Type->AsPrimitive()->ScalarKind, 3));
-	TArrayView<FValue*> Components = Vector.GetMutableComponents();
-	Components[0] = InX;
-	Components[1] = InY;
-	Components[2] = InZ;
-
-	return EmitPrototype(this, Vector);
-}
-
-FValue* FEmitter::EmitVector4(FValue* InX, FValue* InY, FValue* InZ, FValue* InW)
-{
-	if (IsAnyNonNullAPoison(InX, InY, InZ, InW))
-	{
-		return FPoison::Get();
-	}
-
-	check(InX->Type->AsScalar());
-	check(InX->Type == InY->Type);
-	check(InY->Type == InZ->Type);
-	check(InZ->Type == InW->Type);
-
-	TDimensional<4> Vector = MakePrototype<TDimensional<4>>(FPrimitiveType::GetVector(InX->Type->AsPrimitive()->ScalarKind, 4));
-	TArrayView<FValue*> Components = Vector.GetMutableComponents();
-	Components[0] = InX;
-	Components[1] = InY;
-	Components[2] = InZ;
-	Components[3] = InW;
-
-	return EmitPrototype(this, Vector);
+	return FPoison::Get();
 }
 
-FValue* FEmitter::EmitExternalInput(EExternalInput Id)
+FValueRef FEmitter::ExternalInput(EExternalInput Id)
 {
 	MIR::FExternalInput Prototype = MakePrototype<MIR::FExternalInput>(GetExternalInputType(Id));
 	Prototype.Id = Id;
-	return EmitPrototype(this, Prototype);
+	return EmitPrototype(*this, Prototype);
 }
 
-FValue* FEmitter::EmitSubscript(FValue* Value, int32 Index)
-{
-	if (IsAnyNonNullAPoison(Value))
-	{
-		return Value;
-	}
-
-	const FPrimitiveType* PrimitiveType = Value->Type->AsPrimitive();
-	if (!PrimitiveType)
-	{
-		Errorf(TEXT("Value of type '%s' cannot be subscripted."), Value->Type->GetSpelling().GetData());
-		return MIR::FPoison::Get();
-	}
-
-	// Getting first component and Value is already a scalar, just return itself.
-	if (Index == 0 && Value->Type->AsScalar())
-	{
-		return Value;
-	}
-
-	if (Index >= PrimitiveType->GetNumComponents())
-	{
-		Errorf(TEXT("Value of type '%s' has fewer dimensions than subscript index `%d`."), Value->Type->GetSpelling().GetData(), Index);
-		return MIR::FPoison::Get();
-	}
-
-	// Getting first component and Value is already a scalar, just return itself.
-	if (FDimensional* DimensionalValue = As<FDimensional>(Value))
-	{
-		check(Index < DimensionalValue->GetComponents().Num());
-		return DimensionalValue->GetComponents()[Index];
-	}
-	
-	// Avoid subscripting a subscript (e.g. no value.xy.x)
-	if (FSubscript* Subscript = As<FSubscript>(Value))
-	{
-		Value = Subscript->Arg;
-	}
-
-	// We can't resolve it at compile time: emit subscript value.
-	FSubscript Prototype = MakePrototype<FSubscript>(PrimitiveType->ToScalar());
-	Prototype.Arg = Value;
-	Prototype.Index = Index;
-
-	return EmitPrototype(this, Prototype);
-}
-
-FValue* FEmitter::EmitSwizzle(FValue* Value, FSwizzleMask Mask)
-{
-	if (IsAnyNonNullAPoison(Value))
-	{
-		return Value;
-	}
-
-	// At least one component must have been specified.
-	check(Mask.NumComponents > 0);
-
-	// We can only swizzle on non-matrix primitive types.
-	const FPrimitiveType* PrimitiveType = Value->Type->AsVector();
-	if (!PrimitiveType || PrimitiveType->IsMatrix())
-	{
-		Errorf(TEXT("Cannot swizzle a '%s' value."), Value->Type->GetSpelling().GetData());
-		return MIR::FPoison::Get();
-	}
-
-	// Make sure each component in the mask fits the number of components in Value.
-	for (EVectorComponent Component : Mask)
-	{
-		if ((int32)Component >= PrimitiveType->NumRows)
-		{
-			Errorf(TEXT("Value of type '%s' has no component '%s'."), PrimitiveType->Spelling.GetData(), VectorComponentToString(Component));
-			return MIR::FPoison::Get();
-		}
-	}
-
-	// If the requested number of components is the same as Value and the order in which the components
-	// are specified in the mask is sequential (e.g. x, y, z) then this is a no op, simply return Value as is.
-	if (Mask.NumComponents == PrimitiveType->GetNumComponents())
-	{
-		bool InOrder = true;
-		for (int32 i = 0; i < Mask.NumComponents; ++i)
-		{
-			if (Mask.Components[i] != (EVectorComponent)i)
-			{
-				InOrder = false;
-				break;
-			}
-		}
-
-		if (InOrder)
-		{
-			return Value;
-		}
-	}
-	
-	// If only one component is requested, we can use EmitSubscript() to return the single component.
-	if (Mask.NumComponents == 1)
-	{
-		return EmitSubscript(Value, (int32)Mask.Components[0]);
-	}
-
-	// Make the result vector type.
-	const FPrimitiveType* ResultType = FPrimitiveType::GetVector(PrimitiveType->ScalarKind, Mask.NumComponents);
-	FDimensional* Result = NewDimensionalValue(this, ResultType);
-	TArrayView<FValue*> ResultComponents = Result->GetMutableComponents();
-
-	for (int32 i = 0; i < Mask.NumComponents; ++i)
-	{
-		ResultComponents[i] = EmitSubscript(Value, (int32)Mask.Components[i]);
-	}
-
-	return EmitNew(this, Result);
-}
-
-FValue* FEmitter::EmitTextureObject(UTexture* Texture, EMaterialSamplerType SamplerType)
+FValueRef FEmitter::TextureObject(UTexture* Texture, EMaterialSamplerType SamplerType)
 {
 	if (Texture->GetMaterialType() != MCT_Texture2D)
 	{
 		Error(TEXT("Only Textures 2D are supported for now."));
-		return MIR::FPoison::Get();
+		return Poison();
 	}
 
 	FTextureObject Proto = MakePrototype<FTextureObject>(FObjectType::GetTexture2D());
 	Proto.Texture = Texture;
 	Proto.SamplerType = SamplerType;
-	return EmitPrototype(this, Proto);
+	return EmitPrototype(*this, Proto);
 }
 
-FValue* FEmitter::EmitParameter(FName Name, FMaterialParameterMetadata& Metadata, EMaterialSamplerType SamplerType)
+FValueRef FEmitter::Parameter(FName Name, FMaterialParameterMetadata& Metadata, EMaterialSamplerType SamplerType)
 {
 	// Helper local function that register a parameter (info and metadata) to the module, and returns some uint32 ID.
 	auto RegisterParameter = [] (FMaterialIRModule* InModule, FMaterialParameterInfo InInfo, const FMaterialParameterMetadata& InMetadata) -> uint32
@@ -799,28 +611,28 @@ FValue* FEmitter::EmitParameter(FName Name, FMaterialParameterMetadata& Metadata
 		{
 			if (Metadata.PrimitiveDataIndex != INDEX_NONE)
 			{
-				return EmitCustomPrimitiveData(Metadata.PrimitiveDataIndex);
+				return CustomPrimitiveData(Metadata.PrimitiveDataIndex);
 			}
 
-			MIR::FUniformParameter Prototype = MakePrototype<MIR::FUniformParameter>(FPrimitiveType::GetFloat1());
+			MIR::FUniformParameter Prototype = MakePrototype<MIR::FUniformParameter>(FPrimitiveType::GetFloat());
 			Prototype.ParameterIdInModule = RegisterParameter(Module, Info, Metadata);
-			return EmitPrototype(this, Prototype);
+			return EmitPrototype(*this, Prototype);
 		}
 
 		case EMaterialParameterType::Vector:
 		{
 			if (Metadata.PrimitiveDataIndex != INDEX_NONE)
 			{
-				MIR::FValue* X = EmitCustomPrimitiveData(Metadata.PrimitiveDataIndex + 0);
-				MIR::FValue* Y = EmitCustomPrimitiveData(Metadata.PrimitiveDataIndex + 1);
-				MIR::FValue* Z = EmitCustomPrimitiveData(Metadata.PrimitiveDataIndex + 2);
-				MIR::FValue* W = EmitCustomPrimitiveData(Metadata.PrimitiveDataIndex + 3);
-				return EmitVector4(X, Y, Z, W);
+				MIR::FValue* X = CustomPrimitiveData(Metadata.PrimitiveDataIndex + 0);
+				MIR::FValue* Y = CustomPrimitiveData(Metadata.PrimitiveDataIndex + 1);
+				MIR::FValue* Z = CustomPrimitiveData(Metadata.PrimitiveDataIndex + 2);
+				MIR::FValue* W = CustomPrimitiveData(Metadata.PrimitiveDataIndex + 3);
+				return Vector4(X, Y, Z, W);
 			}
 
 			MIR::FUniformParameter Prototype = MakePrototype<MIR::FUniformParameter>(FPrimitiveType::GetFloat4());
 			Prototype.ParameterIdInModule = RegisterParameter(Module, Info, Metadata);
-			return EmitPrototype(this, Prototype);
+			return EmitPrototype(*this, Prototype);
 		}
 
 		case EMaterialParameterType::Texture:
@@ -834,7 +646,7 @@ FValue* FEmitter::EmitParameter(FName Name, FMaterialParameterMetadata& Metadata
 			MIR::FUniformParameter Prototype = MakePrototype<MIR::FUniformParameter>(FObjectType::GetTexture2D());
 			Prototype.ParameterIdInModule = RegisterParameter(Module, Info, Metadata);
 			Prototype.SamplerType = SamplerType;
-			return EmitPrototype(this, Prototype);
+			return EmitPrototype(*this, Prototype);
 		}
 
 		case EMaterialParameterType::StaticSwitch:
@@ -848,7 +660,7 @@ FValue* FEmitter::EmitParameter(FName Name, FMaterialParameterMetadata& Metadata
 					break;
 				}
 			}
-			return EmitConstantBoolScalar(Metadata.Value.Bool[0]);
+			return ConstantBool(Metadata.Value.Bool[0]);
 		}
 
 		default:
@@ -856,24 +668,15 @@ FValue* FEmitter::EmitParameter(FName Name, FMaterialParameterMetadata& Metadata
 	}
 }
 
-FValue* FEmitter::EmitCustomPrimitiveData(uint32 PrimitiveDataIndex)
+FValueRef FEmitter::CustomPrimitiveData(uint32 PrimitiveDataIndex)
 {
 	// UE_MIR_TODO();
-	return nullptr;
+	return {};
 }
 
-FValue* FEmitter::EmitStageSwitch(const FType* Type, TConstArrayView<FValue*> ValuePerStage)
-{
-	check(ValuePerStage.Num() <= NumStages);
-	FStageSwitch Prototype = MakePrototype<FStageSwitch>(Type);
-	for (int i = 0; i < ValuePerStage.Num(); ++i)
-	{
-		Prototype.Args[i] = ValuePerStage[i];
-	}
-	return EmitPrototype(this, Prototype);
-}
+/*------------------------------ Instruction emission ------------------------------*/
 
-FSetMaterialOutput* FEmitter::EmitSetMaterialOutput(EMaterialProperty InProperty, FValue* Arg)
+FSetMaterialOutput* FEmitter::SetMaterialOutput(EMaterialProperty InProperty, FValue* Arg)
 {
 	FSetMaterialOutput Proto = MakePrototype<FSetMaterialOutput>(nullptr);
 	Proto.Property  			= InProperty;
@@ -889,7 +692,7 @@ FSetMaterialOutput* FEmitter::EmitSetMaterialOutput(EMaterialProperty InProperty
 	}
 
 	// Make the instruction
-	FSetMaterialOutput* Instr = static_cast<FSetMaterialOutput*>(EmitPrototype(this, Proto));
+	FSetMaterialOutput* Instr = static_cast<FSetMaterialOutput*>(EmitPrototype(*this, Proto).Value);
 
 	// Add the instruction to list of outputs of the stages it is evaluated in.
 	for (int i = 0; i < NumStages; ++i)
@@ -903,6 +706,86 @@ FSetMaterialOutput* FEmitter::EmitSetMaterialOutput(EMaterialProperty InProperty
 	return Instr;
 }
 
+FValueRef FEmitter::Vector2(FValueRef InX, FValueRef InY)
+{
+	if (IsAnyNotValid(InX, InY))
+	{
+		return Poison();
+	}
+
+	check(InX->Type->AsScalar());
+	check(InX->Type == InY->Type);
+
+	TDimensional<2> Vector = MakePrototype<TDimensional<2>>(FPrimitiveType::GetVector(InX->Type->AsPrimitive()->ScalarKind, 2));
+	TArrayView<FValue*> Components = Vector.GetMutableComponents();
+	Components[0] = InX;
+	Components[1] = InY;
+
+	FValueRef Value = EmitPrototype(*this, Vector);
+	if (InX.Input == InY.Input)
+	{
+		Value.Input = InX.Input;
+	}
+
+	return EmitPrototype(*this, Vector);
+}
+
+FValueRef FEmitter::Vector3(FValueRef InX, FValueRef InY, FValueRef InZ)
+{
+	if (IsAnyNotValid(InX, InY, InZ))
+	{
+		return Poison();
+	}
+
+	check(InX->Type->AsScalar());
+	check(InX->Type == InY->Type);
+	check(InY->Type == InZ->Type);
+
+	TDimensional<3> Vector = MakePrototype<TDimensional<3>>(FPrimitiveType::GetVector(InX->Type->AsPrimitive()->ScalarKind, 3));
+	TArrayView<FValue*> Components = Vector.GetMutableComponents();
+	Components[0] = InX;
+	Components[1] = InY;
+	Components[2] = InZ;
+
+	FValueRef Value = EmitPrototype(*this, Vector);
+	if (InX.Input == InY.Input && InX.Input == InZ.Input)
+	{
+		Value.Input = InX.Input;
+	}
+
+	return Value;
+}
+
+FValueRef FEmitter::Vector4(FValueRef InX, FValueRef InY, FValueRef InZ, FValueRef InW)
+{
+	if (IsAnyNotValid(InX, InY, InZ, InW))
+	{
+		return Poison();
+	}
+
+	check(InX->Type->AsScalar());
+	check(InX->Type == InY->Type);
+	check(InY->Type == InZ->Type);
+	check(InZ->Type == InW->Type);
+
+	TDimensional<4> Vector = MakePrototype<TDimensional<4>>(FPrimitiveType::GetVector(InX->Type->AsPrimitive()->ScalarKind, 4));
+	TArrayView<FValue*> Components = Vector.GetMutableComponents();
+	Components[0] = InX;
+	Components[1] = InY;
+	Components[2] = InZ;
+	Components[3] = InW;
+
+	FValueRef Value = EmitPrototype(*this, Vector);
+	if (InX.Input == InY.Input && InX.Input == InZ.Input && InX.Input == InW.Input)
+	{
+		Value.Input = InX.Input;
+	}
+
+	return EmitPrototype(*this, Vector);
+}
+
+/*--------------------------------- Operator emission ---------------------------------*/
+
 template <typename T>
 static bool FoldComparisonOperatorScalar(EOperator Operator, T A, T B)
 {
@@ -931,7 +814,7 @@ static TInteger FoldScalarBitOperator(EOperator Operator, TInteger A, TInteger B
 }
 
 template <typename T>
-static TOptional<T> FoldScalarArithmeticOperator(FEmitter* Emitter, EOperator Operator, T A, T B, T C)
+static TOptional<T> FoldScalarArithmeticOperator(FEmitter& Emitter, EOperator Operator, T A, T B, T C)
 {
 	if constexpr (std::is_same<T, float>::value || std::is_same<T, double>::value)
 	{
@@ -971,7 +854,7 @@ static TOptional<T> FoldScalarArithmeticOperator(FEmitter* Emitter, EOperator Op
 			{
 				if (A <= 0)
 				{
-					Emitter->Errorf(TEXT("Logarithm of 0 or a negative number is not valid."));
+					Emitter.Errorf(TEXT("Logarithm of 0 or a negative number is not valid."));
 					return NullOpt;
 				}
 				const T Base = Operator == UO_Logarithm ? UE_EULERS_NUMBER
@@ -1029,7 +912,7 @@ static TOptional<T> FoldScalarArithmeticOperator(FEmitter* Emitter, EOperator Op
 // If it returns null, the end result could not be inferred, but the operator could have
 // still been changed to some other (with lower complexity). For example "clamp(x, 0, 1)"
 // will change to "saturate(x)".
-static FValue* TrySimplifyOperator(FEmitter* Emitter, EOperator& Op, FValue*& A, FValue*& B, FValue*& C)
+static FValueRef TrySimplifyOperator(FEmitter& Emitter, EOperator& Op, FValueRef& A, FValueRef& B, FValueRef& C)
 {
 	switch (Op)
 	{
@@ -1047,7 +930,7 @@ static FValue* TrySimplifyOperator(FEmitter* Emitter, EOperator& Op, FValue*& A,
 		case BO_NotEquals:
 			if (A->Equals(B))
 			{
-				return Emitter->EmitConstantFalse();
+				return Emitter.ConstantFalse();
 			}
 			break;
 
@@ -1056,7 +939,7 @@ static FValue* TrySimplifyOperator(FEmitter* Emitter, EOperator& Op, FValue*& A,
 		case BO_Equals:
 			if (A->Equals(B))
 			{
-				return Emitter->EmitConstantTrue();
+				return Emitter.ConstantTrue();
 			}
 			break;
 
@@ -1079,7 +962,7 @@ static FValue* TrySimplifyOperator(FEmitter* Emitter, EOperator& Op, FValue*& A,
 			}
 			else if (A->AreAllNearlyZero())
 			{
-				return Emitter->EmitOperator(UO_Negate, A);
+				return Emitter.Negate(A);
 			}
 			break;
 
@@ -1105,8 +988,8 @@ static FValue* TrySimplifyOperator(FEmitter* Emitter, EOperator& Op, FValue*& A,
 			if (B->AreAllNearlyZero() && C->AreAllNearlyOne())
 			{
 				Op = UO_Saturate;
-				B = nullptr;
-				C = nullptr;
+				B = {};
+				C = {};
 			}
 			break;
 
@@ -1136,17 +1019,17 @@ static FValue* TrySimplifyOperator(FEmitter* Emitter, EOperator& Op, FValue*& A,
 			break;
 	}
 
-	return nullptr;
+	return {};
 }
 
 // Tries to fold (statically evaluate) the operator, assuming that the arguments are all scalar.
 // It returns either the result of the operator or null if it could not be folded.
-static FValue* TryFoldOperatorScalar(FEmitter* Emitter, EOperator Op, FValue* A, FValue* B, FValue* C)
+static FValueRef TryFoldOperatorScalar(FEmitter& Emitter, EOperator Op, FValueRef A, FValueRef B, FValueRef C)
 {
 	const FPrimitiveType* PrimitiveType = A->Type->AsPrimitive();
 
 	// Try to simplify the operator. This could potentially change Op, A, B and C.
-	if (FValue* Simplified = TrySimplifyOperator(Emitter, Op, A, B, C))
+	if (FValueRef Simplified = TrySimplifyOperator(Emitter, Op, A, B, C))
 	{
 		return Simplified;
 	}
@@ -1158,9 +1041,9 @@ static FValue* TryFoldOperatorScalar(FEmitter* Emitter, EOperator Op, FValue* A,
 	}
 
 	// Verify that both lhs and rhs are constants, otherwise we cannot fold the operation.
-	const FConstant* AConstant = As<FConstant>(A);
-	const FConstant* BConstant = As<FConstant>(B);
-	const FConstant* CConstant = As<FConstant>(C);
+	const FConstant* AConstant = As<FConstant>(A.Value);
+	const FConstant* BConstant = As<FConstant>(B.Value);
+	const FConstant* CConstant = As<FConstant>(C.Value);
 	if (!AConstant || (IsBinaryOperator(Op) && !BConstant) || (IsTernaryOperator(Op) && (!BConstant || !CConstant)))
 	{
 		return nullptr;
@@ -1172,18 +1055,18 @@ static FValue* TryFoldOperatorScalar(FEmitter* Emitter, EOperator Op, FValue* A,
 		bool Result;
 		switch (PrimitiveType->ScalarKind)
 		{
-		case ScalarKind_Int:
-			Result = FoldComparisonOperatorScalar<TInteger>(Op, AConstant->Integer, BConstant->Integer);
-			break;
+			case ScalarKind_Int:
+				Result = FoldComparisonOperatorScalar<TInteger>(Op, AConstant->Integer, BConstant->Integer);
+				break;
 
-		case ScalarKind_Float:
-			Result = FoldComparisonOperatorScalar<TFloat>(Op, AConstant->Float, BConstant->Float);
-			break;
+			case ScalarKind_Float:
+				Result = FoldComparisonOperatorScalar<TFloat>(Op, AConstant->Float, BConstant->Float);
+				break;
 
-		default:
-			UE_MIR_UNREACHABLE();
+			default:
+				UE_MIR_UNREACHABLE();
 		}
-		return Emitter->EmitConstantBoolScalar(Result);
+		return Emitter.ConstantBool(Result);
 	}
 	else if (IsBitOperator(Op))
 	{
@@ -1192,13 +1075,13 @@ static FValue* TryFoldOperatorScalar(FEmitter* Emitter, EOperator Op, FValue* A,
 			case ScalarKind_Bool:
 			{
 				bool Result = (bool)FoldScalarBitOperator(Op, AConstant->Boolean, BConstant ? BConstant->Boolean : 0);
-				return Emitter->EmitConstantBoolScalar(Result);
+				return Emitter.ConstantBool(Result);
 			}
 
 			case ScalarKind_Int:
 			{
 				TInteger Result = FoldScalarBitOperator(Op, AConstant->Integer, BConstant ? BConstant->Integer : 0);
-				return Emitter->EmitConstantIntScalar(Result);
+				return Emitter.ConstantInt(Result);
 			}
 			
 			default:
@@ -1212,13 +1095,13 @@ static FValue* TryFoldOperatorScalar(FEmitter* Emitter, EOperator Op, FValue* A,
 			case ScalarKind_Int:
 			{
 				TOptional<TInteger> Result = FoldScalarArithmeticOperator<TInteger>(Emitter, Op, AConstant->Integer, BConstant ? BConstant->Integer : 0, CConstant ? CConstant->Integer : 0);
-				return Result ? Emitter->EmitConstantIntScalar(*Result) : nullptr;
+				return Result ? Emitter.ConstantInt(*Result) : nullptr;
 			}
 
 			case ScalarKind_Float:
 			{
 				TOptional<TFloat> Result = FoldScalarArithmeticOperator<TFloat>(Emitter, Op, AConstant->Float, BConstant ? BConstant->Float : 0, CConstant ? CConstant->Float : 0);
-				return Result ? Emitter->EmitConstantFloatScalar(*Result) : nullptr;
+				return Result ? Emitter.ConstantFloat(*Result) : nullptr;
 			}
 
 			default:
@@ -1236,9 +1119,11 @@ enum EOperatorParameterTypeFilter
 	OPTF_Arithmetic = 1 << 1,				///< Arithmetic primitive type of any dimension (i.e. that supports arithmetic operations)
 	OPTF_NotMatrix = 1 << 2,				///< Any primitive type except matrices
 	OPTF_Vector3 = 1 << 3,					///< 3D vector of any scalar type
+	OPTF_NonNegativeFloatConst = 1 << 5,		///< Validates that if the argument is constant, it is not zero
+	OPTF_NonZeroFloatConst = 1 << 6,		///< Validates that if the argument is constant, it is not zero
 	OPTF_CastToAnyFloat = 1 << 4,			///< Floating point primitive type of any dimension
-	OPTF_CastToFirstArgumentType = 1 << 5,	///< Must be the same type as the first argument
-	OPTF_CastToCommon = 1 << 6,	 			///< An argument with this flag will be automatically cast to the common type between arguments.
+	OPTF_CastToFirstArgumentType = 1 << 7,	///< Must be the same type as the first argument
+	OPTF_CastToCommon = 1 << 8,	 			///< An argument with this flag will be automatically cast to the common type between arguments.
 	OPTF_CastToCommonArithmetic = OPTF_Arithmetic | OPTF_CastToCommon,
 	OPTF_CastToCommonFloat = OPTF_CastToAnyFloat | OPTF_CastToCommon,
 };
@@ -1291,13 +1176,13 @@ static const FOperatorSignature* GetOperatorSignature(EOperator Op)
 		S[UO_Floor] 				= UnaryFloat;
 		S[UO_Frac]					= UnaryFloat;
 		S[UO_Length]				= { { OPTF_Arithmetic | OPTF_NotMatrix | OPTF_CastToAnyFloat }, ORT_FirstArgumentTypeToScalar };
-		S[UO_Logarithm] 			= UnaryFloat;
-		S[UO_Logarithm2] 			= UnaryFloat;
-		S[UO_Logarithm10] 			= UnaryFloat;
+		S[UO_Logarithm] 			= { { OPTF_Arithmetic | OPTF_NonZeroFloatConst | OPTF_NonNegativeFloatConst | OPTF_CastToAnyFloat }, ORT_FirstArgumentType };;
+		S[UO_Logarithm2] 			= { { OPTF_Arithmetic | OPTF_NonZeroFloatConst | OPTF_NonNegativeFloatConst | OPTF_CastToAnyFloat }, ORT_FirstArgumentType };;
+		S[UO_Logarithm10] 			= { { OPTF_Arithmetic | OPTF_NonZeroFloatConst | OPTF_NonNegativeFloatConst | OPTF_CastToAnyFloat }, ORT_FirstArgumentType };;
 		S[UO_Round] 				= UnaryFloat;
 		S[UO_Saturate]				= UnaryFloat;
 		S[UO_Sin] 					= UnaryFloat;
-		S[UO_Sqrt]					= UnaryFloat;
+		S[UO_Sqrt]					= { { OPTF_Arithmetic | OPTF_NonNegativeFloatConst | OPTF_CastToAnyFloat }, ORT_FirstArgumentType };;
 		S[UO_Tan] 					= UnaryFloat;
 		S[UO_Truncate] 				= UnaryFloat;
 		S[BO_GreaterThan] 			= BinaryArithmeticComparison;
@@ -1328,7 +1213,7 @@ static const FOperatorSignature* GetOperatorSignature(EOperator Op)
  
 /// Validates that the types of the arguments are valid for specified operator.
 /// If valid, it returns the type of the result. Otherwise if it is not valid, it returns nullptr.
-static const FPrimitiveType* ValidateOperatorAndGetResultType(FEmitter* Emitter, EOperator Op, FValue*& A, FValue*& B, FValue*& C)
+static const FPrimitiveType* ValidateOperatorAndGetResultType(FEmitter& Emitter, EOperator Op, FValueRef& A, FValueRef& B, FValueRef& C)
 {
 	// Argument A must have always been provided.
 	check(A);
@@ -1343,7 +1228,7 @@ static const FPrimitiveType* ValidateOperatorAndGetResultType(FEmitter* Emitter,
 	check(!MIR::IsTernaryOperator(Op) || C);
 
 	// Verify that the first argument type is primitive.
-	FValue* Arguments[] = { A, B, C, nullptr };
+	FValueRef Arguments[] = { A, B, C, nullptr };
 	const FPrimitiveType* FirstArgumentPrimitiveType = A->Type->AsPrimitive();
 	static const TCHAR* ArgumentsStr[] = { TEXT("first"), TEXT("second"), TEXT("third") };
 	const FOperatorSignature* Signature = GetOperatorSignature(Op);
@@ -1352,16 +1237,14 @@ static const FPrimitiveType* ValidateOperatorAndGetResultType(FEmitter* Emitter,
 	bool bValid = true;
 	for (int32 i = 0; Arguments[i]; ++i)
 	{
-		// Cast this argument type to primitive.
-		const FPrimitiveType* ArgumentPrimitiveType = Arguments[i]->Type->AsPrimitive();
-
-		// All operators have primitive parameter types.
-		if (!ArgumentPrimitiveType)
+		// Check this argument type i primitive.
+		Arguments[i] = Emitter.CheckIsPrimitive(Arguments[i]);
+		if (!Arguments[i].IsValid())
 		{
-			Emitter->Errorf(TEXT("Operator '%s' %s argument does not have primitive type."), LexToString(Op), ArgumentsStr[i]);
 			return nullptr;
 		}
-		
+		const FPrimitiveType* ArgumentPrimitiveType = Arguments[i]->Type->AsPrimitive();
+
 		EOperatorParameterTypeFilter Filter = Signature->ParameterTypes[i];
 		check(Filter != OPTF_Unknown); // No signature specified for this operator.
 
@@ -1371,16 +1254,16 @@ static const FPrimitiveType* ValidateOperatorAndGetResultType(FEmitter* Emitter,
 
 			// Cast this argument to the first argument primitive type
 			ArgumentPrimitiveType = FirstArgumentPrimitiveType;
-			Arguments[i] = Emitter->EmitConstruct(ArgumentPrimitiveType, Arguments[i]);
-			bValid &= !Arguments[i]->IsPoison();
+			Arguments[i] = Emitter.Cast(Arguments[i], ArgumentPrimitiveType);
+			bValid &= Arguments[i];
 		}
 		else if (Filter & OPTF_CastToAnyFloat)
 		{
 			if (!ScalarKindIsAnyFloat(ArgumentPrimitiveType->ScalarKind))
 			{
-				ArgumentPrimitiveType = ArgumentPrimitiveType->ToScalarKind(MIR::ScalarKind_Float);
-				Arguments[i] = Emitter->EmitConstruct(ArgumentPrimitiveType, Arguments[i]);
-				bValid &= !Arguments[i]->IsPoison();
+				ArgumentPrimitiveType = ArgumentPrimitiveType->WithScalarKind(MIR::ScalarKind_Float);
+				Arguments[i] = Emitter.Cast(Arguments[i], ArgumentPrimitiveType);
+				bValid &= Arguments[i];
 			}
 		}
 		
@@ -1388,34 +1271,48 @@ static const FPrimitiveType* ValidateOperatorAndGetResultType(FEmitter* Emitter,
 		{
 			if (ArgumentPrimitiveType->ScalarKind != ScalarKind_Bool)
 			{
-				Emitter->Errorf(TEXT("Operator '%s' %s argument component types are not boolean."), LexToString(Op), ArgumentsStr[i]);
+				Emitter.Errorf(Arguments[i], TEXT("Expected a boolean."));
 				bValid = false;
 			}
 		}
 		
 		if (Filter & OPTF_Arithmetic)
 		{
-			if (!ScalarKindIsArithmetic(ArgumentPrimitiveType->ScalarKind))
+			bValid &= Emitter.CheckIsArithmetic(Arguments[i]).IsValid(); 
+		}
+		
+		if (Filter & OPTF_NotMatrix)
+		{
+			bValid &= Emitter.CheckIsScalarOrVector(Arguments[i]).IsValid(); 
+		}
+
+		if (Filter & OPTF_Vector3)
+		{
+			if (!ArgumentPrimitiveType->IsVector() || ArgumentPrimitiveType->GetNumComponents() != 3)
 			{
-				Emitter->Errorf(TEXT("Operator '%s' %s argument component types expected to be arithmetic (int, float, double, etc)."), LexToString(Op), ArgumentsStr[i]);
+				Emitter.Errorf(Arguments[i], TEXT("Expected a 3D vector."));
 				bValid = false;
 			}
 		}
-		
-		if (Filter & OPTF_NotMatrix)
+
+		if (Filter & OPTF_NonZeroFloatConst)
 		{
-			if (ArgumentPrimitiveType->IsMatrix())
+			check((Filter & OPTF_CastToAnyFloat) || (Filter & OPTF_CastToCommonFloat));
+			FConstant* Constant = Arguments[i]->As<FConstant>();
+			if (Constant && Constant->Float == 0)
 			{
-				Emitter->Errorf(TEXT("Operator '%s' %s argument type expected to be a scalar or vector."), LexToString(Op), ArgumentsStr[i]);
+				Emitter.Errorf(Arguments[i], TEXT("Expected non-zero value."));
 				bValid = false;
 			}
 		}
 
-		if (Filter & OPTF_Vector3)
+		if (Filter & OPTF_NonNegativeFloatConst)
 		{
-			if (!ArgumentPrimitiveType->IsVector() || ArgumentPrimitiveType->GetNumComponents() != 3)
+			check((Filter & OPTF_CastToAnyFloat) || (Filter & OPTF_CastToCommonFloat));
+			FConstant* Constant = Arguments[i]->As<FConstant>();
+			if (Constant && Constant->Float < 0)
 			{
-				Emitter->Errorf(TEXT("Operator '%s' %s argument type expected to be a 3D vector."), LexToString(Op), ArgumentsStr[i]);
+				Emitter.Errorf(Arguments[i], TEXT("Expected non-negative value."));
 				bValid = false;
 			}
 		}
@@ -1423,7 +1320,7 @@ static const FPrimitiveType* ValidateOperatorAndGetResultType(FEmitter* Emitter,
 		// Update the common type.
 		if (i >= 1)
 		{
-			ArgumentsCommonType = Emitter->TryGetCommonType(ArgumentsCommonType, ArgumentPrimitiveType);
+			ArgumentsCommonType = Emitter.TryGetCommonType(ArgumentsCommonType, ArgumentPrimitiveType);
 		}
 	}
 	
@@ -1438,10 +1335,10 @@ static const FPrimitiveType* ValidateOperatorAndGetResultType(FEmitter* Emitter,
 		int32 MaxNumComponents = FMath::Max3(FirstArgumentPrimitiveType->GetNumComponents(), BPrimitiveType->GetNumComponents(), CPrimitiveType->GetNumComponents());
 
 		// Cast the first argument (the boolean condition) to a bool vector of the maximum number of components.
-		Arguments[0] = Emitter->EmitConstruct(FPrimitiveType::Get(ScalarKind_Bool, MaxNumComponents, 1), Arguments[0]);
+		Arguments[0] = Emitter.Cast(Arguments[0], FPrimitiveType::Get(ScalarKind_Bool, MaxNumComponents, 1));
 
 		// Compute the common type between the second and third argument types with a number of components equal to the max of all three.
-		const FType* CommonTypeBetweenSecondAndThirdArguments = Emitter->TryGetCommonType(
+		const FType* CommonTypeBetweenSecondAndThirdArguments = Emitter.TryGetCommonType(
 			FPrimitiveType::Get(BPrimitiveType->ScalarKind, MaxNumComponents, 1),
 			FPrimitiveType::Get(CPrimitiveType->ScalarKind, MaxNumComponents, 1));
 		
@@ -1449,8 +1346,8 @@ static const FPrimitiveType* ValidateOperatorAndGetResultType(FEmitter* Emitter,
 		check(CommonTypeBetweenSecondAndThirdArguments);
 
 		// Cast second and third arguments to their common type.
-		Arguments[1] = Emitter->EmitConstruct(CommonTypeBetweenSecondAndThirdArguments, Arguments[1]);
-		Arguments[2] = Emitter->EmitConstruct(CommonTypeBetweenSecondAndThirdArguments, Arguments[2]);
+		Arguments[1] = Emitter.Cast(Arguments[1], CommonTypeBetweenSecondAndThirdArguments);
+		Arguments[2] = Emitter.Cast(Arguments[2], CommonTypeBetweenSecondAndThirdArguments);
 
 		// Update the valid flag to the results of the conditions above.
 		bValid &= !Arguments[0]->IsPoison() && !Arguments[1]->IsPoison() && !Arguments[2]->IsPoison();
@@ -1464,7 +1361,7 @@ static const FPrimitiveType* ValidateOperatorAndGetResultType(FEmitter* Emitter,
 			if (Filter & OPTF_CastToCommon)
 			{
 				check(ArgumentsCommonType && ArgumentsCommonType->AsPrimitive());
-				Arguments[i] = Emitter->EmitConstruct(ArgumentsCommonType, Arguments[i]);
+				Arguments[i] = Emitter.Cast(Arguments[i], ArgumentsCommonType);
 				bValid &= !Arguments[i]->IsPoison();
 			}
 		}
@@ -1477,7 +1374,7 @@ static const FPrimitiveType* ValidateOperatorAndGetResultType(FEmitter* Emitter,
 
 	if (!bValid)
 	{
-		return nullptr;
+		return {};
 	}
 	
 	// Update the first argument type, as it might have changed.
@@ -1516,7 +1413,7 @@ static bool IsComponentwiseOperator(EOperator Op)
 /// If a value is returned, it will be a dimensional with some component folded to a constant. If some argument
 /// isn't a dimensional, or all arguments components are non-constant, the folding will not be carried out.
 /// If no folding is carried out, this function simply returns nullptr.
-static FValue* TryFoldComponentwiseOperator(FEmitter* Emitter, EOperator Op, FValue* A, FValue* B, FValue* C, const FPrimitiveType* ResultType)
+static FValue* TryFoldComponentwiseOperator(FEmitter& Emitter, EOperator Op, FValue* A, FValue* B, FValue* C, const FPrimitiveType* ResultType)
 {
 	// Check that at least one component of the resulting dimensional value would folded.
 	// If all components of resulting dimensional value are not folded, then instead of emitting
@@ -1534,9 +1431,9 @@ static FValue* TryFoldComponentwiseOperator(FEmitter* Emitter, EOperator Op, FVa
 	for (int32 i = 0; i < ResultType->GetNumComponents(); ++i)
 	{
 		// Extract the arguments individual components
-		FValue* AComponent = Emitter->EmitSubscript(A, i);
-		FValue* BComponent = B ? Emitter->EmitSubscript(B, i) : nullptr;
-		FValue* CComponent = C ? Emitter->EmitSubscript(C, i) : nullptr;
+		FValue* AComponent = Emitter.Subscript(A, i);
+		FValue* BComponent = B ? Emitter.Subscript(B, i) : nullptr;
+		FValue* CComponent = C ? Emitter.Subscript(C, i) : nullptr;
 
 		// Try folding the operation, it may return null
 		FValue* ResultComponent = TryFoldOperatorScalar(Emitter, Op, AComponent, BComponent, CComponent);
@@ -1600,9 +1497,9 @@ static FValue* TryFoldComponentwiseOperator(FEmitter* Emitter, EOperator Op, FVa
 			{
 				FOperator Proto = MakePrototype<FOperator>(ComponentType);
 				Proto.Op = Op;
-				Proto.AArg = Emitter->EmitSubscript(A, i);
-				Proto.BArg = B ? Emitter->EmitSubscript(B, i) : nullptr;
-				Proto.CArg = C ? Emitter->EmitSubscript(C, i) : nullptr;
+				Proto.AArg = Emitter.Subscript(A, i);
+				Proto.BArg = B ? Emitter.Subscript(B, i) : nullptr;
+				Proto.CArg = C ? Emitter.Subscript(C, i) : nullptr;
 				ResultComponents[i] = EmitPrototype(Emitter, Proto);
 			}
 		}
@@ -1610,7 +1507,7 @@ static FValue* TryFoldComponentwiseOperator(FEmitter* Emitter, EOperator Op, FVa
 		return EmitNew(Emitter, Result);
 	}
 
-	return nullptr;
+	return {};
 }
 
 /// If V is a dimensional and all its components are constants, it unpacks the components into OutComponents and returns true.
@@ -1651,7 +1548,7 @@ static MIR::TFloat ConstantDotFloat(TArrayView<FConstant*> AComponents, TArrayVi
 
 /// Tries to fold the operator, that is to evaluate its result now at translation time if its arguments are constant.
 /// If the operator could not be folded in any way, it returns nullptr.
-static FValue* TryFoldOperator(FEmitter* Emitter, EOperator Op, FValue* A, FValue* B, FValue* C, const FPrimitiveType* ResultType)
+static FValue* TryFoldOperator(FEmitter& Emitter, EOperator Op, FValue* A, FValue* B, FValue* C, const FPrimitiveType* ResultType)
 {
 	FConstant* AComponents[4];
 	int32      ANumComponents;
@@ -1670,7 +1567,7 @@ static FValue* TryFoldOperator(FEmitter* Emitter, EOperator Op, FValue* A, FValu
 			if (ResultType->ScalarKind == ScalarKind_Float)
 			{
 				float Result = FMath::Sqrt(ConstantDotFloat(AComponents, AComponents, ANumComponents));
-				return Emitter->EmitConstantFloatScalar(Result);
+				return Emitter.ConstantFloat(Result);
 			}
 			else
 			{
@@ -1687,7 +1584,7 @@ static FValue* TryFoldOperator(FEmitter* Emitter, EOperator Op, FValue* A, FValu
 				if (ResultType->ScalarKind == ScalarKind_Float)
 				{
 					float Result = ConstantDotFloat(AComponents, BComponents, ANumComponents);
-					return Emitter->EmitConstantFloatScalar(Result);
+					return Emitter.ConstantFloat(Result);
 				}
 				else
 				{
@@ -1702,7 +1599,7 @@ static FValue* TryFoldOperator(FEmitter* Emitter, EOperator Op, FValue* A, FValu
 					FVector3f AVector{ AComponents[0]->Float, AComponents[1]->Float, AComponents[2]->Float };
 					FVector3f BVector{ BComponents[0]->Float, BComponents[1]->Float, BComponents[2]->Float };
 					FVector3f Result = AVector.Cross(BVector);
-					return Emitter->EmitConstantFloat3(Result);
+					return Emitter.ConstantFloat3(Result);
 				}
 				else
 				{
@@ -1722,85 +1619,201 @@ static FValue* TryFoldOperator(FEmitter* Emitter, EOperator Op, FValue* A, FValu
 	return nullptr;
 }
 
-FValue* FEmitter::EmitOperator(EOperator Op, FValue* A, FValue* B, FValue* C)
+FValueRef FEmitter::Operator(EOperator Op, FValueRef A, FValueRef B, FValueRef C)
 {
-	if (IsAnyNonNullAPoison(A, B, C))
+	if (!A.IsValid() || (B && !B.IsValid()) || (C && !C.IsValid()))
 	{
-		return FPoison::Get();
+		return Poison();
+	}
+
+	// Validate the operation and retrieve the result type.
+	const FPrimitiveType* ResultType = ValidateOperatorAndGetResultType(*this, Op, A, B, C);
+	if (!ResultType)
+	{
+		return Poison();
+	}
+
+	// Try to apply some operator identity to simplify the operator.
+	if (FValueRef Simplified = TrySimplifyOperator(*this, Op, A, B, C))
+	{
+		return Simplified;
+	}
+
+	// Try folding the operator first.
+	if (FValue* FoldedValue = TryFoldOperator(*this, Op, A, B, C, ResultType))
+	{
+		return FoldedValue;
+	}
+
+	// Otherwise, we must emit a new instruction that executes the operator.
+	FOperator Proto = MakePrototype<FOperator>(ResultType);
+	Proto.Op = Op;
+	Proto.AArg = A;
+	Proto.BArg = B;
+	Proto.CArg = C;
+
+	return EmitPrototype(*this, Proto);
+}
+
+FValueRef FEmitter::Branch(FValueRef Condition, FValueRef True, FValueRef False)
+{
+	if (IsAnyNotValid(Condition, True, False))
+	{
+		return Poison();
+	}
+
+	// Condition must be of type bool
+	Condition = Cast(Condition, FPrimitiveType::GetBool());
+	if (!Condition)
+	{
+		return Poison();
+	}	
+
+	// If the condition is a scalar constant, then simply evaluate the result now.
+	if (const FConstant* ConstCondition = As<FConstant>(Condition))
+	{
+		return ConstCondition->Boolean ? True : False;
+	}
+
+	// If the condition is not static, make both true and false arguments have the same type,
+	// by casting false argument into the true's type.
+	const FType* CommonType = GetCommonType(True->Type, False->Type);
+	if (!CommonType)
+	{
+		return Poison();
+	}
+
+	True = Cast(True, CommonType);
+	False = Cast(False, CommonType);
+	if (!True || !False)
+	{
+		return Poison();
+	}
+
+	// Create the branch instruction.
+	FBranch Proto = MakePrototype<FBranch>(CommonType);
+	Proto.ConditionArg = Condition;
+	Proto.TrueArg = True;
+	Proto.FalseArg = False;
+
+	return EmitPrototype(*this, Proto);
+}
+
+FValueRef FEmitter::Subscript(FValueRef Value, int32 Index)
+{
+	if (!Value.IsValid())
+	{
+		return Value;
+	}
+
+	const FPrimitiveType* PrimitiveType = Value->Type->AsPrimitive();
+	if (!PrimitiveType)
+	{
+		Errorf(Value, TEXT("Value of type '%s' cannot be subscripted."), Value->Type->GetSpelling().GetData());
+		return Value.ToPoison();
 	}
 
-	// Validate the operation and retrieve the result type.
-	const FPrimitiveType* ResultType = ValidateOperatorAndGetResultType(this, Op, A, B, C);
-	if (!ResultType)
+	// Getting first component and Value is already a scalar, just return itself.
+	if (Index == 0 && Value->Type->AsScalar())
 	{
-		return FPoison::Get();
+		return Value;
 	}
 
-	// Try to apply some operator identity to simplify the operator.
-	if (FValue* Simplified = TrySimplifyOperator(this, Op, A, B, C))
+	if (Index >= PrimitiveType->GetNumComponents())
 	{
-		return Simplified;
+		Errorf(Value, TEXT("Value of type '%s' has fewer dimensions than subscript index `%d`."), Value->Type->GetSpelling().GetData(), Index);
+		return Value.ToPoison();
 	}
 
-	// Try folding the operator first.
-	if (FValue* FoldedValue = TryFoldOperator(this, Op, A, B, C, ResultType))
+	// Getting first component and Value is already a scalar, just return itself.
+	if (FDimensional* DimensionalValue = As<FDimensional>(Value))
 	{
-		return FoldedValue;
+		check(Index < DimensionalValue->GetComponents().Num());
+		return Value.To(DimensionalValue->GetComponents()[Index]);
+	}
+	
+	// Avoid subscripting a subscript (e.g. no value.xy.x)
+	if (FSubscript* Subscript = As<FSubscript>(Value))
+	{
+		Value = Value.To(Subscript->Arg);
 	}
 
-	// Otherwise, we must emit a new instruction that executes the operator.
-	FOperator Proto = MakePrototype<FOperator>(ResultType);
-	Proto.Op = Op;
-	Proto.AArg = A;
-	Proto.BArg = B;
-	Proto.CArg = C;
+	// We can't resolve it at compile time: emit subscript value.
+	FSubscript Prototype = MakePrototype<FSubscript>(PrimitiveType->ToScalar());
+	Prototype.Arg = Value;
+	Prototype.Index = Index;
 
-	return EmitPrototype(this, Proto);
+	return Value.To(EmitPrototype(*this, Prototype));
 }
 
-FValue* FEmitter::EmitBranch(FValue* Condition, FValue* True, FValue* False)
+FValueRef FEmitter::Swizzle(FValueRef Value, FSwizzleMask Mask)
 {
-	// Condition must be of type bool
-	Condition = EmitConstruct(FPrimitiveType::GetBool1(), Condition);
-	if (IsAnyNonNullAPoison(Condition))
+	if (!Value.IsValid())
 	{
-		return FPoison::Get();
+		return Value;
 	}
 
-	// If the condition is a scalar constant, then simply evaluate the result now.
-	if (const FConstant* ConstCondition = As<FConstant>(Condition))
+	// At least one component must have been specified.
+	check(Mask.NumComponents > 0);
+
+	// We can only swizzle on non-matrix primitive types.
+	const FPrimitiveType* PrimitiveType = Value->Type->AsVector();
+	if (!PrimitiveType || PrimitiveType->IsMatrix())
 	{
-		return ConstCondition->Boolean ? True : False;
+		Errorf(Value, TEXT("Cannot swizzle a '%s' value."), Value->Type->GetSpelling().GetData());
+		return Value.ToPoison();
 	}
 
-	// Condition must be of type bool
-	check(Condition->Type->IsBoolScalar());
-
-	// If the condition is not static, make both true and false arguments have the same type,
-	// by casting false argument into the true's type.
-	const FType* CommonType = GetCommonType(True->Type, False->Type);
-	if (!CommonType)
+	// Make sure each component in the mask fits the number of components in Value.
+	for (EVectorComponent Component : Mask)
 	{
-		return FPoison::Get();
+		if ((int32)Component >= PrimitiveType->NumRows)
+		{
+			Errorf(Value, TEXT("Value of type '%s' has no component '%s'."), PrimitiveType->Spelling.GetData(), VectorComponentToString(Component));
+			return Value.ToPoison();
+		}
 	}
 
-	True = EmitConstruct(True->Type, True);
-	False = EmitConstruct(True->Type, False);
-	if (IsAnyNonNullAPoison(True, False))
+	// If the requested number of components is the same as Value and the order in which the components
+	// are specified in the mask is sequential (e.g. x, y, z) then this is a no op, simply return Value as is.
+	if (Mask.NumComponents == PrimitiveType->GetNumComponents())
 	{
-		return FPoison::Get();
+		bool InOrder = true;
+		for (int32 i = 0; i < Mask.NumComponents; ++i)
+		{
+			if (Mask.Components[i] != (EVectorComponent)i)
+			{
+				InOrder = false;
+				break;
+			}
+		}
+
+		if (InOrder)
+		{
+			return Value;
+		}
+	}
+	
+	// If only one component is requested, we can use Subscript() to return the single component.
+	if (Mask.NumComponents == 1)
+	{
+		return Value.To(Subscript(Value, (int32)Mask.Components[0]));
 	}
 
-	// Create the branch instruction.
-	FBranch Proto = MakePrototype<FBranch>(True->Type);
-	Proto.ConditionArg = Condition;
-	Proto.TrueArg = True;
-	Proto.FalseArg = False;
+	// Make the result vector type.
+	const FPrimitiveType* ResultType = FPrimitiveType::GetVector(PrimitiveType->ScalarKind, Mask.NumComponents);
+	FDimensional* Result = NewDimensionalValue(*this, ResultType);
+	TArrayView<FValue*> ResultComponents = Result->GetMutableComponents();
+
+	for (int32 i = 0; i < Mask.NumComponents; ++i)
+	{
+		ResultComponents[i] = Subscript(Value, (int32)Mask.Components[i]);
+	}
 
-	return EmitPrototype(this, Proto);
+	return Value.To(EmitNew(*this, Result));
 }
 
-static FValue* CastConstant(FEmitter* Emitter, FConstant* Constant, EScalarKind ConstantScalarKind, EScalarKind TargetKind)
+static FValue* CastConstant(FEmitter& Emitter, FConstant* Constant, EScalarKind ConstantScalarKind, EScalarKind TargetKind)
 {
 	if (ConstantScalarKind == TargetKind)
 	{
@@ -1814,9 +1827,9 @@ static FValue* CastConstant(FEmitter* Emitter, FConstant* Constant, EScalarKind
 		{
 			switch (TargetKind)
 			{
-				case ScalarKind_Bool: return nullptr;
-				case ScalarKind_Int: return Emitter->EmitConstantIntScalar(Constant->Integer);
-				case ScalarKind_Float: return Emitter->EmitConstantFloatScalar((TFloat)Constant->Integer);
+				case ScalarKind_Bool: return {};
+				case ScalarKind_Int: return Emitter.ConstantInt(Constant->Integer);
+				case ScalarKind_Float: return Emitter.ConstantFloat((TFloat)Constant->Integer);
 				default: UE_MIR_UNREACHABLE();
 			}
 		}
@@ -1825,8 +1838,8 @@ static FValue* CastConstant(FEmitter* Emitter, FConstant* Constant, EScalarKind
 		{
 			switch (TargetKind)
 			{
-				case ScalarKind_Bool: return nullptr;
-				case ScalarKind_Int: return Emitter->EmitConstantIntScalar((int32)Constant->Float);
+				case ScalarKind_Bool: return {};
+				case ScalarKind_Int: return Emitter.ConstantInt((int32)Constant->Float);
 				default: UE_MIR_UNREACHABLE();
 			}
 		}
@@ -1837,52 +1850,52 @@ static FValue* CastConstant(FEmitter* Emitter, FConstant* Constant, EScalarKind
 	UE_MIR_UNREACHABLE();
 }
 
-static FValue* ConstructPrimitiveValue(FEmitter* Emitter, const FPrimitiveType* TargetPrimitiveType, FValue* Initializer)
+static FValue* CastValueToPrimitiveType(FEmitter& Emitter, FValueRef Value, const FPrimitiveType* TargetPrimitiveType)
 {
-	const FPrimitiveType* InitializerPrimitiveType = Initializer->Type->AsPrimitive();
-	if (!InitializerPrimitiveType)
+	const FPrimitiveType* ValuePrimitiveType = Value->Type->AsPrimitive();
+	if (!ValuePrimitiveType)
 	{
-		Emitter->Errorf(TEXT("Cannot construct a '%s' from non primitive type '%s'."), TargetPrimitiveType->Spelling.GetData(), Initializer->Type->GetSpelling().GetData());
-		return MIR::FPoison::Get();
+		Emitter.Errorf(Value, TEXT("Cannot construct a '%s' from non primitive type '%s'."), TargetPrimitiveType->Spelling.GetData(), Value->Type->GetSpelling().GetData());
+		return FPoison::Get();
 	}
 
 	// Construct a scalar from another scalar.
 	if (TargetPrimitiveType->IsScalar())
 	{
 		// 
-		Initializer = Emitter->EmitSubscript(Initializer, 0);
-		InitializerPrimitiveType = Initializer->Type->AsPrimitive();
+		Value = Emitter.Subscript(Value, 0);
+		ValuePrimitiveType = Value->Type->AsPrimitive();
 		
 		//
-		if (InitializerPrimitiveType == TargetPrimitiveType)
+		if (ValuePrimitiveType == TargetPrimitiveType)
 		{
-			return Initializer;
+			return Value;
 		}
 
 		// Construct the scalar from a constant.
-		if (FConstant* ConstantInitializer = As<FConstant>(Initializer))
+		if (FConstant* ConstantInitializer = As<FConstant>(Value))
 		{
-			return CastConstant(Emitter, ConstantInitializer, InitializerPrimitiveType->ScalarKind, TargetPrimitiveType->ScalarKind);
+			return CastConstant(Emitter, ConstantInitializer, ValuePrimitiveType->ScalarKind, TargetPrimitiveType->ScalarKind);
 		}
 		else
 		{
 			// Emit the cast to the target type of the subscript value
 			FCast Prototype = MakePrototype<FCast>(TargetPrimitiveType);
-			Prototype.Arg = Initializer;
+			Prototype.Arg = Value;
 			return EmitPrototype(Emitter, Prototype);
 		}
 	}
 
 	// Construct a vector or matrix from a scalar. E.g. 3.14f -> float3(3.14f, 3.14f, 3.14f)
 	// Note: we know target isn't scalar as it's been handled above.
-	if (InitializerPrimitiveType->IsScalar())
+	if (ValuePrimitiveType->IsScalar())
 	{
 		// Create the result dimensional value.
 		FDimensional* Result = NewDimensionalValue(Emitter, TargetPrimitiveType);
 
 		// Create a dimensional and initialize each of its components to the conversion
 		// of initializer value to the single component type.
-		FValue* Component = Emitter->EmitConstruct(TargetPrimitiveType->ToScalar(), Initializer);
+		FValue* Component = Emitter.Cast(Value, TargetPrimitiveType->ToScalar());
 
 		// Get the mutable array of components.
 		TArrayView<FValue*> ResultComponents = Result->GetMutableComponents();
@@ -1899,12 +1912,12 @@ static FValue* ConstructPrimitiveValue(FEmitter* Emitter, const FPrimitiveType*
 	// Construct a vector from another vector. If constructed vector is larger, initialize
 	// remaining components to zero. If it's smaller, truncate initializer vector and only use
 	// the necessary components.
-	if (TargetPrimitiveType->IsVector() && InitializerPrimitiveType->IsVector())
+	if (TargetPrimitiveType->IsVector() && ValuePrimitiveType->IsVector())
 	{
 		// #todo-massimo.tristano Use swizzle when scalartypes are the same, and target num components is less than initializer's.
 
 		int32 TargetNumComponents = TargetPrimitiveType->GetNumComponents();
-		int32 InitializerNumComponents = InitializerPrimitiveType->GetNumComponents();
+		int32 InitializerNumComponents = ValuePrimitiveType->GetNumComponents();
 
 		// Create the result dimensional value.
 		FDimensional* Result = NewDimensionalValue(Emitter, TargetPrimitiveType);
@@ -1920,26 +1933,26 @@ static FValue* ConstructPrimitiveValue(FEmitter* Emitter, const FPrimitiveType*
 		const int32 MinNumComponents = FMath::Min(TargetNumComponents, InitializerNumComponents);
 		for (; Index < MinNumComponents; ++Index)
 		{
-			ResultComponents[Index] = Emitter->EmitConstruct(ResultComponentType, Emitter->EmitSubscript(Initializer, Index));
+			ResultComponents[Index] = Emitter.Cast(Emitter.Subscript(Value, Index), ResultComponentType);
 		}
 
 		// Initialize remaining result dimensional components to zero.
 		for (; Index < TargetNumComponents; ++Index)
 		{
-			ResultComponents[Index] = Emitter->EmitConstantScalarZero(ResultComponentType->ScalarKind);
+			ResultComponents[Index] = Emitter.ConstantZero(ResultComponentType->ScalarKind);
 		}
 
 		return EmitNew(Emitter, Result);
 	}
 	
 	// The two primitive types are identical matrices that differ only by their scalar type.
-	if (TargetPrimitiveType->NumRows 	 == InitializerPrimitiveType->NumRows &&
-		TargetPrimitiveType->NumColumns == InitializerPrimitiveType->NumColumns)
+	if (TargetPrimitiveType->NumRows 	 == ValuePrimitiveType->NumRows &&
+		TargetPrimitiveType->NumColumns == ValuePrimitiveType->NumColumns)
 	{
 		check(TargetPrimitiveType->IsMatrix());
 
 		//  If initializer is a dimensional, we can subscript components of interest.
-		if (FDimensional* DimensionalInitializer = As<FDimensional>(Initializer))
+		if (FDimensional* DimensionalInitializer = As<FDimensional>(Value))
 		{
 			// Create the result dimensional value.
 			FDimensional* Result = NewDimensionalValue(Emitter, TargetPrimitiveType);
@@ -1951,7 +1964,7 @@ static FValue* ConstructPrimitiveValue(FEmitter* Emitter, const FPrimitiveType*
 			// Convert components from the initializer vector.
 			for (int32 Index = 0, Num = TargetPrimitiveType->GetNumComponents(); Index < Num; ++Index)
 			{
-				ResultComponents[Index] = Emitter->EmitConstruct(ResultComponentType, DimensionalInitializer->GetComponents()[Index]);
+				ResultComponents[Index] = Emitter.Cast(DimensionalInitializer->GetComponents()[Index], ResultComponentType);
 			}
 
 			return EmitNew(Emitter, Result);
@@ -1960,52 +1973,84 @@ static FValue* ConstructPrimitiveValue(FEmitter* Emitter, const FPrimitiveType*
 		{
 			// Initializer is an unknown value, construct target value casting initializer.
 			FCast Prototype = MakePrototype<FCast>(TargetPrimitiveType);
-			Prototype.Arg = Initializer;
+			Prototype.Arg = Value;
 			return EmitPrototype(Emitter, Prototype);
 		}
 	}
 
 	// Initializer value cannot be used to construct this primitive type.
-	return MIR::FPoison::Get();
+	return Emitter.Poison();
 }
 
-FValue* FEmitter::EmitConstruct(const FType* Type, FValue* Initializer)
+FValueRef FEmitter::Cast(FValueRef Value, const FType* TargetType)
 {
-	if (Initializer->IsPoison())
+	if (!Value.IsValid())
 	{
-		return FPoison::Get();
+		return Value;
 	}
 
 	// If target type matches initializer's, simply return the same value.
-	const FType* InitializerType = Initializer->Type;
-	if (InitializerType == Type)
+	const FType* InitializerType = Value->Type;
+	if (InitializerType == TargetType)
 	{
-		return Initializer;
+		return Value;
 	}
 	
-	FValue* Result{};
-	if (const FPrimitiveType* PrimitiveType = Type->AsPrimitive())
+	FValue* Result;
+	if (const FPrimitiveType* PrimitiveType = TargetType->AsPrimitive())
 	{
-		Result = ConstructPrimitiveValue(this, PrimitiveType, Initializer);
+		Result = CastValueToPrimitiveType(*this, Value, PrimitiveType);
 	}
-
-	// No other legal conversions applicable. Report error if we haven't converted the value.
-	if (!Result)
+	else // No other legal conversions applicable. Report error if we haven't converted the value.
 	{
-		Errorf(TEXT("Cannot construct a '%s' from a '%s'."), Type->GetSpelling().GetData(), Initializer->Type->GetSpelling().GetData());
-		return FPoison::Get();
+		Errorf(Value, TEXT("Cannot construct a '%s' from a '%s'."), TargetType->GetSpelling().GetData(), Value->Type->GetSpelling().GetData());
+		return Poison();
 	}
 	
 	return Result;
 }
 
-FValue* FEmitter::EmitTextureGather(FValue* Texture, FValue* TexCoord, ETextureReadMode GatherMode, ESamplerSourceMode SamplerSourceMode, EMaterialSamplerType SamplerType)
+FValueRef FEmitter::CastToScalar(FValueRef Value)
+{
+	Value = CheckIsPrimitive(Value);
+	return Value.IsValid()
+		? Cast(Value, FPrimitiveType::GetScalar(Value->Type->AsPrimitive()->ScalarKind))
+		: Value;
+}
+
+FValueRef FEmitter::CastToBool(FValueRef Value, int NumRows)
+{
+	return Cast(Value, FPrimitiveType::GetVector(ScalarKind_Bool, NumRows));
+}
+
+FValueRef FEmitter::CastToInt(FValueRef Value, int NumRows)
+{
+	return Cast(Value, FPrimitiveType::GetVector(ScalarKind_Int, NumRows));
+}
+
+FValueRef FEmitter::CastToFloat(FValueRef Value, int NumRows)
+{
+	return Cast(Value, FPrimitiveType::GetVector(ScalarKind_Float, NumRows));
+}
+
+FValueRef FEmitter::StageSwitch(const FType* Type, TConstArrayView<FValueRef> ValuePerStage)
+{
+	check(ValuePerStage.Num() <= NumStages);
+	FStageSwitch Prototype = MakePrototype<FStageSwitch>(Type);
+	for (int i = 0; i < ValuePerStage.Num(); ++i)
+	{
+		Prototype.Args[i] = ValuePerStage[i];
+	}
+	return EmitPrototype(*this, Prototype);
+}
+
+FValueRef FEmitter::TextureGather(FValueRef Texture, FValueRef TexCoord, ETextureReadMode GatherMode, ESamplerSourceMode SamplerSourceMode, EMaterialSamplerType SamplerType)
 {
 	check(GatherMode >= ETextureReadMode::GatherRed && GatherMode <= ETextureReadMode::GatherAlpha);
 
-	if (IsAnyAPoison(Texture, TexCoord))
+	if (IsAnyNotValid(Texture, TexCoord))
 	{
-		return FPoison::Get();
+		return Poison();
 	}
 
 	FTextureRead Prototype = MakePrototype<FTextureRead>(FPrimitiveType::GetFloat4());
@@ -2015,14 +2060,14 @@ FValue* FEmitter::EmitTextureGather(FValue* Texture, FValue* TexCoord, ETextureR
 	Prototype.SamplerSourceMode = SamplerSourceMode;
 	Prototype.SamplerType = SamplerType;
 
-	return EmitPrototype(this, Prototype);
+	return EmitPrototype(*this, Prototype);
 }
 
-FValue* FEmitter::EmitTextureSample(FValue* Texture, FValue* TexCoord, bool bAutomaticViewMipBias, ESamplerSourceMode SamplerSourceMode, EMaterialSamplerType SamplerType)
+FValueRef FEmitter::TextureSample(FValueRef Texture, FValueRef TexCoord, bool bAutomaticViewMipBias, ESamplerSourceMode SamplerSourceMode, EMaterialSamplerType SamplerType)
 {
-	if (IsAnyAPoison(Texture, TexCoord))
+	if (IsAnyNotValid(Texture, TexCoord))
 	{
-		return FPoison::Get();
+		return Poison();
 	}
 
 	FTextureRead PrototypeHw = MakePrototype<FTextureRead>(FPrimitiveType::GetFloat4());
@@ -2034,31 +2079,31 @@ FValue* FEmitter::EmitTextureSample(FValue* Texture, FValue* TexCoord, bool bAut
 
 	FTextureRead PrototypeAn = PrototypeHw;
 	PrototypeAn.Mode = ETextureReadMode::Derivatives;
-	PrototypeAn.TexCoordDdx = EmitAnalyticalPartialDerivative(TexCoord, EDerivativeAxis::X);
-	PrototypeAn.TexCoordDdy = EmitAnalyticalPartialDerivative(TexCoord, EDerivativeAxis::Y);	
+	PrototypeAn.TexCoordDdx = AnalyticalPartialDerivative(TexCoord, EDerivativeAxis::X);
+	PrototypeAn.TexCoordDdy = AnalyticalPartialDerivative(TexCoord, EDerivativeAxis::Y);	
 
 	if (bAutomaticViewMipBias)
 	{
-		FValue* ViewMaterialTextureMipBias = EmitExternalInput(EExternalInput::ViewMaterialTextureMipBias);
+		FValueRef ViewMaterialTextureMipBias = ExternalInput(EExternalInput::ViewMaterialTextureMipBias);
 		PrototypeHw.Mode = ETextureReadMode::MipBias;
 		PrototypeHw.MipValue = ViewMaterialTextureMipBias;
 
-		FValue* Exp2ViewMaterialTextureMipBias = EmitOperator(UO_Exponential2, ViewMaterialTextureMipBias);
-		PrototypeAn.TexCoordDdx = EmitOperator(BO_Multiply, PrototypeAn.TexCoordDdx, Exp2ViewMaterialTextureMipBias);
-		PrototypeAn.TexCoordDdy = EmitOperator(BO_Multiply, PrototypeAn.TexCoordDdy, Exp2ViewMaterialTextureMipBias);
+		FValueRef Exp2ViewMaterialTextureMipBias = Operator(UO_Exponential2, ViewMaterialTextureMipBias);
+		PrototypeAn.TexCoordDdx = Operator(BO_Multiply, PrototypeAn.TexCoordDdx, Exp2ViewMaterialTextureMipBias);
+		PrototypeAn.TexCoordDdy = Operator(BO_Multiply, PrototypeAn.TexCoordDdy, Exp2ViewMaterialTextureMipBias);
 	}
 
 	FStageSwitch StageSwitch = MakePrototype<FStageSwitch>(PrototypeHw.Type);
-	StageSwitch.SetArgs(EmitPrototype(this, PrototypeHw), EmitPrototype(this, PrototypeAn));
+	StageSwitch.SetArgs(EmitPrototype(*this, PrototypeHw), EmitPrototype(*this, PrototypeAn));
 
-	return EmitPrototype(this, StageSwitch);
+	return EmitPrototype(*this, StageSwitch);
 }
 
-FValue* FEmitter::EmitTextureSampleLevel(FValue* Texture, FValue* TexCoord, FValue* MipLevel, bool bAutomaticViewMipBias, ESamplerSourceMode SamplerSourceMode, EMaterialSamplerType SamplerType)
+FValueRef FEmitter::TextureSampleLevel(FValueRef Texture, FValueRef TexCoord, FValueRef MipLevel, bool bAutomaticViewMipBias, ESamplerSourceMode SamplerSourceMode, EMaterialSamplerType SamplerType)
 {
-	if (IsAnyAPoison(Texture, TexCoord, MipLevel))
+	if (IsAnyNotValid(Texture, TexCoord, MipLevel))
 	{
-		return FPoison::Get();
+		return Poison();
 	}
 
 	FTextureRead Prototype = MakePrototype<FTextureRead>(FPrimitiveType::GetFloat4());
@@ -2071,22 +2116,22 @@ FValue* FEmitter::EmitTextureSampleLevel(FValue* Texture, FValue* TexCoord, FVal
 
 	if (bAutomaticViewMipBias)
 	{
-		Prototype.MipValue = EmitOperator(BO_Add, MipLevel, EmitExternalInput(EExternalInput::ViewMaterialTextureMipBias));
+		Prototype.MipValue = Operator(BO_Add, MipLevel, ExternalInput(EExternalInput::ViewMaterialTextureMipBias));
 	}
 
-	return EmitPrototype(this, Prototype);
+	return EmitPrototype(*this, Prototype);
 }
 
-FValue* FEmitter::EmitTextureSampleBias(FValue* Texture, FValue* TexCoord, FValue* MipBias, bool bAutomaticViewMipBias, ESamplerSourceMode SamplerSourceMode, EMaterialSamplerType SamplerType)
+FValueRef FEmitter::TextureSampleBias(FValueRef Texture, FValueRef TexCoord, FValueRef MipBias, bool bAutomaticViewMipBias, ESamplerSourceMode SamplerSourceMode, EMaterialSamplerType SamplerType)
 {
-	if (IsAnyAPoison(Texture, TexCoord, MipBias))
+	if (IsAnyNotValid(Texture, TexCoord, MipBias))
 	{
-		return FPoison::Get();
+		return Poison();
 	}
 
 	if (bAutomaticViewMipBias)
 	{
-		MipBias = EmitOperator(BO_Add, MipBias, EmitExternalInput(EExternalInput::ViewMaterialTextureMipBias));
+		MipBias = Operator(BO_Add, MipBias, ExternalInput(EExternalInput::ViewMaterialTextureMipBias));
 	}
 
 	FTextureRead PrototypeHw = MakePrototype<FTextureRead>(FPrimitiveType::GetFloat4());
@@ -2100,21 +2145,21 @@ FValue* FEmitter::EmitTextureSampleBias(FValue* Texture, FValue* TexCoord, FValu
 	FTextureRead PrototypeAn = PrototypeHw;
 	PrototypeAn.Mode = ETextureReadMode::Derivatives;
 
-	FValue* Exp2MipBias = EmitOperator(UO_Exponential2, MipBias);
-	PrototypeAn.TexCoordDdx = EmitOperator(BO_Multiply, EmitAnalyticalPartialDerivative(TexCoord, EDerivativeAxis::X), Exp2MipBias);
-	PrototypeAn.TexCoordDdy = EmitOperator(BO_Multiply, EmitAnalyticalPartialDerivative(TexCoord, EDerivativeAxis::Y), Exp2MipBias);	
+	FValueRef Exp2MipBias = Operator(UO_Exponential2, MipBias);
+	PrototypeAn.TexCoordDdx = Operator(BO_Multiply, AnalyticalPartialDerivative(TexCoord, EDerivativeAxis::X), Exp2MipBias);
+	PrototypeAn.TexCoordDdy = Operator(BO_Multiply, AnalyticalPartialDerivative(TexCoord, EDerivativeAxis::Y), Exp2MipBias);	
 
 	FStageSwitch StageSwitch = MakePrototype<FStageSwitch>(PrototypeHw.Type);
-	StageSwitch.SetArgs(EmitPrototype(this, PrototypeHw), EmitPrototype(this, PrototypeAn));
+	StageSwitch.SetArgs(EmitPrototype(*this, PrototypeHw), EmitPrototype(*this, PrototypeAn));
 
-	return EmitPrototype(this, StageSwitch);
+	return EmitPrototype(*this, StageSwitch);
 }
 
-FValue* FEmitter::EmitTextureSampleGrad(FValue* Texture, FValue* TexCoord, FValue* TexCoordDdx, FValue* TexCoordDdy, bool bAutomaticViewMipBias, ESamplerSourceMode SamplerSourceMode, EMaterialSamplerType SamplerType)
+FValueRef FEmitter::TextureSampleGrad(FValueRef Texture, FValueRef TexCoord, FValueRef TexCoordDdx, FValueRef TexCoordDdy, bool bAutomaticViewMipBias, ESamplerSourceMode SamplerSourceMode, EMaterialSamplerType SamplerType)
 {
-	if (IsAnyAPoison(Texture, TexCoord, TexCoordDdx, TexCoordDdy))
+	if (IsAnyNotValid(Texture, TexCoord, TexCoordDdx, TexCoordDdy))
 	{
-		return FPoison::Get();
+		return Poison();
 	}
 
 	FTextureRead Prototype = MakePrototype<MIR::FTextureRead>(FPrimitiveType::GetFloat4());
@@ -2128,112 +2173,47 @@ FValue* FEmitter::EmitTextureSampleGrad(FValue* Texture, FValue* TexCoord, FValu
 
 	if (bAutomaticViewMipBias)
 	{
-		FValue* ViewMaterialTextureDerivativeMultiply = EmitExternalInput(EExternalInput::ViewMaterialTextureDerivativeMultiply);
-		Prototype.TexCoordDdx = EmitOperator(BO_Multiply, Prototype.TexCoordDdx, ViewMaterialTextureDerivativeMultiply);
-		Prototype.TexCoordDdy = EmitOperator(BO_Multiply, Prototype.TexCoordDdy, ViewMaterialTextureDerivativeMultiply);
-	}
-
-	return EmitPrototype(this, Prototype);
-}
-
-FValue* FEmitter::EmitInlineHLSL(const FType* Type, FString Code, TConstArrayView<FValue*> InArguments, EValueFlags ValueFlags, EGraphProperties UsedGraphProperties)
-{
-	MIR::FInlineHLSL Prototype = MakePrototype<MIR::FInlineHLSL>(Type);
-	Prototype.Flags = ValueFlags | EValueFlags::HasDynamicHLSLCode;
-	Prototype.GraphProperties = UsedGraphProperties;
-	Prototype.Code = Module->PushUserString(MoveTemp(Code));
-
-	if (!InArguments.IsEmpty())
-	{
-		checkf(InArguments.Num() < MIR::FInlineHLSL::MaxNumArguments, TEXT("Number of arguments for inline-HLSL out of bounds: %d was specified, but upper bound is %d"), InArguments.Num(), MIR::FInlineHLSL::MaxNumArguments);
-		Prototype.NumArguments = InArguments.Num();
-		FMemory::Memcpy(&Prototype.Arguments[0], InArguments.GetData(), sizeof(FValue*) * InArguments.Num());
+		FValueRef ViewMaterialTextureDerivativeMultiply = ExternalInput(EExternalInput::ViewMaterialTextureDerivativeMultiply);
+		Prototype.TexCoordDdx = Operator(BO_Multiply, Prototype.TexCoordDdx, ViewMaterialTextureDerivativeMultiply);
+		Prototype.TexCoordDdy = Operator(BO_Multiply, Prototype.TexCoordDdy, ViewMaterialTextureDerivativeMultiply);
 	}
 
-	return EmitPrototype(this, Prototype);
+	return EmitPrototype(*this, Prototype);
 }
 
-FValue* FEmitter::EmitInlineHLSL(const FMaterialExternalCodeDeclaration* InExternalCodeDeclaration, TConstArrayView<FValue*> InArguments, EValueFlags ValueFlags, EGraphProperties UsedGraphProperties)
-{
-	check(InExternalCodeDeclaration != nullptr);
-	const FType* ReturnType = FType::FromMaterialValueType(InExternalCodeDeclaration->GetReturnTypeValue());
-
-	MIR::FInlineHLSL Prototype = MakePrototype<MIR::FInlineHLSL>(ReturnType);
-	Prototype.Flags = ValueFlags;
-	Prototype.GraphProperties = UsedGraphProperties;
-	Prototype.ExternalCodeDeclaration = InExternalCodeDeclaration;
-
-	if (!InArguments.IsEmpty())
-	{
-		checkf(InArguments.Num() < MIR::FInlineHLSL::MaxNumArguments, TEXT("Number of arguments for inline-HLSL out of bounds: %d was specified, but upper bound is %d"), InArguments.Num(), MIR::FInlineHLSL::MaxNumArguments);
-		Prototype.NumArguments = InArguments.Num();
-		FMemory::Memcpy(&Prototype.Arguments[0], InArguments.GetData(), sizeof(FValue*) * InArguments.Num());
-	}
-
-	return EmitPrototype(this, Prototype);
-}
-
-const FType* FEmitter::TryGetCommonType(const FType* A, const FType* B)
-{
-	// Trivial case: types are equal
-	if (A == B)
-	{
-		return A;
-	}
-
-	const FPrimitiveType* PrimitiveA = A->AsPrimitive();
-	const FPrimitiveType* PrimitiveB = B->AsPrimitive();
-	if (!PrimitiveA || !PrimitiveB)
-	{
-		return nullptr;
-	}
-	
-	// If both A and B are matrices, their dimensions must match (equality check above didn't trigger).
-	if (PrimitiveA->IsMatrix() || PrimitiveB->IsMatrix())
-	{
-		return nullptr;
-	}
-
-	// Neither A nor B are matrices, but single scalar or vector. Return the largest.
-	check(PrimitiveA->NumColumns == 1 && PrimitiveB->NumColumns == 1);
-	EScalarKind ScalarKind = FMath::Max(PrimitiveA->ScalarKind, PrimitiveB->ScalarKind);
-	int32 NumRows = FMath::Max(PrimitiveA->NumRows, PrimitiveB->NumRows);
-	return FPrimitiveType::Get(ScalarKind, NumRows, 1);
-}
-
-FValue* FEmitter::EmitPartialDerivative(FValue* Value, EDerivativeAxis Axis)
+FValueRef FEmitter::PartialDerivative(FValueRef Value, EDerivativeAxis Axis)
 {
 	// Any operation on poison arguments is a poison.
-	if (Value->IsPoison())
+	if (!Value.IsValid())
 	{
-		return FPoison::Get();
+		return Value;
 	}
 
 	// Differentiation is only valid on primitive types.
 	const MIR::FPrimitiveType* ValuePrimitiveType = Value->Type->AsPrimitive();
 	if (!ValuePrimitiveType || !ScalarKindIsAnyFloat(ValuePrimitiveType->ScalarKind))
 	{
-		Errorf(TEXT("Trying to differentiate a value of type `%s` is invalid. Expected a float type."), Value->Type->GetSpelling().GetData());
-		return FPoison::Get();
+		Errorf(Value, TEXT("Trying to differentiate a value of type `%s` is invalid. Expected a float type."), Value->Type->GetSpelling().GetData());
+		return Poison();
 	}
 
 	// Make the hardware derivative instruction.
 	FHardwarePartialDerivative HwDerivativeProto = MakePrototype<FHardwarePartialDerivative>(Value->Type);
 	HwDerivativeProto.Arg = Value;
 	HwDerivativeProto.Axis = Axis;
-	FValue* HwDerivative = EmitPrototype(this, HwDerivativeProto);
+	FValueRef HwDerivative = EmitPrototype(*this, HwDerivativeProto);
 
 	// Compute the analytical derivative for stages that don't support hardware derivatives.
-	FValue* AnalyticalDerivative = EmitAnalyticalPartialDerivative(Value, Axis);
+	FValueRef AnalyticalDerivative = AnalyticalPartialDerivative(Value, Axis);
 
 	// Emit the stage switch instruction so that hardware derivatives are used on stages
 	// that support it and analytical deriatives in the other stages.
-	FValue* StageValues[NumStages] = {};
+	FValueRef StageValues[NumStages] = {};
 	for (int i = 0; i < NumStages; ++i)
 	{
 		StageValues[i] = (i == Stage_Pixel) ? HwDerivative : AnalyticalDerivative;
 	}
-	return EmitStageSwitch(Value->Type, StageValues);
+	return StageSwitch(Value->Type, StageValues);
 }
 
 static FValue* DifferentiateExternalInput(FEmitter& Emitter, FExternalInput* ExternalInput, EDerivativeAxis Axis)
@@ -2247,35 +2227,15 @@ static FValue* DifferentiateExternalInput(FEmitter& Emitter, FExternalInput* Ext
 			? EExternalInput((int32)EExternalInput::TexCoord0_Ddx + TexCoordIndex)
 			: EExternalInput((int32)EExternalInput::TexCoord0_Ddy + TexCoordIndex);
 
-		return Emitter.EmitExternalInput(PartialDerivativeExternalInput);
+		return Emitter.ExternalInput(PartialDerivativeExternalInput);
 	}
 
 	// All other inputs are assumed constant.
 	const MIR::FPrimitiveType* PrimType = static_cast<const MIR::FPrimitiveType*>(ExternalInput->Type);
-	return Emitter.EmitConstruct(PrimType, Emitter.EmitConstantScalarZero(PrimType->ScalarKind));
-}
-
-static FValue* EmitDifferentialInlineHLSL(FEmitter& Emitter, const FPrimitiveType* PrimitiveType, FInlineHLSL* InlineHLSL, EDerivativeAxis Axis)
-{
-	if (InlineHLSL->ExternalCodeDeclaration)
-	{
-		if (InlineHLSL->ExternalCodeDeclaration->Derivative == EDerivativeStatus::Valid)
-		{
-			const TConstArrayView<FValue*> InlineHLSLArguments(InlineHLSL->Arguments, InlineHLSL->NumArguments);
-			return Emitter.EmitInlineHLSL(InlineHLSL->ExternalCodeDeclaration, InlineHLSLArguments, Axis == EDerivativeAxis::X ? EValueFlags::DerivativeDDX : EValueFlags::DerivativeDDY);
-		}
-		else if (InlineHLSL->ExternalCodeDeclaration->Derivative == EDerivativeStatus::NotValid)
-		{
-			Emitter.Errorf(TEXT("External code declaration `%s` does not provide a valid derivative."), *InlineHLSL->ExternalCodeDeclaration->Name.ToString());
-			return FPoison::Get();
-		}
-	}
-
-	// If inline HLSL was dynamically generated and not provided from an external code declaration, simply return constant zero as no derivative is available.
-	return Emitter.EmitConstruct(PrimitiveType, Emitter.EmitConstantScalarZero(PrimitiveType->ScalarKind));
+	return Emitter.Cast(Emitter.ConstantZero(PrimType->ScalarKind), PrimType);
 }
 
-static FValue* DifferentiateOperator(FEmitter& Emitter, FOperator* Op, EDerivativeAxis Axis)
+static FValue* DifferentiateOperator(FEmitter& E, FOperator* Op, EDerivativeAxis Axis)
 {
 	const MIR::FPrimitiveType* PrimType = static_cast<const MIR::FPrimitiveType*>(Op->Type);
 
@@ -2284,21 +2244,15 @@ static FValue* DifferentiateOperator(FEmitter& Emitter, FOperator* Op, EDerivati
 	FValue* F = Op->AArg;
 	FValue* G = Op->BArg;
 	FValue* H = Op->CArg;
-	FValue* dF = F && !F->Type->IsBoolean() ? Emitter.EmitAnalyticalPartialDerivative(F, Axis) : nullptr; // Note: select's first argument is a boolean, avoid making the derivative then
-	FValue* dG = G ? Emitter.EmitAnalyticalPartialDerivative(G, Axis) : nullptr;
-	FValue* dH = H ? Emitter.EmitAnalyticalPartialDerivative(H, Axis) : nullptr;
+	FValue* dF = F && !F->Type->IsBoolean() ? E.AnalyticalPartialDerivative(F, Axis) : nullptr; // Note: select's first argument is a boolean, avoid making the derivative then
+	FValue* dG = G ? E.AnalyticalPartialDerivative(G, Axis) : nullptr;
+	FValue* dH = H ? E.AnalyticalPartialDerivative(H, Axis) : nullptr;
 
 	// Convenience local functions as multiplications and division operations are common in derivatives.
-	auto Zero = [&Emitter, PrimType] () { return Emitter.EmitConstantScalarZero(PrimType->ScalarKind); };
-	auto One = [&Emitter, PrimType] () { return Emitter.EmitConstantScalarOne(PrimType->ScalarKind); };
-	auto Constant = [&Emitter, PrimType] (TFloat FromFloat) { return Emitter.EmitConstantScalar(PrimType->ScalarKind, FromFloat); };
-	auto Neg = [&Emitter] (FValue* A) { return Emitter.EmitOperator(UO_Negate, A); };
-	auto Add = [&Emitter] (FValue* A, FValue* B) { return Emitter.EmitOperator(BO_Add, A, B); };
-	auto Sub = [&Emitter] (FValue* A, FValue* B) { return Emitter.EmitOperator(BO_Subtract, A, B); };
-	auto Mul = [&Emitter] (FValue* A, FValue* B) { return Emitter.EmitOperator(BO_Multiply, A, B); };
-	auto Div = [&Emitter] (FValue* A, FValue* B) { return Emitter.EmitOperator(BO_Divide, A, B); };
-	auto Sqrt = [&Emitter] (FValue* A) { return Emitter.EmitOperator(UO_Sqrt, A); };
-	
+	auto Zero = [&E, PrimType] () { return E.ConstantZero(PrimType->ScalarKind); };
+	auto One = [&E, PrimType] () { return E.ConstantOne(PrimType->ScalarKind); };
+	auto Constant = [&E, PrimType] (TFloat FromFloat) { return E.ConstantScalar(PrimType->ScalarKind, FromFloat); };
+
 	// Some constants
 	constexpr TFloat Ln2 = 0.69314718055994530941723212145818;
 	constexpr TFloat Ln10 = 2.3025850929940456840179914546844;
@@ -2307,35 +2261,35 @@ static FValue* DifferentiateOperator(FEmitter& Emitter, FOperator* Op, EDerivati
 	{
 		// d/dx -f(x) = -f'(x)
 		case UO_Negate:
-			return Neg(dF);
+			return E.Negate(dF);
 
 		// d/dx |f(x)| = f(x) f'(x) / |f(x)|
 		case UO_Abs:
-			return Div(Mul(F, dF), Op);
+			return E.Divide(E.Multiply(F, dF), Op);
 
 		// d/dx arccos(f(x)) = -1 / sqrt(1 - f(x)^2) * f'(x)
 		case UO_ACos:
-			return Neg(Div(dF, Sqrt(Sub(One(), Mul(F, F)))));
+			return E.Negate(E.Divide(dF, E.Sqrt(E.Subtract(One(), E.Multiply(F, F)))));
 
 		// d/dx arcsin(f(x)) = 1 / sqrt(1 - f(x)^2) * f'(x)
 		case UO_ASin:
-			return Div(dF, Sqrt(Sub(One(), Mul(F, F))));
+			return E.Divide(dF, E.Sqrt(E.Subtract(One(), E.Multiply(F, F))));
 
 		// d/dx arctan(f(x)) = 1 / (1 + f(x)^2) * f'(x)
 		case UO_ATan:
-			return Div(dF, Add(One(), Mul(F, F)));
+			return E.Divide(dF, E.Add(One(), E.Multiply(F, F)));
 
 		// d/dx cos(f(x)) = -sin(f(x)) * f'(x)
 		case UO_Cos:
-			return Neg(Mul(Emitter.EmitOperator(UO_Sin, F), dF));
+			return E.Negate(E.Multiply(E.Sin(F), dF));
 
 		// d/dx e^f(x) = e^f(x) * f'(x)
 		case UO_Exponential:
-			return Mul(Op, dF);
+			return E.Multiply(Op, dF);
 
 		// d/dx 2^f(x) = ln(2) * 2^f(x) * f'(x)
 		case UO_Exponential2:
-			return Mul(Mul(Constant(Ln2), Op), dF);
+			return E.Multiply(E.Multiply(Constant(Ln2), Op), dF);
 
 		// d/dx frac(f(x)) = f'(x), since frac(x) = x - floor(x)
 		case UO_Frac:
@@ -2343,41 +2297,40 @@ static FValue* DifferentiateOperator(FEmitter& Emitter, FOperator* Op, EDerivati
 
 		// d/dx |f(x)| (length in vector case) = f(x) f'(x) / |f(x)|
 		case UO_Length:
-			return Div(Mul(F, dF), Op);
+			return E.Divide(E.Multiply(F, dF), Op);
 
 		// d/dx log(f(x)) = 1 / f(x) * f'(x)
 		case UO_Logarithm:
-			return Div(dF, F);
+			return E.Divide(dF, F);
 
 		// d/dx log2(f(x)) = 1 / (f(x) * ln(2)) * f'(x)
 		case UO_Logarithm2:
-			return Div(dF, Mul(F, Constant(Ln2)));
+			return E.Divide(dF, E.Multiply(F, Constant(Ln2)));
 
 		// d/dx log10(f(x)) = 1 / (f(x) * ln(10)) * f'(x)
 		case UO_Logarithm10:
-			return Div(dF, Mul(F, Constant(Ln10)));
+			return E.Divide(dF, E.Multiply(F, Constant(Ln10)));
 
 		// d/dx saturate(f(x)) = f'(x) if f(x) is inside (0-1) range, 0 otherwise
 		case UO_Saturate:
-			return Emitter.EmitOperator(TO_Select,
-						Emitter.EmitOperator(BO_And,
-							Emitter.EmitOperator(BO_LessThan, Zero(), F), // 0 < f(x)
-							Emitter.EmitOperator(BO_LessThan, F, One())), // f(x) < 1
+			return E.Select(E.And(
+							E.LessThan(Zero(), F), // 0 < f(x)
+							E.LessThan(F, One())), // f(x) < 1
 						dF, Zero());
 
 		// d/dx sin(f(x)) = cos(f(x)) * f'(x)
 		case UO_Sin:
-			return Mul(Emitter.EmitOperator(UO_Cos, F), dF);
+			return E.Multiply(E.Cos(F), dF);
 
 		// d/dx sqrt(f(x)) = 1 / (2 * sqrt(f(x))) * f'(x)
 		case UO_Sqrt:
-			return Div(dF, Mul(Emitter.EmitConstantScalar(PrimType->ScalarKind, 2), Sqrt(F)));
+			return E.Divide(dF, E.Multiply(Constant(2), E.Sqrt(F)));
 
 		// d/dx tan(f(x)) = 1 / cos^2(f(x)) * f'(x)
 		case UO_Tan:
 		{
-			FValue* CosVal = Emitter.EmitOperator(UO_Cos, F);
-			return Div(dF, Mul(CosVal, CosVal));
+			FValue* CosVal = E.Cos(F);
+			return E.Divide(dF, E.Multiply(CosVal, CosVal));
 		}
 
 		// These functions are piecewise constant, that is mostly constant with some
@@ -2391,85 +2344,73 @@ static FValue* DifferentiateOperator(FEmitter& Emitter, FOperator* Op, EDerivati
 
 		// d/dx (f(x) + g(x)) = f'(x) + g'(x)
 		case BO_Add:
-			return Add(dF, dG);
+			return E.Add(dF, dG);
 
 		// d/dx (f(x) - g(x)) = f'(x) - g'(x)
 		case BO_Subtract:
-			return Sub(dF, dG);
+			return E.Subtract(dF, dG);
 
 		// d/dx (f(x) * g(x)) = f'(x) * g(x) + f(x) * g'(x)
 		case BO_Multiply:
-			return Add(Mul(dF, G), Mul(F, dG));
+			return E.Add(E.Multiply(dF, G), E.Multiply(F, dG));
 
 		// d/dx (f(x) / g(x)) = (f'(x) * g(x) - f(x) * g'(x)) / g(x)^2
 		case BO_Divide:
-			return Div(Sub(Mul(dF, G), Mul(F, dG)), Mul(G, G));
+			return E.Divide(E.Subtract(E.Multiply(dF, G), E.Multiply(F, dG)), E.Multiply(G, G));
 
 		// fmod(f(x), g(x)) = f(x) - g(x) * floor(f(x) / g(x)).
 		// Thus:
 		//     d/dx fmod(f(x), g(x)) = f'(x) - g(x) * floor(f(x) / g(x))
 		// since `floor` is piecewise constant.
 		case BO_Fmod:
-			return Sub(dF, Mul(dG, Emitter.EmitOperator(UO_Floor, Div(F, G))));
+			return E.Subtract(dF, E.Multiply(dG, E.Operator(UO_Floor, E.Divide(F, G))));
 
 		// d/dx max(f(x), g(x)) = f'(x) if f(x) > g(x), else g'(x)
 		case BO_Max:
-			return Emitter.EmitOperator(TO_Select,
-						Emitter.EmitOperator(BO_GreaterThan, F, G),
-						dF,
-						dG);
+			return E.Select(E.Operator(BO_GreaterThan, F, G), dF, dG);
 
 		// d/dx min(f(x), g(x)) = f'(x) if f(x) < g(x), else g'(x)
 		case BO_Min:
-			return Emitter.EmitOperator(TO_Select,
-						Emitter.EmitOperator(BO_LessThan, F, G),
-						dF,
-						dG);
+			return E.Select(E.LessThan( F, G), dF, dG);
 
 		// The multiplication rule applies for the dot product too.
 		// d/dx (f(x) • g(x)) = f'(x) • g(x) + f(x) • g'(x)
 		case BO_Dot:
-			return Add(Emitter.EmitOperator(BO_Dot, dF, G), 
-					   Emitter.EmitOperator(BO_Dot, F, dG));
+			return E.Add(E.Operator(BO_Dot, dF, G), E.Operator(BO_Dot, F, dG));
 
-		
 		// The multiplication rule applies for the cross product too.
 		// d/dx (f(x) × g(x)) = f'(x) × g(x) + f(x) × g'(x)
 		case BO_Cross:
-			return Add(Emitter.EmitOperator(BO_Cross, dF, G), 
-					   Emitter.EmitOperator(BO_Cross, F, dG));
+			return E.Add(E.Operator(BO_Cross, dF, G), E.Operator(BO_Cross, F, dG));
 
 		// clamp(x, min, max) (F=x, min=G, max=H)
 		// The derivative is defined when x is between min and max (f'(x)). At and outside
 		// bounds, the clamp result is constant and thus the derivative is zero.
 		case TO_Clamp:
-			return Emitter.EmitOperator(TO_Select,
-				Emitter.EmitOperator(BO_And,
-					Emitter.EmitOperator(BO_LessThan, G, F),
-					Emitter.EmitOperator(BO_LessThan, F, H)),
+			return E.Select(E.And(
+					E.LessThan(G, F), 
+					E.LessThan(F, H)),
 				dF, Zero());
 
-
 		// lerp(a, b, t) = a + t * (b - a)
 		// d/dx lerp(f(x), g(x), h(x)) = f'(x) + d/dx (h(x) * ((g(x) - f(x)))
 		// d/dx (h(x) * ((g(x) - f(x))) = h'(x) * ((g(x) - f(x))) + h(x) * (g'(x) - f'(x))
 		case TO_Lerp:
-			return Add(dF, 
-					   Add(Mul(dH, Sub(G, F)), 
-						   Mul(H, Sub(dG, dF))));
+			return E.Add(dF, 
+					   E.Add(E.Multiply(dH, E.Subtract(G, F)), 
+						   E.Multiply(H, E.Subtract(dG, dF))));
 
-		// d/dx select(cond, g(x), h(x)) ≈ select(cond, g'(x), h'(x))
+		// d/dx select(F, g(x), h(x)) ≈ select(F, g'(x), h'(x))
 		case TO_Select:
-			return Emitter.EmitOperator(TO_Select, dG, dH);
-
+			return E.Select(F, dG, dH);
 
 		// smoothstep(f(x), g(x), h(x)) = 3 z^2 - 2 z^3  with z = saturate((h - f) / (g - f))
 		case TO_Smoothstep:
 		{
-			FValue* Z  = Emitter.EmitOperator(UO_Saturate, Div(Sub(H, F), Sub(G, F)));
-			FValue* dZ = Emitter.EmitAnalyticalPartialDerivative(Z, Axis);
+			FValue* Z  = E.Saturate(E.Divide(E.Subtract(H, F), E.Subtract(G, F)));
+			FValue* dZ = E.AnalyticalPartialDerivative(Z, Axis);
 			// d/dx 3 z(x)^2 - 2 z(x)^3 = 6 * z(x) * z'(x) - 6 * z(x)^2 * z'(x) = 6 * (z(x) - z(x)^2) * z'(x)
-			return Mul(dZ, Mul(Constant(6), Sub(Z, Mul(Z, Z))));
+			return E.Multiply(dZ, E.Multiply(Constant(6), E.Subtract(Z, E.Multiply(Z, Z))));
 		}
 
 		default:
@@ -2477,115 +2418,164 @@ static FValue* DifferentiateOperator(FEmitter& Emitter, FOperator* Op, EDerivati
 	}
 }
 
-FValue* FEmitter::EmitAnalyticalPartialDerivative(FValue* Value, EDerivativeAxis Axis)
+FValueRef FEmitter::AnalyticalPartialDerivative(FValueRef Value, EDerivativeAxis Axis)
 {
 	// Any operation on poison arguments is a poison.
-	if (Value->IsPoison())
+	if (!Value.IsValid())
 	{
-		return FPoison::Get();
+		return Value;
 	}
 
 	// Differentiation is only valid on primitive types.
 	const MIR::FPrimitiveType* ValuePrimitiveType = Value->Type->AsPrimitive();
 	if (!ValuePrimitiveType || !ScalarKindIsAnyFloat(ValuePrimitiveType->ScalarKind))
 	{
-		Errorf(TEXT("Trying to differentiate a value of type `%s` is invalid. Expected a float type."), Value->Type->GetSpelling().GetData());
-		return FPoison::Get();
+		Errorf(Value, TEXT("Trying to differentiate a value of type `%s` is invalid. Expected a float type."), Value->Type->GetSpelling().GetData());
+		return Poison();
 	}
 
 	switch (Value->Kind)
 	{
 		case MIR::VK_ExternalInput:
-			return DifferentiateExternalInput(*this, static_cast<FExternalInput*>(Value), Axis);
+			return DifferentiateExternalInput(*this, Value->As<FExternalInput>(), Axis);
 
 		case MIR::VK_Dimensional:
 		{
 			TConstArrayView<FValue*> ValueComponents = Value->As<FDimensional>()->GetComponents();
-			FDimensional* Derivative = NewDimensionalValue(this, ValuePrimitiveType);
+			FDimensional* Derivative = NewDimensionalValue(*this, ValuePrimitiveType);
 			TArrayView<FValue*> DerivativeComponents = Derivative->GetMutableComponents();
 			for (int32 i = 0; i < ValueComponents.Num(); ++i)
 			{
-				DerivativeComponents[i] = EmitAnalyticalPartialDerivative(ValueComponents[i], Axis);
+				DerivativeComponents[i] = AnalyticalPartialDerivative(ValueComponents[i], Axis);
 			}
 			return Derivative;
 		}
 
 		case MIR::VK_Operator:
-			return DifferentiateOperator(*this, static_cast<FOperator*>(Value), Axis);
+			return DifferentiateOperator(*this, Value->As<FOperator>(), Axis);
 		
 		case MIR::VK_Branch:
 		{
-			FBranch* Branch = Value->As<FBranch>();
-			return EmitBranch(Branch->ConditionArg,
-				EmitAnalyticalPartialDerivative(Branch->TrueArg, Axis),
-				EmitAnalyticalPartialDerivative(Branch->FalseArg, Axis));
+			FBranch* AsBranch = Value->As<FBranch>();
+			return Branch(AsBranch->ConditionArg,
+				AnalyticalPartialDerivative(AsBranch->TrueArg, Axis),
+				AnalyticalPartialDerivative(AsBranch->FalseArg, Axis));
 		}
 
 		case MIR::VK_Subscript:
 		{
-			FSubscript* Subscript = Value->As<FSubscript>();
-			return EmitSubscript(EmitAnalyticalPartialDerivative(Subscript->Arg, Axis), Subscript->Index);
+			FSubscript* AsSubscript = Value->As<FSubscript>();
+			return Subscript(AnalyticalPartialDerivative(AsSubscript->Arg, Axis), AsSubscript->Index);
 		}
 
 		case MIR::VK_Cast:
 		{
-			FCast* Cast = Value->As<FCast>();
-			return EmitConstruct(Cast->Type, EmitAnalyticalPartialDerivative(Cast->Arg, Axis));
+			FCast* AsCast = Value->As<FCast>();
+			return Cast(AnalyticalPartialDerivative(AsCast->Arg, Axis), AsCast->Type);
 		}
 
 		// These values are uniform (constant), thus their value is always zero.
 		case MIR::VK_Constant:
 		case MIR::VK_UniformParameter:
 		case MIR::VK_TextureRead:
-			return EmitConstruct(ValuePrimitiveType, EmitConstantScalarZero(ValuePrimitiveType->ScalarKind));
-
 		case MIR::VK_InlineHLSL:
-		{
-			FInlineHLSL* InlineHLSL = Value->As<FInlineHLSL>();
-			return EmitDifferentialInlineHLSL(*this, ValuePrimitiveType, InlineHLSL, Axis);
-		}
+			return Cast(ConstantZero(ValuePrimitiveType->ScalarKind), ValuePrimitiveType);
 
 		default:
 			UE_MIR_UNREACHABLE();
 	}
 }
 
-const FType* FEmitter::GetCommonType(const FType* A, const FType* B)
+static FValueRef EmitInlineHLSL(FEmitter& Emitter, const FType* Type, const FMaterialExternalCodeDeclaration* InExternalCodeDeclaration, const TCHAR* Code, TConstArrayView<FValueRef> InputValues, EValueFlags ValueFlags, EGraphProperties UsedGraphProperties)
 {
-	if (const FType* CommonType = TryGetCommonType(A, B))
+	MIR::FInlineHLSL Prototype = MakePrototype<MIR::FInlineHLSL>(Type);
+	Prototype.Type = Type;
+	Prototype.Flags = ValueFlags;
+	Prototype.GraphProperties = UsedGraphProperties;
+	
+	if (InExternalCodeDeclaration)
 	{
-		return CommonType;
+		check(!Code);
+		Prototype.ExternalCodeDeclaration = InExternalCodeDeclaration;
+	}
+	else
+	{
+		Prototype.Code = Code;
 	}
 
-	Errorf(TEXT("No common type between '%s' and '%s'."), A->GetSpelling().GetData(), B->GetSpelling().GetData());
-	return nullptr;
+	if (!InputValues.IsEmpty())
+	{
+		checkf(InputValues.Num() < MIR::FInlineHLSL::MaxNumArguments, TEXT("Number of arguments for inline-HLSL out of bounds: %d was specified, but upper bound is %d"), InputValues.Num(), MIR::FInlineHLSL::MaxNumArguments);
+		Prototype.NumArguments = InputValues.Num();
+		for (int32 i = 0; i < InputValues.Num(); ++i)
+		{
+			Prototype.NumArguments = InputValues[i];
+		}
+	}
+
+	return EmitPrototype(Emitter, Prototype);
 }
 
-void FEmitter::Error(FStringView Message)
+FValueRef FEmitter::InlineHLSL(const FType* Type, FString Code, TConstArrayView<FValueRef> InputValues, EValueFlags ValueFlags, EGraphProperties UsedGraphProperties)
 {
-	FMaterialIRModule::FError Error;
-	Error.Expression = Expression;
-
-	// Add the node type to the error message
-	const int32 ChopCount = FCString::Strlen(TEXT("MaterialExpression"));
-	const FString& ErrorClassName = Expression->GetClass()->GetName();
+	if (IsAnyNotValid(InputValues))
+	{
+		return Poison();
+	}
 	
-	Error.Message.Appendf(TEXT("(Node %s) %s"), *ErrorClassName + ChopCount, Message.GetData());
+	return EmitInlineHLSL(*this, Type, nullptr, Module->PushUserString(MoveTemp(Code)), InputValues, ValueFlags | EValueFlags::HasDynamicHLSLCode, UsedGraphProperties);
+}
+
+FValueRef FEmitter::InlineHLSL(const FMaterialExternalCodeDeclaration* InExternalCodeDeclaration, TConstArrayView<FValueRef> InputValues, EValueFlags ValueFlags, EGraphProperties UsedGraphProperties)
+{
+	if (IsAnyNotValid(InputValues))
+	{
+		return Poison();
+	}
+
+	check(InExternalCodeDeclaration != nullptr);
+	const FType* ReturnType = FType::FromMaterialValueType(InExternalCodeDeclaration->GetReturnTypeValue());
+	return EmitInlineHLSL(*this, ReturnType, InExternalCodeDeclaration, nullptr, InputValues, ValueFlags, UsedGraphProperties);
+}
+
+const FType* FEmitter::TryGetCommonType(const FType* A, const FType* B)
+{
+	// Trivial case: types are equal
+	if (A == B)
+	{
+		return A;
+	}
+
+	const FPrimitiveType* PrimitiveA = A->AsPrimitive();
+	const FPrimitiveType* PrimitiveB = B->AsPrimitive();
+	if (!PrimitiveA || !PrimitiveB)
+	{
+		return nullptr;
+	}
 	
-	Module->Errors.Push(Error);
-	bCurrentExpressionHasErrors = true;
+	// If both A and B are matrices, their dimensions must match (equality check above didn't trigger).
+	if (PrimitiveA->IsMatrix() || PrimitiveB->IsMatrix())
+	{
+		return nullptr;
+	}
+
+	// Neither A nor B are matrices, but single scalar or vector. Return the largest.
+	check(PrimitiveA->NumColumns == 1 && PrimitiveB->NumColumns == 1);
+	EScalarKind ScalarKind = FMath::Max(PrimitiveA->ScalarKind, PrimitiveB->ScalarKind);
+	int32 NumRows = FMath::Max(PrimitiveA->NumRows, PrimitiveB->NumRows);
+	return FPrimitiveType::Get(ScalarKind, NumRows, 1);
 }
 
 void FEmitter::Initialize()
 {
 	// Create and reference the true/false constants.
-	FConstant Temp = MakePrototype<FConstant>(FPrimitiveType::GetBool1());
+	FConstant Temp = MakePrototype<FConstant>(FPrimitiveType::GetBool());
 
 	Temp.Boolean = true;
-	ConstantTrue = EmitPrototype(this, Temp);
+	TrueConstant = EmitPrototype(*this, Temp);
 
 	Temp.Boolean = false;
-	ConstantFalse = EmitPrototype(this, Temp);
+	FalseConstant = EmitPrototype(*this, Temp);
 }
 
 bool FEmitter::FValueKeyFuncs::Matches(KeyInitType A, KeyInitType B)
diff --git a/Engine/Source/Runtime/Engine/Private/Materials/MaterialIRInternal.cpp b/Engine/Source/Runtime/Engine/Private/Materials/MaterialIRInternal.cpp
--- a/Engine/Source/Runtime/Engine/Private/Materials/MaterialIRInternal.cpp
+++ b/Engine/Source/Runtime/Engine/Private/Materials/MaterialIRInternal.cpp
@@ -62,19 +62,19 @@ MIR::FValue* CreateMaterialAttributeDefaultValue(FEmitter& Emitter, const UMater
 
 	switch (Type)
 	{
-		case MCT_ShadingModel: return Emitter.EmitConstantIntScalar(Material->GetShadingModels().GetFirstShadingModel());
+		case MCT_ShadingModel: return Emitter.ConstantInt(Material->GetShadingModels().GetFirstShadingModel());
 
-		case MCT_Float1: return Emitter.EmitConstantFloatScalar(  DefaultValue.X );
-		case MCT_Float2: return Emitter.EmitConstantFloat2({ DefaultValue.X, DefaultValue.Y });
-		case MCT_Float3: return Emitter.EmitConstantFloat3({ DefaultValue.X, DefaultValue.Y, DefaultValue.Z });
+		case MCT_Float1: return Emitter.ConstantFloat(  DefaultValue.X );
+		case MCT_Float2: return Emitter.ConstantFloat2({ DefaultValue.X, DefaultValue.Y });
+		case MCT_Float3: return Emitter.ConstantFloat3({ DefaultValue.X, DefaultValue.Y, DefaultValue.Z });
 		case MCT_Float:
-		case MCT_Float4: return Emitter.EmitConstantFloat4(DefaultValue);
+		case MCT_Float4: return Emitter.ConstantFloat4(DefaultValue);
 
-		case MCT_UInt1: return Emitter.EmitConstantIntScalar(  (int32)DefaultValue.X );
-		case MCT_UInt2: return Emitter.EmitConstantInt2({ (int32)DefaultValue.X, (int32)DefaultValue.Y });
-		case MCT_UInt3: return Emitter.EmitConstantInt3({ (int32)DefaultValue.X, (int32)DefaultValue.Y, (int32)DefaultValue.Z });
+		case MCT_UInt1: return Emitter.ConstantInt((TInteger)DefaultValue.X);
+		case MCT_UInt2: return Emitter.ConstantInt2({ (TInteger)DefaultValue.X, (TInteger)DefaultValue.Y });
+		case MCT_UInt3: return Emitter.ConstantInt3({ (TInteger)DefaultValue.X, (TInteger)DefaultValue.Y, (TInteger)DefaultValue.Z });
 		case MCT_UInt:
-		case MCT_UInt4: return Emitter.EmitConstantInt4({ (int32)DefaultValue.X, (int32)DefaultValue.Y, (int32)DefaultValue.Z, (int32)DefaultValue.W });
+		case MCT_UInt4: return Emitter.ConstantInt4({ (TInteger)DefaultValue.X, (TInteger)DefaultValue.Y, (TInteger)DefaultValue.Z, (TInteger)DefaultValue.W });
 
 		default: UE_MIR_UNREACHABLE();
 	}
diff --git a/Engine/Source/Runtime/Engine/Private/Materials/MaterialIRModuleBuilder.cpp b/Engine/Source/Runtime/Engine/Private/Materials/MaterialIRModuleBuilder.cpp
--- a/Engine/Source/Runtime/Engine/Private/Materials/MaterialIRModuleBuilder.cpp
+++ b/Engine/Source/Runtime/Engine/Private/Materials/MaterialIRModuleBuilder.cpp
@@ -103,11 +103,11 @@ struct FMaterialIRModuleBuilderImpl
 		FMaterialInputDescription Input;
 		Builder->Material->GetExpressionInputDescription(Property, Input);
 
-		MIR::FSetMaterialOutput* Output = Emitter->EmitSetMaterialOutput(Property, nullptr);
+		MIR::FSetMaterialOutput* Output = Emitter->SetMaterialOutput(Property, nullptr);
 
 		if (Input.bUseConstant)
 		{
-			Output->Arg = Emitter->EmitConstantFromShaderValue(Input.ConstantValue);
+			Output->Arg = Emitter->ConstantFromShaderValue(Input.ConstantValue);
 		}
 		else if (!Input.Input->IsConnected())
 		{
@@ -228,7 +228,7 @@ struct FMaterialIRModuleBuilderImpl
 			if (Value)
 			{
 				const MIR::FType* Type = MIR::FType::FromMaterialValueType(FunctionInput.ExpressionInput->GetInputValueType(0));
-				CallInputValues[i] = Emitter->EmitConstruct(Type, Value);
+				CallInputValues[i] = Emitter->Cast(Value, Type);
 			}
 		}
 
@@ -287,7 +287,7 @@ struct FMaterialIRModuleBuilderImpl
 			const MIR::FType* OutputType = MIR::FType::FromMaterialValueType(FunctionOutput.ExpressionOutput->GetOutputValueType(0));
 
 			// Cast the value to the expected output type. This may fail (value will be poison). 
-			Value = Emitter->EmitConstruct(OutputType, Value);
+			Value = Emitter->Cast(Value, OutputType);
 
 			// And flow it to the relative function *call* output in the parent context.
 			ParentContext.SetOutputValue(Call->GetOutput(i), Value);
@@ -319,7 +319,7 @@ struct FMaterialIRModuleBuilderImpl
 					MIR::Internal::BindValueToExpressionInput(this, Input.Input, *ValuePtr);
 
 					const MIR::FType* OutputArgType = MIR::FType::FromShaderType(Input.Type);
-					Output->Arg = Emitter->EmitConstruct(OutputArgType, *ValuePtr);
+					Output->Arg = Emitter->Cast(*ValuePtr, OutputArgType);
 				}
 
 				// Push this connection insight
diff --git a/Engine/Source/Runtime/Engine/Private/Materials/MaterialIRTypes.cpp b/Engine/Source/Runtime/Engine/Private/Materials/MaterialIRTypes.cpp
--- a/Engine/Source/Runtime/Engine/Private/Materials/MaterialIRTypes.cpp
+++ b/Engine/Source/Runtime/Engine/Private/Materials/MaterialIRTypes.cpp
@@ -90,7 +90,7 @@ const FType* FType::FromMaterialParameterType(EMaterialParameterType Type)
 {
 	switch (Type)
 	{
-		case EMaterialParameterType::Scalar: return FPrimitiveType::GetFloat1();
+		case EMaterialParameterType::Scalar: return FPrimitiveType::GetFloat();
 		case EMaterialParameterType::Vector: return FPrimitiveType::GetFloat4();
 		case EMaterialParameterType::DoubleVector: UE_MIR_TODO();
 		case EMaterialParameterType::Texture: return FObjectType::GetTexture2D();
@@ -98,7 +98,7 @@ const FType* FType::FromMaterialParameterType(EMaterialParameterType Type)
 		case EMaterialParameterType::Font: UE_MIR_TODO();
 		case EMaterialParameterType::RuntimeVirtualTexture: UE_MIR_TODO();
 		case EMaterialParameterType::SparseVolumeTexture: UE_MIR_TODO();
-		case EMaterialParameterType::StaticSwitch: return FPrimitiveType::GetBool1();
+		case EMaterialParameterType::StaticSwitch: return FPrimitiveType::GetBool();
 		default: UE_MIR_UNREACHABLE();
 	}
 }
@@ -195,7 +195,7 @@ bool FType::IsInteger() const
 
 bool FType::IsBoolScalar() const
 {
-	return this == FPrimitiveType::GetBool1();
+	return this == FPrimitiveType::GetBool();
 }
 
 bool FType::IsObjectOfKind(EObjectKind ObjectKind) const
@@ -221,6 +221,13 @@ const FPrimitiveType* FType::AsPrimitive() const
 	return Kind == TypeKind_Primitive ? static_cast<const FPrimitiveType*>(this) : nullptr; 
 }
 
+const FPrimitiveType* FType::AsArithmetic() const
+{
+	return Kind == TypeKind_Primitive && static_cast<const FPrimitiveType*>(this)->IsArithmetic()
+		? static_cast<const FPrimitiveType*>(this)
+		: nullptr; 
+}
+
 const FObjectType* FType::AsObject() const
 {
 	return Kind == TypeKind_Object ? static_cast<const FObjectType*>(this) : nullptr; 
@@ -265,17 +272,17 @@ const TCHAR* ScalarKindToString(EScalarKind Kind)
 	}
 }
 
-const FPrimitiveType* FPrimitiveType::GetBool1()
+const FPrimitiveType* FPrimitiveType::GetBool()
 {
 	return GetScalar(ScalarKind_Bool);
 }
 
-const FPrimitiveType* FPrimitiveType::GetInt1()
+const FPrimitiveType* FPrimitiveType::GetInt()
 {
 	return GetScalar(ScalarKind_Int);
 }
 
-const FPrimitiveType* FPrimitiveType::GetFloat1()
+const FPrimitiveType* FPrimitiveType::GetFloat()
 {
 	return GetScalar(ScalarKind_Float);
 }
@@ -300,17 +307,10 @@ const FPrimitiveType* FPrimitiveType::GetScalar(EScalarKind InScalarKind)
 	return Get(InScalarKind, 1, 1);
 }
 
-const FPrimitiveType* FPrimitiveType::GetVector(EScalarKind InScalarKind, int NumComponents)
+const FPrimitiveType* FPrimitiveType::GetVector(EScalarKind InScalarKind, int NumRows)
 {
-	check(NumComponents >= 1 && NumComponents <= 4);
-	return Get(InScalarKind, NumComponents, 1);
-}
-
-const FPrimitiveType* FPrimitiveType::GetMatrix(EScalarKind InScalarKind, int NumRows, int NumColumns)
-{
-	check(NumColumns > 1 && NumColumns <= 4);
-	check(NumRows > 1 && NumRows <= 4);
-	return Get(InScalarKind, NumRows, NumColumns);
+	check(NumRows >= 1 && NumRows <= 4);
+	return Get(InScalarKind, NumRows, 1);
 }
 
 const FPrimitiveType* FPrimitiveType::Get(EScalarKind InScalarKind, int NumRows, int NumColumns)
@@ -380,11 +380,16 @@ const FPrimitiveType* FPrimitiveType::ToScalar() const
 	return FPrimitiveType::GetScalar(ScalarKind);
 }
 
-const FPrimitiveType* FPrimitiveType::ToScalarKind(EScalarKind InScalarKind) const
+const FPrimitiveType* FPrimitiveType::WithScalarKind(EScalarKind InScalarKind) const
 {
 	return FPrimitiveType::Get(InScalarKind, NumRows, NumColumns);
 }
 
+const FPrimitiveType* FPrimitiveType::ToVector(int InNumRows) const
+{
+	return FPrimitiveType::GetVector(ScalarKind, InNumRows);
+}
+
 const FType* FObjectType::GetTexture2D()
 {
 	static const FObjectType Instance { { TypeKind_Object }, ObjectKind_Texture2D };
diff --git a/Engine/Source/Runtime/Engine/Private/Materials/MaterialIRValueAnalyzer.h b/Engine/Source/Runtime/Engine/Private/Materials/MaterialIRValueAnalyzer.h
--- a/Engine/Source/Runtime/Engine/Private/Materials/MaterialIRValueAnalyzer.h
+++ b/Engine/Source/Runtime/Engine/Private/Materials/MaterialIRValueAnalyzer.h
@@ -44,9 +44,6 @@ struct FMaterialIRValueAnalyzer
 		
 	//
 	FMaterialInsights* Insights{}; // Optional
-
-	//
-	FMaterialIRStateTracker* StateTracker{};
 };
 
 #endif // #if WITH_EDITOR
diff --git a/Engine/Source/Runtime/Engine/Public/Materials/MaterialIR.h b/Engine/Source/Runtime/Engine/Public/Materials/MaterialIR.h
--- a/Engine/Source/Runtime/Engine/Public/Materials/MaterialIR.h
+++ b/Engine/Source/Runtime/Engine/Public/Materials/MaterialIR.h
@@ -11,10 +11,8 @@ struct FMaterialExternalCodeDeclaration;
 namespace MIR
 {
 
-/**
- * @brief Identifies the block of execution in which an instruction runs.
- * @note If you introduce a new stage, make sure you update EValueFlags accordingly.
- */
+// Identifies the block of execution in which an instruction runs.
+// Note: If you introduce a new stage, make sure you update EValueFlags accordingly.
 enum EStage
 {
 	Stage_Vertex,
@@ -23,45 +21,39 @@ enum EStage
 	NumStages
 };
 
-/**
- * Returns the string representation of given stage.
- */
+// Returns the string representation of given stage.
 const TCHAR* LexToString(EStage Stage);
 
-/**
- * @brief Returns whether specified property is evaluated in specified stage.
- */
+// Returns whether specified property is evaluated in specified stage.
 bool MaterialOutputEvaluatesInStage(EMaterialProperty Property, EStage Stage);
 
-/**
- * @brief A collection of bit flags for a specific Value instance.
- */
+// A collection of bit flags for a specific Value instance.
 enum class EValueFlags : uint8
 {
 	None = 0,
 
-	/** Value has been analyzed for the vertex stage. */
+	// Value has been analyzed for the vertex stage.
 	AnalyzedForStageVertex = 1 << 0,
 
-	/** Value has been analyzed for the pixel stage. */
+	// Value has been analyzed for the pixel stage.
 	AnalyzedForStagePixel = 1 << 1,
 
-	/** Value has been analyzed for the compute stage. */
+	// Value has been analyzed for the compute stage.
 	AnalyzedForStageCompute = 1 << 2,
 
-	/** Mask that includes all "analyzed in stage" flags. */
+	// Mask that includes all "analyzed in stage" flags.
 	AnalyzedInAnyStageMask = AnalyzedForStageVertex | AnalyzedForStagePixel | AnalyzedForStageCompute,
 
-	/** Whether to skip substution of abstract tokens like "<PREV>". Only true for custom nodes. */
+	// Whether to skip substution of abstract tokens like "<PREV>". Only true for custom nodes.
 	SubstituteTagsInInlineHLSL = 1 << 3,
 
-	/** Whether 'FInlineHLSL::Code' field is used. Otherwise, 'ExternalCodeDeclaration' is used. */
+	// Whether 'FInlineHLSL::Code' field is used. Otherwise, 'ExternalCodeDeclaration' is used.
 	HasDynamicHLSLCode = 1 << 4, 
 
-	/** External code declaration must emit the secondary definition for the DDX derivative. */
+	// External code declaration must emit the secondary definition for the DDX derivative.
 	DerivativeDDX = 1 << 5,
 
-	/** External code declaration must emit the secondary definition for the DDY derivative. */
+	// External code declaration must emit the secondary definition for the DDY derivative.
 	DerivativeDDY = 1 << 6,
 };
 
@@ -69,32 +61,28 @@ ENUM_CLASS_FLAGS(EValueFlags);
 
 static_assert(uint32(1 << NumStages) - 1 == uint32(EValueFlags::AnalyzedInAnyStageMask), "Discrepancy between the number of stages in NumStages and the flags in EValueFlags.");
 
-/** 
- * @brief A collection of graph properties used by a value.
- * 
- * If a graph property is true, it means that either the value itself makes has of that
- * property, or one of its dependencies (direct or indirect) has it. This entails these
- * flags are automatically propagated to all dependant values (values that depend on a given one).
- * As an example, if "ReadsPixelNormal" is true on a specific value it means that either
- * that value itself or some other upstream value that value is dependent on reads
- * the pixel normal.
- */
+// A collection of graph properties used by a value.
+//
+// If a graph property is true, it means that either the value itself makes has of that
+// property, or one of its dependencies (direct or indirect) has it. This entails these
+// flags are automatically propagated to all dependant values (values that depend on a given one).
+// As an example, if "ReadsPixelNormal" is true on a specific value it means that either
+// that value itself or some other upstream value that value is dependent on reads
+// the pixel normal.
 enum class EGraphProperties : uint8
 {
 	None = 0,
 
-	/** Some value reads the pixel normal. */
+	// Some value reads the pixel normal.
 	ReadsPixelNormal = 1 << 0,
 };
 
 ENUM_CLASS_FLAGS(EGraphProperties);
 
-/**
- * @brief Enumeration of all different structs deriving from FValue. Used for dynamic casting.
- */
+// Enumeration of all different structs deriving from FValue. Used for dynamic casting.
 enum EValueKind
 {
-	/* Values */
+	// Values
 
 	VK_Poison,
 	VK_Constant,
@@ -102,12 +90,12 @@ enum EValueKind
 	VK_TextureObject,
 	VK_UniformParameter,
 
-	/* Instructions */
+	// Instructions
 
 	VK_InstructionBegin,
 
-	VK_Dimensional,
 	VK_SetMaterialOutput,
+	VK_Dimensional,
 	VK_Operator,
 	VK_Branch,
 	VK_Subscript,
@@ -120,173 +108,149 @@ enum EValueKind
 	VK_InstructionEnd,
 };
 
-/** @return The string representation of given value kind. */
+// Returns the string representation of given value kind.
 const TCHAR* LexToString(EValueKind Kind);
 
-/* Values */
-
-/**
- * @brief Base entity of all IR graph nodes.
- * 
- * An IR module is a graph of values, connected by their "uses" relations. The graph of IR
- * values is built by the `MaterialIRModuleBuilder` as the result of crawling through and 
- * analyizing the MaterialExpression graph contained in the translated Material. During 
- * this processing, IR values are emitted, and linked together. After the graph is 
- * constructed, it is itself analyzed: the builder will call `FMaterialIRValueAnalyzer::Analyze()`
- * in each active* (i.e. truly used) value in the graph, making sure a value is analyzed only 
- * after its dependencies have been analyzed.
- * A few notes:
- *  - FValues are automatically zero-initialized.
- *  - FValues are intended to be simple and inert data records. They cannot have non-trivia
- *    ctor, dtor or copy operators.
- *  - The ModuleBuilder relies on this property to efficiently hashing values so that it
- *    will reuse the same value instance instead of creating multiple instances of the
- *    same computation (for a more efficient output shader).
- *  - All values have a MIR type.
- *  - Pure FValue instances are values that do not have other dependencies (called "uses").
- *  - If a value has some other value as dependency, it means that it is the result of a
- *    calculation on those values. Values that have dependencies are Instructions (they
- *    derive from FInstruction).
- */
+// Values
+
+// Base entity of all IR graph nodes.
+//
+// An IR module is a graph of values, connected by their "uses" relations. The graph of IR
+// values is built by the `MaterialIRModuleBuilder` as the result of crawling through and
+// analyizing the MaterialExpression graph contained in the translated Material. During
+// this processing, IR values are emitted, and linked together. After the graph is
+// constructed, it is itself analyzed: the builder will call `FMaterialIRValueAnalyzer::Analyze()`
+// in each active* (i.e. truly used) value in the graph, making sure a value is analyzed only
+// after its dependencies have been analyzed.
+// A few notes:
+// - FValues are automatically zero-initialized.
+// - FValues are intended to be simple and inert data records. They cannot have non-trivia
+//   ctor, dtor or copy operators.
+// - The ModuleBuilder relies on this property to efficiently hashing values so that it
+//   will reuse the same value instance instead of creating multiple instances of the
+//   same computation (for a more efficient output shader).
+// - All values have a MIR type.
+// - Pure FValue instances are values that do not have other dependencies (called "uses").
+// - If a value has some other value as dependency, it means that it is the result of a
+//   calculation on those values. Values that have dependencies are Instructions (they
+//   derive from FInstruction).
 struct FValue
 {
-	/** Used to discern the concrete C++ type of this value (e.g. Subscript) */
+	// Used to discern the concrete C++ type of this value (e.g. Subscript)
 	EValueKind Kind : 8;
 
-	/** Set of fundamental flags true for this value. */
+	// Set of fundamental flags true for this value.
 	EValueFlags Flags;
 
-	/** The set of properties that are true for this value. Some flags might have been
-	 * set but some upstream dependency leading to this value and not this value directly.
-	 * See `EGraphProperties` for more information. */
+	// The set of properties that are true for this value. Some flags might have been
+	// set but some upstream dependency leading to this value and not this value directly.
+	// See `EGraphProperties` for more information.
 	EGraphProperties GraphProperties;
 
-	/** The runtime type this value has. */
+	// The runtime type this value has.
 	const FType* Type;
 
-	/** @return Whether this value has been analyzed for specified stage. */
+	// Returns whether this value has been analyzed for specified stage.
 	bool IsAnalyzed(EStage State) const;
 
-	/** @return Whether specified flags are true for this value. */
+	// Returns whether specified flags are true for this value.
 	bool HasFlags(EValueFlags InFlags) const;
 
-	/** Enables the specified flags without affecting others. */
+	// Enables the specified flags without affecting others.
 	void SetFlags(EValueFlags InFlags);
 
-	/** Disables the specified flags without affecting others. */
+	// Disables the specified flags without affecting others.
 	void ClearFlags(EValueFlags InFlags);
 
-	/** @return Whether specified properties are true for this value. */
+	// Returns whether specified properties are true for this value.
 	bool HasSubgraphProperties(EGraphProperties Properties) const;
 
-	/** Enables specified properties for this value. */
+	// Enables specified properties for this value.
 	void UseSubgraphProperties(EGraphProperties Properties);
 
-	/** @return The size in bytes of this value instance. */
+	// Returns the size in bytes of this value instance.
 	uint32 GetSizeInBytes() const;
-	
-	/** @return Whether this value is of specified kind. */
+
+	// Returns whether this value is of specified kind.
 	bool IsA(EValueKind InKind) const;
 
-	/**
-	 * Gets the immutable array of all this value uses.
-	 * 
-	 * An use is another value referenced by this one (e.g. the operands of a binary expression).
-	 * 
-	 * @return The immutable array of uses.
-	 */
+	// Gets the immutable array of all this value uses.
+	// An use is another value referenced by this one (e.g. the operands of a binary expression).
+	// Returns the immutable array of uses.
 	TConstArrayView<FValue*> GetUses() const;
-	
-	/**
-	 * Gets the immutable array of this value uses filtered for a specific stage stage.
-	 * 
-	 * @return The immutable array of uses.
-	 */
+
+	// Gets the immutable array of this value uses filtered for a specific stage stage.
+	//
+	// Returns the immutable array of uses.
 	TConstArrayView<FValue*> GetUsesForStage(MIR::EStage Stage) const;
 
-	/** @return Whether this value is poison. */
+	// Returns whether this value is poison.
 	bool IsPoison() const;
 
-	/** @return Whether this value exactly equals Other. */
+	// Returns whether this value exactly equals Other.
 	bool Equals(const FValue* Other) const;
 
-	/** @return Whether this value is a scalar (its type is Primitive with exactly 1 component). */
+	// Returns whether this value is a scalar (its type is Primitive with exactly 1 component).
 	bool IsScalar() const;
 
-	/** @return Whether this value is a vector (tis type is Primitive with 1-4 rows and exactly 1 column). */
+	// Returns whether this value is a vector (tis type is Primitive with 1-4 rows and exactly 1 column).
 	bool IsVector() const;
 
-	/** @return Whether this value is a constant boolean with value true. */
+	// Returns whether this value is a constant boolean with value true.
 	bool IsTrue() const;
 
-	/** @return Whether this value is a constant boolean with value false. */
+	// Returns whether this value is a constant boolean with value false.
 	bool IsFalse() const;
 
-	/** @return Whether this value is boolean and all components are true. */
+	// Returns whether this value is boolean and all components are true.
 	bool AreAllTrue() const;
 
-	/** @return Whether this value is boolean and all components are false. */
+	// Returns whether this value is boolean and all components are false.
 	bool AreAllFalse() const;
 
-	/** @return Whether this value is arithmetic and all components are exactly zero. */
+	// Returns whether this value is arithmetic and all components are exactly zero.
 	bool AreAllExactlyZero() const;
 
-	/** @return Whether this value is arithmetic and all components are approximately zero. */
+	// Returns whether this value is arithmetic and all components are approximately zero.
 	bool AreAllNearlyZero() const;
 
-	/** @return Whether this value is arithmetic and all components are exactly one. */
+	// Returns whether this value is arithmetic and all components are exactly one.
 	bool AreAllExactlyOne() const;
 
-	/** @return Whether this value is arithmetic and all components are approximately one. */
+	// Returns whether this value is arithmetic and all components are approximately one.
 	bool AreAllNearlyOne() const;
-	
-	/** Tries to cast this value to specified type T and returns the casted pointer, if possible (nullptr otherwise). */
+
+	// Tries to cast this value to specified type T and returns the casted pointer, if possible (nullptr otherwise).
 	template <typename T>
 	T* As() { return IsA(T::TypeKind) ? static_cast<T*>(this) : nullptr; }
 
-	/** Tries to cast this value to specified type T and returns the casted pointer, if possible (nullptr otherwise). */
+	// Tries to cast this value to specified type T and returns the casted pointer, if possible (nullptr otherwise).
 	template <typename T>
 	const T* As() const { return IsA(T::TypeKind) ? static_cast<const T*>(this) : nullptr; }
 };
 
-/**
- * Tries to cast a value to a derived type.
- * 
- * @return  If specified value is not null, it tries to cast this value T and returns it.
- *          Otherwise, it returns null.
- */
+// Tries to cast a value to a derived type.
+// If specified value is not null, it tries to cast this value T and returns it. Otherwise, it returns null.
 template <typename T>
 T* As(FValue* Value)
 {
 	return Value && Value->IsA(T::TypeKind) ? static_cast<T*>(Value) : nullptr;
 }
 
-/**
- * Tries to cast a value to a derived type.
- * 
- * @return  If specified value is not null, it tries to cast this value T and returns it.
- *          Otherwise, it returns null.
- */
+// Tries to cast a value to a derived type.
+// If specified value is not null, it tries to cast this value T and returns it. Otherwise, it returns null.
 template <typename T>
 const T* As(const FValue* Value)
 {
 	return Value && Value->IsA(T::TypeKind) ? static_cast<const T*>(Value) : nullptr;
 }
 
-/**
- * Casts a value to an instruction.
- * 
- * @return If specified value is not null, it tries to cast this value to an instruction
- *         and returns it. Otherwise, it returns null.
- */
+// Casts a value to an instruction.
+// If specified value is not null, it tries to cast this value to an instruction and returns it. Otherwise, it returns null.
 FInstruction* AsInstruction(FValue* Value);
 
-/**
- * Casts a value to an instruction.
- * 
- * @return If specified value is not null, it tries to cast this value to an instruction
- *         and returns it. Otherwise, it returns null.
- */
+// Casts a value to an instruction.
+// If specified value is not null, it tries to cast this value to an instruction and returns it. Otherwise, it returns null.
 const FInstruction* AsInstruction(const FValue* Value);
 
 template <EValueKind TTypeKind>
@@ -295,57 +259,28 @@ struct TValue : FValue
 	static constexpr EValueKind TypeKind = TTypeKind;
 };
 
-/**
- * @brief A placeholder for an invalid value.
- * 
- * A poison value represents an invalid value. It is produced by the emitter when an
- * invalid operation is performed. Poison values can be passed as arguments to other operations,
- * but they are "contagious": any instruction emitted with a poison value as an argument
- * will itself produce a poison value.
- */
+// A placeholder for an invalid value.
+//
+// A poison value represents an invalid value. It is produced by the emitter when an
+// invalid operation is performed. Poison values can be passed as arguments to other operations,
+// but they are "contagious": any instruction emitted with a poison value as an argument
+// will itself produce a poison value.
 struct FPoison : TValue<VK_Poison>
 {
 	static FPoison* Get();
 };
 
-static inline bool IsAnyNonNullAPoison()
-{
-	return false;
-}
-
-/** @return Whether there is a poison in the values specified, which could be null. */
-template <typename... TTail>
-bool IsAnyNonNullAPoison(const FValue* Head, TTail... Tail)
-{
-	return (Head && Head->IsPoison()) || IsAnyNonNullAPoison(Tail...);
-}
-
-static inline bool IsAnyAPoison()
-{
-	return false;
-}
-
-/** @return Whether there is a poison in the values specified among non-null values */
-template <typename... TTail>
-bool IsAnyAPoison(const FValue* Head, TTail... Tail)
-{
-	check(Head && "A non-null value must have been provided.");
-	return Head->IsPoison() || IsAnyAPoison(Tail...);
-}
-
-/** The integer type used inside MIR. */
+// The integer type used inside MIR.
 using TInteger = int64_t;
 
-/** The floating point type used inside MIR. */
+// The floating point type used inside MIR.
 using TFloat = float;
 
-/**
- * @brief A constant value.
- * 
- * A constant represents a translation-time known scalar primitive value. Operations on
- * constant values can be folded by the builder, that is they can be evaluated statically
- * while the builder constructs the IR graph of an input material.
- */
+// A constant value.
+//
+// A constant represents a translation-time known scalar primitive value. Operations on
+// constant values can be folded by the builder, that is they can be evaluated statically
+// while the builder constructs the IR graph of an input material.
 struct FConstant : TValue<VK_Constant>
 {
 	union
@@ -355,16 +290,16 @@ struct FConstant : TValue<VK_Constant>
 		TFloat 		Float;
 	};
 
-	/** @return Whether this constant is a boolean. */
+	// Returns whether this constant is a boolean.
 	bool IsBool() const;
 
-	/** @return Whether this constant is an integer. */
+	// Returns whether this constant is an integer.
 	bool IsInteger() const;
-	
-	/** @return Whether this constant is a float. */
+
+	// Returns whether this constant is a float.
 	bool IsFloat() const;
 
-	/** @return The constant value of given type T. The type must be bool, integral or floating point. */
+	// Returns the constant value of given type T. The type must be bool, integral or floating point.
 	template <typename T>
 	T Get() const
 	{
@@ -387,9 +322,7 @@ struct FConstant : TValue<VK_Constant>
 	}
 };
 
-/**
- * @brief Enumeration of all supported material external inputs.
- */
+// Enumeration of all supported material external inputs.
 enum class EExternalInput
 {
 	None,
@@ -427,159 +360,148 @@ enum class EExternalInput
 	Count,
 };
 
-/** The number of texcoord external input groups (uv, ddx, ddy). */
+// The number of texcoord external input groups (uv, ddx, ddy).
 static constexpr int32 TexCoordGroups = 3;
 
-/** Maximum number of supported texture coordinates. */
+// Maximum number of supported texture coordinates.
 static constexpr int32 TexCoordMaxNum = 8;
 
-/** @return The string representation of given given external input. */
+// Returns the string representation of given given external input.
 const TCHAR* LexToString(EExternalInput Input);
 
-/** @return The given external input type. */
+// Returns the given external input type.
 const FType* GetExternalInputType(EExternalInput Id);
 
-/** Converts the given texture coordinate index to its corresponding external input mapping. */
+// Converts the given texture coordinate index to its corresponding external input mapping.
 MIR::EExternalInput TexCoordIndexToExternalInput(int32 TexCoordIndex);
 
-/** @return The index of the specified texture coordinates external input.  */
+// Returns the index of the specified texture coordinates external input.
 int32 ExternalInputToTexCoordIndex(EExternalInput Id);
 
-/** @return Whether the given external is a texture coordinate. */
+// Returns whether the given external is a texture coordinate.
 bool IsExternalInputTexCoord(EExternalInput Id);
 
-/** @return Whether the given external is a texture coordinate ddx. */
+// Returns whether the given external is a texture coordinate ddx.
 bool IsExternalInputTexCoordDdx(EExternalInput Id);
 
-/** @return Whether the given external is a texture coordinate ddy. */
+// Returns whether the given external is a texture coordinate ddy.
 bool IsExternalInputTexCoordDdy(EExternalInput Id);
 
-/** @return Whether the given external is a texture coordinate, a ddx or ddy. */
+// Returns whether the given external is a texture coordinate, a ddx or ddy.
 bool IsExternalInputTexCoordOrPartialDerivative(EExternalInput Id);
 
-/**
- * @brief It represents an external input (like texture coordinates) value.
- */
+// It represents an external input (like texture coordinates) value.
 struct FExternalInput : TValue<VK_ExternalInput>
 {
 	EExternalInput Id;
 };
 
-/**
- * @brief It represents a texture object value. 
- */
+// It represents a texture object value.
 struct FTextureObject : TValue<VK_TextureObject>
 {
-	/** The texture object. */
+	// The texture object.
 	UTexture* Texture;
 
-	/** The sampler type associated to this texture object. */
+	// The sampler type associated to this texture object.
 	EMaterialSamplerType SamplerType;
 
-	/** Index of this parameter in the uniform expression set.
-	 * @note This field is automatically set by the builder. */
+	// Index of this parameter in the uniform expression set.
+	// Note: This field is automatically set by the builder.
 	uint32 Analysis_UniformParameterIndex;
 };
 
-/**
- * @brief It represents a material uniform parameter.
- */
+// It represents a material uniform parameter.
 struct FUniformParameter: TValue<VK_UniformParameter>
 {
-	/** Index of the parameter as registered in the module. */
+	// Index of the parameter as registered in the module.
 	uint32 ParameterIdInModule;
 
-	/** Eventual sampler type to use to sample this parameter's texture (if it is one) */
+	// Eventual sampler type to use to sample this parameter's texture (if it is one)
 	EMaterialSamplerType SamplerType;
 
-	/** Index of this parameter in the uniform expression set
-	 * @note This field is automatically set by the builder. */
+	// Index of this parameter in the uniform expression set
+	// Note: This field is automatically set by the builder.
 	uint32 Analysis_UniformParameterIndex;
 };
 
 /*------------------------------------ Instructions ------------------------------------*/
 
-/**
- * @brief A block of instructions.
- * 
- * A block is a sequence of instructions in order of execution. Blocks are organized in a
- * tree-like structure, used to model blocks nested inside other blocks.
- */
+// A block of instructions.
+//
+// A block is a sequence of instructions in order of execution. Blocks are organized in a
+// tree-like structure, used to model blocks nested inside other blocks.
 struct FBlock
 {
-	/** This block's parent block, if any. If this is null, this is a *root block. */
+	// This block's parent block, if any. If this is null, this is a *root block.
 	FBlock* Parent;
 
-	/** The linked-list head of the instructions contained in this block.
-	 * Links are contained in the FInstruction::Next field. */
+	// The linked-list head of the instructions contained in this block.
+	// Links are contained in the FInstruction::Next field.
 	FInstruction* Instructions;
 
-	/** Depth of this block in the tree structure (root blocks have level zero). */
+	// Depth of this block in the tree structure (root blocks have level zero).
 	int32 Level;
 
-	/** Finds and returns the common block between this and Other, if any. It returns null
-	 * if no common block was found.
-	 * @note This has O(n) complexity, where n is the maximum depth of the tree structure. */
+	// Finds and returns the common block between this and Other, if any. It returns null
+	// if no common block was found.
+	// Note: This has O(n) complexity, where n is the maximum depth of the tree structure.
 	FBlock* FindCommonParentWith(MIR::FBlock* Other);
 };
 
-/**
- * @brief Base struct of an instruction value.
- * 
- * An instruction is a value in a well defined order of evaluation.
- * Instructions have a parent block and are organized in a linked list: the Next field
- * indicates the instruction that will execute immediately after this one.
- * Since the material shader has multiple stages, the same instruction can belong to two
- * different graphs of execution, which explains why all fields in this structure have
- * a different possible value per stage.
- * 
- * @note All fields in this struct are not expected to be set by the user when emitting 
- *       an instruction, and are instead automatically populated by the builder.
- */
+// Base struct of an instruction value.
+//
+// An instruction is a value in a well defined order of evaluation.
+// Instructions have a parent block and are organized in a linked list: the Next field
+// indicates the instruction that will execute immediately after this one.
+// Since the material shader has multiple stages, the same instruction can belong to two
+// different graphs of execution, which explains why all fields in this structure have
+// a different possible value per stage.
+//
+// Note: All fields in this struct are not expected to be set by the user when emitting
+//      an instruction, and are instead automatically populated by the builder.
 struct FInstruction : FValue
 {
-	/** The next instruction executing after this one in all stages of execution. */
+	// The next instruction executing after this one in all stages of execution.
 	FInstruction* Next[NumStages];
 
-	/** This instruction parent block in all stages of execution. */
+	// This instruction parent block in all stages of execution.
 	FBlock* Block[NumStages];
 
-	/** How many users (i.e., dependencies) this instruction has in each stage of execution. */
+	// How many users (i.e., dependencies) this instruction has in each stage of execution.
 	uint32 NumUsers[NumStages];
 
-	/** The number of users processed during instruction graph linking in each stage of execution.
-	 * @note This information combined with NumUsers is used by the builder to push instructions
-	 *       in the appropriate block automatically. */
+	// The number of users processed during instruction graph linking in each stage of execution.
+	// Note: This information combined with NumUsers is used by the builder to push instructions
+	//      in the appropriate block automatically.
 	uint32 NumProcessedUsers[NumStages];
 
-	/** @return The block in which the dependency with specified index should execute. */
+	// Returns the block in which the dependency with specified index should execute.
 	FBlock* GetDesiredBlockForUse(EStage Stage, int32 UseIndex);
 };
 
 template <EValueKind TTypeKind, uint32 TNumStaticUses>
 struct TInstruction : FInstruction
 {
-	/** The kind of this instruction. */
+	// The kind of this instruction.
 	static constexpr EValueKind TypeKind = TTypeKind;
 
-	/** The number of values this instruction uses statically. Some instructions have a dynamic number of uses, in which case NumStaticUses is 0. */
+	// The number of values this instruction uses statically. Some instructions have a
+	// dynamic number of uses, in which case NumStaticUses is 0.
 	static constexpr uint32 NumStaticUses = TNumStaticUses;
 };
 
-/**
- * @brief An aggregate of other values.
- * 
- * A dimensional is a fixed array of other values. This value is used to model vectors and matrices.
- */
+// An aggregate of other values.
+//
+// A dimensional is a fixed array of other values. This value is used to model vectors and matrices.
 struct FDimensional : TInstruction<VK_Dimensional, 0>
 {
-	/** @return The constant array of component values.  */
+	// Returns the constant array of component values.
 	TConstArrayView<FValue*> GetComponents() const;
 
-	/** @return The mutable array of component values.  */
+	// Returns the mutable array of component values.
 	TArrayView<FValue*> GetMutableComponents();
 
-	/** @return Whether all components are constant (i.e., they're instances of FConstant). */
+	// Returns whether all components are constant (i.e., they're instances of FConstant).
 	bool AreComponentsConstant() const;
 };
 
@@ -589,61 +511,55 @@ struct TDimensional : FDimensional
 	FValue* Components[TDimension];
 };
 
-/**
- * @brief Instruction that sets material attribute (i.e., "BaseColor") to its value.
- */
+// Instruction that sets material attribute (i.e., "BaseColor") to its value.
 struct FSetMaterialOutput : TInstruction<VK_SetMaterialOutput, 1>
 {
-	/** The value this material attribute should be set to. */
+	// The value this material attribute should be set to.
 	FValue* Arg;
 
-	/** The material attribute to set. */
+	// The material attribute to set.
 	EMaterialProperty Property;
 };
 
-/**
- * @brief Inline HLSL instruction.
- * 
- * A value that is the result of an arbitrary snippet of HLSL code.
- * 
- * @note See BaseMaterialExpressions.ini file. 
- */
+// Inline HLSL instruction.
+//
+// A value that is the result of an arbitrary snippet of HLSL code.
+//
+// Note: See BaseMaterialExpressions.ini file.
 struct FInlineHLSL : TInstruction<VK_InlineHLSL, 0>
 {
 	static constexpr int32 MaxNumArguments = 16;
 
-	/** Array of argument values. */
+	// Array of argument values.
 	FValue* Arguments[MaxNumArguments];
 
-	/** Number of arguments from another expression node. */
+	// Number of arguments from another expression node.
 	int32 NumArguments;
 
 	union
 	{
-		/** The declaration this instruction refers to. */
+		// The declaration this instruction refers to.
 		const FMaterialExternalCodeDeclaration* ExternalCodeDeclaration;
 
-		/** The arbitrary inlined HLSL code snippet */
+		// The arbitrary inlined HLSL code snippet
 		const TCHAR* Code;
 	};
 };
 
-/**
- * @brief Operator enumeration.
- * 
- * @note If you modify this enum, update the implementations of the helper functions below.
- */
+// Operator enumeration.
+//
+// Note: If you modify this enum, update the implementations of the helper functions below.
 enum EOperator
 {
 	O_Invalid,
 
-	/* Unary */
+	// Unary
 	UO_FirstUnaryOperator,
 
 	UO_Negate = UO_FirstUnaryOperator, 	// Arithmetic negation: negate(5) -> -5
 	UO_Not,								// Logical negation: not(true) -> false
-	
-	/* Unary Intrinsic Operators */
+
+	// Unary Intrinsic Operators
 	UO_Abs,
 	UO_ACos,
 	UO_ASin,
@@ -665,7 +581,7 @@ enum EOperator
 	UO_Tan,
 	UO_Truncate,
 
-	/* Binary Comparison Operators */
+	// Binary Comparison Operators
 	BO_FirstBinaryOperator,
 
 	BO_GreaterThan = BO_FirstBinaryOperator,
@@ -675,7 +591,7 @@ enum EOperator
 	BO_Equals,
 	BO_NotEquals,
 
-	/* Binary Arithmetic Operators */
+	// Binary Arithmetic Operators
 	BO_Add,
 	BO_Subtract,
 	BO_Multiply,
@@ -684,15 +600,15 @@ enum EOperator
 	BO_Max,
 	BO_Min,
 
-	/* Binary Logical Operators */
-	BO_Or,
+	// Binary Logical Operators
 	BO_And,
+	BO_Or,
 
-	/* Binary Intrinsic Operators */
+	// Binary Intrinsic Operators
 	BO_Dot,
 	BO_Cross,
 
-	/* Ternary Arithmetic Operators */
+	// Ternary Arithmetic Operators
 	TO_FirstTernaryOperator,
 
 	TO_Clamp = TO_FirstTernaryOperator,
@@ -703,201 +619,181 @@ enum EOperator
 	OperatorCount,
 };
 
-/** Whether the given operator identifies a comparison operation (e.g., ">=", "=="). */
+// Whether the given operator identifies a comparison operation (e.g., ">=", "==").
 bool IsComparisonOperator(EOperator Op);
 
-/** Whether the given operator identifies a bit operation (e.g., "&", "|"). */
+// Whether the given operator identifies a bit operation (e.g., "&", "|").
 bool IsBitOperator(EOperator Op);
 
-/** Whether the given operator identifies a unary operator. */
+// Whether the given operator identifies a unary operator.
 bool IsUnaryOperator(EOperator Op);
 
-/** Whether the given operator identifies a binary operator. */
+// Whether the given operator identifies a binary operator.
 bool IsBinaryOperator(EOperator Op);
 
-/** Whether the given operator identifies a ternary operator. */
+// Whether the given operator identifies a ternary operator.
 bool IsTernaryOperator(EOperator Op);
 
-/** @return The string representation of the given operator. */
+// Returns the string representation of the given operator.
 const TCHAR* LexToString(EOperator Op);
 
-/**
- * @brief A mathematical operator instruction.
- * 
- * This instruction identifies a built-in operation on one, two or three argument values.
- */
+// A mathematical operator instruction.
+//
+// This instruction identifies a built-in operation on one, two or three argument values.
 struct FOperator : TInstruction<VK_Operator, 3>
 {
-	/** The first argument of the operation. This value is never null. */
+	// The first argument of the operation. This value is never null.
 	FValue* AArg;
 
-	/** The second argument of the operation. This value is null for unary operators. */
+	// The second argument of the operation. This value is null for unary operators.
 	FValue* BArg;
 
-	/** The third argument of the operation. This value is null for unary and binary operators. */
+	// The third argument of the operation. This value is null for unary and binary operators.
 	FValue* CArg;
 
-	/** It identifies which supported operation to carry. */
+	// It identifies which supported operation to carry.
 	EOperator Op;
 };
 
-/**
- * @brief A branch instruction.
- * 
- * This instruction evaluates to one or another argument based on whether a third boolean
- * condition argument is true or false. The builder will automatically place as many instructions
- * as possible in the true/false inner blocks whilst respecting dependency requirements.
- * This is done in an effort to avoid the unnecessary computation of the input value that was
- * not selected by the condition.
- */
+// A branch instruction.
+//
+// This instruction evaluates to one or another argument based on whether a third boolean
+// condition argument is true or false. The builder will automatically place as many
+// instructions as possible in the true/false inner blocks whilst respecting dependency
+// requirements. This is done in an effort to avoid the unnecessary computation of the
+// input value that was not selected by the condition.
 struct FBranch : TInstruction<VK_Branch, 3>
 {
-	/** The boolean condition argument used to forward the "true" or "false" argument. */
+	// The boolean condition argument used to forward the "true" or "false" argument.
 	FValue* ConditionArg;
 
-	/** Value this branch evaluates to when the condition is true. */
+	// Value this branch evaluates to when the condition is true.
 	FValue* TrueArg;
 
-	/** Value this branch evaluates to when the condition is false. */
+	// Value this branch evaluates to when the condition is false.
 	FValue* FalseArg;
 
-	/** The inner block (in each execution stage) the subgraph evaluating `TrueArg` should be placed in. */
+	// The inner block (in each execution stage) the subgraph evaluating `TrueArg` should be placed in.
 	FBlock TrueBlock[NumStages];
 
-	/** The inner block (in each execution stage) the subgraph evaluating `FalseArg` should be placed in. */
+	// The inner block (in each execution stage) the subgraph evaluating `FalseArg` should be placed in.
 	FBlock FalseBlock[NumStages];
 };
 
-/**
- * @brief A subscript instruction.
- * 
- * This instruction is used to pull the inner value making a compound one. For example,
- * it is used to extract an individual component of a vector value.
- */
+// A subscript instruction.
+//
+// This instruction is used to pull the inner value making a compound one. For example,
+// it is used to extract an individual component of a vector value.
 struct FSubscript : TInstruction<VK_Subscript, 1>
 {
-	/** The argument to subscript. */
+	// The argument to subscript.
 	FValue* Arg;
 
-	/** The subscript index, i.e. the index of the component to extract. */
+	// The subscript index, i.e. the index of the component to extract.
 	int Index;
 };
 
-/**
- * @brief A type cast instruction.
- * 
- * This instruction casts an argument value to a different type.
- */
+// A type cast instruction.
+//
+// This instruction casts an argument value to a different type.
 struct FCast : TInstruction<VK_Cast, 1>
 {
-	/** The argument value to cast. */
+	// The argument value to cast.
 	FValue* Arg;
 };
 
-/**
- * What texture gather mode to use in a texture read instruction (none indicates a sample).
- */
+// What texture gather mode to use in a texture read instruction (none indicates a sample).
 enum class ETextureReadMode
 {
-	/** Gather the four red components in a 2x2 pixel block. */
+	// Gather the four red components in a 2x2 pixel block.
 	GatherRed,
 
-	/** Gather the four green components in a 2x2 pixel block. */
+	// Gather the four green components in a 2x2 pixel block.
 	GatherGreen,
 
-	/** Gather the four blue components in a 2x2 pixel block. */
+	// Gather the four blue components in a 2x2 pixel block.
 	GatherBlue,
 
-	/** Gather the four alpha components in a 2x2 pixel block. */
+	// Gather the four alpha components in a 2x2 pixel block.
 	GatherAlpha,
 
-	/** Texture gather with automatically calculated mip level */
+	// Texture gather with automatically calculated mip level
 	MipAuto,
 
-	/** Texture gather with user specified mip level */
+	// Texture gather with user specified mip level
 	MipLevel,
 
-	/** Texture gather with automatically calculated mip level plus user specified bias */
+	// Texture gather with automatically calculated mip level plus user specified bias
 	MipBias,
 
-	/** Texture gather using automatically caluclated mip level based on user provided partial derivatives */
+	// Texture gather using automatically caluclated mip level based on user provided partial derivatives
 	Derivatives,
 };
 
-/** Returns the string representation of given mode. */
+// Returns the string representation of given mode.
 const TCHAR* LexToString(ETextureReadMode Mode);
 
-/**
- * @brief This instruction performs texture read operaation (sample or gather).
- */
+// This instruction performs texture read operaation (sample or gather).
 struct FTextureRead : TInstruction<VK_TextureRead, 5>
 {
-	/** The texture object to sample. */
+	// The texture object to sample.
 	FValue* TextureObject;
 
-	/** The texture coordinate at which to sample. */
+	// The texture coordinate at which to sample.
 	FValue* TexCoord;
 
-	/** Optional. The mip index to sample, if any provided. */
+	// Optional. The mip index to sample, if any provided.
 	FValue* MipValue;
 
-	/** Optional. The analytical partial derivative of the coordinates along the X axis. */
+	// Optional. The analytical partial derivative of the coordinates along the X axis.
 	FValue* TexCoordDdx;
-		
-	/** Optional. The analytical partial derivative of the coordinates along the Y axis. */
+
+	// Optional. The analytical partial derivative of the coordinates along the Y axis.
 	FValue* TexCoordDdy;
 
-	/** The mip value mode to use for sampling. */
+	// The mip value mode to use for sampling.
 	ETextureReadMode Mode : 8;
 
-	/** The sampler source mode to use for sampling. */
+	// The sampler source mode to use for sampling.
 	ESamplerSourceMode SamplerSourceMode : 8;
 
-	/** The sampler type to use for sampling. */
+	// The sampler type to use for sampling.
 	EMaterialSamplerType SamplerType : 8;
 };
 
-/**
- * @brief Utility value for selecting a different value based on the execution stage.
- */
+// Utility value for selecting a different value based on the execution stage.
 struct FStageSwitch : TInstruction<VK_StageSwitch, 1>
 {
-	/** The argument for to be bypassed in each stage. e*/
+	// The argument for to be bypassed in each stage. e
 	FValue* Args[NumStages];
 
-	/**
-	 * Use the specified value argument in the pixel stage, and another specified
-	 * argument for other stages.
-	 */
+	// Use the specified value argument in the pixel stage, and another specified
+	// argument for other stages.
 	void SetArgs(FValue* PixelStageArg, FValue* OtherStagesArg);
 };
 
-/**
- * @brief Specifies the axis for computing screen-space derivatives.
- *
- * This enum is used to indicate whether a partial derivative should be taken 
- * along the X or Y screen-space direction, corresponding to HLSL's ddx and ddy
- * functions.
- */
+// Specifies the axis for computing screen-space derivatives.
+//
+// This enum is used to indicate whether a partial derivative should be taken
+// along the X or Y screen-space direction, corresponding to HLSL's ddx and ddy
+// functions.
 enum class EDerivativeAxis
 {
 	X, // Corresponding to d/dx
 	Y, // Corresponding to d/dy
 };
 
-/**
- * @brief Instruction that maps to hardware ddx()/ddy().
- * 
- * @remark This is only available in stages that support hardware derivatives (e.g. pixel shader)
- */
+// Instruction that maps to hardware ddx()/ddy().
+//
+// Note: This is only available in stages that support hardware derivatives (e.g. pixel shader)
 struct FHardwarePartialDerivative : TInstruction<VK_HardwarePartialDerivative, 1>
 {
-	/** The value argument of ddx()/ddy() */
+	// The value argument of ddx()/ddy()
 	FValue* Arg;
 
-	/** The direction of partial derivative */
+	// The direction of partial derivative
 	EDerivativeAxis Axis;
 };
 
 } // namespace MIR
-#endif
+#endif // WITH_EDITOR
diff --git a/Engine/Source/Runtime/Engine/Public/Materials/MaterialIRCommon.h b/Engine/Source/Runtime/Engine/Public/Materials/MaterialIRCommon.h
--- a/Engine/Source/Runtime/Engine/Public/Materials/MaterialIRCommon.h
+++ b/Engine/Source/Runtime/Engine/Public/Materials/MaterialIRCommon.h
@@ -24,7 +24,6 @@ struct FMaterialInputDescription;
 struct FMaterialInsights;
 struct FMaterialIRModuleBuilder;
 struct FMaterialIRModuleBuilderImpl;
-struct FMaterialIRStateTracker;
 struct FShaderCompilerEnvironment;
 struct FStaticParameterSet;
 
diff --git a/Engine/Source/Runtime/Engine/Public/Materials/MaterialIREmitter.h b/Engine/Source/Runtime/Engine/Public/Materials/MaterialIREmitter.h
--- a/Engine/Source/Runtime/Engine/Public/Materials/MaterialIREmitter.h
+++ b/Engine/Source/Runtime/Engine/Public/Materials/MaterialIREmitter.h
@@ -7,7 +7,7 @@
 
 #if WITH_EDITOR
 
-#define UE_MIR_CHECKPOINT() if (Emitter.CurrentExpressionHasErrors()) return
+#define UE_MIR_CHECKPOINT(Emitter) if (Emitter.CurrentExpressionHasErrors()) return
 
 namespace MIR {
 
@@ -16,15 +16,19 @@ enum class EVectorComponent : uint8_t
 	X, Y, Z, W, 
 };
 
-// Return the lower case string representation of specified component (e.g. "x")
+// Return the lower case string representation of given component (e.g. "x")
 const TCHAR* VectorComponentToString(EVectorComponent);
 
-//
+// This utility data structure is used to define how a swizzle operation should be performed.
 struct FSwizzleMask
 {
+	// Which component should be extracted from the argument, in order. E.g. [Z, Z, X] to model "MyVec.zzx".
 	EVectorComponent Components[4];
+
+	// How many components have been defined for swizzle should be made of (maximum four).
 	int NumComponents{};
 
+	// Convenience ".xyz" swizzle mask.
 	static FSwizzleMask XYZ();
 
 	FSwizzleMask() {}
@@ -33,311 +37,504 @@ struct FSwizzleMask
 	FSwizzleMask(EVectorComponent X, EVectorComponent Y, EVectorComponent Z);
 	FSwizzleMask(EVectorComponent X, EVectorComponent Y, EVectorComponent Z, EVectorComponent W);
 
+	// Pushes a component to this swizzle mask.
 	void Append(EVectorComponent Component);
+
 	const EVectorComponent* begin() const { return Components; }
 	const EVectorComponent* end() const { return Components + NumComponents; }
 };
 
+// A lightweight wrapper around FValue* that also optionally tracks the source expression input.
+// Used to carry both value and its origin in the expression graph.
+// A null reference evaluates has the FValue* set to nullptr and evaluates false. It is
+// returned by the FEmitter::TryInput() value to indicate a missing value. An invalid reference
+// is either null or has a poison value. It is usually generated as the result of an invalid
+// operation.
+// Important: The Emitter APIs can safely take both null and poison values. Operations that
+// take a single argument simply return the argument if it is invalid (null or poison),
+// whereas operations that take multiple arguments return poison if any is invalid.
 //
-class FEmitter
+// When your intention is to test for whether a value is not-null (available), leverage the
+// operator bool()
+//
+//		if (Value) { /* handle case when value is available */ }
+//
+// Otherwise, if you would like the value to be available AND valid (not-poison) you should
+// test for IsValid()
+//
+//		if (Value.IsValid()) {
+//			/* handle case when the operations that produced the value were all succesful */
+//		}
+//
+// Note: In user Expression::Build() implementations, you should avoid testing values validity
+// directly and instead leverage the UE_MIR_CHECKPOINT() macro to verify that all emitter
+// performed so far were all valid, and thus the values you have computed so far are what
+// you expect.
+struct FValueRef
 {
-public:
-	/* Defaults */
-	
-	//
-	FEmitter& DefaultInputToFloatScalarZero(const FExpressionInput* Input);
+	// Default constructor. Creates an invalid reference.
+	FValueRef() {}
+
+	// Constructs from a value pointer.
+	FValueRef(FValue* Value) : Value{ Value } {}
+
+	// Constructs from a value pointer and its corresponding input source.
+	FValueRef(FValue* Value, const FExpressionInput* Input) : Value{ Value }, Input{ Input } {}
+
+	// Returns whether this reference is valid, that is it references to a non-null value.
+	operator bool() const { return Value != nullptr; }
+
+	// Returns the referenced value, which must be non-null.
+	operator FValue*() const { return Value; }
+
+	// Returns the referenced value, which must be non-null.
+	FValue* operator->() const { return Value; }
 	
-	//
-	FEmitter& DefaultInputToBoolScalar(const FExpressionInput* Input, bool Value);
+	// Checks if the reference holds a valid non-null value.
+	bool IsValid() const;
+
+	// Creates a new reference with the same input but a different value.
+	FValueRef To(FValue* Value) const;
 
-	//
-	FEmitter& DefaultInputToIntScalar(const FExpressionInput* Input, TInteger Value);
+	// Creates a poison version of the current reference (same input, but poison value).
+	FValueRef ToPoison() const;
 	
-	//
-	FEmitter& DefaultInputToFloatScalar(const FExpressionInput* Input, TFloat Value);
+	// The referenced value.
+	FValue* Value{};
+
+	// The expression input being read.
+	const FExpressionInput* Input{};
+};
+
+// The FEmitter is a helper object responsible for emitting MIR values and instructions.
+// It is primarily used within UMaterialExpression::Build() implementations to lower the
+// high-level semantics of a material expression into corresponding lower-level MIR nodes.
+//
+// During the emission process, the emitter automatically performs various
+// simplifications and constant folding optimizations where possible.
+//
+// An instance of the emitter is created and managed by the FMaterialIRModuleBuilder.
+// This instance is then passed to the UMaterialExpression::Build() functions
+// so that they can emit their corresponding IR values.
+//
+// IMPORTANT: Users should *not* assume the exact MIR::FType of the FValue* returned
+// by EmitXXX() functions. The emitter is free to optimize operations, which may
+// result in a returned value having a different type than naively expected (e.g.,
+// folding an operation into a constant).
+//
+// The emitter handles invalid inputs gracefully. If an operation is attempted
+// with invalid arguments (e.g., acting on arguments of unexpected types), the emitter
+// automatically reports an error and returns the MIR::FPoison value. Furthermore, all
+// EmitXXX() functions are robust to receiving MIR::FPoison as input; they will simply
+// propagate the poison value without attempting the operation.
+class FEmitter
+{
+public:
+	/*--------------------------------- Error handling ---------------------------------*/
+
+	// Returns whether the current expression generated an error.
+	// You can also use UE_MIR_CHECKPOINT() in your code to check this and automatically return if an error occurred.
+	bool CurrentExpressionHasErrors() const { return bCurrentExpressionHasErrors; }
+
+	// Reports an error using printf-like format and arguments to format the error message.
+	template <typename... TArgs>
+	void Errorf(FValueRef Source, UE::Core::TCheckedFormatString<FString::FmtCharType, TArgs...> Format, TArgs... Args)
+	{
+		Error(Source, FString::Printf(Format, Args...));
+	}
+
+	// Reports an error using printf-like format and arguments to format the error message.
+	template <typename... TArgs>
+	void Errorf(UE::Core::TCheckedFormatString<FString::FmtCharType, TArgs...> Format, TArgs... Args)
+	{
+		Error(FString::Printf(Format, Args...));
+	}
 
-	//
-	FEmitter& DefaultInputToFloat3(const FExpressionInput* Input, FVector3f Value);
+	// Reports an error with given error message.
+	void Error(FValueRef Source, FStringView Message);
 
-	//
-	FEmitter& DefaultInputToFloat4(const FExpressionInput* Input, FVector4f Value);
+	// Reports an error with given error message.
+	void Error(FStringView Message);
 
-	//
+	/*--------------------------------- Type handling ----------------------------------*/
 
-	/* Input */
+	// Tries to find a common type between A and B.
+	// The trivial case is when A and B are the same type, in which case that type is
+	// the common type. When A and B are primitive, the common type is defined as the type
+	// with the highest dimensions between the two and with a scalar kind that follows the
+	// progression: bool -> int -> float
+	// As an example, the common type between a float3 and a bool4 will be a float4.
+	// Note; that this may cause a loss of data when casting from int to float.
+	// It returns the common type if it exists, nullptr otherwise.
+	const FType* TryGetCommonType(const FType* A, const FType* B);
 
-	// Tries to get a value from specified input. If it has no value (for instance
-	// because it is not connected), it returns null.
-	FValue* TryInput(const FExpressionInput* Input);
+	// Gets the common type between A and B and reports an error if no such type exists.
+	// See TryGetCommonType()
+	const FType* GetCommonType(const FType* A, const FType* B);
 
-	// Gets and returns the value flowing into input with specified `Input`.
-	// If no value is flowing in (e.g. bacause it is not connected) it reports
-	// an error and returns a poison value.
-	FValue* Input(const FExpressionInput* Input);
+	/*-------------------------------- Input management --------------------------------*/
+
+	// Retrieves the value from the given input, if present. A null reference otherwise.
+	// Note: it does not report an error if the value is missing.
+	// It is safe to pass returned value to other Emitter functions as argument even if null.
+	// For instance, the following is safe and idiomatic
+	//     
+	//    FValueRef Value = Emitter.CheckIsArithmetic(Emitter.TryInput(&A));
+	//    UE_MIR_CHECKPOINT(Emitter);
+	//    if (Value) { /* handle case when value is provided */ }
+	// 
+	// Keep in mind that a reference can be null (its Value is null) or invalid (which means
+	// that it's null OR poison).
+	FValueRef TryInput(const FExpressionInput* Input);
+
+	// Retrieves the value flowing into given input.
+	// If the value is missing, it reports an error and returns poison.
+	FValueRef Input(const FExpressionInput* Input);
+
+	// If the input has no value flowing in (e.g., because the input is disconnected),
+	// these DefaultXXX() functions will emit a constant value, bind it to the input, and
+	// return it.
+	// Note that these functions may return poison if the value flowing into this input
+	// was poison in	FValueRef DefaultBool(FValueRef Value, bool Default);
+	FValueRef InputDefaultBool(const FExpressionInput* Input, bool Default);
+	FValueRef InputDefaultInt(const FExpressionInput* Input, TInteger Default);
+	FValueRef InputDefaultInt2(const FExpressionInput* Input, UE::Math::TIntVector2<TInteger> Default);
+	FValueRef InputDefaultInt3(const FExpressionInput* Input, UE::Math::TIntVector3<TInteger> Default);
+	FValueRef InputDefaultInt4(const FExpressionInput* Input, UE::Math::TIntVector4<TInteger> Default);
+	FValueRef InputDefaultFloat(const FExpressionInput* Input, TFloat Default);
+	FValueRef InputDefaultFloat2(const FExpressionInput* Input, UE::Math::TVector2<TFloat> Default);
+	FValueRef InputDefaultFloat3(const FExpressionInput* Input, UE::Math::TVector<TFloat> Default);
+	FValueRef InputDefaultFloat4(const FExpressionInput* Input, UE::Math::TVector4<TFloat> Default);
+
+	// Retrieves the value from the given input of given type kind.
+	// If the value is missing or its type does not have the given kind, it reports an error and returns poison.
+	FValueRef CheckTypeIsKind(FValueRef Value, ETypeKind Kind);
 	
-	// Gets and returns the value flowing into input with specified `Input`.
-	// Before returning the value, it checks that its type is of specified `Kind`.
-	FValue* InputOfType(const FExpressionInput* Input, ETypeKind Kind);
+	// Validates that the value is of a primitive type. If not, reports an error and returns poison.
+	FValueRef CheckIsPrimitive(FValueRef Value);
 
-	// Gets and returns the value flowing into input with specified `Input`.
-	// Before returning the value, it checks that the value is a constant boolean scalar.
-	bool InputConstantBoolScalar(const FExpressionInput* Input);
-	
-	// Gets and returns the value flowing into input with specified `Input`.
-	// Before returning the value, it checks that it is a boolean scalar.
-	FValue* InputBoolScalar(const FExpressionInput* Input);
+	// Validates that the value is of an arithmetic type (int or float). Reports error and returns poison otherwise.
+	FValueRef CheckIsArithmetic(FValueRef Value);
 
-	// Gets and returns the value flowing into input with specified `Input`.
-	// Before returning the value, it checks that it is a float scalar.
-	FValue* InputFloatScalar(const FExpressionInput* Input);
+	// Validates that the value is a scalar. If not, reports error and returns poison.
+	FValueRef CheckIsScalar(FValueRef Value);
 
-	// Gets and returns the value flowing into input with specified `Input`.
-	// Before returning the value, it checks that it is a scalar.
-	FValue* InputScalar(const FExpressionInput* Input);
+	// Validates that the value is either a scalar or a vector. Reports error and returns poison otherwise.
+	FValueRef CheckIsScalarOrVector(FValueRef Value);
 	
-	// Gets and returns the value flowing into input with specified `Input`.
-	// Before returning the value, it checks that it has a primitive type.
-	FValue* InputPrimitive(const FExpressionInput* Input);
+	// Retrieves the texture value flowing into given input.
+	// If the value is missing or cannot be cast to a texture, it reports an error and returns poison.
+	FValueRef CheckIsTexture(FValueRef Value);
 
-	// Gets and returns the value flowing into input with specified `Input`.
-	// Before returning the value, it checks that it is a float vector.
-	FValue* InputFloatVector(const FExpressionInput* Input, int32 NumComponents);
+	// Casts the value flowing into this input to a constant boolean and returns it.
+	// If the value is missing, or is not a constant boolean it reports an error and returns false.
+	bool ToConstantBool(FValueRef Value);
 
-	//
-	FValue* InputTexture(const FExpressionInput* Input);
-
-	/* Output */
+	/*-------------------------------- Output management -------------------------------*/
 	
-	// Flows specified `Value` out of output with specified `OutputIndex`.
-	FEmitter& Output(int32 OutputIndex, FValue* Value);
+	// Flows given `Value` out of the expression output with given `OutputIndex`.
+	FEmitter& Output(int32 OutputIndex, FValueRef Value);
 
-	// Flows specified `Value` out of specified expression `Output`.
-	FEmitter& Output(const FExpressionOutput* Output, FValue* Value);
+	// Flows given `Value` out of given expression `Output`.
+	FEmitter& Output(const FExpressionOutput* Output, FValueRef Value);
 	
-	/* Value manipulation */
+	/*------------------------------- Constants emission -------------------------------*/
 
-	// Checks that specified value is a constant bool scalar and returns its value.
-	// If value type is incorrect, it reports an error.
-	bool ToConstantBoolScalar(MIR::FValue* Value);
+	// Converts the given UE::Shader::FValue to a MIR constant value and returns it.
+	FValueRef ConstantFromShaderValue(const UE::Shader::FValue& InValue);
 
-	/* Error Checking */
+	// Emits a scalar constant value of given kind with value 0 and returns it.
+	FValueRef ConstantZero(EScalarKind Kind);
 
-	//
-	void CheckInputIsScalar(const FExpressionInput* Input, FValue* InputValue);
+	// Emits a scalar constant value of given kind with value 1 and returns it.
+	FValueRef ConstantOne(EScalarKind Kind);
+	
+	// Casts the given float value to the given scalar kind and returns its value.
+	FValueRef ConstantScalar(EScalarKind Kind, TFloat FromFloat);
 
-	//
-	void CheckInputIsScalar(const FExpressionInput* Input, FValue* InputValue, EScalarKind Kind);
+	// Returns the constant boolean scalar `true`.
+	FValueRef ConstantTrue();
 
-	//
-	void CheckInputIsVector(const FExpressionInput* Input, FValue* InputValue, int32 NumComponents, EScalarKind Kind);
+	// Returns the constant boolean scalar `false`.
+	FValueRef ConstantFalse();
 
-	// Checks that the type of the value `InputValue` flowing into `Input` is of
-	// specified type `Kind. If it isn't it reports an error. You may check whether
-	// an error occurred with `IsInvalid()`.
-	void CheckInputTypeIs(const FExpressionInput* Input, FValue* InputValue, ETypeKind Kind);
+	// Returns the given constant boolean scalar.
+	FValueRef ConstantBool(bool InX);
+
+	// Returns the given constant integer scalar.
+	FValueRef ConstantInt(TInteger InX);
 	
-	//
-	void CheckValueIsScalar(const FValue* Value, EScalarKind Kind);
-
-	//
-	void CheckValueIsTexture(const FValue* Value);
-
-	/* Constants */
-
-	FValue* EmitConstantFromShaderValue(const UE::Shader::FValue& InValue);
-	FValue* EmitConstantScalarZero(EScalarKind Kind);
-	FValue* EmitConstantScalarOne(EScalarKind Kind);
-	FValue* EmitConstantTrue();
-	FValue* EmitConstantFalse();
-	FValue* EmitConstantBoolScalar(bool InX);
-	FValue* EmitConstantIntScalar(TInteger InX);
-	FValue* EmitConstantFloatScalar(TFloat InX);
-	FValue* EmitConstantScalar(EScalarKind Kind, TFloat FromFloat);
-	FValue* EmitConstantFloat2(const FVector2f& InValue);
-	FValue* EmitConstantFloat3(const FVector3f& InValue);
-	FValue* EmitConstantFloat4(const FVector4f& InValue);
-	FValue* EmitConstantInt2(const FIntVector2& InValue);
-	FValue* EmitConstantInt3(const FIntVector3& InValue);
-	FValue* EmitConstantInt4(const FIntVector4& InValue);
-	FValue* EmitVector2(FValue* InX, FValue* InY);
-	FValue* EmitVector3(FValue* InX, FValue* InY, FValue* InZ);
-	FValue* EmitVector4(FValue* InX, FValue* InY, FValue* InZ, FValue* InW);
-	FValue* EmitExternalInput(EExternalInput Id);
-
-	/* Other Values */
-
-	FValue* EmitSubscript(FValue* Value, int32 ComponentIndex);
-	FValue* EmitSwizzle(FValue* Value, FSwizzleMask Mask);
-	FValue* EmitTextureObject(UTexture* Texture, EMaterialSamplerType SamplerType);
-	FValue* EmitParameter(FName Name, FMaterialParameterMetadata& Metadata, EMaterialSamplerType SamplerType = SAMPLERTYPE_Color);
-	FValue* EmitCustomPrimitiveData(uint32 PrimitiveDataIndex);
-
-	/**
-	 * Emits a StageSwitch value.
-	 * 
-	 * @param Type the type of the emitted value.
-	 * @param ValuePerStage the array of values for each MIR::EStage entry.
-	 * 
-	 * @remark ValuePerStage must have no more than MIR::NumStages elements.
-	 */
-	FValue* EmitStageSwitch(const FType* Type, TConstArrayView<FValue*> ValuePerStage);
-
-	/* Instructions */
-
-	FSetMaterialOutput* EmitSetMaterialOutput(EMaterialProperty InProperty, FValue* Arg);
-	FValue* EmitOperator(EOperator Operator, FValue* A, FValue* B = nullptr, FValue* C = nullptr);
-	FValue* EmitBranch(FValue* Condition, FValue* True, FValue* False);
-	FValue* EmitConstruct(const FType* Type, FValue* Initializer);
-
-	/**
-	 * @brief Emits a texture gather instruction.
-	 * 
-	 * @param Texture The texture to sample. Must be of type MIR::FObjectType and represent a valid texture.
-	 * @param TexCoord The UV coordinates at which to sample the texture.
-	 * @param GatherMode What channel to gather (any ::GatherXXX value).
-	 * @param SamplerSourceMode Which sampler source mode to use.
-	 * @param SamplerType Which sampler type to use.
-	 */
-	FValue* EmitTextureGather(FValue* Texture, FValue* TexCoord, ETextureReadMode GatherMode, ESamplerSourceMode SamplerSourceMode, EMaterialSamplerType SamplerType);
-
-	/**
-	 * @brief Emits a texture sample instruction.
-	 * 
-	 * @param Texture The texture to sample. Must be of type MIR::FObjectType and represent a valid texture.
-	 * @param TexCoord The UV coordinates at which to sample the texture.
-	 * @param bAddViewMaterialMipBias Whether to add the view material texture mip bias to the mip level.
-	 * @param SamplerSourceMode Which sampler source mode to use.
-	 * @param SamplerType Which sampler type to use.
-	 */
-	FValue* EmitTextureSample(FValue* Texture, FValue* TexCoord, bool bAddViewMaterialMipBias, ESamplerSourceMode SamplerSourceMode, EMaterialSamplerType SamplerType);
-
-	/**
-	 * @brief Emits a texture sample instruction with a manually specified mip level.
-	 * 
-	 * @param Texture The texture to sample. Must be of type MIR::FObjectType and represent a valid texture.
-	 * @param TexCoord The UV coordinates at which to sample the texture.
-	 * @param MipLevel Which mip level to use for this sample.
-	 * @param bAddViewMaterialMipBias Whether to add the view material texture mip bias to the mip level.
-	 * @param SamplerSourceMode Which sampler source mode to use.
-	 * @param SamplerType Which sampler type to use.
-	 */
-	FValue* EmitTextureSampleLevel(FValue* Texture, FValue* TexCoord, FValue* MipLevel, bool bAddViewMaterialMipBias, ESamplerSourceMode SamplerSourceMode, EMaterialSamplerType SamplerType);
-
-	/**
-	 * @brief Emits a texture sample instruction with bias added to the automatically selected mip level.
-	 * 
-	 * @param Texture The texture to sample. Must be of type MIR::FObjectType and represent a valid texture.
-	 * @param TexCoord The UV coordinates at which to sample the texture.
-	 * @param MipBias The bias to add to the automatically selected mip level.
-	 * @param bAddViewMaterialMipBias Whether to add the view material texture mip bias to the mip level.
-	 * @param SamplerSourceMode Which sampler source mode to use.
-	 * @param SamplerType Which sampler type to use.
-	 */
-	FValue* EmitTextureSampleBias(FValue* Texture, FValue* TexCoord, FValue* MipBias, bool bAddViewMaterialMipBias, ESamplerSourceMode SamplerSourceMode, EMaterialSamplerType SamplerType);
+	// Emits a constant integer 2D column vector and returns it.
+	FValueRef ConstantInt2(UE::Math::TIntVector2<TInteger> InValue);
 	
-	/**
-	 * @brief Emits a texture sample instruction using user provided partial derivatives.
-	 * 
-	 * @param Texture The texture to sample. Must be of type MIR::FObjectType and represent a valid texture.
-	 * @param TexCoord The UV coordinates at which to sample the texture.
-	 * @param TexCoordDdx The partial derivative of the texture coordinates along the screen space x axis.
-	 * @param TexCoordDdx The partial derivative of the texture coordinates along the screen space y axis.
-	 * @param bAddViewMaterialMipBias Whether to add the view material texture mip bias to the mip level.
-	 * @param SamplerSourceMode Which sampler source mode to use.
-	 * @param SamplerType Which sampler type to use.
-	 */
-	FValue* EmitTextureSampleGrad(FValue* Texture, FValue* TexCoord, FValue* TexCoordDdx, FValue* TexCoordDdy, bool bAddViewMaterialMipBias, ESamplerSourceMode SamplerSourceMode, EMaterialSamplerType SamplerType);
+	// Emits a constant integer 3D column vector and returns it.
+	FValueRef ConstantInt3(UE::Math::TIntVector3<TInteger>  InValue);
+
+	// Emits a constant integer 4D column vector and returns it.
+	FValueRef ConstantInt4(UE::Math::TIntVector4<TInteger>  InValue);
+
+	// Returns the given constant float scalar.
+	FValueRef ConstantFloat(TFloat InX);
+
+	// Emits a constant float 2D column vector and returns it.
+	FValueRef ConstantFloat2(UE::Math::TVector2<TFloat> InValue);
+
+	// Emits a constant float 3D column vector and returns it.
+	FValueRef ConstantFloat3(UE::Math::TVector<TFloat> InValue);
+
+	// Emits a constant float 4D column vector and returns it.
+	FValueRef ConstantFloat4(UE::Math::TVector4<TFloat> InValue);
 	
-	FValue* EmitInlineHLSL(const FType* Type, FString Code, TConstArrayView<FValue*> InputValues, EValueFlags ValueFlags = EValueFlags::None, EGraphProperties UsedGraphProperties = EGraphProperties::None);
-	FValue* EmitInlineHLSL(const FMaterialExternalCodeDeclaration* InExternalCodeDeclaration, TConstArrayView<FValue*> InputValues, EValueFlags ValueFlags = EValueFlags::None, EGraphProperties UsedGraphProperties = EGraphProperties::None);
-
-	/**
-	 * @brief Emits the partial derivative of given value.
-	 * 
-	 * This function differentiates the specified value with respect to the given axis.
-	 * If differentiation is not possible, it returns a poison value to indicate an error.
-	 * 
-	 * This API can be used regardless of the stages value will execute in. Hardware partial
-	 * derivatives will be used by default where available. If not available, the analytical
-	 * derivative will be computed instead.
-	 * 
-	 * @param Value The value to differentiate.
-	 * @param Axis The axis along which to compute the partial derivative (X for ddx, Y for ddy).
-	 * 
-	 * @return A new FValue* representing the partial derivative, or a poison value if invalid.
-	 */
-	FValue* EmitPartialDerivative(FValue* Value, EDerivativeAxis Axis);
-
-	/**
-	 * @brief Computes and emits the analytical partial derivative of a give value.
-	 * 
-	 * @param Value The value to differentiate.
-	 * @param Axis The axis along which to compute the partial derivative (X for ddx, Y for ddy).
-	 * 
-	 * @return A new FValue* representing the analytical derivative, or a poison value if invalid.
-	 */
-	FValue* EmitAnalyticalPartialDerivative(FValue* Value, EDerivativeAxis Axis);
-
-	/* Types */
-
-	/**
-	 * Tries to find a common type between A and B.
-	 * 
-	 * The trivial case is when A and B are the same type, in which case that type is
-	 * the common type. When A and B are primitive, the common type is defined as the type
-	 * with the highest dimensions between the two and with a scalar kind that follows the
-	 * progression:
-	 * 
-	 * 	bool -> int -> float
-	 * 
-	 * As an example, the common type between a float3 and a bool4 will be a float4.
-	 * 
-	 * Note that this may cause a loss of data when casting from int to float.
-	 */
-	const FType* TryGetCommonType(const FType* A, const FType* B);
+	/*--------------------- Other non-instruction values emission ---------------------*/
 
-	/**
-	 * Gets the common type between A and B and reports an error if no such type exists.
-	 * 
-	 * @remarks See `TryGetCommonType()`.
-	 */
-	const FType* GetCommonType(const FType* A, const FType* B);
+	// Emits the poison value.
+	FValueRef Poison();
 
+	// Emits an external input value of given input Id and returns it.
+	FValueRef ExternalInput(EExternalInput Id);
 
-	/* Error reporting */
+	// Emits a texture object value of fiven sampler type and returns it.
+	FValueRef TextureObject(UTexture* Texture, EMaterialSamplerType SamplerType);
 
-	bool CurrentExpressionHasErrors() const { return bCurrentExpressionHasErrors; }
+	// Emits a material parameter value with given name, metadata and sampler type (if it is a texture) and returns it.
+	FValueRef Parameter(FName Name, FMaterialParameterMetadata& Metadata, EMaterialSamplerType SamplerType = SAMPLERTYPE_Color);
 
-	template <typename... TArgs>
-	void Errorf(UE::Core::TCheckedFormatString<FString::FmtCharType, TArgs...> Format, TArgs... Args)
-	{
-		Error(FString::Printf(Format, Args...));
-	}
+	// Emits a custom primitive data value with given index.
+	FValueRef CustomPrimitiveData(uint32 PrimitiveDataIndex);
+	
+	/*----------------------------- Instructions emission -----------------------------*/
 
-	void Error(FStringView Message);
+	// Emits an instruction that sets the material attribute with given property to given argument.
+	// Note: this API is used by the builder and you should not use it manually.
+	FSetMaterialOutput* SetMaterialOutput(EMaterialProperty InProperty, FValue* Arg);
 
-	struct FPrivate;
+	// Emits a 2D column vector dimensional value and returns it.
+	FValueRef Vector2(FValueRef InX, FValueRef InY);
 
-private:
-	FEmitter() {}
+	// Emits a 3D column vector dimensional value and returns it.
+	FValueRef Vector3(FValueRef InX, FValueRef InY, FValueRef InZ);
 
-	void Initialize();
+	// Emits a 4D column vector dimensional value and returns it.
+	FValueRef Vector4(FValueRef InX, FValueRef InY, FValueRef InZ, FValueRef InW);
+	
+	// Emits a mathematical operation instruction with given operator and arguments.
+	// Note: this function will try to simplify the operation at translation time if possible.
+	// The returned value is therefore not guaranteed to be an FOperator instruction instance.
+	FValueRef Operator(EOperator Operator, FValueRef A, FValueRef B = {}, FValueRef C = {});
+	
+	// Emits a branch instruction.
+	// When the result of the Condition argument is true, the instruction will evaluate
+	// the True argument, otherwise the False argument.
+	// This instruction will place the as much of the other instructions whose results
+	// serve the computation of True or False within separate inner scopes, in order to
+	// avoid unnecessarily computing the inactive argument.
+	FValueRef Branch(FValueRef Condition, FValueRef True, FValueRef False);
+		
+	// Casts the given value to the specified target type. Returns poison if the cast is invalid.
+	FValueRef Cast(FValueRef Value, const FType* TargetType);
+	
+	// Casts the given arithmetic value to a scalar. If this is a vector, it will convert to its first component.
+	// Returns if the cast is invalid.
+	FValueRef CastToScalar(FValueRef Value);
+
+	// These functions cast the given value to the a primitive scalar or column vector type of given rows.
+	// If the value is missing ore could not be cast to the relative primitive type, they repors an error and return poison.
+	FValueRef CastToBool(FValueRef Value, int NumRows);
+	FValueRef CastToInt(FValueRef Value, int NumRows);
+	FValueRef CastToFloat(FValueRef Value, int NumRows);
+	
+	// Extracts the component with given index from the given argument value.
+	// The argument value must be a scalar or a vector primitive type, otherwise it reports an error and returns poison.
+	// Note that if no error occurred, the returned value is guaranteed to be a primitive scalar.
+	FValueRef Subscript(FValueRef Value, int32 Index);
+
+	// Swizzles the given argument by given mask.
+	// The argument value must be a scalar or vector type. If the operation is invalid on
+	// given argument and mask it reports an error and returns poison.
+	FValueRef Swizzle(FValueRef Value, FSwizzleMask Mask);
+
+	// Emits a stage switch instruction. Use this instruction to select a different value
+	// in each execution stage.
+	// Note: ValuePerStage must have no more than MIR::NumStages elements.
+	FValueRef StageSwitch(const FType* Type, TConstArrayView<FValueRef> ValuePerStage);
+
+	// Emits a texture gather instruction.
+	// 
+	// @param Texture The texture to sample. Must be of type MIR::FObjectType and represent a valid texture.
+	// @param TexCoord The UV coordinates at which to sample the texture.
+	// @param GatherMode What channel to gather (any ::GatherXXX value).
+	// @param SamplerSourceMode Which sampler source mode to use.
+	// @param SamplerType Which sampler type to use.
+	FValueRef TextureGather(FValueRef Texture, FValueRef TexCoord, ETextureReadMode GatherMode, ESamplerSourceMode SamplerSourceMode, EMaterialSamplerType SamplerType);
+
+	// Emits a texture sample instruction.
+	// 
+	// @param Texture The texture to sample. Must be of type MIR::FObjectType and represent a valid texture.
+	// @param TexCoord The UV coordinates at which to sample the texture.
+	// @param bAddViewMaterialMipBias Whether to add the view material texture mip bias to the mip level.
+	// @param SamplerSourceMode Which sampler source mode to use.
+	// @param SamplerType Which sampler type to use.
+	FValueRef TextureSample(FValueRef Texture, FValueRef TexCoord, bool bAddViewMaterialMipBias, ESamplerSourceMode SamplerSourceMode, EMaterialSamplerType SamplerType);
+
+	// Emits a texture sample instruction with a manually given mip level.
+	// 
+	// @param Texture The texture to sample. Must be of type MIR::FObjectType and represent a valid texture.
+	// @param TexCoord The UV coordinates at which to sample the texture.
+	// @param MipLevel Which mip level to use for this sample.
+	// @param bAddViewMaterialMipBias Whether to add the view material texture mip bias to the mip level.
+	// @param SamplerSourceMode Which sampler source mode to use.
+	// @param SamplerType Which sampler type to use.
+	FValueRef TextureSampleLevel(FValueRef Texture, FValueRef TexCoord, FValueRef MipLevel, bool bAddViewMaterialMipBias, ESamplerSourceMode SamplerSourceMode, EMaterialSamplerType SamplerType);
+
+	// Emits a texture sample instruction with bias added to the automatically selected mip level.
+	// 
+	// @param Texture The texture to sample. Must be of type MIR::FObjectType and represent a valid texture.
+	// @param TexCoord The UV coordinates at which to sample the texture.
+	// @param MipBias The bias to add to the automatically selected mip level.
+	// @param bAddViewMaterialMipBias Whether to add the view material texture mip bias to the mip level.
+	// @param SamplerSourceMode Which sampler source mode to use.
+	// @param SamplerType Which sampler type to use.
+	FValueRef TextureSampleBias(FValueRef Texture, FValueRef TexCoord, FValueRef MipBias, bool bAddViewMaterialMipBias, ESamplerSourceMode SamplerSourceMode, EMaterialSamplerType SamplerType);
+	
+	// Emits a texture sample instruction using user provided partial derivatives.
+	// 
+	// @param Texture The texture to sample. Must be of type MIR::FObjectType and represent a valid texture.
+	// @param TexCoord The UV coordinates at which to sample the texture.
+	// @param TexCoordDdx The partial derivative of the texture coordinates along the screen space x axis.
+	// @param TexCoordDdx The partial derivative of the texture coordinates along the screen space y axis.
+	// @param bAddViewMaterialMipBias Whether to add the view material texture mip bias to the mip level.
+	// @param SamplerSourceMode Which sampler source mode to use.
+	// @param SamplerType Which sampler type to use.
+	FValueRef TextureSampleGrad(FValueRef Texture, FValueRef TexCoord, FValueRef TexCoordDdx, FValueRef TexCoordDdy, bool bAddViewMaterialMipBias, ESamplerSourceMode SamplerSourceMode, EMaterialSamplerType SamplerType);
+	
+	// Emits the partial derivative of given value.
+	// This function differentiates the given value with respect to the given axis.
+	// If differentiation is not possible, it returns a poison value to indicate an error.
+	// This API can be used regardless of the stages value will execute in. Hardware partial
+	// derivatives will be used by default where available. If not available, the analytical
+	// derivative will be computed instead.
+	// @param Value The value to differentiate.
+	// @param Axis The axis along which to compute the partial derivative (X for ddx, Y for ddy).
+	// @return A new FValueRef representing the partial derivative, or a poison value if invalid.
+	FValueRef PartialDerivative(FValueRef Value, EDerivativeAxis Axis);
+
+	// Computes and emits the analytical partial derivative of a give value.
+	// @param Value The value to differentiate.
+	// @param Axis The axis along which to compute the partial derivative (X for ddx, Y for ddy).
+	// @return A new FValueRef representing the analytical derivative, or a poison value if invalid.
+	FValueRef AnalyticalPartialDerivative(FValueRef Value, EDerivativeAxis Axis);
+
+	// Emit a new HLSL function inline within the material shader.
+	// @param Type The expected output type of the HLSL code.
+	// @param Code The HLSL code to be emitted.
+	// @param InputValues An array of input values required by the HLSL code.
+	// @param ValueFlags Flags that specify additional properties for this value (default is None).
+	// @param UsedGraphProperties Properties used during graph processing (default is None).
+	FValueRef InlineHLSL(const FType* Type, FString Code, TConstArrayView<FValueRef> InputValues, EValueFlags ValueFlags = EValueFlags::None, EGraphProperties UsedGraphProperties = EGraphProperties::None);
+
+	// Emit a new HLSL function inline within the material shader based on an external code declaration.
+	// @param InExternalCodeDeclaration The external HLSL code declaration to use for emission.
+	// @param InputValues An array of input values required by the HLSL code.
+	// @param UsedGraphProperties Properties used during graph processing (default is None).
+	FValueRef InlineHLSL(const FMaterialExternalCodeDeclaration* InExternalCodeDeclaration, TConstArrayView<FValueRef> InputValues, EValueFlags ValueFlags = EValueFlags::None, EGraphProperties UsedGraphProperties = EGraphProperties::None);
+
+	/*--------------------------- Unary operators shortcuts ----------------------------*/
+
+	FValueRef Negate(FValueRef A) { return Operator(UO_Negate, A); }
+	FValueRef Not(FValueRef A) { return Operator(UO_Not, A); }
+	FValueRef Abs(FValueRef A) { return Operator(UO_Abs, A); }
+	FValueRef ACos(FValueRef A) { return Operator(UO_ACos, A); }
+	FValueRef ASin(FValueRef A) { return Operator(UO_ASin, A); }
+	FValueRef ATan(FValueRef A) { return Operator(UO_ATan, A); }
+	FValueRef Ceil(FValueRef A) { return Operator(UO_Ceil, A); }
+	FValueRef Cos(FValueRef A) { return Operator(UO_Cos, A); }
+	FValueRef Exponential(FValueRef A) { return Operator(UO_Exponential, A); }
+	FValueRef Exponential2(FValueRef A) { return Operator(UO_Exponential2, A); }
+	FValueRef Floor(FValueRef A) { return Operator(UO_Floor, A); }
+	FValueRef Frac(FValueRef A) { return Operator(UO_Frac, A); }
+	FValueRef Length(FValueRef A) { return Operator(UO_Length, A); }
+	FValueRef Logarithm(FValueRef A) { return Operator(UO_Logarithm, A); }
+	FValueRef Logarithm2(FValueRef A) { return Operator(UO_Logarithm2, A); }
+	FValueRef Logarithm10(FValueRef A) { return Operator(UO_Logarithm10, A); }
+	FValueRef Round(FValueRef A) { return Operator(UO_Round, A); }
+	FValueRef Saturate(FValueRef A) { return Operator(UO_Saturate, A); }
+	FValueRef Sin(FValueRef A) { return Operator(UO_Sin, A); }
+	FValueRef Sqrt(FValueRef A) { return Operator(UO_Sqrt, A); }
+	FValueRef Tan(FValueRef A) { return Operator(UO_Tan, A); }
+	FValueRef Truncate(FValueRef A) { return Operator(UO_Truncate, A); }
+
+	/*-------------------------- Binary operators shortcuts ----------------------------*/
+	
+	FValueRef GreaterThan(FValueRef A, FValueRef B) { return Operator(BO_GreaterThan, A, B); }
+	FValueRef GreaterThanOrEquals(FValueRef A, FValueRef B) { return Operator(BO_GreaterThanOrEquals, A, B); }
+	FValueRef LessThan(FValueRef A, FValueRef B) { return Operator(BO_LessThan, A, B); }
+	FValueRef LessThanOrEquals(FValueRef A, FValueRef B) { return Operator(BO_LessThanOrEquals, A, B); }
+	FValueRef Equals(FValueRef A, FValueRef B) { return Operator(BO_Equals, A, B); }
+	FValueRef NotEquals(FValueRef A, FValueRef B) { return Operator(BO_NotEquals, A, B); }
+	FValueRef Add(FValueRef A, FValueRef B) { return Operator(BO_Add, A, B); }
+	FValueRef Subtract(FValueRef A, FValueRef B) { return Operator(BO_Subtract, A, B); }
+	FValueRef Multiply(FValueRef A, FValueRef B) { return Operator(BO_Multiply, A, B); }
+	FValueRef Divide(FValueRef A, FValueRef B) { return Operator(BO_Divide, A, B); }
+	FValueRef Fmod(FValueRef A, FValueRef B) { return Operator(BO_Fmod, A, B); }
+	FValueRef Max(FValueRef A, FValueRef B) { return Operator(BO_Max, A, B); }
+	FValueRef Min(FValueRef A, FValueRef B) { return Operator(BO_Min, A, B); }
+	FValueRef And(FValueRef A, FValueRef B) { return Operator(BO_And, A, B); }
+	FValueRef Or(FValueRef A, FValueRef B) { return Operator(BO_Or, A, B); }
+	FValueRef Dot(FValueRef A, FValueRef B) { return Operator(BO_Dot, A, B); }
+	FValueRef Cross(FValueRef A, FValueRef B) { return Operator(BO_Cross, A, B); }
+
+	/*-------------------------- Ternary operators shortcuts ---------------------------*/
+
+	FValueRef Clamp(FValueRef A, FValueRef B, FValueRef C) { return Operator(TO_Clamp, A, B, C); }
+	FValueRef Lerp(FValueRef A, FValueRef B, FValueRef C) { return Operator(TO_Lerp, A, B, C); }
+	FValueRef Select(FValueRef A, FValueRef B, FValueRef C) { return Operator(TO_Select, A, B, C); }
+	FValueRef Smoothstep(FValueRef A, FValueRef B, FValueRef C) { return Operator(TO_Smoothstep, A, B, C); }
+
+	// Internal
+	struct FPrivate;
 
+private:
 	struct FValueKeyFuncs : DefaultKeyFuncs<FValue*>
 	{
 		static bool Matches(KeyInitType A, KeyInitType B);
 		static uint32 GetKeyHash(KeyInitType Key);
 	};
 
+private:
+	FEmitter() {}
+
+	// Initializes the emitter, called by the friend builder.
+	void Initialize();
+
+private:
+	// Pointer to the builder internal implementation.
 	FMaterialIRModuleBuilderImpl* BuilderImpl{};
+
+	// The material being translated.
 	UMaterial* Material{};
+
+	// The IR module being built.
 	FMaterialIRModule* Module{};
+
+	// The set of static parameter assignments to use during this translation.
 	const FStaticParameterSet* StaticParameterSet{};
+
+	// The current expression being built (set by the builder).
 	UMaterialExpression* Expression{};
+
+	// Whether the current expression has reported any error.
 	bool bCurrentExpressionHasErrors = false;
-	FValue* ConstantTrue;
-	FValue* ConstantFalse;
-	TSet<FValue*, FValueKeyFuncs> ValueSet;
+
+	// Global "true" constant.
+	FValue* TrueConstant{};
+
+	// Global "false" constant.
+	FValue* FalseConstant{};
+
+	// The set of all values previously emitted. It is used to avoid duplicating identical
+	// values, explointing the strict SSA data-flow paradigm of MIR to efficiently reuse
+	// calculations.
+	TSet<FValue*, FValueKeyFuncs> ValueSet{};
 
 	friend FMaterialIRModuleBuilder;
 	friend FMaterialIRModuleBuilderImpl;
diff --git a/Engine/Source/Runtime/Engine/Public/Materials/MaterialIRTypes.h b/Engine/Source/Runtime/Engine/Public/Materials/MaterialIRTypes.h
--- a/Engine/Source/Runtime/Engine/Public/Materials/MaterialIRTypes.h
+++ b/Engine/Source/Runtime/Engine/Public/Materials/MaterialIRTypes.h
@@ -66,6 +66,9 @@ struct FType
 
 	/// Returns this type upcast this type to PrimitiveType if it is one, otherwise nullptr.
 	const FPrimitiveType* AsPrimitive() const;
+	
+	/// Returns this type upcast this type to PrimitiveType if it is one, otherwise nullptr.
+	const FPrimitiveType* AsArithmetic() const;
 
 	/// Returns this type upcast this type to PrimitiveType if it's a scalar, otherwise nullptr.
 	const FPrimitiveType* AsScalar() const;
@@ -120,24 +123,56 @@ struct FPrimitiveType : FType
 	/// Number of columns (rows > 1 and columns == 1, this is a vector).
 	int NumColumns;
 
-	static const FPrimitiveType* GetBool1();
-	static const FPrimitiveType* GetInt1();
-	static const FPrimitiveType* GetFloat1();
+	// Returns the boolean scalar type.
+	static const FPrimitiveType* GetBool();
+	
+	// Returns the integer scalar type.
+	static const FPrimitiveType* GetInt();
+	
+	// Returns the floating point scalar type.
+	static const FPrimitiveType* GetFloat();
+	
+	// Returns the floating point 2D vector type.
 	static const FPrimitiveType* GetFloat2();
+	
+	// Returns the floating point 3D vector type.
 	static const FPrimitiveType* GetFloat3();
+	
+	// Returns the floating point 4D vector type.
 	static const FPrimitiveType* GetFloat4();
 
+	// Returns the scalar type with given kind.
 	static const FPrimitiveType* GetScalar(EScalarKind InScalarKind);
+
+	// Returns the column vector type with given kind and number of rows.
 	static const FPrimitiveType* GetVector(EScalarKind InScalarKind, int NumRows);
-	static const FPrimitiveType* GetMatrix(EScalarKind InScalarKind, int NumColumns, int NumRows);
+	
+	// Returns the primitive type with given kind and number of columns and rows.
 	static const FPrimitiveType* Get(EScalarKind InScalarKind, int NumRows, int NumColumns);
 
+	// Returns the number of components in this primitive type.
 	int  GetNumComponents() const { return NumRows * NumColumns; }
+
+	// Whether this primitive type is scalar.
 	bool IsScalar() const { return GetNumComponents() == 1; }
+	
+	// Whether this primitive type is a column vector.
 	bool IsVector() const { return NumRows > 1 && NumColumns == 1; }
+	
+	// Whether this primitive type is a matrix.
 	bool IsMatrix() const { return NumRows > 1 && NumColumns > 1; }
+	
+	// Whether this primitive type is arithmetic (it supports arithemtic operations like addition).
+	bool IsArithmetic() const { return ScalarKindIsArithmetic(ScalarKind); }
+
+	// Returns this primitive type with a different scalar kind.
+	const FPrimitiveType* WithScalarKind(EScalarKind InScalarKind) const;
+	
+	// Returns the scalar type with this type scalar kind
 	const FPrimitiveType* ToScalar() const;
-	const FPrimitiveType* ToScalarKind(EScalarKind InScalarKind) const;
+	
+	// Returns the vector type with this type scalar kind
+	const FPrimitiveType* ToVector(int NumRows) const;
 };
 
 struct FObjectType : FType
