Asset Manager
- Remove MissingDownloadParams error. We now get errors from async loading when trying to load packages that are not installed.

diff --git a/Engine/Source/Runtime/Engine/Classes/Engine/Experimental/StreamableManagerError.h b/Engine/Source/Runtime/Engine/Classes/Engine/Experimental/StreamableManagerError.h
--- a/Engine/Source/Runtime/Engine/Classes/Engine/Experimental/StreamableManagerError.h
+++ b/Engine/Source/Runtime/Engine/Classes/Engine/Experimental/StreamableManagerError.h
@@ -14,8 +14,6 @@ DECLARE_ERROR(PackageLoadCancelled, 2, StreamableManager, NSLOCTEXT("StreamableM
 DECLARE_ERROR(DownloadError, 3, StreamableManager, NSLOCTEXT("StreamableManager", "DownloadError", "Failed to download"));
 DECLARE_ERROR_ONEPARAM(PackageNameInvalid, 4, StreamableManager, NSLOCTEXT("StreamableManager", "PackageNameInvalid", "Found invalid package name {InvalidPackageName}"), FString, InvalidPackageName, TEXT("Unknown"));
 
-
-DECLARE_ERROR_ONEPARAM(MissingDownloadParams, 5, StreamableManager, NSLOCTEXT("StreamableManager", "MissingDownloadParams", "Missing download params for primary asset {PrimaryAssetName}"), FString, PrimaryAssetName, TEXT("Unknown"));
 DECLARE_ERROR(IoStoreNotFound, 6, StreamableManager, NSLOCTEXT("StreamableManager", "IoStoreNotFound", "IoStore did not load correctly."));
 DECLARE_ERROR_ONEPARAM(SyncLoadIncomplete, 7, StreamableManager, NSLOCTEXT("StreamableManager", "SyncLoadIncomplete", "Sync load did not complete correctly for {DebugName}."), FString, DebugName, TEXT("Unknown"));
 
diff --git a/Engine/Source/Runtime/Engine/Private/AssetManager.cpp b/Engine/Source/Runtime/Engine/Private/AssetManager.cpp
--- a/Engine/Source/Runtime/Engine/Private/AssetManager.cpp
+++ b/Engine/Source/Runtime/Engine/Private/AssetManager.cpp
@@ -1828,13 +1828,6 @@ TSharedPtr<FStreamableHandle> UAssetManager::ChangeBundleStateForPrimaryAssets(
 		{
 			FPlatformMisc::PumpEssentialAppMessages();
 
-			// If a request has ever been made with a download params, require a download for any subsequent requests.
-			// This can help prevent user error if downloads aren't consistently requested and ensures that the correct
-			// content remains pinned in the iostore download cache.
-			const bool bDownloadRequired = 
-				(NameData->PendingState.Handle && NameData->PendingState.Handle->GetDownloadCachePin().IsValid()) ||
-				(NameData->CurrentState.Handle && NameData->CurrentState.Handle->GetDownloadCachePin().IsValid());
-
 			// Iterate list of changes, compute new bundle set
 
 			// Use pending state if valid
@@ -1933,42 +1926,29 @@ TSharedPtr<FStreamableHandle> UAssetManager::ChangeBundleStateForPrimaryAssets(
 				DebugName << TEXT(")");
 			}
 
-			if (!LoadParams.DownloadParams && bDownloadRequired)
+			// Set pending bundle state before issuing the request, because it can be resolved immediately.
+			// OnAssetStateChangeCompleted requires this data to handle state transition correctly.
+			NameData->PendingState.BundleNames = NewBundleState;
+
+			FAssetManagerLoadParams AssetListLoadParams;
+			AssetListLoadParams.Priority = LoadParams.Priority;
+			if (LoadParams.DownloadParams)
 			{
-				// The user previously use download params for requests with this primary asset, so we expect them going forward.
-				// See comment above where bDownloadRequired is delcared.
-				FStreamableAsyncLoadParams StreamableLoadParams;
-				StreamableLoadParams.TargetsToStream = MoveTemp(PathsToLoad);
-				NewHandle = StreamableManager.CreateErrorHandle(
-					UE::UnifiedError::StreamableManager::MissingDownloadParams::MakeError(*PrimaryAssetId.ToString()),
-					MoveTemp(StreamableLoadParams), FString(DebugName), CopyTemp(Location));
+				AssetListLoadParams.DownloadParams = CopyTemp(*LoadParams.DownloadParams);
+			}
+
+			NewHandle = LoadAssetList(MoveTemp(PathsToLoad), MoveTemp(AssetListLoadParams), FString(DebugName), CopyTemp(Location));
+			if (NewHandle && !NewHandle->HasLoadCompleted())
+			{
+				// save handle so possible duplicate requests will not override this one
+				// set only if loading is not completed yet because completion callback might already reset the PendingState
+				NameData->PendingState.Handle = NewHandle;
+				NewHandle->BindCompleteDelegate(FStreamableDelegateWithHandle::CreateUObject(this, &UAssetManager::OnAssetStateChangeCompleted, PrimaryAssetId));
+				NewHandle->BindCancelDelegate(FStreamableDelegateWithHandle::CreateUObject(this, &UAssetManager::OnAssetStateChangeCompleted, PrimaryAssetId));
 			}
 			else
 			{
-				// Set pending bundle state before issuing the request, because it can be resolved immediately.
-				// OnAssetStateChangeCompleted requires this data to handle state transition correctly.
-				NameData->PendingState.BundleNames = NewBundleState;
-
-				FAssetManagerLoadParams AssetListLoadParams;
-				AssetListLoadParams.Priority = LoadParams.Priority;
-				if (LoadParams.DownloadParams)
-				{
-					AssetListLoadParams.DownloadParams = CopyTemp(*LoadParams.DownloadParams);
-				}
-
-				NewHandle = LoadAssetList(MoveTemp(PathsToLoad), MoveTemp(AssetListLoadParams), FString(DebugName), CopyTemp(Location));
-				if (NewHandle && !NewHandle->HasLoadCompleted())
-				{
-					// save handle so possible duplicate requests will not override this one
-					// set only if loading is not completed yet because completion callback might already reset the PendingState
-					NameData->PendingState.Handle = NewHandle;
-					NewHandle->BindCompleteDelegate(FStreamableDelegateWithHandle::CreateUObject(this, &UAssetManager::OnAssetStateChangeCompleted, PrimaryAssetId));
-					NewHandle->BindCancelDelegate(FStreamableDelegateWithHandle::CreateUObject(this, &UAssetManager::OnAssetStateChangeCompleted, PrimaryAssetId));
-				}
-				else
-				{
-					OnAssetStateChangeCompleted(NewHandle, PrimaryAssetId);
-				}
+				OnAssetStateChangeCompleted(NewHandle, PrimaryAssetId);
 			}
 
 			if (!NewHandle.IsValid())
