Texture Blueprint_GetBuiltTextureSize

diff --git a/Engine/Source/Runtime/Engine/Classes/Engine/Texture.h b/Engine/Source/Runtime/Engine/Classes/Engine/Texture.h
--- a/Engine/Source/Runtime/Engine/Classes/Engine/Texture.h
+++ b/Engine/Source/Runtime/Engine/Classes/Engine/Texture.h
@@ -1771,6 +1771,13 @@ public:
 		int32 IgnoredSizeZ;
 		GetBuiltTextureSize(TargetPlatformSettings,TargetPlatformControls,OutSizeX,OutSizeY,IgnoredSizeZ);
 	}
+	
+	/**
+	 * Get the dimensions of the largest mip of the texture when built, accounting for LODBias and other constraints
+	 *	 Returns the 3d size when built for the current platform.  Does not cause texture to get built.
+	 */
+	UFUNCTION(BlueprintCallable, meta = (DisplayName = "GetBuiltTextureSize"), Category = "Rendering|Texture")
+	FVector3f Blueprint_GetBuiltTextureSize();
 
 	/**
 	 * Serializes cooked platform data.
diff --git a/Engine/Source/Runtime/Engine/Private/Texture.cpp b/Engine/Source/Runtime/Engine/Private/Texture.cpp
--- a/Engine/Source/Runtime/Engine/Private/Texture.cpp
+++ b/Engine/Source/Runtime/Engine/Private/Texture.cpp
@@ -21,6 +21,7 @@
 #include "IImageWrapperModule.h"
 #include "Engine/Engine.h"
 #include "Interfaces/ITargetPlatform.h"
+#include "Interfaces/ITargetPlatformManagerModule.h"
 #include "AssetCompilingManager.h"
 #include "Engine/TextureLODSettings.h"
 #include "RenderUtils.h"
@@ -3958,7 +3959,8 @@ void UTexture::GetBuiltTextureSize(const ITargetPlatformSettings* TargetPlatform
 	int32 SizeX=0,SizeY=0,SizeZ=0;
 
 #if WITH_EDITORONLY_DATA
-	if (Source.IsValid())
+	// TargetPlatformControls not currently used
+	if ( Source.IsValid() && TargetPlatformSettings )
 	{
 		FIntPoint SourceSize = Source.GetLogicalSize();
 		SizeX = SourceSize.X;
@@ -4012,37 +4014,70 @@ void UTexture::GetBuiltTextureSize(const ITargetPlatformSettings* TargetPlatform
 			SizeZ = FMath::Max<int32>(SizeZ >> LODBiasNoCinematics, 1);
 		}
 	}
-#else // WITH_EDITORONLY_DATA
+#endif // WITH_EDITORONLY_DATA
 
-	// no Editor data
-	//	can't support a query about arbitrary target platform
-	//	must be a query of current running platform
+	if ( SizeX == 0 )
+	{
+		// no Editor data
+		//	can't support a query about arbitrary target platform
+		//	must be a query of current running platform
 
-	// note that using PlatformData size is not right in Editor
-	//	because of the different way that LODBias acts in Editor vs cook (as a mip selection, not a size change)
+		// note that using PlatformData size is not right in Editor
+		//	because of the different way that LODBias acts in Editor vs cook (as a mip selection, not a size change)
 
-	FTexturePlatformData ** PPlatformData = const_cast<UTexture *>(this)->GetRunningPlatformData();
-	if ( PPlatformData && *PPlatformData )
-	{
-		SizeX = (*PPlatformData)->SizeX;
-		SizeY = (*PPlatformData)->SizeY;
-		SizeZ = (*PPlatformData)->GetNumSlices();
-	}
-	else
-	{
-		SizeX = SizeY = SizeZ = 0;
+		FTexturePlatformData ** PPlatformData = const_cast<UTexture *>(this)->GetRunningPlatformData();
+		if ( PPlatformData && *PPlatformData )
+		{
+			SizeX = (*PPlatformData)->SizeX;
+			SizeY = (*PPlatformData)->SizeY;
+			SizeZ = (*PPlatformData)->GetNumSlices();
+		}
 	}
 
-#endif // WITH_EDITORONLY_DATA
-
 	OutSizeX = SizeX;
 	OutSizeY = SizeY;
 	OutSizeZ = SizeZ;
 }
+
 void UTexture::GetBuiltTextureSize(const ITargetPlatform* TargetPlatform, int32& OutSizeX, int32& OutSizeY, int32& OutSizeZ) const
 {
-	return GetBuiltTextureSize(TargetPlatform->GetTargetPlatformSettings(), TargetPlatform->GetTargetPlatformControls(), OutSizeX, OutSizeY, OutSizeZ);
+	const ITargetPlatformSettings* TargetPlatformSettings = nullptr;
+	const ITargetPlatformControls* TargetPlatformControls = nullptr;
+	
+	if ( TargetPlatform )
+	{
+		TargetPlatformSettings = TargetPlatform->GetTargetPlatformSettings();
+		TargetPlatformControls = TargetPlatform->GetTargetPlatformControls();
+	}
+
+	GetBuiltTextureSize(TargetPlatformSettings,TargetPlatformControls, OutSizeX, OutSizeY, OutSizeZ);
+}
+
+
+FVector3f UTexture::Blueprint_GetBuiltTextureSize()
+{
+	const ITargetPlatformSettings* TargetPlatformSettings = nullptr;
+	const ITargetPlatformControls* TargetPlatformControls = nullptr;
+
+	#if WITH_EDITOR
+	{
+		const ITargetPlatform* TargetPlatform = GetTargetPlatformManagerRef().GetRunningTargetPlatform();
+		if ( TargetPlatform )
+		{
+			TargetPlatformSettings = TargetPlatform->GetTargetPlatformSettings();
+			TargetPlatformControls = TargetPlatform->GetTargetPlatformControls();
+		}
+	}
+	#endif
+
+	// when called in non-Editor context, GetBuiltTextureSize will return the cooked data size in PlatformData
+
+	int32 OutSizeX=0, OutSizeY=0, OutSizeZ=0;
+	GetBuiltTextureSize(TargetPlatformSettings,TargetPlatformControls,OutSizeX,OutSizeY,OutSizeZ);
+
+	return FVector3f(OutSizeX,OutSizeY,OutSizeZ);
 }
+
 // this should not be called directly; it is called from TargetPlatform GetTextureFormats
 //	entry point API is GetPlatformTextureFormatNamesWithPrefix
 FName GetDefaultTextureFormatName( const ITargetPlatformSettings* TargetPlatformSettings, const ITargetPlatformControls* TargetPlatformControls, const UTexture* Texture, int32 LayerIndex, 
@@ -4099,7 +4134,7 @@ FName GetDefaultTextureFormatName( const ITargetPlatformSettings* TargetPlatform
 	if (!bNoCompression)
 	{
 		int32 SizeX,SizeY;
-		Texture->GetBuiltTextureSize(TargetPlatformSettings, TargetPlatformControls,SizeX,SizeY);
+		Texture->GetBuiltTextureSize(TargetPlatformSettings,TargetPlatformControls,SizeX,SizeY);
 
 	
 		// Don't compress textures smaller than the DXT block size.
