Add a function to create an FWeakObjectPtr from an FObjectKey.
Change UNetDriver::ReplicationChangeListMap to be keyed on FObjectKey instead of raw object pointers. Even though the value type contained a weak pointer, checking that for validity isn't a reliable proxy for the key being valid when objects can migrate with UE_WITH_REMOTE_OBJECT_HANDLE.

diff --git a/Engine/Source/Runtime/CoreUObject/Public/UObject/ObjectKey.h b/Engine/Source/Runtime/CoreUObject/Public/UObject/ObjectKey.h
--- a/Engine/Source/Runtime/CoreUObject/Public/UObject/ObjectKey.h
+++ b/Engine/Source/Runtime/CoreUObject/Public/UObject/ObjectKey.h
@@ -171,6 +171,21 @@ public:
 		return WeakPtr.GetEvenIfUnreachable();
 	}
 
+	/**
+	 * Create an FWeakObjectPtr that points to the same object as this key.
+	 * @return The new FWeakObjectPtr.
+	 */
+	FWeakObjectPtr GetWeakObjectPtr() const
+	{
+		FWeakObjectPtr WeakPtr;
+		WeakPtr.ObjectIndex = ObjectIndex;
+		WeakPtr.ObjectSerialNumber = ObjectSerialNumber;
+#if UE_WITH_REMOTE_OBJECT_HANDLE
+		WeakPtr.ObjectRemoteId = RemoteId;
+#endif
+		return WeakPtr;
+	}
+
 	/** Hash function */
 	[[nodiscard]] friend uint32 GetTypeHash(const FObjectKey& Key)
 	{
diff --git a/Engine/Source/Runtime/Engine/Classes/Engine/NetDriver.h b/Engine/Source/Runtime/Engine/Classes/Engine/NetDriver.h
--- a/Engine/Source/Runtime/Engine/Classes/Engine/NetDriver.h
+++ b/Engine/Source/Runtime/Engine/Classes/Engine/NetDriver.h
@@ -13,6 +13,7 @@
 #include "UObject/ObjectMacros.h"
 #include "UObject/UObjectGlobals.h"
 #include "UObject/Object.h"
+#include "UObject/ObjectKey.h"
 #include "Misc/NetworkGuid.h"
 #include "UObject/CoreNet.h"
 #include "Channel.h"
@@ -1328,6 +1329,9 @@ private:
 	/** Cached list of analytic attributes that can be appended via FNetAnalyticsAggregator::AppendGameInstanceAttributes */
 	TMap<FString, FString> CachedNetAnalyticsAttributes;
 
+	/** Maps an object to the respective FReplicationChangelistMgr */
+	TMap<FObjectKey, TSharedPtr<FReplicationChangelistMgr>>	ReplicationChangeListMap;
+
 public:
 
 	const TSet<FNetworkGUID>& GetDestroyedStartupOrDormantActors(const FName& LevelName)
@@ -1354,29 +1358,6 @@ public:
 	/** Maps FRepLayout to the respective UClass */
 	TMap<TWeakObjectPtr<UObject>, TSharedPtr<FRepLayout>, FDefaultSetAllocator, TWeakObjectPtrMapKeyFuncs<TWeakObjectPtr<UObject>, TSharedPtr<FRepLayout> > >	RepLayoutMap;
 
-	class FReplicationChangelistMgrWrapper
-	{
-	public:
-		FReplicationChangelistMgrWrapper(UObject* Obj, const TSharedPtr<FReplicationChangelistMgr>& InReplicationChangelistMgr) : ReplicationChangelistMgr(InReplicationChangelistMgr), WeakObjectPtr(Obj) {}
-
-		const FReplicationChangelistMgr* operator->() const { return ReplicationChangelistMgr.Get(); }
-		FReplicationChangelistMgr* operator->() { return ReplicationChangelistMgr.Get(); }
-
-		bool IsValid() const { return ReplicationChangelistMgr.IsValid(); }
-		bool IsObjectValid() const { return WeakObjectPtr.IsValid(); }
-
-		TWeakObjectPtr<UObject> GetWeakObjectPtr() const { return WeakObjectPtr; }
-
-		TSharedPtr<FReplicationChangelistMgr> ReplicationChangelistMgr;
-
-		void CountBytes(FArchive& Ar) const;
-
-	private:
-		TWeakObjectPtr<UObject> WeakObjectPtr;
-	};
-	/** Maps an object to the respective FReplicationChangelistMgr */
-	TMap< UObject*, FReplicationChangelistMgrWrapper >	ReplicationChangeListMap;
-
 	/** Creates if necessary, and returns a FRepLayout that maps to the passed in UClass */
 	TSharedPtr< FRepLayout >	GetObjectClassRepLayout( UClass * InClass );
 
diff --git a/Engine/Source/Runtime/Engine/Private/NetDriver.cpp b/Engine/Source/Runtime/Engine/Private/NetDriver.cpp
--- a/Engine/Source/Runtime/Engine/Private/NetDriver.cpp
+++ b/Engine/Source/Runtime/Engine/Private/NetDriver.cpp
@@ -1330,7 +1330,7 @@ void UNetDriver::CleanupStaleDormantReplicators()
 
 		for (auto It = ReplicationChangeListMap.CreateIterator(); It; ++It)
 		{
-			if (!It.Value().IsObjectValid())
+			if (!It.Key().GetWeakObjectPtr().IsValid())
 			{
 				It.RemoveCurrent();
 			}
@@ -3235,15 +3235,6 @@ void UNetDriver::SetAnalyticsProvider(TSharedPtr<IAnalyticsProvider> InProvider)
 	}
 }
 
-void UNetDriver::FReplicationChangelistMgrWrapper::CountBytes(FArchive& Ar) const
-{
-	if (FReplicationChangelistMgr const * const ChangelistMgr = ReplicationChangelistMgr.Get())
-	{
-		Ar.CountBytes(sizeof(FReplicationChangelistMgr), sizeof(FReplicationChangelistMgr));
-		ChangelistMgr->CountBytes(Ar);
-	}
-}
-
 void UNetDriver::Serialize( FArchive& Ar )
 {
 	GRANULAR_NETWORK_MEMORY_TRACKING_INIT(Ar, "UNetDriver::Serialize");
@@ -3327,7 +3318,11 @@ void UNetDriver::Serialize( FArchive& Ar )
 
 			for (const auto& ReplicationChangeListPair : ReplicationChangeListMap)
 			{
-				ReplicationChangeListPair.Value.CountBytes(Ar);
+				if (FReplicationChangelistMgr const* const ChangelistMgr = ReplicationChangeListPair.Value.Get())
+				{
+					Ar.CountBytes(sizeof(FReplicationChangelistMgr), sizeof(FReplicationChangelistMgr));
+					ChangelistMgr->CountBytes(Ar);
+				}
 			}
 		);
 
@@ -4634,7 +4629,7 @@ void UNetDriver::PostGarbageCollect()
 
 	for (auto It = ReplicationChangeListMap.CreateIterator(); It; ++It)
 	{
-		if (!It.Value().IsObjectValid())
+		if (!It.Key().GetWeakObjectPtr().IsValid())
 		{
 			It.RemoveCurrent();
 		}
@@ -7582,27 +7577,19 @@ TSharedPtr<FRepLayout> UNetDriver::GetStructRepLayout( UStruct * Struct )
 	return *RepLayoutPtr;
 }
 
-TSharedPtr< FReplicationChangelistMgr > UNetDriver::GetReplicationChangeListMgr( UObject* Object )
+TSharedPtr<FReplicationChangelistMgr> UNetDriver::GetReplicationChangeListMgr(UObject* Object)
 {
 	check(IsServer() || MaySendProperties());
 
-	FReplicationChangelistMgrWrapper* ReplicationChangeListMgrPtr = ReplicationChangeListMap.Find(Object);
-
-	// Object can be a new object with a pointer that matches an old, no longer valid, object
-	if (ReplicationChangeListMgrPtr != nullptr && !ReplicationChangeListMgrPtr->IsObjectValid())
-	{
-		ReplicationChangeListMap.Remove(Object);
-		ReplicationChangeListMgrPtr = nullptr;
-	}
+	TSharedPtr<FReplicationChangelistMgr>* ReplicationChangeListMgrPtr = ReplicationChangeListMap.Find(Object);
 
 	if (!ReplicationChangeListMgrPtr)
 	{
 		const TSharedPtr<const FRepLayout> RepLayout = GetObjectClassRepLayout(Object->GetClass());
-		FReplicationChangelistMgrWrapper Wrapper(Object, RepLayout->CreateReplicationChangelistMgr(Object, GetCreateReplicationChangelistMgrFlags()));
-		ReplicationChangeListMgrPtr = &ReplicationChangeListMap.Add(Object, Wrapper);
+		ReplicationChangeListMgrPtr = &ReplicationChangeListMap.Add(Object, RepLayout->CreateReplicationChangelistMgr(Object, GetCreateReplicationChangelistMgrFlags()));
 	}
 
-	return ReplicationChangeListMgrPtr->ReplicationChangelistMgr;
+	return *ReplicationChangeListMgrPtr;
 }
 
 void UNetDriver::RemoveDestroyedGuidsByLevel(const ULevel* Level, const TArray<FNetworkGUID>& RemovedGUIDs)
