Show slot nodes sequences, weights and play back time in rewind debugger

diff --git a/Engine/Source/Runtime/Engine/Private/Animation/AnimInstanceProxy.cpp b/Engine/Source/Runtime/Engine/Private/Animation/AnimInstanceProxy.cpp
--- a/Engine/Source/Runtime/Engine/Private/Animation/AnimInstanceProxy.cpp
+++ b/Engine/Source/Runtime/Engine/Private/Animation/AnimInstanceProxy.cpp
@@ -1975,6 +1975,28 @@ void FAnimInstanceProxy::SlotEvaluatePose(const FName& SlotNodeName, const FAnim
 			{
 				AdditivePoses.Add(FSlotEvaluationPose(MoveTemp(NewPose)));
 			}
+
+#if ANIM_TRACE_ENABLED
+			if (UE_TRACE_CHANNELEXPR_IS_ENABLED(AnimationChannel))
+			{
+				FAnimationBaseContext Context;
+				Context.AnimInstanceProxy = this;
+
+				const float ClampedTime = FMath::Clamp(ExtractionContext.CurrentTime, 0.f, AnimTrack->GetLength());
+
+				if (const FAnimSegment* const AnimSegment = AnimTrack->GetSegmentAtTime(ClampedTime))
+				{
+					float CurrentTime = 0;
+					FAnimTickRecord TickRecord;
+					TickRecord.SourceAsset = AnimSegment->GetAnimationData(ClampedTime, CurrentTime);
+					TickRecord.PlayRateMultiplier = AnimSegment->AnimPlayRate;
+					TickRecord.TimeAccumulator = &CurrentTime;
+					TickRecord.DeltaTimeRecord = &ExtractionContext.DeltaTimeRecord;
+					TickRecord.EffectiveBlendWeight = MontageWeight;
+					TRACE_ANIM_TICK_RECORD(Context, TickRecord);
+				}
+			}
+#endif
 		} // if IsValidSlot
 	} // for all MontageEvaluationData
 
