Add debug text to show camera modifiers when using ShowDebug Camera.
Hex color used is grey.

#rnx

diff --git a/Engine/Source/Runtime/Engine/Private/PlayerCameraManager.cpp b/Engine/Source/Runtime/Engine/Private/PlayerCameraManager.cpp
--- a/Engine/Source/Runtime/Engine/Private/PlayerCameraManager.cpp
+++ b/Engine/Source/Runtime/Engine/Private/PlayerCameraManager.cpp
@@ -1094,8 +1094,12 @@ void APlayerCameraManager::DisplayDebug(class UCanvas* Canvas, const FDebugDispl
 	const FString BlendStr = FString::Printf(TEXT("   ViewTarget Blend: From %s to %s, time remaining = %f, pct = %f"), *GetNameSafe(ViewTarget.Target), *GetNameSafe(PendingViewTarget.Target), BlendTimeToGo, DurationPct);
 	DisplayDebugManager.DrawString(BlendStr);
 	
+	DisplayDebugManager.DrawString(TEXT("Modifiers:"));
+	const FColor DisabledModifierColor = FColor::FromHex(TEXT("#949494"));
 	for (UCameraModifier* Modifier : ModifierList)
 	{
+		DisplayDebugManager.SetDrawColor(Modifier->IsDisabled() ? DisabledModifierColor : FColor::Green);
+		DisplayDebugManager.DrawString(Modifier->GetName());
 		Modifier->DisplayDebug(Canvas, DebugDisplay, YL, YPos);
 	}
 }
