Fix subtitle asset variables in BP getting overridden.

diff --git a/Engine/Source/Runtime/Engine/Private/Subtitles/SubtitlesAndClosedCaptionsDelegates.cpp b/Engine/Source/Runtime/Engine/Private/Subtitles/SubtitlesAndClosedCaptionsDelegates.cpp
--- a/Engine/Source/Runtime/Engine/Private/Subtitles/SubtitlesAndClosedCaptionsDelegates.cpp
+++ b/Engine/Source/Runtime/Engine/Private/Subtitles/SubtitlesAndClosedCaptionsDelegates.cpp
@@ -5,5 +5,4 @@
 TDelegate<void(const FQueueSubtitleParameters&, const ESubtitleTiming)> FSubtitlesAndClosedCaptionsDelegates::QueueSubtitle;
 TDelegate<bool(const UAssetUserData&)> FSubtitlesAndClosedCaptionsDelegates::IsSubtitleActive;
 TDelegate<void(const UAssetUserData&)> FSubtitlesAndClosedCaptionsDelegates::StopSubtitle;
-TDelegate<void()> FSubtitlesAndClosedCaptionsDelegates::StopAllSubtitles;
-TDelegate<void(FCanvas&)> FSubtitlesAndClosedCaptionsDelegates::DisplaySubtitle;
\ No newline at end of file
+TDelegate<void()> FSubtitlesAndClosedCaptionsDelegates::StopAllSubtitles;
\ No newline at end of file
diff --git a/Engine/Source/Runtime/Engine/Public/Subtitles/SubtitlesAndClosedCaptionsDelegates.h b/Engine/Source/Runtime/Engine/Public/Subtitles/SubtitlesAndClosedCaptionsDelegates.h
--- a/Engine/Source/Runtime/Engine/Public/Subtitles/SubtitlesAndClosedCaptionsDelegates.h
+++ b/Engine/Source/Runtime/Engine/Public/Subtitles/SubtitlesAndClosedCaptionsDelegates.h
@@ -34,7 +34,4 @@ public:
 	static TDelegate<void(const UAssetUserData&)> StopSubtitle;
 
 	static TDelegate<void()> StopAllSubtitles;
-
-	// Have the subtitle subsystem to display the highest priority subtitle. DisplaySubtitles is expected to be called/drawn per frame.
-	static TDelegate<void(FCanvas&)> DisplaySubtitle;
 };
