Minor improvement to how disabled camera modifier debug text color is initialized.
Addresses feedback from previous CL 41460866.

#rnx

diff --git a/Engine/Source/Runtime/Engine/Private/PlayerCameraManager.cpp b/Engine/Source/Runtime/Engine/Private/PlayerCameraManager.cpp
--- a/Engine/Source/Runtime/Engine/Private/PlayerCameraManager.cpp
+++ b/Engine/Source/Runtime/Engine/Private/PlayerCameraManager.cpp
@@ -1095,7 +1095,7 @@ void APlayerCameraManager::DisplayDebug(class UCanvas* Canvas, const FDebugDispl
 	DisplayDebugManager.DrawString(BlendStr);
 	
 	DisplayDebugManager.DrawString(TEXT("Modifiers:"));
-	const FColor DisabledModifierColor = FColor::FromHex(TEXT("#949494"));
+	const FColor DisabledModifierColor(0xff949494);
 	for (UCameraModifier* Modifier : ModifierList)
 	{
 		DisplayDebugManager.SetDrawColor(Modifier->IsDisabled() ? DisabledModifierColor : FColor::Green);
