AudioLink spatialize all sounds cvar

diff --git a/Engine/Source/Runtime/Engine/Private/ActiveSound.cpp b/Engine/Source/Runtime/Engine/Private/ActiveSound.cpp
--- a/Engine/Source/Runtime/Engine/Private/ActiveSound.cpp
+++ b/Engine/Source/Runtime/Engine/Private/ActiveSound.cpp
@@ -68,6 +68,13 @@ FAutoConsoleVariableRef CVarIncludeModulationInConcurrencyEvaluation(
 	TEXT("0: Ignore Audio Modulation Volume when evaluating concurrency (Default), 1: Apply Modulation Destination to logical volume for one-shots only, 2: Apply Modulation Destination to logical volume on all sounds."),
 	ECVF_Default);
 
+static int32 bAudioLinkSpatializeAllSoundsCVar = 0;
+FAutoConsoleVariableRef CVarAudioLinkSpatializeAllSounds(
+	TEXT("au.AudioLink.SpatializeAllSounds"),
+	bAudioLinkSpatializeAllSoundsCVar,
+	TEXT("0 (off), !=0 (on)"),
+	ECVF_Default);
+
 FTraceDelegate FActiveSound::ActiveSoundTraceDelegate;
 TMap<FTraceHandle, FActiveSound::FAsyncTraceDetails> FActiveSound::TraceToActiveSoundMap;
 
@@ -2128,7 +2135,7 @@ void FActiveSound::UpdateAttenuation(float DeltaTime, FSoundParseParameters& Par
 	ParseParams.NonSpatializedRadiusMode = Settings->NonSpatializedRadiusMode;
 	ParseParams.StereoSpread = Settings->StereoSpread;
 	ParseParams.bApplyNormalizationToStereoSounds = Settings->bApplyNormalizationToStereoSounds;
-	ParseParams.bUseSpatialization |= Settings->bSpatialize;
+	ParseParams.bUseSpatialization |= Settings->bSpatialize || bAudioLinkSpatializeAllSoundsCVar;
 	ParseParams.bEnableSourceDataOverride |= Settings->bEnableSourceDataOverride;
 	ParseParams.bEnableSendToAudioLink |= Settings->bEnableSendToAudioLink;
 
