World Partition - HLOD - Instancing: Added options to perform instance filtering based on the instance size

diff --git a/Engine/Plugins/Editor/WorldPartitionHLODUtilities/Source/Private/WorldPartition/HLOD/Builders/HLODBuilderInstancing.cpp b/Engine/Plugins/Editor/WorldPartitionHLODUtilities/Source/Private/WorldPartition/HLOD/Builders/HLODBuilderInstancing.cpp
--- a/Engine/Plugins/Editor/WorldPartitionHLODUtilities/Source/Private/WorldPartition/HLOD/Builders/HLODBuilderInstancing.cpp
+++ b/Engine/Plugins/Editor/WorldPartitionHLODUtilities/Source/Private/WorldPartition/HLOD/Builders/HLODBuilderInstancing.cpp
@@ -11,6 +11,10 @@
 UHLODBuilderInstancingSettings::UHLODBuilderInstancingSettings(const FObjectInitializer& ObjectInitializer)
 	: Super(ObjectInitializer)
 	, bDisallowNanite(false)
+	, InstanceFilteringType(EInstanceFilteringType::FilterNone)
+	, MinimumExtent(0)
+	, MinimumArea(0)
+	, MinimumVolume(0)
 {
 }
 
@@ -27,6 +31,25 @@ uint32 UHLODBuilderInstancingSettings::GetCRC() const
 	Ar << This.bDisallowNanite;
 	UE_LOG(LogHLODBuilder, VeryVerbose, TEXT(" - bDisallowNanite = %d"), Ar.GetCrc());
 
+	if (InstanceFilteringType != EInstanceFilteringType::FilterNone)
+	{
+		Ar << This.InstanceFilteringType;
+		switch (InstanceFilteringType)
+		{
+		case EInstanceFilteringType::FilterMinimumExtent:
+			Ar << This.MinimumExtent;
+			break;
+		case EInstanceFilteringType::FilterMinimumArea:
+			Ar << This.MinimumArea;
+			break;
+		case EInstanceFilteringType::FilterMinimumVolume:
+			Ar << This.MinimumVolume;
+			break;
+		default:
+			unimplemented();
+		}
+	}
+
 	return Ar.GetCrc();
 }
 
@@ -45,7 +68,48 @@ TArray<UActorComponent*> UHLODBuilderInstancing::Build(const FHLODBuildContext&
 {
 	const UHLODBuilderInstancingSettings& InstancingSettings = *CastChecked<UHLODBuilderInstancingSettings>(HLODBuilderSettings);
 
-	TArray<UActorComponent*> HLODComponents = UHLODBuilder::BatchInstances(InSourceComponents);
+	int32 NumInstancesTotal = 0;
+	int32 NumInstancesRejected = 0;
+
+	auto FilterInstances = [&InstancingSettings, &NumInstancesTotal, &NumInstancesRejected](const FBox& InInstanceBounds)
+	{
+		bool bPassFilter = true;
+
+		switch (InstancingSettings.InstanceFilteringType)
+		{
+			case EInstanceFilteringType::FilterMinimumExtent:
+			{
+				double MaxExtent = InInstanceBounds.GetExtent().GetMax();
+				bPassFilter = MaxExtent >= InstancingSettings.MinimumExtent;
+				break;
+			}
+			case EInstanceFilteringType::FilterMinimumArea:
+			{
+				FVector Extent = InInstanceBounds.GetExtent();
+				double Area = 8 * (Extent.X * Extent.Y + Extent.X * Extent.Z + Extent.Y * Extent.Z);
+				bPassFilter = Area >= InstancingSettings.MinimumArea;
+				break;
+			}
+			case EInstanceFilteringType::FilterMinimumVolume:
+			{
+				double Volume = InInstanceBounds.GetVolume();
+				bPassFilter = Volume >= InstancingSettings.MinimumVolume;
+				break;
+			}
+		}
+
+		NumInstancesTotal++;
+		NumInstancesRejected += bPassFilter ? 0 : 1;
+
+		return bPassFilter;
+	};
+
+	TArray<UActorComponent*> HLODComponents = UHLODBuilder::BatchInstances(InSourceComponents, FilterInstances);
+
+	if (NumInstancesRejected > 0)
+	{
+		UE_LOG(LogHLODBuilder, Log, TEXT("UHLODBuilderInstancing: Filter rejected %d out of %d instances"), NumInstancesRejected, NumInstancesTotal);
+	}
 
 	// If requested, disallow Nanite on components whose mesh is Nanite enabled
 	if (InstancingSettings.bDisallowNanite)
@@ -68,4 +132,3 @@ TArray<UActorComponent*> UHLODBuilderInstancing::Build(const FHLODBuildContext&
 
 	return HLODComponents;
 }
-
diff --git a/Engine/Plugins/Editor/WorldPartitionHLODUtilities/Source/Public/WorldPartition/HLOD/Builders/HLODBuilderInstancing.h b/Engine/Plugins/Editor/WorldPartitionHLODUtilities/Source/Public/WorldPartition/HLOD/Builders/HLODBuilderInstancing.h
--- a/Engine/Plugins/Editor/WorldPartitionHLODUtilities/Source/Public/WorldPartition/HLOD/Builders/HLODBuilderInstancing.h
+++ b/Engine/Plugins/Editor/WorldPartitionHLODUtilities/Source/Public/WorldPartition/HLOD/Builders/HLODBuilderInstancing.h
@@ -6,6 +6,14 @@
 #include "WorldPartition/HLOD/HLODBuilder.h"
 #include "HLODBuilderInstancing.generated.h"
 
+UENUM()
+enum class EInstanceFilteringType : uint8
+{
+	FilterNone UMETA(DisplayName = "No Filtering"),
+	FilterMinimumExtent,
+	FilterMinimumArea,
+	FilterMinimumVolume,
+};
 
 UCLASS(Blueprintable, Config = Engine, PerObjectConfig)
 class WORLDPARTITIONHLODUTILITIES_API UHLODBuilderInstancingSettings : public UHLODBuilderSettings
@@ -25,6 +33,24 @@ class WORLDPARTITIONHLODUTILITIES_API UHLODBuilderInstancingSettings : public UH
 	 */
 	UPROPERTY(EditAnywhere, AdvancedDisplay, BlueprintReadWrite, Category = Rendering)
 	bool bDisallowNanite;
+
+	/**
+	 * Type of filtering to apply to instances. Instances that do not pass the filter will be excluded from the HLOD. 
+	 */
+	UPROPERTY(EditAnywhere, AdvancedDisplay, BlueprintReadWrite, Category = Filtering)
+	EInstanceFilteringType InstanceFilteringType;
+
+	// Minimum extent for an instance to be included in the HLOD. 0 means no minimum.
+	UPROPERTY(EditAnywhere, AdvancedDisplay, BlueprintReadWrite, Category = Filtering, Meta = (DisplayName = "Minimum Extent", ToolTip = "Minimum Extent in Unreal Units", EditCondition = "InstanceFilteringType==EInstanceFilteringType::FilterMinimumExtent", EditConditionHides))
+	double MinimumExtent;
+
+	// Minimum area for an instance to be included in the HLOD. 0 means no minimum. Value is in Unreal Unit ^ 2
+	UPROPERTY(EditAnywhere, AdvancedDisplay, BlueprintReadWrite, Category = Filtering, Meta = (DisplayName = "Minimum Area (uu²)", ToolTip = "Minimum area in square Unreal Units", EditCondition = "InstanceFilteringType==EInstanceFilteringType::FilterMinimumArea", EditConditionHides))
+	double MinimumArea;
+
+	// Minimum volume for an instance to be included in the HLOD. 0 means no minimum. Value is in Unreal Unit ^ 3
+	UPROPERTY(EditAnywhere, AdvancedDisplay, BlueprintReadWrite, Category = Filtering, Meta = (DisplayName = "Minimum Volume (uu³)", ToolTip = "Minimum volume in cubic Unreal Units", EditCondition = "InstanceFilteringType==EInstanceFilteringType::FilterMinimumVolume", EditConditionHides))
+	double MinimumVolume;
 };
 
 
diff --git a/Engine/Source/Runtime/Engine/Private/ISMPartition/ISMComponentBatcher.cpp b/Engine/Source/Runtime/Engine/Private/ISMPartition/ISMComponentBatcher.cpp
--- a/Engine/Source/Runtime/Engine/Private/ISMPartition/ISMComponentBatcher.cpp
+++ b/Engine/Source/Runtime/Engine/Private/ISMPartition/ISMComponentBatcher.cpp
@@ -2,21 +2,32 @@
 
 #include "ISMPartition/ISMComponentBatcher.h"
 #include "Components/InstancedSkinnedMeshComponent.h"
+#include "Engine/StaticMesh.h"
 #include "Serialization/ArchiveCrc32.h"
 #include "Misc/TransformUtilities.h"
 #include "Templates/TypeHash.h"
 
 void FISMComponentBatcher::Add(const UActorComponent* InComponent)
 {
-	AddInternal(InComponent, TOptional<TFunctionRef<FTransform(const FTransform&)>>());
+	AddInternal(InComponent, TOptional<TFunctionRef<FTransform(const FTransform&)>>(), TOptional<TFunctionRef<bool(const FBox&)>>());
 }
 
 void FISMComponentBatcher::Add(const UActorComponent* InComponent, TFunctionRef<FTransform(const FTransform&)> InTransformFunc)
 {
-	AddInternal(InComponent, InTransformFunc);
+	AddInternal(InComponent, InTransformFunc, TOptional<TFunctionRef<bool(const FBox&)>>());
 }
 
-void FISMComponentBatcher::AddInternal(const UActorComponent* InComponent, TOptional<TFunctionRef<FTransform(const FTransform&)>> InTransformFunc)
+void FISMComponentBatcher::Add(const UActorComponent* InComponent, TFunctionRef<bool(const FBox&)> InFilterFunc)
+{
+	AddInternal(InComponent, TOptional<TFunctionRef<FTransform(const FTransform&)>>(), InFilterFunc);
+}
+
+void FISMComponentBatcher::Add(const UActorComponent* InComponent, TFunctionRef<FTransform(const FTransform&)> InTransformFunc, TFunctionRef<bool(const FBox&)> InFilterFunc)
+{
+	AddInternal(InComponent, InTransformFunc, InFilterFunc);
+}
+
+void FISMComponentBatcher::AddInternal(const UActorComponent* InComponent, TOptional<TFunctionRef<FTransform(const FTransform&)>> InTransformFunc, TOptional<TFunctionRef<bool(const FBox&)>> InFilterFunc)
 {
 	Hash = 0; // Invalidate
 	int32 NewNumCustomDataFloats = 0;
@@ -39,6 +50,8 @@ void FISMComponentBatcher::AddInternal(const UActorComponent* InComponent, TOpti
 		NewNumInstances = ISKMC->GetInstanceCount();
 	}
 
+	// If we must increase the number of custom data floats, this means we have to add the proper padding between each existing
+	// custom data float set.
 	if (NewNumCustomDataFloats > NumCustomDataFloats)
 	{
 		TArray<float> NewInstancesCustomData;
@@ -58,30 +71,52 @@ void FISMComponentBatcher::AddInternal(const UActorComponent* InComponent, TOpti
 		NumCustomDataFloats = NewNumCustomDataFloats;
 	}
 
-	NumInstances += NewNumInstances;
-	InstancesTransformsWS.Reserve(NumInstances);
-	InstancesCustomData.Reserve(NumInstances * NumCustomDataFloats);
+	int32 ReserveNum = NumInstances + NewNumInstances;
+	InstancesTransformsWS.Reserve(ReserveNum);
+	InstancesCustomData.Reserve(ReserveNum * NumCustomDataFloats);
 
 	// Add instances
 	if (const UInstancedStaticMeshComponent* ISMC = Cast<UInstancedStaticMeshComponent>(InComponent))
 	{
-		// Add random seeds
-		RandomSeeds.Add({ NumInstances, ISMC->InstancingRandomSeed });
-		for (const FInstancedStaticMeshRandomSeed& AdditionalRandomSeed : ISMC->AdditionalRandomSeeds)
-		{
-			RandomSeeds.Add({ NumInstances + AdditionalRandomSeed.StartInstanceIndex, AdditionalRandomSeed.RandomSeed });
-		}
+		FBox StaticMeshBox = ISMC->GetStaticMesh() ? ISMC->GetStaticMesh()->GetBounds().GetBox() : FBox();
+
+		// Instancing Random Seed
+		InstancingRandomSeed  = ISMC->InstancingRandomSeed;
 
 		// Add each instance
-		for (int32 InstanceIdx = 0; InstanceIdx < ISMC->GetInstanceCount(); InstanceIdx++)
+		int32 NextRandomSeedsIndex = ISMC->AdditionalRandomSeeds.IsValidIndex(0) ? 0 : INDEX_NONE;
+		int32 NextRandomSeedsInstanceStart = NextRandomSeedsIndex != INDEX_NONE ? ISMC->AdditionalRandomSeeds[NextRandomSeedsIndex].StartInstanceIndex : INDEX_NONE;
+
+		for (int32 InstanceIdx = 0, RemappedInstanceIdx=0; InstanceIdx < ISMC->GetInstanceCount(); InstanceIdx++, RemappedInstanceIdx++)
 		{
-			// Add instance transform
+			// Add additional random seeds
+			if (InstanceIdx == NextRandomSeedsInstanceStart)
+			{
+				AdditionalRandomSeeds.Add({ RemappedInstanceIdx, ISMC->AdditionalRandomSeeds[NextRandomSeedsIndex].RandomSeed });
+				NextRandomSeedsIndex = ISMC->AdditionalRandomSeeds.IsValidIndex(NextRandomSeedsIndex+1) ? NextRandomSeedsIndex+1 : INDEX_NONE;
+				NextRandomSeedsInstanceStart = NextRandomSeedsIndex != INDEX_NONE ? ISMC->AdditionalRandomSeeds[NextRandomSeedsIndex].StartInstanceIndex : INDEX_NONE;
+			}
+			
 			FTransform InstanceTransformWS;
 			ISMC->GetInstanceTransform(InstanceIdx, InstanceTransformWS, /*bWorldSpace*/ true);
 			if (InTransformFunc.IsSet())
 			{
-				InstanceTransformWS = InTransformFunc->operator()(InstanceTransformWS);
+				InstanceTransformWS = InTransformFunc.GetValue()(InstanceTransformWS);
+			}
+
+			if (InFilterFunc.IsSet())
+			{
+				bool bFilterPass = InFilterFunc.GetValue()(StaticMeshBox.TransformBy(InstanceTransformWS));
+				if (!bFilterPass)
+				{
+					RemappedInstanceIdx--;
+					continue;
+				}
 			}
+
+			NumInstances++;
+
+			// Add instance transform
 			InstancesTransformsWS.Add(InstanceTransformWS);
 
 			// Add per instance custom data, if any
@@ -99,12 +134,25 @@ void FISMComponentBatcher::AddInternal(const UActorComponent* InComponent, TOpti
 	}
 	else if (const UStaticMeshComponent* SMC = Cast<UStaticMeshComponent>(InComponent))
 	{
+		FBox StaticMeshBox = SMC->GetStaticMesh() ? SMC->GetStaticMesh()->GetBounds().GetBox() : FBox();
+
 		FTransform InstanceTransformWS = SMC->GetComponentTransform();
 		if (InTransformFunc.IsSet())
 		{
-			InstanceTransformWS = InTransformFunc->operator()(InstanceTransformWS);
+			InstanceTransformWS = InTransformFunc.GetValue()(InstanceTransformWS);
 		}
 
+		if (InFilterFunc.IsSet())
+		{
+			bool bFilterPass = InFilterFunc.GetValue()(StaticMeshBox.TransformBy(InstanceTransformWS));
+			if (!bFilterPass)
+			{
+				return;
+			}
+		}
+
+		NumInstances++;
+
 		// Add transform
 		InstancesTransformsWS.Add(InstanceTransformWS);
 
@@ -118,13 +166,15 @@ void FISMComponentBatcher::AddInternal(const UActorComponent* InComponent, TOpti
 		FPrimitiveInstanceId InstanceId;
 		for (InstanceId.Id = 0; InstanceId.Id < ISKMC->GetInstanceCount(); InstanceId.Id++)
 		{
+			NumInstances++;
+
 			// Add instance transform
 			FTransform InstanceTransformWS;
 
 			ISKMC->GetInstanceTransform(InstanceId, InstanceTransformWS, /*bWorldSpace*/ true);
 			if (InTransformFunc.IsSet())
 			{
-				InstanceTransformWS = InTransformFunc->operator()(InstanceTransformWS);
+				InstanceTransformWS = InTransformFunc.GetValue()(InstanceTransformWS);
 			}
 			InstancesTransformsWS.Add(InstanceTransformWS);
 
@@ -151,15 +201,8 @@ void FISMComponentBatcher::InitComponent(UInstancedStaticMeshComponent* ISMCompo
 	ISMComponent->AddInstances(InstancesTransformsWS, /*bShouldReturnIndices*/false, /*bWorldSpace*/true);
 	ISMComponent->PerInstanceSMCustomData = InstancesCustomData;
 
-	if (!RandomSeeds.IsEmpty())
-	{
-		ISMComponent->InstancingRandomSeed = RandomSeeds[0].RandomSeed;
-
-		if (RandomSeeds.Num() > 1)
-		{
-			ISMComponent->AdditionalRandomSeeds = TArrayView<FInstancedStaticMeshRandomSeed>(&const_cast<FISMComponentBatcher*>(this)->RandomSeeds[1], RandomSeeds.Num() - 1);
-		}
-	}		
+	ISMComponent->InstancingRandomSeed = InstancingRandomSeed;
+	ISMComponent->AdditionalRandomSeeds = AdditionalRandomSeeds;
 }
 
 void FISMComponentBatcher::InitComponent(UInstancedSkinnedMeshComponent* ISMComponent) const
@@ -189,7 +232,9 @@ void FISMComponentBatcher::ComputeHash() const
 	FISMComponentBatcher& This = *const_cast<FISMComponentBatcher*>(this);
 
 	Ar << This.InstancesCustomData;
-	Ar << This.RandomSeeds;
+	Ar << This.InstancingRandomSeed;
+	Ar << This.AdditionalRandomSeeds;
+	Ar << This.BankIndices;
 
 	Hash = Ar.GetCrc();
 }
diff --git a/Engine/Source/Runtime/Engine/Private/WorldPartition/HLOD/HLODBuilder.cpp b/Engine/Source/Runtime/Engine/Private/WorldPartition/HLOD/HLODBuilder.cpp
--- a/Engine/Source/Runtime/Engine/Private/WorldPartition/HLOD/HLODBuilder.cpp
+++ b/Engine/Source/Runtime/Engine/Private/WorldPartition/HLOD/HLODBuilder.cpp
@@ -145,6 +145,11 @@ namespace
 }
 
 TArray<UActorComponent*> UHLODBuilder::BatchInstances(const TArray<UActorComponent*>& InSourceComponents)
+{
+	return BatchInstances(InSourceComponents, [](const FBox& InBox) { return true; });
+}
+
+TArray<UActorComponent*> UHLODBuilder::BatchInstances(const TArray<UActorComponent*>& InSourceComponents, TFunctionRef<bool(const FBox&)> InFilterFunc)
 {
 	TRACE_CPUPROFILER_EVENT_SCOPE(UHLODBuilderInstancing::Build);
 
@@ -215,16 +220,19 @@ TArray<UActorComponent*> UHLODBuilder::BatchInstances(const TArray<UActorCompone
 				InstanceBatch.ISMComponentDescriptor = MoveTemp(ISMComponentDescriptor);
 			}
 
-			InstanceBatch.ISMComponentBatcher.Add(SMC);
+			InstanceBatch.ISMComponentBatcher.Add(SMC, InFilterFunc);
 		}
 
 		// Create an ISMC for each SM asset we found		
 		for (auto& Entry : InstancesData)
 		{
 			const FInstancedStaticMeshBatch& InstanceBatch = Entry.Value;
-			UInstancedStaticMeshComponent* ISMComponent = InstanceBatch.ISMComponentDescriptor->CreateComponent(GetTransientPackage());
-			InstanceBatch.ISMComponentBatcher.InitComponent(ISMComponent);
-			HLODComponents.Add(ISMComponent);
+			if (InstanceBatch.ISMComponentBatcher.GetNumInstances() > 0)
+			{
+				UInstancedStaticMeshComponent* ISMComponent = InstanceBatch.ISMComponentDescriptor->CreateComponent(GetTransientPackage());
+				InstanceBatch.ISMComponentBatcher.InitComponent(ISMComponent);
+				HLODComponents.Add(ISMComponent);
+			}
 		};
 	}
 
@@ -250,16 +258,19 @@ TArray<UActorComponent*> UHLODBuilder::BatchInstances(const TArray<UActorCompone
 				InstanceBatch.ISMComponentDescriptor = MoveTemp(ISMComponentDescriptor);
 			}
 
-			InstanceBatch.ISMComponentBatcher.Add(ISKMC);
+			InstanceBatch.ISMComponentBatcher.Add(ISKMC, InFilterFunc);
 		}
 
 		// Create an ISKMC for each SKM asset we found
 		for (auto& Entry : InstancesData)
 		{
 			const FInstancedSkinnedMeshBatch& InstanceBatch = Entry.Value;
-			UInstancedSkinnedMeshComponent* ISMComponent = InstanceBatch.ISMComponentDescriptor->CreateComponent(GetTransientPackage());
-			InstanceBatch.ISMComponentBatcher.InitComponent(ISMComponent);
-			HLODComponents.Add(ISMComponent);
+			if (InstanceBatch.ISMComponentBatcher.GetNumInstances() > 0)
+			{
+				UInstancedSkinnedMeshComponent* ISMComponent = InstanceBatch.ISMComponentDescriptor->CreateComponent(GetTransientPackage());
+				InstanceBatch.ISMComponentBatcher.InitComponent(ISMComponent);
+				HLODComponents.Add(ISMComponent);
+			}
 		};
 	}
 
diff --git a/Engine/Source/Runtime/Engine/Public/ISMPartition/ISMComponentBatcher.h b/Engine/Source/Runtime/Engine/Public/ISMPartition/ISMComponentBatcher.h
--- a/Engine/Source/Runtime/Engine/Public/ISMPartition/ISMComponentBatcher.h
+++ b/Engine/Source/Runtime/Engine/Public/ISMPartition/ISMComponentBatcher.h
@@ -26,6 +26,21 @@ public:
 	 * @param	InTransformFunc		Function that takes the world space transform of an instance and modifies it. Must return a world space transform.
 	 */
 	ENGINE_API void Add(const UActorComponent* InActorComponent, TFunctionRef<FTransform(const FTransform&)> InTransformFunc);
+
+	/**
+	 * Add a single component to be batched
+	 * @param	InActorComponent	Component to be batched
+	 * @param	InFilterFunc		Function that can be used to filter out instances based on their world bounds.
+	 */
+	ENGINE_API void Add(const UActorComponent* InComponent, TFunctionRef<bool(const FBox&)> InFilterFunc);
+
+	/**
+	 * Add a single component to be batched
+	 * @param	InActorComponent	Component to be batched
+	 * @param	InTransformFunc		Function that takes the world space transform of an instance and modifies it. Must return a world space transform.
+	 * @param	InFilterFunc		Function that can be used to filter out instances based on their world bounds.
+	 */
+	ENGINE_API void Add(const UActorComponent* InComponent, TFunctionRef<FTransform(const FTransform&)> InTransformFunc, TFunctionRef<bool(const FBox&)> InFilterFunc);
 	
 	/**
 	 * Add an array of component to be batched
@@ -36,7 +51,7 @@ public:
 	{
 		for (const UActorComponent* InComponent : InComponents)
 		{
-			AddInternal(InComponent, TOptional<TFunctionRef<FTransform(const FTransform&)>>());
+			Add(InComponent);
 		}
 	}
 
@@ -50,7 +65,36 @@ public:
 	{
 		for (const UActorComponent* InComponent : InComponents)
 		{
-			AddInternal(InComponent, InTransformFunc);
+			Add(InComponent, InTransformFunc);
+		}
+	}
+
+	/**
+	 * Add an array of components to be batched
+	 * @param	InComponents		Components to be batched
+	 * @param	InFilterFunc		Function that can be used to filter out instances based on their world bounds.
+	 */
+	template<typename TComponentClass = UActorComponent, typename = typename TEnableIf<TIsDerivedFrom<TComponentClass, UActorComponent>::IsDerived>::Type>
+	void Append(const TArray<TComponentClass*>& InComponents, TFunctionRef<bool(const FBox&)> InFilterFunc)
+	{
+		for (const UActorComponent* InComponent : InComponents)
+		{
+			Add(InComponent, InFilterFunc);
+		}
+	}
+
+	/**
+	 * Add an array of components to be batched
+	 * @param	InComponents		Components to be batched
+	 * @param	InFilterFunc		Function that can be used to filter out instances based on their world bounds.
+	 * @param	InTransformFunc		Function that takes the world space transform of an instance and modifies it. Must return a world space transform.
+	 */
+	template<typename TComponentClass = UActorComponent, typename = typename TEnableIf<TIsDerivedFrom<TComponentClass, UActorComponent>::IsDerived>::Type>
+	void Append(const TArray<TComponentClass*>& InComponents, TFunctionRef<FTransform(const FTransform&)> InTransformFunc, TFunctionRef<bool(const FBox&)> InFilterFunc)
+	{
+		for (const UActorComponent* InComponent : InComponents)
+		{
+			Add(InComponent, InTransformFunc, InFilterFunc);
 		}
 	}
 
@@ -81,7 +125,7 @@ public:
 	}
 
 protected:
-	void ENGINE_API AddInternal(const UActorComponent* InComponent, TOptional<TFunctionRef<FTransform(const FTransform&)>> InTransformFunc);
+	void ENGINE_API AddInternal(const UActorComponent* InComponent, TOptional<TFunctionRef<FTransform(const FTransform&)>> InTransformFunc, TOptional<TFunctionRef<bool(const FBox&)>> InFilterFunc);
 	void ENGINE_API ComputeHash() const;
 
 	mutable uint32 Hash = 0;
@@ -91,6 +135,10 @@ protected:
 	TArray<FTransform> InstancesTransformsWS;
 	TArray<float> InstancesCustomData;
 
-	TArray<FInstancedStaticMeshRandomSeed> RandomSeeds; // For ISMC only
-	TArray<int32> BankIndices;							// For ISKMC only
+	// For ISMC only
+	int32 InstancingRandomSeed = 0;
+	TArray<FInstancedStaticMeshRandomSeed> AdditionalRandomSeeds; 
+
+	// For ISKMC only
+	TArray<int32> BankIndices;
 };
diff --git a/Engine/Source/Runtime/Engine/Public/WorldPartition/HLOD/HLODBuilder.h b/Engine/Source/Runtime/Engine/Public/WorldPartition/HLOD/HLODBuilder.h
--- a/Engine/Source/Runtime/Engine/Public/WorldPartition/HLOD/HLODBuilder.h
+++ b/Engine/Source/Runtime/Engine/Public/WorldPartition/HLOD/HLODBuilder.h
@@ -168,7 +168,8 @@ public:
 protected:
 	virtual bool ShouldIgnoreBatchingPolicy() const { return false; }
 
-	static ENGINE_API TArray<UActorComponent*> BatchInstances(const TArray<UActorComponent*>& InSubComponents);
+	static ENGINE_API TArray<UActorComponent*> BatchInstances(const TArray<UActorComponent*>& InSourceComponents);
+	static ENGINE_API TArray<UActorComponent*> BatchInstances(const TArray<UActorComponent*>& InSourceComponents, TFunctionRef<bool(const FBox&)> InFilterFunc);
 
 	template <typename TComponentClass>
 	static inline TArray<TComponentClass*> FilterComponents(const TArray<UActorComponent*>& InSourceComponents)
