Texture.cpp attempt to disable unhelpful PVS warning

Issue 861522

diff --git a/Engine/Source/Runtime/Engine/Private/Texture.cpp b/Engine/Source/Runtime/Engine/Private/Texture.cpp
--- a/Engine/Source/Runtime/Engine/Private/Texture.cpp
+++ b/Engine/Source/Runtime/Engine/Private/Texture.cpp
@@ -4016,7 +4016,7 @@ void UTexture::GetBuiltTextureSize(const ITargetPlatformSettings* TargetPlatform
 	}
 #endif // WITH_EDITORONLY_DATA
 
-	if ( SizeX == 0 )
+	if ( SizeX == 0 ) //-V547
 	{
 		// no Editor data
 		//	can't support a query about arbitrary target platform
