World Partition - HLOD: When building HLOD static meshes, silence warning regarding mesh bounds differences between RenderData bounds and MeshDescription bounds
* Nanite meshes can sometime contain smaller triangles that could be discarded in the proxy mesh, depending on the simplification settings
* If those triangles are found on the outskirt of a merged mesh, the static mesh build will complain about the bounds delta being too high.
* As HLODs are not hand authored meshes, there's not much we can do to prevent those bounds delta.

diff --git a/Engine/Plugins/Experimental/MeshModelingToolsetExp/Source/GeometryProcessingAdapters/Private/ApproximateActorsImpl.cpp b/Engine/Plugins/Experimental/MeshModelingToolsetExp/Source/GeometryProcessingAdapters/Private/ApproximateActorsImpl.cpp
--- a/Engine/Plugins/Experimental/MeshModelingToolsetExp/Source/GeometryProcessingAdapters/Private/ApproximateActorsImpl.cpp
+++ b/Engine/Plugins/Experimental/MeshModelingToolsetExp/Source/GeometryProcessingAdapters/Private/ApproximateActorsImpl.cpp
@@ -35,6 +35,7 @@
 #include "Scene/SceneCapturePhotoSet.h"
 #include "AssetUtils/Texture2DBuilder.h"
 
+#include "Engine/StaticMesh.h"
 #include "MaterialDomain.h"
 #include "Materials/Material.h"
 #include "Materials/MaterialInstanceConstant.h"
@@ -1417,10 +1418,23 @@ UStaticMesh* FApproximateActorsImpl::EmitGeneratedMeshAsset(
 	{
 		MeshAssetOptions.AssetMaterials.Add(UMaterial::GetDefaultMaterial(MD_Surface));
 	}
+
+	MeshAssetOptions.bDeferPostEditChange = true; // Defer static mesh build to control options
+
 	FStaticMeshResults MeshAssetOutputs;
 	ECreateStaticMeshResult ResultCode = UE::AssetUtils::CreateStaticMeshAsset(MeshAssetOptions, MeshAssetOutputs);
 	ensure(ResultCode == ECreateStaticMeshResult::Ok);
 
+	
+	UStaticMesh::FBuildParameters BuildParameters;
+
+	// Nanite meshes can sometime contain smaller triangles that could be discarded in the proxy mesh.
+	// If those are found on the outskirt of a merged mesh, the static mesh build will complain about
+	// the bounds delta being too high.
+	BuildParameters.bIgnoreBoundsDiff = Options.bGenerateNaniteEnabledMesh; 
+
+	MeshAssetOutputs.StaticMesh->Build(BuildParameters);
+
 	ResultsOut.NewMeshAssets.Add(MeshAssetOutputs.StaticMesh);
 
 	if (DebugMesh != nullptr)
diff --git a/Engine/Source/Developer/MeshMergeUtilities/Private/MeshMergeUtilities.cpp b/Engine/Source/Developer/MeshMergeUtilities/Private/MeshMergeUtilities.cpp
--- a/Engine/Source/Developer/MeshMergeUtilities/Private/MeshMergeUtilities.cpp
+++ b/Engine/Source/Developer/MeshMergeUtilities/Private/MeshMergeUtilities.cpp
@@ -2738,7 +2738,17 @@ void FMeshMergeUtilities::MergeComponentsToStaticMesh(const TArray<UPrimitiveCom
 			StaticMesh->CalculateExtendedBounds();
 		}		
 
-		StaticMesh->PostEditChange();
+		UStaticMesh::FBuildParameters BuildParameters;
+		BuildParameters.bInSilent = true;
+		BuildParameters.bInRebuildUVChannelData = true;
+		BuildParameters.bInEnforceLightmapRestrictions = true;
+		
+		// Nanite meshes can sometime contain smaller triangles that could be discarded in the proxy mesh.
+		// If those are found on the outskirt of a merged mesh, the static mesh build will complain about
+		// the bounds delta being too high.
+		BuildParameters.bIgnoreBoundsDiff = StaticMesh->NaniteSettings.bEnabled;
+
+		StaticMesh->Build(BuildParameters);
 
 		OutAssetsToSync.Add(StaticMesh);
 		OutMergedActorLocation = MergedAssetPivot;
diff --git a/Engine/Source/Developer/MeshMergeUtilities/Private/ProxyGenerationProcessor.cpp b/Engine/Source/Developer/MeshMergeUtilities/Private/ProxyGenerationProcessor.cpp
--- a/Engine/Source/Developer/MeshMergeUtilities/Private/ProxyGenerationProcessor.cpp
+++ b/Engine/Source/Developer/MeshMergeUtilities/Private/ProxyGenerationProcessor.cpp
@@ -343,8 +343,6 @@ void FProxyGenerationProcessor::ProcessJob(const FGuid& JobGuid, FProxyGeneratio
 		SectionIndex++;
 	}
 
-	StaticMesh->Build();
-
 	if (ImposterBounds.IsValid)
 	{
 		const FBox StaticMeshBox = StaticMesh->GetBoundingBox();
@@ -354,7 +352,14 @@ void FProxyGenerationProcessor::ProcessJob(const FGuid& JobGuid, FProxyGeneratio
 		StaticMesh->CalculateExtendedBounds();
 	}
 
-	StaticMesh->PostEditChange();	
+	UStaticMesh::FBuildParameters BuildParameters;
+
+	// Nanite meshes can sometime contain smaller triangles that could be discarded in the proxy mesh.
+	// If those are found on the outskirt of a merged mesh, the static mesh build will complain about
+	// the bounds delta being too high.
+	BuildParameters.bIgnoreBoundsDiff = StaticMesh->NaniteSettings.bEnabled;
+
+	StaticMesh->Build(BuildParameters);
 
 	OutAssetsToSync.Add(StaticMesh);
 
diff --git a/Engine/Source/Runtime/Engine/Classes/Engine/StaticMesh.h b/Engine/Source/Runtime/Engine/Classes/Engine/StaticMesh.h
--- a/Engine/Source/Runtime/Engine/Classes/Engine/StaticMesh.h
+++ b/Engine/Source/Runtime/Engine/Classes/Engine/StaticMesh.h
@@ -1802,6 +1802,7 @@ public:
 			, OutErrors(nullptr)
 			, bInRebuildUVChannelData(false)
 			, bInEnforceLightmapRestrictions(false)
+			, bIgnoreBoundsDiff(false)
 		{
 		}
 #endif
@@ -1817,6 +1818,9 @@ public:
 
 		// Whether to call EnforceLightmapRestrictions as part of the build process.
 		bool bInEnforceLightmapRestrictions = false;
+
+		// Whether to avoid outputting a warning if the render data and MeshDescription bounds differ too much.
+		bool bIgnoreBoundsDiff = false;
 	};
 
 	/**
@@ -2184,7 +2188,7 @@ private:
 	/**
 	 * Complete the static mesh building process - Can't be done in parallel.
 	 */
-	void FinishBuildInternal(const TArray<IStaticMeshComponent*>& InAffectedComponents, bool bHasRenderDataChanged, bool bShouldComputeExtendedBounds = true);
+	void FinishBuildInternal(const TArray<IStaticMeshComponent*>& InAffectedComponents, bool bHasRenderDataChanged, bool bIgnoreBoundsDiff, bool bShouldComputeExtendedBounds = true);
 
 #if WITH_EDITORONLY_DATA
 	/**
diff --git a/Engine/Source/Runtime/Engine/Private/StaticMeshBuild.cpp b/Engine/Source/Runtime/Engine/Private/StaticMeshBuild.cpp
--- a/Engine/Source/Runtime/Engine/Private/StaticMeshBuild.cpp
+++ b/Engine/Source/Runtime/Engine/Private/StaticMeshBuild.cpp
@@ -344,7 +344,7 @@ void UStaticMesh::BatchBuild(const TArray<UStaticMesh*>& InStaticMeshes, const F
 				// Wait the result of the async task
 				const bool bHasRenderDataChanged = Task.Future.Get();
 
-				Task.StaticMesh->FinishBuildInternal(StaticMeshComponents.FindChecked(Task.StaticMesh), bHasRenderDataChanged);
+				Task.StaticMesh->FinishBuildInternal(StaticMeshComponents.FindChecked(Task.StaticMesh), bHasRenderDataChanged, BuildParameters.bIgnoreBoundsDiff);
 				FinalizeStaticMesh(Task.StaticMesh);
 			}
 		}
@@ -365,7 +365,7 @@ void UStaticMesh::BatchBuild(const TArray<UStaticMesh*>& InStaticMeshes, const F
 
 					const bool bHasRenderDataChanged = StaticMesh->ExecuteBuildInternal(BuildParameters);
 
-					StaticMesh->FinishBuildInternal(StaticMeshComponents.FindChecked(StaticMesh), bHasRenderDataChanged);
+					StaticMesh->FinishBuildInternal(StaticMeshComponents.FindChecked(StaticMesh), bHasRenderDataChanged, BuildParameters.bIgnoreBoundsDiff);
 					FinalizeStaticMesh(StaticMesh);
 				}
 			}
@@ -632,7 +632,7 @@ bool UStaticMesh::ExecuteBuildInternal(const FBuildParameters& BuildParameters)
 	return bHasRenderDataChanged;
 }
 
-void UStaticMesh::FinishBuildInternal(const TArray<IStaticMeshComponent*>& InAffectedComponents, bool bHasRenderDataChanged, bool bShouldComputeExtendedBounds)
+void UStaticMesh::FinishBuildInternal(const TArray<IStaticMeshComponent*>& InAffectedComponents, bool bHasRenderDataChanged, bool bIgnoreBoundsDiff, bool bShouldComputeExtendedBounds)
 {
 	TRACE_CPUPROFILER_EVENT_SCOPE(UStaticMesh::PostBuildInternal);
 
@@ -657,7 +657,7 @@ void UStaticMesh::FinishBuildInternal(const TArray<IStaticMeshComponent*>& InAff
 		// Calculate extended bounds
 		CalculateExtendedBounds();
 	}
-	else if (!IsNaniteLandscape())
+	else if (!IsNaniteLandscape() && !bIgnoreBoundsDiff)
 	{
 		// We don't care about minor differences but might want to highlight if big differences are noted.
 		const FBoxSphereBounds RenderBounds = GetRenderData()->Bounds;
diff --git a/Engine/Source/Runtime/Engine/Private/StaticMeshCompiler.cpp b/Engine/Source/Runtime/Engine/Private/StaticMeshCompiler.cpp
--- a/Engine/Source/Runtime/Engine/Private/StaticMeshCompiler.cpp
+++ b/Engine/Source/Runtime/Engine/Private/StaticMeshCompiler.cpp
@@ -402,6 +402,7 @@ void FStaticMeshCompilingManager::PostCompilation(UStaticMesh* StaticMesh)
 				StaticMesh->FinishBuildInternal(
 					ComponentsToUpdate,
 					LocalAsyncTask->GetTask().BuildContext->bHasRenderDataChanged,
+					LocalAsyncTask->GetTask().BuildContext->BuildParameters.bIgnoreBoundsDiff,
 					LocalAsyncTask->GetTask().BuildContext->bShouldComputeExtendedBounds
 				);
 
