Linux/clang compile fixes for UE_WITH_REMOTE_OBJECT_HANDLE

diff --git a/Engine/Source/Runtime/Engine/Private/Components/SkeletalMeshComponent.cpp b/Engine/Source/Runtime/Engine/Private/Components/SkeletalMeshComponent.cpp
--- a/Engine/Source/Runtime/Engine/Private/Components/SkeletalMeshComponent.cpp
+++ b/Engine/Source/Runtime/Engine/Private/Components/SkeletalMeshComponent.cpp
@@ -3957,7 +3957,7 @@ void USkeletalMeshComponent::ApplyMigratedPhysicsState(Chaos::FSerializedDataBuf
 		MigratedState.Serialize(MultiBodyStateReader);
 
 		const Chaos::FSerializedMultiPhysicsState::FHeader& MigratedStateHeader = MigratedState.GetHeader();
-		const bool bIsValidMigratedState = MigratedStateHeader.IsValid() && MigratedStateHeader.NumBodies == Bodies.Num() || MigratedStateHeader.NumConstraints == Constraints.Num();
+		const bool bIsValidMigratedState = MigratedStateHeader.IsValid() && MigratedStateHeader.NumBodies == Bodies.Num() && MigratedStateHeader.NumConstraints == Constraints.Num();
 		if (!ensure(bIsValidMigratedState))
 		{
 			return;
diff --git a/Engine/Source/Runtime/Engine/Private/PlayerController.cpp b/Engine/Source/Runtime/Engine/Private/PlayerController.cpp
--- a/Engine/Source/Runtime/Engine/Private/PlayerController.cpp
+++ b/Engine/Source/Runtime/Engine/Private/PlayerController.cpp
@@ -4946,10 +4946,8 @@ void APlayerController::PostMigrate(const FUObjectMigrationContext& MigrationCon
 				}
 			}
 		}
-		else
-		{
-			return 0;
-		}
+
+		return 0;
 	};
 
 	auto GetConnectionWithPlayerController = [](UNetDriver* Driver, APlayerController* PlayerController)
