IncludeWhatYouUse: AssetManager: Replace embedded class keywords used in variable declarations with forward declares of those classes, to match the current standard for forward declares.
#rnx

diff --git a/Engine/Source/Runtime/Engine/Classes/Engine/AssetManager.h b/Engine/Source/Runtime/Engine/Classes/Engine/AssetManager.h
--- a/Engine/Source/Runtime/Engine/Classes/Engine/AssetManager.h
+++ b/Engine/Source/Runtime/Engine/Classes/Engine/AssetManager.h
@@ -17,10 +17,16 @@
 DECLARE_LOG_CATEGORY_EXTERN(LogAssetManager, Log, All);
 
 /** Defined in C++ file */
+class IAssetRegistry;
+class ITargetPlatform;
+class FAssetManagerEditorModule;
+class FTimerManager;
+class UAssetManagerSettings;
 struct FPrimaryAssetTypeData;
 struct FPrimaryAssetData;
 struct FPrimaryAssetRulesCustomOverride;
 namespace UE::Cook { class ICookInfo; }
+struct FCompiledAssetManagerSearchRules;
 
 /** Delegate called when acquiring resources/chunks for assets, first parameter will be true if all resources were acquired, false if any failed.
 	Second parameter will contain a list of chunks not yet downloaded */
@@ -742,13 +748,13 @@ public:
 	ENGINE_API virtual void FinishInitialLoading();
 
 	/** Accessor for asset registry */
-	ENGINE_API class IAssetRegistry& GetAssetRegistry() const;
+	ENGINE_API IAssetRegistry& GetAssetRegistry() const;
 
 	/** Return settings object */
-	ENGINE_API const class UAssetManagerSettings& GetSettings() const;
+	ENGINE_API const UAssetManagerSettings& GetSettings() const;
 
 	/** Returns a timer manager that is safe to use for asset loading actions. This will either be the editor or game instance one, or null during very early startup */
-	ENGINE_API class FTimerManager* GetTimerManager() const;
+	ENGINE_API FTimerManager* GetTimerManager() const;
 
 	// Overrides
 	ENGINE_API virtual void BeginDestroy() override;
@@ -827,7 +833,7 @@ public:
 	 * @param OutChunkList List of chunks to actually assign this to
 	 * @param OutOverrideChunkList List of chunks that were added due to override rules, not just normal primary asset rules. Tools use this for dependency checking
 	 */
-	ENGINE_API virtual bool GetPackageChunkIds(FName PackageName, const class ITargetPlatform* TargetPlatform, TArrayView<const int32> ExistingChunkList, TArray<int32>& OutChunkList, TArray<int32>* OutOverrideChunkList = nullptr) const;
+	ENGINE_API virtual bool GetPackageChunkIds(FName PackageName, const ITargetPlatform* TargetPlatform, TArrayView<const int32> ExistingChunkList, TArray<int32>& OutChunkList, TArray<int32>* OutOverrideChunkList = nullptr) const;
 
 	/** 
 	 * Retrieve the encryption key guid for a given chunk ID
@@ -854,7 +860,7 @@ public:
 	virtual int32 GetContentEncryptionGroupChunkID(FName InGroupName) const { return INDEX_NONE; }
 
 	/** Returns the list of chunks assigned to the list of primary assets, which is usually a manager list. This is called by GetPackageChunkIds */
-	ENGINE_API virtual bool GetPrimaryAssetSetChunkIds(const TSet<FPrimaryAssetId>& PrimaryAssetSet, const class ITargetPlatform* TargetPlatform, TArrayView<const int32> ExistingChunkList, TArray<int32>& OutChunkList) const;
+	ENGINE_API virtual bool GetPrimaryAssetSetChunkIds(const TSet<FPrimaryAssetId>& PrimaryAssetSet, const ITargetPlatform* TargetPlatform, TArrayView<const int32> ExistingChunkList, TArray<int32>& OutChunkList) const;
 
 	/** Resets all asset manager data, called in the editor to reinitialize the config */
 	ENGINE_API virtual void ReinitializeFromConfig();
@@ -875,7 +881,7 @@ public:
 	 * Initializes asset bundle data from a passed in struct or class, this will read the AssetBundles metadata off the UProperties. As an example this property definition:
 	 *
 	 * UPROPERTY(EditDefaultsOnly, Category=Data, meta = (AssetBundles = "Client,Server"))
-	 * TSoftObjectPtr<class UCurveTable> CurveTableReference;
+	 * TSoftObjectPtr<UCurveTable> CurveTableReference;
 	 *
 	 * Would add the value of CurveTableReference to both the Client and Server asset bundles
 	 *
@@ -894,7 +900,7 @@ public:
 	/** 
 	  * Called immediately before saving the asset registry during cooking
 	  */
-	virtual void PreSaveAssetRegistry(const class ITargetPlatform* TargetPlatform, const TSet<FName>& InCookedPackages) {}
+	virtual void PreSaveAssetRegistry(const ITargetPlatform* TargetPlatform, const TSet<FName>& InCookedPackages) {}
 
 	/** 
 	  * Called immediately after saving the asset registry during cooking
@@ -941,10 +947,10 @@ public:
 	ENGINE_API void GetPakStringToChunkIdMapping(TMap<FString, int32>& StringOverrideToChunkId);
 
 protected:
-	friend class FAssetManagerEditorModule;
+	friend FAssetManagerEditorModule;
 
 	/** Internal helper function that attempts to get asset data from the specified path; Accounts for possibility of blueprint classes ending in _C */
-	ENGINE_API virtual void GetAssetDataForPathInternal(class IAssetRegistry& AssetRegistry, const FString& AssetPath, OUT FAssetData& OutAssetData) const;
+	ENGINE_API virtual void GetAssetDataForPathInternal(IAssetRegistry& AssetRegistry, const FString& AssetPath, OUT FAssetData& OutAssetData) const;
 
 	/** Updates the asset data cached on the name data; returns false if the asset is not a valid primary asset. */
 	ENGINE_API virtual bool TryUpdateCachedAssetData(const FPrimaryAssetId& PrimaryAssetId, const FAssetData& NewAssetData, bool bAllowDuplicates);
@@ -1209,10 +1215,10 @@ private:
 	TMap<FName, TArray<FTopLevelAssetPath>> AssetBundlePathsForPackage;
 #endif
 
-	mutable class IAssetRegistry* CachedAssetRegistry;
-	mutable const class UAssetManagerSettings* CachedSettings;
+	mutable IAssetRegistry* CachedAssetRegistry;
+	mutable const UAssetManagerSettings* CachedSettings;
 
-	friend struct FCompiledAssetManagerSearchRules;
+	friend FCompiledAssetManagerSearchRules;
 };
 
 inline TSharedPtr<FStreamableHandle> UAssetManager::LoadPrimaryAssets(const TArray<FPrimaryAssetId>& AssetsToLoad, const TArray<FName>& LoadBundles, FStreamableDelegate DelegateToCall, TAsyncLoadPriority Priority, UE::FSourceLocation Location)
