Fix switching preview SKM in Physics asset editor by marking bRequiredBonesUpToDate = false before recreating physics state inside USkeletalMeshComponent::SetSkeletalMesh.

Conservatively just setting it here as well as the place where it's set to false later in the function. I don't know if something  in between should be dirtying this state.

Will follow-up with a more holistic fix, but this should unblock people.

diff --git a/Engine/Source/Runtime/Engine/Private/Components/SkeletalMeshComponent.cpp b/Engine/Source/Runtime/Engine/Private/Components/SkeletalMeshComponent.cpp
--- a/Engine/Source/Runtime/Engine/Private/Components/SkeletalMeshComponent.cpp
+++ b/Engine/Source/Runtime/Engine/Private/Components/SkeletalMeshComponent.cpp
@@ -3340,6 +3340,10 @@ void USkeletalMeshComponent::SetSkeletalMesh(USkeletalMesh* InSkelMesh, bool bRe
 
 		if(IsPhysicsStateCreated())
 		{
+			// Make sure that required bones are invalidated as we have just changed our mesh
+			// RecalcRequiredBones will be called by RecreatePhysicsState (which may be called by UpdateBoneBodyMapping).
+			bRequiredBonesUpToDate = false;
+
 			if(GetPhysicsAsset() == OldPhysAsset && OldPhysAsset && Bodies.Num() == OldPhysAsset->SkeletalBodySetups.Num())	//Make sure that we actually created all the bodies for the asset (needed for old assets in editor)
 			{
 				UpdateBoneBodyMapping();
