[Deformer Graph] Fixed crash when the first section of the mesh is disabled. Also removed a few checkes that are no longer valid since DG switched to unified dispatch by default and special case mulitple sub invocations

diff --git a/Engine/Plugins/Animation/DeformerGraph/Source/OptimusCore/Private/DataInterfaces/OptimusDataInterfaceConnectivity.cpp b/Engine/Plugins/Animation/DeformerGraph/Source/OptimusCore/Private/DataInterfaces/OptimusDataInterfaceConnectivity.cpp
--- a/Engine/Plugins/Animation/DeformerGraph/Source/OptimusCore/Private/DataInterfaces/OptimusDataInterfaceConnectivity.cpp
+++ b/Engine/Plugins/Animation/DeformerGraph/Source/OptimusCore/Private/DataInterfaces/OptimusDataInterfaceConnectivity.cpp
@@ -249,10 +249,6 @@ bool FOptimusConnectivityDataProviderProxy::IsValid(FValidationData const& InVal
 	{
 		return false;
 	}
-	if (SkeletalMeshObject->GetSkeletalMeshRenderData().LODRenderData[SkeletalMeshObject->GetLOD()].RenderSections.Num() != InValidationData.NumInvocations)
-	{
-		return false;
-	}
 
 	return true;
 }
diff --git a/Engine/Plugins/Animation/DeformerGraph/Source/OptimusCore/Private/DataInterfaces/OptimusDataInterfaceHalfEdge.cpp b/Engine/Plugins/Animation/DeformerGraph/Source/OptimusCore/Private/DataInterfaces/OptimusDataInterfaceHalfEdge.cpp
--- a/Engine/Plugins/Animation/DeformerGraph/Source/OptimusCore/Private/DataInterfaces/OptimusDataInterfaceHalfEdge.cpp
+++ b/Engine/Plugins/Animation/DeformerGraph/Source/OptimusCore/Private/DataInterfaces/OptimusDataInterfaceHalfEdge.cpp
@@ -267,11 +267,6 @@ bool FOptimusHalfEdgeDataProviderProxy::IsValid(FValidationData const& InValidat
 
 	FSkeletalMeshLODRenderData const* LodRenderData = &SkeletalMeshObject->GetSkeletalMeshRenderData().LODRenderData[SkeletalMeshObject->GetLOD()];	
 	
-	if (LodRenderData->RenderSections.Num() != InValidationData.NumInvocations)
-	{
-		return false;
-	}
-
 	// Invalid if there is no cooked buffer and run time generation isn't possible either
 	const FRawStaticIndexBuffer16or32Interface* IndexBuffer = LodRenderData->MultiSizeIndexContainer.GetIndexBuffer();
 	const FPositionVertexBuffer& VertexBuffer = LodRenderData->StaticVertexBuffers.PositionVertexBuffer;
diff --git a/Engine/Plugins/Animation/DeformerGraph/Source/OptimusCore/Private/DataInterfaces/OptimusDataInterfaceMorphTarget.cpp b/Engine/Plugins/Animation/DeformerGraph/Source/OptimusCore/Private/DataInterfaces/OptimusDataInterfaceMorphTarget.cpp
--- a/Engine/Plugins/Animation/DeformerGraph/Source/OptimusCore/Private/DataInterfaces/OptimusDataInterfaceMorphTarget.cpp
+++ b/Engine/Plugins/Animation/DeformerGraph/Source/OptimusCore/Private/DataInterfaces/OptimusDataInterfaceMorphTarget.cpp
@@ -123,10 +123,7 @@ bool FOptimusMorphTargetDataProviderProxy::IsValid(FValidationData const& InVali
 	{
 		return false;
 	}
-	if (SkeletalMeshObject->GetSkeletalMeshRenderData().LODRenderData[SkeletalMeshObject->GetLOD()].RenderSections.Num() != InValidationData.NumInvocations)
-	{
-		return false;
-	}
+
 
 	return true;
 }
diff --git a/Engine/Plugins/Animation/DeformerGraph/Source/OptimusCore/Private/DataInterfaces/OptimusDataInterfaceSkinWeightsAsVertexMask.cpp b/Engine/Plugins/Animation/DeformerGraph/Source/OptimusCore/Private/DataInterfaces/OptimusDataInterfaceSkinWeightsAsVertexMask.cpp
--- a/Engine/Plugins/Animation/DeformerGraph/Source/OptimusCore/Private/DataInterfaces/OptimusDataInterfaceSkinWeightsAsVertexMask.cpp
+++ b/Engine/Plugins/Animation/DeformerGraph/Source/OptimusCore/Private/DataInterfaces/OptimusDataInterfaceSkinWeightsAsVertexMask.cpp
@@ -451,10 +451,6 @@ bool FOptimusSkinWeightsAsVertexMaskDataProviderProxy::IsValid(FValidationData c
 	const int32 LodIndex = SkeletalMeshObject->GetLOD();
 	FSkeletalMeshRenderData const& SkeletalMeshRenderData = SkeletalMeshObject->GetSkeletalMeshRenderData();
 	FSkeletalMeshLODRenderData const* LodRenderData = &SkeletalMeshRenderData.LODRenderData[LodIndex];
-	if (LodRenderData->RenderSections.Num() != InValidationData.NumInvocations)
-	{
-		return false;
-	}
 	
 	FSkinWeightVertexBuffer const* WeightBuffer = LodRenderData->GetSkinWeightVertexBuffer();
 	if (LodRenderData->SkinWeightProfilesData.ContainsProfile(SkinWeightProfile))
diff --git a/Engine/Plugins/Animation/DeformerGraph/Source/OptimusCore/Private/DataInterfaces/OptimusDataInterfaceSkinnedMesh.cpp b/Engine/Plugins/Animation/DeformerGraph/Source/OptimusCore/Private/DataInterfaces/OptimusDataInterfaceSkinnedMesh.cpp
--- a/Engine/Plugins/Animation/DeformerGraph/Source/OptimusCore/Private/DataInterfaces/OptimusDataInterfaceSkinnedMesh.cpp
+++ b/Engine/Plugins/Animation/DeformerGraph/Source/OptimusCore/Private/DataInterfaces/OptimusDataInterfaceSkinnedMesh.cpp
@@ -161,10 +161,6 @@ bool FOptimusSkinnedMeshDataProviderProxy::IsValid(FValidationData const& InVali
 	{
 		return false;
 	}
-	if (SkeletalMeshObject->GetSkeletalMeshRenderData().LODRenderData[SkeletalMeshObject->GetLOD()].RenderSections.Num() != InValidationData.NumInvocations)
-	{
-		return false;
-	}
 
 	return true;
 }
diff --git a/Engine/Plugins/Animation/DeformerGraph/Source/OptimusCore/Private/DataInterfaces/OptimusDataInterfaceSkinnedMeshRead.cpp b/Engine/Plugins/Animation/DeformerGraph/Source/OptimusCore/Private/DataInterfaces/OptimusDataInterfaceSkinnedMeshRead.cpp
--- a/Engine/Plugins/Animation/DeformerGraph/Source/OptimusCore/Private/DataInterfaces/OptimusDataInterfaceSkinnedMeshRead.cpp
+++ b/Engine/Plugins/Animation/DeformerGraph/Source/OptimusCore/Private/DataInterfaces/OptimusDataInterfaceSkinnedMeshRead.cpp
@@ -197,11 +197,13 @@ bool FOptimusSkinnedMeshReadDataProviderProxy::IsValid(FValidationData const& In
 	{
 		return false;
 	}
-	if (SkeletalMeshObject->GetSkeletalMeshRenderData().LODRenderData[SkeletalMeshObject->GetLOD()].RenderSections.Num() != InValidationData.NumInvocations)
+
+	if (FSkeletalMeshDeformerHelpers::GetIndexOfFirstAvailableSection(SkeletalMeshObject, SkeletalMeshObject->GetLOD()) == INDEX_NONE)
 	{
 		return false;
 	}
 
+
 	return true;
 }
 
diff --git a/Engine/Plugins/Animation/DeformerGraph/Source/OptimusCore/Private/DataInterfaces/OptimusDataInterfaceSkinnedMeshVertexAttribute.cpp b/Engine/Plugins/Animation/DeformerGraph/Source/OptimusCore/Private/DataInterfaces/OptimusDataInterfaceSkinnedMeshVertexAttribute.cpp
--- a/Engine/Plugins/Animation/DeformerGraph/Source/OptimusCore/Private/DataInterfaces/OptimusDataInterfaceSkinnedMeshVertexAttribute.cpp
+++ b/Engine/Plugins/Animation/DeformerGraph/Source/OptimusCore/Private/DataInterfaces/OptimusDataInterfaceSkinnedMeshVertexAttribute.cpp
@@ -198,10 +198,6 @@ bool FOptimusSkinnedMeshVertexAttributeDataProviderProxy::IsValid(FValidationDat
 	{
 		return false;
 	}
-	if (SkeletalMeshObject->GetSkeletalMeshRenderData().LODRenderData[SkeletalMeshObject->GetLOD()].RenderSections.Num() != InValidationData.NumInvocations)
-	{
-		return false;
-	}
 
 	return true;
 }
diff --git a/Engine/Plugins/Animation/DeformerGraph/Source/OptimusCore/Private/DataInterfaces/OptimusDataInterfaceSkinnedMeshWrite.cpp b/Engine/Plugins/Animation/DeformerGraph/Source/OptimusCore/Private/DataInterfaces/OptimusDataInterfaceSkinnedMeshWrite.cpp
--- a/Engine/Plugins/Animation/DeformerGraph/Source/OptimusCore/Private/DataInterfaces/OptimusDataInterfaceSkinnedMeshWrite.cpp
+++ b/Engine/Plugins/Animation/DeformerGraph/Source/OptimusCore/Private/DataInterfaces/OptimusDataInterfaceSkinnedMeshWrite.cpp
@@ -210,7 +210,8 @@ bool FOptimusSkinnedMeshWriteDataProviderProxy::IsValid(FValidationData const& I
 	{
 		return false;
 	}
-	if (SkeletalMeshObject->GetSkeletalMeshRenderData().LODRenderData[SkeletalMeshObject->GetLOD()].RenderSections.Num() != InValidationData.NumInvocations)
+	
+	if (FSkeletalMeshDeformerHelpers::GetIndexOfFirstAvailableSection(SkeletalMeshObject, SkeletalMeshObject->GetLOD()) == INDEX_NONE)
 	{
 		return false;
 	}
diff --git a/Engine/Source/Runtime/Engine/Private/SkeletalMeshDeformerHelpers.cpp b/Engine/Source/Runtime/Engine/Private/SkeletalMeshDeformerHelpers.cpp
--- a/Engine/Source/Runtime/Engine/Private/SkeletalMeshDeformerHelpers.cpp
+++ b/Engine/Source/Runtime/Engine/Private/SkeletalMeshDeformerHelpers.cpp
@@ -139,6 +139,22 @@ FRDGBuffer* FSkeletalMeshDeformerHelpers::GetAllocatedColorBuffer(FRDGBuilder& G
 	return ColorBuffer;
 }
 
+int32 FSkeletalMeshDeformerHelpers::GetIndexOfFirstAvailableSection(FSkeletalMeshObject* InMeshObject, int32 InLodIndex)
+{
+	int32 Result = INDEX_NONE;
+	const TArray<FSkelMeshRenderSection>& RenderSections = InMeshObject->GetRenderSections(InLodIndex);
+	for (int32 Index = 0; Index < RenderSections.Num(); Index++)
+	{
+		if (!RenderSections[Index].bDisabled)
+		{
+			Result = Index;
+			break;
+		}
+	}
+
+	return Result;
+}
+
 FRDGBuffer* FSkeletalMeshDeformerHelpers::AllocateVertexFactoryPositionBuffer(FRDGBuilder& GraphBuilder, FRDGExternalAccessQueue& ExternalAccesQueue, FSkeletalMeshObject* InMeshObject, int32 InLodIndex, TCHAR const* InBufferName)
 {
 	if (InMeshObject->IsCPUSkinned())
@@ -152,7 +168,7 @@ FRDGBuffer* FSkeletalMeshDeformerHelpers::AllocateVertexFactoryPositionBuffer(FR
 	const FRDGBufferDesc BufferDesc = FRDGBufferDesc::CreateBufferDesc(PosBufferBytesPerElement, NumVertices * PosBufferElementMultiplier);
 
 	FSkeletalMeshObjectGPUSkin* MeshObjectGPU = static_cast<FSkeletalMeshObjectGPUSkin*>(InMeshObject);
-	FGPUBaseSkinVertexFactory const* BaseVertexFactory = MeshObjectGPU->GetBaseSkinVertexFactory(InLodIndex, 0);
+	FGPUBaseSkinVertexFactory const* BaseVertexFactory = MeshObjectGPU->GetBaseSkinVertexFactory(InLodIndex, GetIndexOfFirstAvailableSection(InMeshObject, InLodIndex));
 	uint32 Frame = BaseVertexFactory->GetShaderData().UpdatedFrameNumber;
 
 	FMeshDeformerGeometry& DeformerGeometry = MeshObjectGPU->GetDeformerGeometry(InLodIndex);
@@ -244,7 +260,7 @@ FRDGBuffer* FSkeletalMeshDeformerHelpers::AllocateVertexFactoryTangentBuffer(FRD
 	const FRDGBufferDesc BufferDesc = FRDGBufferDesc::CreateBufferDesc(TangentBufferBytesPerElement, NumVertices * TangentBufferElementMultiplier);
 
 	FSkeletalMeshObjectGPUSkin* MeshObjectGPU = static_cast<FSkeletalMeshObjectGPUSkin*>(InMeshObject);
-	FGPUBaseSkinVertexFactory const* BaseVertexFactory = MeshObjectGPU->GetBaseSkinVertexFactory(InLodIndex, 0);
+	FGPUBaseSkinVertexFactory const* BaseVertexFactory = MeshObjectGPU->GetBaseSkinVertexFactory(InLodIndex, GetIndexOfFirstAvailableSection(InMeshObject, InLodIndex));
 	uint32 Frame = BaseVertexFactory->GetShaderData().UpdatedFrameNumber;
 
 	FMeshDeformerGeometry& DeformerGeometry = MeshObjectGPU->GetDeformerGeometry(InLodIndex);
@@ -293,7 +309,7 @@ FRDGBuffer* FSkeletalMeshDeformerHelpers::AllocateVertexFactoryColorBuffer(FRDGB
 	const FRDGBufferDesc BufferDesc = FRDGBufferDesc::CreateBufferDesc(ColorBufferBytesPerElement, NumVertices);
 
 	FSkeletalMeshObjectGPUSkin* MeshObjectGPU = static_cast<FSkeletalMeshObjectGPUSkin*>(InMeshObject);
-	FGPUBaseSkinVertexFactory const* BaseVertexFactory = MeshObjectGPU->GetBaseSkinVertexFactory(InLodIndex, 0);
+	FGPUBaseSkinVertexFactory const* BaseVertexFactory = MeshObjectGPU->GetBaseSkinVertexFactory(InLodIndex, GetIndexOfFirstAvailableSection(InMeshObject, InLodIndex));
 	uint32 Frame = BaseVertexFactory->GetShaderData().UpdatedFrameNumber;
 
 	FMeshDeformerGeometry& DeformerGeometry = MeshObjectGPU->GetDeformerGeometry(InLodIndex);
@@ -369,7 +385,7 @@ void FSkeletalMeshDeformerHelpers::UpdateVertexFactoryBufferOverrides(FRHIComman
 	}
 
 	const FSkeletalMeshObjectGPUSkin::FSkeletalMeshObjectLOD& LOD = MeshObjectGPU->LODs[InLodIndex];
-	FGPUBaseSkinVertexFactory const* BaseVertexFactory = MeshObjectGPU->GetBaseSkinVertexFactory(InLodIndex, 0);
+	FGPUBaseSkinVertexFactory const* BaseVertexFactory = MeshObjectGPU->GetBaseSkinVertexFactory(InLodIndex, GetIndexOfFirstAvailableSection(InMeshObject, InLodIndex));
 	FGPUSkinPassthroughVertexFactory* TargetVertexFactory = LOD.GPUSkinVertexFactories.PassthroughVertexFactory.Get();
 
 	// The passthrough vertex factory should exist if we got this far, but prefer skipping the update to crashing if that assumption fails.
diff --git a/Engine/Source/Runtime/Engine/Public/SkeletalMeshDeformerHelpers.h b/Engine/Source/Runtime/Engine/Public/SkeletalMeshDeformerHelpers.h
--- a/Engine/Source/Runtime/Engine/Public/SkeletalMeshDeformerHelpers.h
+++ b/Engine/Source/Runtime/Engine/Public/SkeletalMeshDeformerHelpers.h
@@ -62,6 +62,9 @@ public:
 
 	/** Returns the allocated writable color buffer if one has been allocated */	
 	ENGINE_API static FRDGBuffer* GetAllocatedColorBuffer(FRDGBuilder& GraphBuilder, FSkeletalMeshObject* InMeshObject, int32 InLodIndex);
+
+	/** Returns the index of the first section that is not disabled */
+	ENGINE_API static int32 GetIndexOfFirstAvailableSection(FSkeletalMeshObject* InMeshObject, int32 InLodIndex);
 	
 #pragma endregion GetInternals
 
