[NetPhysics] Move the PostInitialize_Internal callback earlier to just after AsyncInput has been made available and before bPendingDelete gets set if the ISimCallbackObject is marked for deletion.

diff --git a/Engine/Source/Runtime/Engine/Private/PhysicsEngine/NetworkPhysicsComponent.cpp b/Engine/Source/Runtime/Engine/Private/PhysicsEngine/NetworkPhysicsComponent.cpp
--- a/Engine/Source/Runtime/Engine/Private/PhysicsEngine/NetworkPhysicsComponent.cpp
+++ b/Engine/Source/Runtime/Engine/Private/PhysicsEngine/NetworkPhysicsComponent.cpp
@@ -547,6 +547,18 @@ void UNetworkPhysicsComponent::InitializeComponent()
 				NetworkPhysicsComponent_Internal->bCompareInputToTriggerRewind = bCompareInputToTriggerRewind;
 				CreateAsyncDataHistory();
 				UpdateAsyncComponent(true);
+
+				/** Run OnInitialize_Internal on the ISimCallbackObject first thing on the next physics thread frame */
+				FAsyncNetworkPhysicsComponent* AsyncNetworkPhysicsComponent = NetworkPhysicsComponent_Internal;
+				Solver->EnqueueCommandImmediate(
+					[AsyncNetworkPhysicsComponent]()
+					{
+						if (AsyncNetworkPhysicsComponent)
+						{
+							AsyncNetworkPhysicsComponent->OnInitialize_Internal();
+						}
+					}
+				);
 			}
 		}
 	}
@@ -1508,7 +1520,7 @@ FAsyncNetworkPhysicsComponent::FAsyncNetworkPhysicsComponent() : TSimCallbackObj
 {
 }
 
-void FAsyncNetworkPhysicsComponent::OnPostInitialize_Internal()
+void FAsyncNetworkPhysicsComponent::OnInitialize_Internal()
 {
 	if (Chaos::FPhysicsSolverBase* BaseSolver = GetSolver())
 	{
diff --git a/Engine/Source/Runtime/Engine/Public/Physics/NetworkPhysicsComponent.h b/Engine/Source/Runtime/Engine/Public/Physics/NetworkPhysicsComponent.h
--- a/Engine/Source/Runtime/Engine/Public/Physics/NetworkPhysicsComponent.h
+++ b/Engine/Source/Runtime/Engine/Public/Physics/NetworkPhysicsComponent.h
@@ -1453,7 +1453,7 @@ public:
 
 private:
 	// Initialize, bind delegates etc.
-	virtual void OnPostInitialize_Internal() override;
+	void OnInitialize_Internal();
 
 	// Uninitialize, unbind delegates etc.
 	void OnUninitialize_Internal();
diff --git a/Engine/Source/Runtime/Experimental/Chaos/Private/PBDRigidsSolver.cpp b/Engine/Source/Runtime/Experimental/Chaos/Private/PBDRigidsSolver.cpp
--- a/Engine/Source/Runtime/Experimental/Chaos/Private/PBDRigidsSolver.cpp
+++ b/Engine/Source/Runtime/Experimental/Chaos/Private/PBDRigidsSolver.cpp
@@ -1966,10 +1966,7 @@ namespace Chaos
 		
 		for (ISimCallbackObject* SimCallbackObject : PushData.SimCallbackObjectsToAdd)
 		{
-			if (!SimCallbackObject->bPendingDelete)
-			{
-				SimCallbackObject->PostInitialize_Internal();
-			}
+			SimCallbackObject->PostInitialize_Internal();
 		}
 
 		if (MRewindCallback && !IsShuttingDown())
