Fix propagation of inertialization requests with tags has some configuration could cause them to not be handled
Fix inertialization requests compare operator that was missing the tag

diff --git a/Engine/Plugins/Animation/BlendStack/Source/Runtime/Private/AnimNode_BlendStack.cpp b/Engine/Plugins/Animation/BlendStack/Source/Runtime/Private/AnimNode_BlendStack.cpp
--- a/Engine/Plugins/Animation/BlendStack/Source/Runtime/Private/AnimNode_BlendStack.cpp
+++ b/Engine/Plugins/Animation/BlendStack/Source/Runtime/Private/AnimNode_BlendStack.cpp
@@ -969,26 +969,7 @@ static void RequestInertialBlend(const FAnimationUpdateContext& Context, float B
 
 	if (BlendTime > 0.0f)
 	{
-		UE::Anim::IInertializationRequester* InertializationRequester = nullptr;
-		// If was have an inertialization node tag, find the corresponding node
-		if (InertialBlendNodeTag != NAME_None)
-		{
-			InertializationRequester = Context.FindMessage<UE::Anim::IInertializationRequester>([&InertialBlendNodeTag](UE::Anim::IInertializationRequester& InMessage)
-				{
-					if (InMessage.GetTag() == InertialBlendNodeTag)
-					{
-						return true;
-					}
-					return false;
-				});
-			ensureMsgf(InertializationRequester, TEXT("FAnimNode_BlendStack with inertialization node tag '%s' could not find a matching node to inertialize"), *InertialBlendNodeTag.ToString());
-
-		}
-		else
-		{
-			InertializationRequester = Context.GetMessage<UE::Anim::IInertializationRequester>();
-		}
-		
+		UE::Anim::IInertializationRequester* InertializationRequester = Context.GetMessage<UE::Anim::IInertializationRequester>();
 		if (InertializationRequester)
 		{
 			FInertializationRequest Request;
diff --git a/Engine/Source/Runtime/Engine/Classes/Animation/AnimInertializationRequest.h b/Engine/Source/Runtime/Engine/Classes/Animation/AnimInertializationRequest.h
--- a/Engine/Source/Runtime/Engine/Classes/Animation/AnimInertializationRequest.h
+++ b/Engine/Source/Runtime/Engine/Classes/Animation/AnimInertializationRequest.h
@@ -39,7 +39,8 @@ struct FInertializationRequest
 			(A.BlendProfile == B.BlendProfile) &&
 			(A.bUseBlendMode == B.bUseBlendMode) &&
 			(A.BlendMode == B.BlendMode) &&
-			(A.CustomBlendCurve == B.CustomBlendCurve);
+			(A.CustomBlendCurve == B.CustomBlendCurve) &&
+			(A.Tag == B.Tag);
 	}
 
 	friend bool operator!=(const FInertializationRequest& A, const FInertializationRequest& B)
diff --git a/Engine/Source/Runtime/Engine/Private/Animation/AnimNode_DeadBlending.cpp b/Engine/Source/Runtime/Engine/Private/Animation/AnimNode_DeadBlending.cpp
--- a/Engine/Source/Runtime/Engine/Private/Animation/AnimNode_DeadBlending.cpp
+++ b/Engine/Source/Runtime/Engine/Private/Animation/AnimNode_DeadBlending.cpp
@@ -708,11 +708,7 @@ void FAnimNode_DeadBlending::RequestInertialization(const FInertializationReques
 {
 	if (Request.Duration >= 0.0f)
 	{
-		// If the request has a specific tag, it needs to match ours
-		if (Request.Tag == NAME_None || Request.Tag == GetTag())
-		{
-			RequestQueue.AddUnique(Request);
-		}
+		RequestQueue.AddUnique(Request);
 	}
 }
 
@@ -934,104 +930,111 @@ void FAnimNode_DeadBlending::Evaluate_AnyThread(FPoseContext& Output)
 
 	if (!RequestQueue.IsEmpty())
 	{
-		const int32 NumSkeletonBones = UE::Anim::DeadBlending::Private::GetNumSkeletonBones(Output.AnimInstanceProxy->GetRequiredBones());
 
 		// Find shortest request
-
 		int32 ShortestRequestIdx = INDEX_NONE;
 		float ShortestRequestDuration = UE_MAX_FLT;
 		for (int32 RequestIdx = 0; RequestIdx < RequestQueue.Num(); RequestIdx++)
 		{
-			if (RequestQueue[RequestIdx].Duration < ShortestRequestDuration)
+			const FInertializationRequest& Request = RequestQueue[RequestIdx];
+			
+			// If the request has a specific tag, it needs to match ours
+			if (Request.Tag == NAME_None || Request.Tag == GetTag())
 			{
-				ShortestRequestIdx = RequestIdx;
-				ShortestRequestDuration = RequestQueue[RequestIdx].Duration;
+				if (Request.Duration < ShortestRequestDuration)
+				{
+					ShortestRequestIdx = RequestIdx;
+					ShortestRequestDuration = Request.Duration;
+				}
 			}
 		}
 
-		// Record Request
+		if (ShortestRequestIdx != INDEX_NONE)
+		{
+			const int32 NumSkeletonBones = UE::Anim::DeadBlending::Private::GetNumSkeletonBones(Output.AnimInstanceProxy->GetRequiredBones());
 
-		InertializationTime = 0.0f;
-		InertializationDuration = BlendTimeMultiplier * RequestQueue[ShortestRequestIdx].Duration;
-		InertializationDurationPerBone.Init(InertializationDuration, NumSkeletonBones);
+			// Record Request
 
-		const USkeleton* TargetSkeleton = Output.AnimInstanceProxy->GetRequiredBones().GetSkeletonAsset();
-		if (RequestQueue[ShortestRequestIdx].BlendProfile)
-		{
-			RequestQueue[ShortestRequestIdx].BlendProfile->FillSkeletonBoneDurationsArray(InertializationDurationPerBone, InertializationDuration, TargetSkeleton);
-		}
-		else if (DefaultBlendProfile)
-		{
-			DefaultBlendProfile->FillSkeletonBoneDurationsArray(InertializationDurationPerBone, InertializationDuration, TargetSkeleton);
-		}
+			InertializationTime = 0.0f;
+			InertializationDuration = BlendTimeMultiplier * RequestQueue[ShortestRequestIdx].Duration;
+			InertializationDurationPerBone.Init(InertializationDuration, NumSkeletonBones);
 
-		// Cache the maximum duration across all bones (so we know when to deactivate the inertialization request)
-		InertializationMaxDuration = FMath::Max(InertializationDuration, *Algo::MaxElement(InertializationDurationPerBone));
+			const USkeleton* TargetSkeleton = Output.AnimInstanceProxy->GetRequiredBones().GetSkeletonAsset();
+			if (RequestQueue[ShortestRequestIdx].BlendProfile)
+			{
+				RequestQueue[ShortestRequestIdx].BlendProfile->FillSkeletonBoneDurationsArray(InertializationDurationPerBone, InertializationDuration, TargetSkeleton);
+			}
+			else if (DefaultBlendProfile)
+			{
+				DefaultBlendProfile->FillSkeletonBoneDurationsArray(InertializationDurationPerBone, InertializationDuration, TargetSkeleton);
+			}
 
-		if (RequestQueue[ShortestRequestIdx].bUseBlendMode)
-		{
-			InertializationBlendMode = RequestQueue[ShortestRequestIdx].BlendMode;
-			InertializationCustomBlendCurve = RequestQueue[ShortestRequestIdx].CustomBlendCurve;
-		}
+			// Cache the maximum duration across all bones (so we know when to deactivate the inertialization request)
+			InertializationMaxDuration = FMath::Max(InertializationDuration, *Algo::MaxElement(InertializationDurationPerBone));
+
+			if (RequestQueue[ShortestRequestIdx].bUseBlendMode)
+			{
+				InertializationBlendMode = RequestQueue[ShortestRequestIdx].BlendMode;
+				InertializationCustomBlendCurve = RequestQueue[ShortestRequestIdx].CustomBlendCurve;
+			}
 
 #if ANIM_TRACE_ENABLED
-		InertializationRequestDescription = RequestQueue[ShortestRequestIdx].DescriptionString;
-		InertializationRequestNodeId = RequestQueue[ShortestRequestIdx].NodeId;
-		InertializationRequestAnimInstance = RequestQueue[ShortestRequestIdx].AnimInstance;
+			InertializationRequestDescription = RequestQueue[ShortestRequestIdx].DescriptionString;
+			InertializationRequestNodeId = RequestQueue[ShortestRequestIdx].NodeId;
+			InertializationRequestAnimInstance = RequestQueue[ShortestRequestIdx].AnimInstance;
 #endif
 
-		// Override with defaults
+			// Override with defaults
 
-		if (bAlwaysUseDefaultBlendSettings)
-		{
-			InertializationDuration = BlendTimeMultiplier * DefaultBlendDuration;
-			InertializationDurationPerBone.Init(BlendTimeMultiplier * DefaultBlendDuration, NumSkeletonBones);
-			InertializationMaxDuration = BlendTimeMultiplier * DefaultBlendDuration;
-			InertializationBlendMode = DefaultBlendMode;
-			InertializationCustomBlendCurve = DefaultCustomBlendCurve;
-		}
+			if (bAlwaysUseDefaultBlendSettings)
+			{
+				InertializationDuration = BlendTimeMultiplier * DefaultBlendDuration;
+				InertializationDurationPerBone.Init(BlendTimeMultiplier * DefaultBlendDuration, NumSkeletonBones);
+				InertializationMaxDuration = BlendTimeMultiplier * DefaultBlendDuration;
+				InertializationBlendMode = DefaultBlendMode;
+				InertializationCustomBlendCurve = DefaultCustomBlendCurve;
+			}
 
-		check(InertializationDuration != UE_MAX_FLT);
-		check(InertializationMaxDuration != UE_MAX_FLT);
+			check(InertializationDuration != UE_MAX_FLT);
+			check(InertializationMaxDuration != UE_MAX_FLT);
 
-		// Reset Request Queue
+			// Initialize the recorded pose state at the point of transition
 
-		RequestQueue.Reset();
+			if (!PrevPoseSnapshot.IsEmpty() && !CurrPoseSnapshot.IsEmpty())
+			{
+				// We have two previous poses and so can initialize as normal.
+
+				InitFrom(
+					Output.Pose,
+					Output.Curve,
+					Output.CustomAttributes,
+					PrevPoseSnapshot,
+					CurrPoseSnapshot);
+			}
+			else if (!CurrPoseSnapshot.IsEmpty())
+			{
+				// We only have a single previous pose. Repeat this pose assuming zero velocity.
+
+				InitFrom(
+					Output.Pose,
+					Output.Curve,
+					Output.CustomAttributes,
+					CurrPoseSnapshot,
+					CurrPoseSnapshot);
+			}
+			else
+			{
+				// This should never happen because we are not able to issue an inertialization 
+				// requested until we have at least one pose recorded in the snapshots.
+				check(false);
+			}
 
-		// Initialize the recorded pose state at the point of transition
+			// Set state to active
 
-		if (!PrevPoseSnapshot.IsEmpty() && !CurrPoseSnapshot.IsEmpty())
-		{
-			// We have two previous poses and so can initialize as normal.
-
-			InitFrom(
-				Output.Pose,
-				Output.Curve,
-				Output.CustomAttributes,
-				PrevPoseSnapshot,
-				CurrPoseSnapshot);
-		}
-		else if (!CurrPoseSnapshot.IsEmpty())
-		{
-			// We only have a single previous pose. Repeat this pose assuming zero velocity.
-
-			InitFrom(
-				Output.Pose,
-				Output.Curve,
-				Output.CustomAttributes,
-				CurrPoseSnapshot,
-				CurrPoseSnapshot);
-		}
-		else
-		{
-			// This should never happen because we are not able to issue an inertialization 
-			// requested until we have at least one pose recorded in the snapshots.
-			check(false);
+			InertializationState = EInertializationState::Active;
 		}
-
-		// Set state to active
-
-		InertializationState = EInertializationState::Active;
+		// Reset Request Queue
+		RequestQueue.Reset();
 	}
 
 	// Update Time Since Transition and deactivate if blend is over
diff --git a/Engine/Source/Runtime/Engine/Private/Animation/AnimNode_Inertialization.cpp b/Engine/Source/Runtime/Engine/Private/Animation/AnimNode_Inertialization.cpp
--- a/Engine/Source/Runtime/Engine/Private/Animation/AnimNode_Inertialization.cpp
+++ b/Engine/Source/Runtime/Engine/Private/Animation/AnimNode_Inertialization.cpp
@@ -274,11 +274,7 @@ void FAnimNode_Inertialization::RequestInertialization(const FInertializationReq
 {
 	if (Request.Duration >= 0.0f)
 	{
-		// If the request has a specific tag, it needs to match ours
-		if (Request.Tag == NAME_None || Request.Tag == GetTag())
-		{
-			RequestQueue.AddUnique(Request);
-		}
+		RequestQueue.AddUnique(Request);
 	}
 }
 
@@ -467,6 +463,16 @@ void FAnimNode_Inertialization::Evaluate_AnyThread(FPoseContext& Output)
 		return;
 	}
 
+	// Filter requests with tags that do not match ours
+	for (int32 RequestIndex = RequestQueue.Num()-1; RequestIndex >= 0; --RequestIndex)
+	{
+		const FInertializationRequest& Request = RequestQueue[RequestIndex];
+		if (Request.Tag != NAME_None && Request.Tag != GetTag())
+		{
+			RequestQueue.RemoveAtSwap(RequestIndex);
+		}
+	}
+
 	// Update the inertialization state if a new inertialization request is pending
 	const int32 NumRequests = RequestQueue.Num();
 	if (NumRequests > 0 && !CurrPoseSnapshot.IsEmpty())
