Audio Bus pre/post sends - Hide irrelevant properties based on Send Level Control Method

diff --git a/Engine/Source/Runtime/Engine/Classes/Sound/SoundSourceBusSend.h b/Engine/Source/Runtime/Engine/Classes/Sound/SoundSourceBusSend.h
--- a/Engine/Source/Runtime/Engine/Classes/Sound/SoundSourceBusSend.h
+++ b/Engine/Source/Runtime/Engine/Classes/Sound/SoundSourceBusSend.h
@@ -29,42 +29,46 @@ struct FSoundSourceBusSendInfo
 
 	/*
 		Manual: Use Send Level only
-		Linear: Interpolate between Min and Max Send Levels based on listener distance (between Distance Min and Distance Max)
-		Custom Curve: Use the float curve to map Send Level to distance (0.0-1.0 on curve maps to Distance Min - Distance Max)
+		Linear: Interpolate between Min and Max Send Levels based on listener distance (between Min/Max Send Distance)
+		Custom Curve: Use the float curve to map Send Level to distance (0.0-1.0 on curve maps to Min/Max Send Distance)
 	*/
-	UPROPERTY(EditAnywhere, Category = BusSend)
+	UPROPERTY(EditAnywhere, Category = BusSend, meta = (DisplayName = "Send Level Control Method"))
 	ESourceBusSendLevelControlMethod SourceBusSendLevelControlMethod;
 
-	// A source Bus to send the audio to. Source buses sonify (make audible) the audio sent to it and are themselves sounds which take up a voice slot in the audio engine.
+	// The Source Bus to send the audio to
 	UPROPERTY(EditAnywhere, Category = BusSend)
 	TObjectPtr<USoundSourceBus> SoundSourceBus;
 
-	// An audio bus to send the audio to. Audio buses can be used to route audio to DSP effects or other purposes. E.g. side-chaining, analysis, etc. Audio buses are not audible unless hooked up to a source bus.
+	// The Audio Bus to send the audio to
 	UPROPERTY(EditAnywhere, Category = BusSend)
 	TObjectPtr<UAudioBus> AudioBus;
 
-	// The amount of audio to send to the bus.
-	UPROPERTY(EditAnywhere, Category = BusSend, meta = (DisplayName = "Manual Send Level"))
+	// Manually set the amount of audio to send to the bus
+	UPROPERTY(EditAnywhere, Category = BusSend, meta = (DisplayName = "Manual Send Level", EditCondition = "SourceBusSendLevelControlMethod == ESourceBusSendLevelControlMethod::Manual", EditConditionHides))
 	float SendLevel;
 
-	// The amount to send to the bus when sound is located at a distance equal to value specified in the min send distance.
-	UPROPERTY(EditAnywhere, BlueprintReadWrite, Category = BusSend)
+	// The amount to send to the bus when sound is located at a distance less than or equal to value specified in the Min Send Distance
+	UPROPERTY(EditAnywhere, BlueprintReadWrite, Category = BusSend, meta = (EditCondition = "SourceBusSendLevelControlMethod == ESourceBusSendLevelControlMethod::Linear", EditConditionHides))
 	float MinSendLevel;
 
-	// The amount to send to the bus when sound is located at a distance equal to value specified in the max send distance.
-	UPROPERTY(EditAnywhere, BlueprintReadWrite, Category = BusSend)
+	// The amount to send to the bus when sound is located at a distance greater than or equal to value specified in the Max Send Distance
+	UPROPERTY(EditAnywhere, BlueprintReadWrite, Category = BusSend, meta = (EditCondition = "SourceBusSendLevelControlMethod == ESourceBusSendLevelControlMethod::Linear", EditConditionHides))
 	float MaxSendLevel;
 
-	// The distance at which the min send Level is sent to the bus
-	UPROPERTY(EditAnywhere, BlueprintReadWrite, Category = BusSend)
+	/** The distance at which to start mapping between to Min/Max Send Level
+	 *  Distances LESS than this will result in a clamped Min Send Level
+	 */ 
+	UPROPERTY(EditAnywhere, BlueprintReadWrite, Category = BusSend, meta = (EditCondition = "SourceBusSendLevelControlMethod != ESourceBusSendLevelControlMethod::Manual", EditConditionHides))
 	float MinSendDistance;
 
-	// The distance at which the max send level is sent to the bus
-	UPROPERTY(EditAnywhere, BlueprintReadWrite, Category = BusSend)
+	/** The distance at which to stop mapping between Min/Max Send Level
+	 *  Distances GREATER than this will result in a clamped Max Send Level
+	 */
+	UPROPERTY(EditAnywhere, BlueprintReadWrite, Category = BusSend, meta = (EditCondition = "SourceBusSendLevelControlMethod != ESourceBusSendLevelControlMethod::Manual", EditConditionHides))
 	float MaxSendDistance;
 
-	// The custom curve to use for distance-based bus send level.
-	UPROPERTY(EditAnywhere, BlueprintReadWrite, Category = BusSend)
+	// The custom send curve to use for distance-based send level. (0.0-1.0 on the curve's X-axis maps to Min/Max Send Distance)
+	UPROPERTY(EditAnywhere, BlueprintReadWrite, Category = BusSend, meta = (EditCondition = "SourceBusSendLevelControlMethod == ESourceBusSendLevelControlMethod::CustomCurve", EditConditionHides))
 	FRuntimeFloatCurve CustomSendLevelCurve;
 
 	FSoundSourceBusSendInfo()
diff --git a/Engine/Source/Runtime/Engine/Classes/Sound/SoundSubmixSend.h b/Engine/Source/Runtime/Engine/Classes/Sound/SoundSubmixSend.h
--- a/Engine/Source/Runtime/Engine/Classes/Sound/SoundSubmixSend.h
+++ b/Engine/Source/Runtime/Engine/Classes/Sound/SoundSubmixSend.h
@@ -94,7 +94,7 @@ struct FSoundSubmixSendInfoBase
 	UPROPERTY(EditAnywhere, BlueprintReadWrite, Category = SubmixSend)
 	ESendLevelControlMethod SendLevelControlMethod;
 	
-	// The submix to send the audio to
+	// The Submix to send the audio to
 	UPROPERTY(EditAnywhere, BlueprintReadWrite, Category = SubmixSend)
 	TObjectPtr<USoundSubmixBase> SoundSubmix;
 
@@ -126,7 +126,7 @@ struct FSoundSubmixSendInfoBase
 	UPROPERTY(EditAnywhere, BlueprintReadWrite, Category = SubmixSend, meta = (EditCondition = "SendLevelControlMethod != ESendLevelControlMethod::Manual", EditConditionHides))
 	float MaxSendDistance;
 
-	// The custom reverb send curve to use for distance-based send level. (0.0-1.0 on curve maps to Min/Max Send Distance)
+	// The custom send curve to use for distance-based send level. (0.0-1.0 on the curve's X-axis maps to Min/Max Send Distance)
 	UPROPERTY(EditAnywhere, BlueprintReadWrite, Category = SubmixSend, meta = (EditCondition = "SendLevelControlMethod == ESendLevelControlMethod::CustomCurve", EditConditionHides))
 	FRuntimeFloatCurve CustomSendLevelCurve;
 };
