Restrict safezone broadcast only when changed in preview platform

diff --git a/Engine/Source/Runtime/Engine/Private/LocalPlayer.cpp b/Engine/Source/Runtime/Engine/Private/LocalPlayer.cpp
--- a/Engine/Source/Runtime/Engine/Private/LocalPlayer.cpp
+++ b/Engine/Source/Runtime/Engine/Private/LocalPlayer.cpp
@@ -1276,16 +1276,22 @@ bool ULocalPlayer::GetProjectionData(FViewport* Viewport, FSceneViewProjectionDa
 	{
 		FSlateApplication::Get().OnConstrainedAspectRatioChanged.Broadcast(GEditor->PreviewPlatform.GetConstrainedAspectRatio());
 
-		FVector4f Safezones = FVector4f::Zero();
-		Safezones = GEditor->PreviewPlatform.GetSafeZones();
-		FSlateApplication::Get().SetCustomSafeZone(Safezones);
-		FSlateApplication::Get().OnDebugSafeZoneChanged.Broadcast(Safezones, true);
+		FVector4f Safezones = GEditor->PreviewPlatform.GetSafeZones();
+		if (Safezones != FSlateApplication::Get().GetCustomSafeZone())
+		{
+			FSlateApplication::Get().SetCustomSafeZone(Safezones);
+			FSlateApplication::Get().OnDebugSafeZoneChanged.Broadcast(Safezones, true);
+		}
 	}
 	else
 	{
 		FSlateApplication::Get().OnConstrainedAspectRatioChanged.Broadcast(0.f);
-		FSlateApplication::Get().ResetCustomSafeZone();
-		FSlateApplication::Get().OnDebugSafeZoneChanged.Broadcast(FMargin(), true);
+
+		if (FVector4f::Zero() != FSlateApplication::Get().GetCustomSafeZone())
+		{
+			FSlateApplication::Get().ResetCustomSafeZone();
+			FSlateApplication::Get().OnDebugSafeZoneChanged.Broadcast(FMargin(), true);
+		}
 	}
 #endif
 
