Fix CIS Linux

diff --git a/Engine/Source/Runtime/Engine/Private/LocalPlayer.cpp b/Engine/Source/Runtime/Engine/Private/LocalPlayer.cpp
--- a/Engine/Source/Runtime/Engine/Private/LocalPlayer.cpp
+++ b/Engine/Source/Runtime/Engine/Private/LocalPlayer.cpp
@@ -1276,7 +1276,7 @@ bool ULocalPlayer::GetProjectionData(FViewport* Viewport, FSceneViewProjectionDa
 	{
 		FSlateApplication::Get().OnConstrainedAspectRatioChanged.Broadcast(GEditor->PreviewPlatform.GetConstrainedAspectRatio());
 
-		FVector4f Safezones = GEditor->PreviewPlatform.GetSafeZones();
+		FMargin Safezones(GEditor->PreviewPlatform.GetSafeZones());
 		if (Safezones != FSlateApplication::Get().GetCustomSafeZone())
 		{
 			FSlateApplication::Get().SetCustomSafeZone(Safezones);
@@ -1287,7 +1287,7 @@ bool ULocalPlayer::GetProjectionData(FViewport* Viewport, FSceneViewProjectionDa
 	{
 		FSlateApplication::Get().OnConstrainedAspectRatioChanged.Broadcast(0.f);
 
-		if (FVector4f::Zero() != FSlateApplication::Get().GetCustomSafeZone())
+		if (FMargin() != FSlateApplication::Get().GetCustomSafeZone())
 		{
 			FSlateApplication::Get().ResetCustomSafeZone();
 			FSlateApplication::Get().OnDebugSafeZoneChanged.Broadcast(FMargin(), true);
