AudioLink fix Cvar au.AudioLink.SpatializeAllSounds

diff --git a/Engine/Source/Runtime/Engine/Private/ActiveSound.cpp b/Engine/Source/Runtime/Engine/Private/ActiveSound.cpp
--- a/Engine/Source/Runtime/Engine/Private/ActiveSound.cpp
+++ b/Engine/Source/Runtime/Engine/Private/ActiveSound.cpp
@@ -1040,7 +1040,7 @@ void FActiveSound::UpdateWaveInstances(TArray<FWaveInstance*> &InWaveInstances,
 	if (FadeOut == EFadeOut::None || ComponentVolumeFader.IsActive())
 	{
 		bool bReverbSendLevelWasSet = false;
-		if (bHasAttenuationSettings)
+		if (bHasAttenuationSettings || bAudioLinkSpatializeAllSoundsCVar)
 		{
 			UpdateAttenuation(DeltaTime, ParseParams, ClosestListenerIndex);
 			bReverbSendLevelWasSet = true;
