[FixupForPIE]
- Fix extremely deep callstack by avoiding recursion

diff --git a/Engine/Source/Runtime/Engine/Private/PieFixupSerializer.cpp b/Engine/Source/Runtime/Engine/Private/PieFixupSerializer.cpp
--- a/Engine/Source/Runtime/Engine/Private/PieFixupSerializer.cpp
+++ b/Engine/Source/Runtime/Engine/Private/PieFixupSerializer.cpp
@@ -55,9 +55,29 @@ FArchive& FPIEFixupSerializer::operator<<(UObject*& Object)
 		// Skip instanced static mesh component as their impact on serialization is enormous and they don't contain lazy ptrs.
 		if (!Cast<UInstancedStaticMeshComponent>(Object))
 		{
-			Object->Serialize(*this);
+			if (RecursionCount > 0)
+			{
+				DeferredSerialize.Add(Object);
+			}
+			else
+			{
+				RecursionCount++;
+				Object->Serialize(*this);
+				RecursionCount--;
+			}
 		}
 	}
+
+	if (RecursionCount == 0)
+	{
+		RecursionCount++;
+		while (DeferredSerialize.Num())
+		{
+			DeferredSerialize.Pop(EAllowShrinking::No)->Serialize(*this);
+		}
+		RecursionCount--;
+	}
+
 	return *this;
 }
 
diff --git a/Engine/Source/Runtime/Engine/Private/PieFixupSerializer.h b/Engine/Source/Runtime/Engine/Private/PieFixupSerializer.h
--- a/Engine/Source/Runtime/Engine/Private/PieFixupSerializer.h
+++ b/Engine/Source/Runtime/Engine/Private/PieFixupSerializer.h
@@ -21,6 +21,8 @@ public:
 private:
 	TFunctionRef<void(int32, FSoftObjectPath&)> SoftObjectPathFixupFunction;
 	TSet<UObject*> VisitedObjects;
+	TArray<UObject*, TInlineAllocator<32>> DeferredSerialize;
+	int32 RecursionCount = 0;
 	UObject* Root;
 	int32 PIEInstanceID;
 };
\ No newline at end of file
