Fix tooltips for UBodySetup::bNeverNeedsCookedCollisionData so in the editor we show the full description for that field

diff --git a/Engine/Source/Runtime/Engine/Classes/PhysicsEngine/BodySetup.h b/Engine/Source/Runtime/Engine/Classes/PhysicsEngine/BodySetup.h
--- a/Engine/Source/Runtime/Engine/Classes/PhysicsEngine/BodySetup.h
+++ b/Engine/Source/Runtime/Engine/Classes/PhysicsEngine/BodySetup.h
@@ -196,8 +196,8 @@ class UBodySetup : public UBodySetupCore
 	/** Indicates whether this setup has any cooked collision data. This is set when writing to or reading from cooked data, but not authoritative for assets in the editor. */
 	uint8 bHasCookedCollisionData:1;
 
-	/** Indicates that we will never use convex or trimesh shapes. This is an optimization to skip checking for binary data. */
 	/** 
+	 * Indicates that we will never use convex or trimesh shapes. This is an optimization to skip checking for binary data.
 	 * TODO Chaos this is to opt out of CreatePhysicsMeshes for certain meshes
 	 * Better long term mesh is to not call CreatePhysicsMeshes until it is known there is a mesh instance that needs it.
 	 */
