[Truncation Warnings] Update EdGraphSchema, EdGraphActions, and SGraphEditor to use FVector2f

diff --git a/Engine/Source/Editor/BlueprintGraph/Private/EdGraphSchema_K2_Actions.cpp b/Engine/Source/Editor/BlueprintGraph/Private/EdGraphSchema_K2_Actions.cpp
--- a/Engine/Source/Editor/BlueprintGraph/Private/EdGraphSchema_K2_Actions.cpp
+++ b/Engine/Source/Editor/BlueprintGraph/Private/EdGraphSchema_K2_Actions.cpp
@@ -325,7 +325,7 @@ UEdGraphNode* FEdGraphSchemaAction_K2NewNode::CreateNode(
 	return ResultNode;
 }
 
-UEdGraphNode* FEdGraphSchemaAction_K2NewNode::PerformAction(class UEdGraph* ParentGraph, UEdGraphPin* FromPin, const FVector2D Location, bool bSelectNewNode/* = true*/)
+UEdGraphNode* FEdGraphSchemaAction_K2NewNode::PerformAction(class UEdGraph* ParentGraph, UEdGraphPin* FromPin, const FVector2f& Location, bool bSelectNewNode/* = true*/)
 {
 	EK2NewNodeFlags Options = EK2NewNodeFlags::None;
 	if (bSelectNewNode)
@@ -336,10 +336,10 @@ UEdGraphNode* FEdGraphSchemaAction_K2NewNode::PerformAction(class UEdGraph* Pare
 	{
 		Options |= EK2NewNodeFlags::GotoNewNode;
 	}
-	return CreateNode(ParentGraph, MakeArrayView(&FromPin, 1), Location, NodeTemplate, Options);
+	return CreateNode(ParentGraph, MakeArrayView(&FromPin, 1), FDeprecateSlateVector2D(Location), NodeTemplate, Options);
 }
 
-UEdGraphNode* FEdGraphSchemaAction_K2NewNode::PerformAction(class UEdGraph* ParentGraph, TArray<UEdGraphPin*>& FromPins, const FVector2D Location, bool bSelectNewNode/* = true*/) 
+UEdGraphNode* FEdGraphSchemaAction_K2NewNode::PerformAction(class UEdGraph* ParentGraph, TArray<UEdGraphPin*>& FromPins, const FVector2f& Location, bool bSelectNewNode/* = true*/) 
 {
 	EK2NewNodeFlags Options = EK2NewNodeFlags::None;
 	if (bSelectNewNode)
@@ -350,7 +350,7 @@ UEdGraphNode* FEdGraphSchemaAction_K2NewNode::PerformAction(class UEdGraph* Pare
 	{
 		Options |= EK2NewNodeFlags::GotoNewNode;
 	}
-	return CreateNode(ParentGraph, FromPins, Location, NodeTemplate, Options);
+	return CreateNode(ParentGraph, FromPins, FDeprecateSlateVector2D(Location), NodeTemplate, Options);
 }
 
 void FEdGraphSchemaAction_K2NewNode::AddReferencedObjects( FReferenceCollector& Collector )
@@ -364,7 +364,7 @@ void FEdGraphSchemaAction_K2NewNode::AddReferencedObjects( FReferenceCollector&
 /////////////////////////////////////////////////////
 // FEdGraphSchemaAction_K2ViewNode
 
-UEdGraphNode* FEdGraphSchemaAction_K2ViewNode::PerformAction(class UEdGraph* ParentGraph, UEdGraphPin* FromPin, const FVector2D Location, bool bSelectNewNode/* = true*/)
+UEdGraphNode* FEdGraphSchemaAction_K2ViewNode::PerformAction(class UEdGraph* ParentGraph, UEdGraphPin* FromPin, const FVector2f& Location, bool bSelectNewNode/* = true*/)
 {
 	// If the node is valid, select it
 	if (NodePtr)
@@ -376,7 +376,7 @@ UEdGraphNode* FEdGraphSchemaAction_K2ViewNode::PerformAction(class UEdGraph* Par
 	return nullptr;
 }
 
-UEdGraphNode* FEdGraphSchemaAction_K2ViewNode::PerformAction(class UEdGraph* ParentGraph, TArray<UEdGraphPin*>& FromPins, const FVector2D Location, bool bSelectNewNode/* = true*/) 
+UEdGraphNode* FEdGraphSchemaAction_K2ViewNode::PerformAction(class UEdGraph* ParentGraph, TArray<UEdGraphPin*>& FromPins, const FVector2f& Location, bool bSelectNewNode/* = true*/) 
 {
 	PerformAction(ParentGraph, nullptr, Location, bSelectNewNode);
 	return nullptr;
@@ -420,15 +420,15 @@ UEdGraphNode* FEdGraphSchemaAction_K2AssignDelegate::AssignDelegate(class UK2Nod
 	return BindNode;
 }
 
-UEdGraphNode* FEdGraphSchemaAction_K2AssignDelegate::PerformAction(class UEdGraph* ParentGraph, UEdGraphPin* FromPin, const FVector2D Location, bool bSelectNewNode/* = true*/)
+UEdGraphNode* FEdGraphSchemaAction_K2AssignDelegate::PerformAction(class UEdGraph* ParentGraph, UEdGraphPin* FromPin, const FVector2f& Location, bool bSelectNewNode/* = true*/)
 {
-	return AssignDelegate(NodeTemplate, ParentGraph, FromPin, Location, bSelectNewNode);
+	return AssignDelegate(NodeTemplate, ParentGraph, FromPin, FDeprecateSlateVector2D(Location), bSelectNewNode);
 }
 
 /////////////////////////////////////////////////////
 // FEdGraphSchemaAction_EventFromFunction
 
-UEdGraphNode* FEdGraphSchemaAction_EventFromFunction::PerformAction(class UEdGraph* ParentGraph, UEdGraphPin* FromPin, const FVector2D Location, bool bSelectNewNode/* = true*/)
+UEdGraphNode* FEdGraphSchemaAction_EventFromFunction::PerformAction(class UEdGraph* ParentGraph, UEdGraphPin* FromPin, const FVector2f& Location, bool bSelectNewNode/* = true*/)
 {
 	UK2Node_CustomEvent* EventNode = nullptr;
 	if (SignatureFunction)
@@ -455,7 +455,7 @@ UEdGraphNode* FEdGraphSchemaAction_EventFromFunction::PerformAction(class UEdGra
 			FromPin->Modify();
 		}
 
-		EventNode = UK2Node_CustomEvent::CreateFromFunction(Location, ParentGraph, SignatureFunction->GetName() + TEXT("_Event"), SignatureFunction, bSelectNewNode);
+		EventNode = UK2Node_CustomEvent::CreateFromFunction(FDeprecateSlateVector2D(Location), ParentGraph, SignatureFunction->GetName() + TEXT("_Event"), SignatureFunction, bSelectNewNode);
 		EventNode->AutowireNewNode(FromPin);
 
 		UBlueprint* Blueprint = FBlueprintEditorUtils::FindBlueprintForGraphChecked(ParentGraph);
@@ -464,7 +464,7 @@ UEdGraphNode* FEdGraphSchemaAction_EventFromFunction::PerformAction(class UEdGra
 	return EventNode;
 }
 
-UEdGraphNode* FEdGraphSchemaAction_EventFromFunction::PerformAction(class UEdGraph* ParentGraph, TArray<UEdGraphPin*>& FromPins, const FVector2D Location, bool bSelectNewNode/* = true*/)
+UEdGraphNode* FEdGraphSchemaAction_EventFromFunction::PerformAction(class UEdGraph* ParentGraph, TArray<UEdGraphPin*>& FromPins, const FVector2f& Location, bool bSelectNewNode/* = true*/)
 {
 	UEdGraphNode* ResultNode = nullptr;
 
@@ -497,7 +497,7 @@ void FEdGraphSchemaAction_EventFromFunction::AddReferencedObjects(FReferenceColl
 /////////////////////////////////////////////////////
 // FEdGraphSchemaAction_K2AddComponent
 
-UEdGraphNode* FEdGraphSchemaAction_K2AddComponent::PerformAction(class UEdGraph* ParentGraph, UEdGraphPin* FromPin, const FVector2D Location, bool bSelectNewNode/* = true*/)
+UEdGraphNode* FEdGraphSchemaAction_K2AddComponent::PerformAction(class UEdGraph* ParentGraph, UEdGraphPin* FromPin, const FVector2f& Location, bool bSelectNewNode/* = true*/)
 {
 	if (ComponentClass == nullptr)
 	{
@@ -554,7 +554,7 @@ void FEdGraphSchemaAction_K2AddComponent::AddReferencedObjects( FReferenceCollec
 /////////////////////////////////////////////////////
 // FEdGraphSchemaAction_K2AddEvent
 
-UEdGraphNode* FEdGraphSchemaAction_K2AddEvent::PerformAction(class UEdGraph* ParentGraph, UEdGraphPin* FromPin, const FVector2D Location, bool bSelectNewNode/* = true*/)
+UEdGraphNode* FEdGraphSchemaAction_K2AddEvent::PerformAction(class UEdGraph* ParentGraph, UEdGraphPin* FromPin, const FVector2f& Location, bool bSelectNewNode/* = true*/)
 {
 	UEdGraphNode* NewNode = nullptr;
 	const FScopedTransaction Transaction( NSLOCTEXT("UnrealEd", "K2_Event", "Add Event") );	
@@ -596,7 +596,7 @@ bool FEdGraphSchemaAction_K2AddEvent::EventHasAlreadyBeenPlaced(UBlueprint const
 /////////////////////////////////////////////////////
 // FEdGraphSchemaAction_K2AddCustomEvent
 
-UEdGraphNode* FEdGraphSchemaAction_K2AddCustomEvent::PerformAction(class UEdGraph* ParentGraph, UEdGraphPin* FromPin, const FVector2D Location, bool bSelectNewNode/* = true*/)
+UEdGraphNode* FEdGraphSchemaAction_K2AddCustomEvent::PerformAction(class UEdGraph* ParentGraph, UEdGraphPin* FromPin, const FVector2f& Location, bool bSelectNewNode/* = true*/)
 {
 	const FScopedTransaction Transaction( NSLOCTEXT("UnrealEd", "K2_CustomEvent", "Add Custom Event") );
 	UBlueprint* Blueprint = FBlueprintEditorUtils::FindBlueprintForGraphChecked(ParentGraph);
@@ -613,16 +613,16 @@ UEdGraphNode* FEdGraphSchemaAction_K2AddCustomEvent::PerformAction(class UEdGrap
 /////////////////////////////////////////////////////
 // FEdGraphSchemaAction_K2AddCallOnActor
 
-UEdGraphNode* FEdGraphSchemaAction_K2AddCallOnActor::PerformAction(class UEdGraph* ParentGraph, UEdGraphPin* FromPin, const FVector2D Location, bool bSelectNewNode/* = true*/)
+UEdGraphNode* FEdGraphSchemaAction_K2AddCallOnActor::PerformAction(class UEdGraph* ParentGraph, UEdGraphPin* FromPin, const FVector2f& Location, bool bSelectNewNode/* = true*/)
 {
 	// Snap the node placement location to the grid, ensures calculations later match up better
 	const uint32 GridSnapSize = GetDefault<UEditorStyleSettings>()->GridSnapSize;
-	FVector2D LocalLocation;
+	FVector2f LocalLocation;
 	LocalLocation.X = FMath::GridSnap(Location.X, GridSnapSize);
 	LocalLocation.Y = FMath::GridSnap(Location.Y, GridSnapSize);
 
 	// First use the base functionality to spawn the 'call function' node
-	UEdGraphNode* CallNode = FEdGraphSchemaAction_K2NewNode::PerformAction(ParentGraph, FromPin, FVector2D(LocalLocation));
+	UEdGraphNode* CallNode = FEdGraphSchemaAction_K2NewNode::PerformAction(ParentGraph, FromPin, LocalLocation);
 	const float FunctionNodeHeightUnsnapped = UEdGraphSchema_K2::EstimateNodeHeight( CallNode );
 
 	// this is the guesstimate of the function node's height, snapped to grid units
@@ -678,14 +678,14 @@ void FEdGraphSchemaAction_K2AddCallOnActor::AddReferencedObjects( FReferenceColl
 /////////////////////////////////////////////////////
 // FEdGraphSchemaAction_K2AddComment
 
-UEdGraphNode* FEdGraphSchemaAction_K2AddComment::PerformAction(class UEdGraph* ParentGraph, UEdGraphPin* FromPin, const FVector2D Location, bool bSelectNewNode/* = true*/)
+UEdGraphNode* FEdGraphSchemaAction_K2AddComment::PerformAction(class UEdGraph* ParentGraph, UEdGraphPin* FromPin, const FVector2f& Location, bool bSelectNewNode/* = true*/)
 {
 	// Add menu item for creating comment boxes
 	UEdGraphNode_Comment* CommentTemplate = NewObject<UEdGraphNode_Comment>();
 
 	UBlueprint* Blueprint = FBlueprintEditorUtils::FindBlueprintForGraph(ParentGraph);
 
-	FVector2D SpawnLocation = Location;
+	FVector2f SpawnLocation = Location;
 
 	FSlateRect Bounds;
 	if ((Blueprint != nullptr) && FKismetEditorUtilities::GetBoundsForSelectedNodes(Blueprint, Bounds, 50.0f))
@@ -710,7 +710,7 @@ UEdGraphNode* FEdGraphSchemaAction_K2AddComment::PerformAction(class UEdGraph* P
 /////////////////////////////////////////////////////
 // FEdGraphSchemaAction_K2TargetNode
 
-UEdGraphNode* FEdGraphSchemaAction_K2TargetNode::PerformAction( class UEdGraph* ParentGraph, UEdGraphPin* FromPin, const FVector2D Location, bool bSelectNewNode/* = true*/ ) 
+UEdGraphNode* FEdGraphSchemaAction_K2TargetNode::PerformAction( class UEdGraph* ParentGraph, UEdGraphPin* FromPin, const FVector2f& Location, bool bSelectNewNode/* = true*/ ) 
 {
 	FKismetEditorUtilities::BringKismetToFocusAttentionOnObject(NodeTemplate);
 	return nullptr;
@@ -719,9 +719,9 @@ UEdGraphNode* FEdGraphSchemaAction_K2TargetNode::PerformAction( class UEdGraph*
 /////////////////////////////////////////////////////
 // FEdGraphSchemaAction_K2PasteHere
 
-UEdGraphNode* FEdGraphSchemaAction_K2PasteHere::PerformAction( class UEdGraph* ParentGraph, UEdGraphPin* FromPin, const FVector2D Location, bool bSelectNewNode/* = true*/ ) 
+UEdGraphNode* FEdGraphSchemaAction_K2PasteHere::PerformAction( class UEdGraph* ParentGraph, UEdGraphPin* FromPin, const FVector2f& Location, bool bSelectNewNode/* = true*/ ) 
 {
-	FKismetEditorUtilities::PasteNodesHere(ParentGraph, Location);
+	FKismetEditorUtilities::PasteNodesHere(ParentGraph, FDeprecateSlateVector2D(Location));
 	return nullptr;
 }
 
diff --git a/Engine/Source/Editor/GraphEditor/Private/SGraphEditorImpl.cpp b/Engine/Source/Editor/GraphEditor/Private/SGraphEditorImpl.cpp
--- a/Engine/Source/Editor/GraphEditor/Private/SGraphEditorImpl.cpp
+++ b/Engine/Source/Editor/GraphEditor/Private/SGraphEditorImpl.cpp
@@ -105,7 +105,7 @@ FAlignmentData FAlignmentHelper::GetAlignmentDataForNode(UEdGraphNode* Node)
 /////////////////////////////////////////////////////
 // SGraphEditorImpl
 
-FVector2D SGraphEditorImpl::GetPasteLocation() const
+UE::Slate::FDeprecateVector2DResult SGraphEditorImpl::GetPasteLocation2f() const
 {
 	return GraphPanel->GetPastePosition();
 }
@@ -1018,7 +1018,7 @@ void SGraphEditorImpl::AddContextMenuCommentSection(UToolMenu* InMenu)
 
 							if (Action.IsValid())
 							{
-								Action->PerformAction(Graph, nullptr, FVector2D());
+								Action->PerformAction(Graph, nullptr, FVector2f());
 							}
 						}
 					}
@@ -1447,7 +1447,7 @@ TSharedPtr<SWidget> SGraphEditorImpl::GetTitleBar() const
 	return TitleBar;
 }
 
-void SGraphEditorImpl::SetViewLocation( const FVector2D& Location, float ZoomAmount, const FGuid& BookmarkId ) 
+void SGraphEditorImpl::SetViewLocation( const FVector2f& Location, float ZoomAmount, const FGuid& BookmarkId ) 
 {
 	if( GraphPanel.IsValid() &&  EdGraphObj && (!IsLocked() || !GraphPanel->HasDeferredObjectFocus()))
 	{
@@ -1455,7 +1455,7 @@ void SGraphEditorImpl::SetViewLocation( const FVector2D& Location, float ZoomAmo
 	}
 }
 
-void SGraphEditorImpl::GetViewLocation( FVector2D& Location, float& ZoomAmount ) 
+void SGraphEditorImpl::GetViewLocation( FVector2f& Location, float& ZoomAmount ) 
 {
 	if( GraphPanel.IsValid() &&  EdGraphObj && (!IsLocked() || !GraphPanel->HasDeferredObjectFocus()))
 	{
@@ -1747,7 +1747,7 @@ void DistributeNodes(TArray<FAlignmentData>& InData, bool bIsHorizontal)
 				{
 					FAlignmentData& Entry = InData[Index];
 
-					FVector2D Target2DPosition(Entry.Node->NodePosX, Entry.Node->NodePosY);
+					FVector2f Target2DPosition(Entry.Node->NodePosX, Entry.Node->NodePosY);
  
 					if (bIsHorizontal)
 					{
@@ -1852,20 +1852,20 @@ void StackNodes(TArray<FAlignmentData>& InData, bool bIsHorizontal)
 				{
 					FAlignmentData& Entry = InData[Index];
 
-					FVector2D Target2DPosition(Entry.Node->NodePosX, Entry.Node->NodePosY);
+					FVector2f AdjustedTargetPosition(Entry.Node->NodePosX, Entry.Node->NodePosY);
 
 					if (bIsHorizontal)
 					{
-						Target2DPosition.X = TargetPosition;
-						Target2DPosition.Y = StackAxisPosition;
+						AdjustedTargetPosition.X = TargetPosition;
+						AdjustedTargetPosition.Y = StackAxisPosition;
 					}
 					else
 					{
-						Target2DPosition.X = StackAxisPosition;
-						Target2DPosition.Y = TargetPosition;
+						AdjustedTargetPosition.X = StackAxisPosition;
+						AdjustedTargetPosition.Y = TargetPosition;
 					}
 
-					Schema->SetNodePosition(Entry.Node, Target2DPosition);
+					Schema->SetNodePosition(Entry.Node, AdjustedTargetPosition);
 
 					TargetPosition = Entry.GetTarget() + PaddingAmount;
 				}
diff --git a/Engine/Source/Editor/GraphEditor/Private/SGraphEditorImpl.h b/Engine/Source/Editor/GraphEditor/Private/SGraphEditorImpl.h
--- a/Engine/Source/Editor/GraphEditor/Private/SGraphEditorImpl.h
+++ b/Engine/Source/Editor/GraphEditor/Private/SGraphEditorImpl.h
@@ -123,7 +123,7 @@ struct FAlignmentHelper
 					for (FAlignmentData& Entry : AlignmentData)
 					{
 						int32 TargetProperty = Target - Entry.TargetOffset; 
-						FVector2D TargetPosition(Entry.Node->NodePosX, Entry.Node->NodePosY);
+						FVector2f TargetPosition(Entry.Node->NodePosX, Entry.Node->NodePosY);
 
 						if (Orientation == EOrientation::Orient_Horizontal)
 						{
@@ -317,7 +317,7 @@ public:
 	virtual void SetNodeSelection(UEdGraphNode* Node, bool bSelect) override;
 	virtual void SelectAllNodes() override;
 
-	virtual FVector2D GetPasteLocation() const override;
+	virtual UE::Slate::FDeprecateVector2DResult GetPasteLocation2f() const override;
 	virtual bool IsNodeTitleVisible( const UEdGraphNode* Node, bool bRequestRename ) override;
 	virtual void JumpToNode( const UEdGraphNode* JumpToMe, bool bRequestRename = false, bool bSelectNode = true ) override;
 	virtual void JumpToPin( const UEdGraphPin* JumpToMe ) override;
@@ -328,8 +328,8 @@ public:
 	virtual bool GetBoundsForNode( const UEdGraphNode* InNode, class FSlateRect& Rect, float Padding) const override;
 	virtual void NotifyGraphChanged() override;
 	virtual TSharedPtr<SWidget> GetTitleBar() const override;
-	virtual void SetViewLocation(const FVector2D& Location, float ZoomAmount, const FGuid& BookmarkId = FGuid()) override;
-	virtual void GetViewLocation(FVector2D& Location, float& ZoomAmount) override;
+	virtual void SetViewLocation(const FVector2f& Location, float ZoomAmount, const FGuid& BookmarkId = FGuid()) override;
+	virtual void GetViewLocation(FVector2f& Location, float& ZoomAmount) override;
 	virtual void GetViewBookmark(FGuid& BookmarkId) override;
 	virtual void LockToGraphEditor(TWeakPtr<SGraphEditor> Other) override;
 	virtual void UnlockFromGraphEditor(TWeakPtr<SGraphEditor> Other) override;
diff --git a/Engine/Source/Editor/GraphEditor/Private/SGraphPanel.cpp b/Engine/Source/Editor/GraphEditor/Private/SGraphPanel.cpp
--- a/Engine/Source/Editor/GraphEditor/Private/SGraphPanel.cpp
+++ b/Engine/Source/Editor/GraphEditor/Private/SGraphPanel.cpp
@@ -1070,7 +1070,7 @@ FReply SGraphPanel::OnMouseButtonDoubleClick(const FGeometry& MyGeometry, const
 		const FVector2f DoubleClickPositionInGraphSpace = PanelCoordToGraphCoord(MyGeometry.AbsoluteToLocal(MouseEvent.GetScreenSpacePosition()));
 
 		const UEdGraphSchema* Schema = GraphObj->GetSchema();
-		Schema->OnPinConnectionDoubleCicked(Pin1, Pin2, FVector2D(DoubleClickPositionInGraphSpace));
+		Schema->OnPinConnectionDoubleCicked(Pin1, Pin2, DoubleClickPositionInGraphSpace);
 	}
 	else if (!PreviousFrameSplineOverlap.GetCloseToSpline())
 	{
@@ -1382,7 +1382,7 @@ FReply SGraphPanel::OnDrop( const FGeometry& MyGeometry, const FDragDropEvent& D
 			{
 				if (PassesAssetReferenceFilter(DroppedAssetData))
 				{
-					GraphObj->GetSchema()->DroppedAssetsOnGraph( DroppedAssetData, FVector2D(NodeAddPosition), GraphObj );
+					GraphObj->GetSchema()->DroppedAssetsOnGraph( DroppedAssetData, NodeAddPosition, GraphObj );
 				}
 				return FReply::Handled();
 			}
@@ -1910,7 +1910,7 @@ private:
 			const UEdGraphSchema* Schema = GraphObj->GetSchema();
 
 			float NewNodePosY = NodeToPins.Key->NodePosY + (AlignmentDelta / NodeToPins.Value.Num());
-			Schema->SetNodePosition(NodeToPins.Key, FVector2D(NodeToPins.Key->NodePosX, NewNodePosY));
+			Schema->SetNodePosition(NodeToPins.Key, FVector2f(NodeToPins.Key->NodePosX, NewNodePosY));
 	
 			VisitedNodes.Add(Node);
 			VisitedNodes.Add(NodeToPins.Key);
diff --git a/Engine/Source/Editor/GraphEditor/Private/SGraphPin.cpp b/Engine/Source/Editor/GraphEditor/Private/SGraphPin.cpp
--- a/Engine/Source/Editor/GraphEditor/Private/SGraphPin.cpp
+++ b/Engine/Source/Editor/GraphEditor/Private/SGraphPin.cpp
@@ -887,7 +887,7 @@ FReply SGraphPin::OnDrop( const FGeometry& MyGeometry, const FDragDropEvent& Dra
 			TSharedPtr<FAssetDragDropOp> AssetOp = StaticCastSharedPtr<FAssetDragDropOp>(Operation);
 			if (AssetOp->HasAssets())
 			{
-				Node->GetSchema()->DroppedAssetsOnPin(AssetOp->GetAssets(), FVector2D(NodeWidget->GetPosition2f() + MyGeometry.Position), GraphPinObj);
+				Node->GetSchema()->DroppedAssetsOnPin(AssetOp->GetAssets(), NodeWidget->GetPosition2f() + MyGeometry.Position, GraphPinObj);
 			}
 		}
 		return FReply::Handled();
diff --git a/Engine/Source/Editor/Kismet/Private/BlueprintActionMenuItem.cpp b/Engine/Source/Editor/Kismet/Private/BlueprintActionMenuItem.cpp
--- a/Engine/Source/Editor/Kismet/Private/BlueprintActionMenuItem.cpp
+++ b/Engine/Source/Editor/Kismet/Private/BlueprintActionMenuItem.cpp
@@ -55,7 +55,7 @@ namespace FBlueprintMenuActionItemImpl
 	 * @param  Bindings			Any bindings you want applied after the node has been spawned.
 	 * @return The spawned node (could be an existing one if the event was already placed).
 	 */
-	static UEdGraphNode* InvokeAction(const UBlueprintNodeSpawner* Action, UEdGraph* ParentGraph, FVector2D const Location, IBlueprintNodeBinder::FBindingSet const& Bindings, bool& bOutNewNode);
+	static UEdGraphNode* InvokeAction(const UBlueprintNodeSpawner* Action, UEdGraph* ParentGraph, const FVector2f& Location, IBlueprintNodeBinder::FBindingSet const& Bindings, bool& bOutNewNode);
 
 	/**
 	 * 
@@ -106,11 +106,11 @@ static void FBlueprintMenuActionItemImpl::DirtyBlueprintFromNewNode(UEdGraphNode
 }
 
 //------------------------------------------------------------------------------
-static UEdGraphNode* FBlueprintMenuActionItemImpl::InvokeAction(const UBlueprintNodeSpawner* Action, UEdGraph* ParentGraph, FVector2D const Location, IBlueprintNodeBinder::FBindingSet const& Bindings, bool& bOutNewNode)
+static UEdGraphNode* FBlueprintMenuActionItemImpl::InvokeAction(const UBlueprintNodeSpawner* Action, UEdGraph* ParentGraph, const FVector2f& Location, IBlueprintNodeBinder::FBindingSet const& Bindings, bool& bOutNewNode)
 {
 	int32 const PreSpawnNodeCount = ParentGraph->Nodes.Num();
 	// this could return an existing node
-	UEdGraphNode* SpawnedNode = Action->Invoke(ParentGraph, Bindings, Location);
+	UEdGraphNode* SpawnedNode = Action->Invoke(ParentGraph, Bindings, FDeprecateSlateVector2D(Location));
 
 	// if a returned node wasn't one that previously existed in the graph
 	const bool bNewNode = PreSpawnNodeCount < ParentGraph->Nodes.Num();
@@ -232,12 +232,12 @@ FBlueprintActionMenuItem::FBlueprintActionMenuItem(UBlueprintNodeSpawner const*
 }
 
 //------------------------------------------------------------------------------
-UEdGraphNode* FBlueprintActionMenuItem::PerformAction(UEdGraph* ParentGraph, UEdGraphPin* FromPin, FVector2D const Location, bool bSelectNewNode/* = true*/)
+UEdGraphNode* FBlueprintActionMenuItem::PerformAction(UEdGraph* ParentGraph, UEdGraphPin* FromPin, const FVector2f& Location, bool bSelectNewNode/* = true*/)
 {
 	using namespace FBlueprintMenuActionItemImpl;
 	FScopedTransaction Transaction(LOCTEXT("AddNodeTransaction", "Add Node"));
 	
-	FVector2D ModifiedLocation = Location;
+	FVector2f ModifiedLocation = Location;
 	if (FromPin != nullptr)
 	{
 		// for input pins, a new node will generally overlap the node being
@@ -246,9 +246,9 @@ UEdGraphNode* FBlueprintActionMenuItem::PerformAction(UEdGraph* ParentGraph, UEd
 		{
 			UEdGraphNode* FromNode = FromPin->GetOwningNode();
 			check(FromNode != nullptr);
-			const double FromNodeX = FromNode->NodePosX;
+			const float FromNodeX = FromNode->NodePosX;
 
-			static const double MinNodeDistance = 60.0; // min distance between spawned nodes (to keep them from overlapping)
+			static const float MinNodeDistance = 60.0f; // min distance between spawned nodes (to keep them from overlapping)
 			if (MinNodeDistance > FMath::Abs(FromNodeX - Location.X))
 			{
 				ModifiedLocation.X = FromNodeX - MinNodeDistance;
@@ -319,7 +319,7 @@ UEdGraphNode* FBlueprintActionMenuItem::PerformAction(UEdGraph* ParentGraph, UEd
 }
 
 //------------------------------------------------------------------------------
-UEdGraphNode* FBlueprintActionMenuItem::PerformAction(UEdGraph* ParentGraph, TArray<UEdGraphPin*>& FromPins, FVector2D const Location, bool bSelectNewNode/* = true*/)
+UEdGraphNode* FBlueprintActionMenuItem::PerformAction(UEdGraph* ParentGraph, TArray<UEdGraphPin*>& FromPins, const FVector2f& Location, bool bSelectNewNode/* = true*/)
 {
 	UEdGraphPin* FromPin = nullptr;
 	if (FromPins.Num() > 0)
diff --git a/Engine/Source/Editor/Kismet/Private/BlueprintDragDropMenuItem.cpp b/Engine/Source/Editor/Kismet/Private/BlueprintDragDropMenuItem.cpp
--- a/Engine/Source/Editor/Kismet/Private/BlueprintDragDropMenuItem.cpp
+++ b/Engine/Source/Editor/Kismet/Private/BlueprintDragDropMenuItem.cpp
@@ -41,7 +41,7 @@ FBlueprintDragDropMenuItem::FBlueprintDragDropMenuItem(FBlueprintActionContext c
 }
 
 //------------------------------------------------------------------------------
-UEdGraphNode* FBlueprintDragDropMenuItem::PerformAction(UEdGraph* ParentGraph, UEdGraphPin* FromPin, FVector2D const Location, bool bSelectNewNode/* = true*/)
+UEdGraphNode* FBlueprintDragDropMenuItem::PerformAction(UEdGraph* ParentGraph, UEdGraphPin* FromPin, const FVector2f& Location, bool bSelectNewNode/* = true*/)
 {
 	// we shouldn't get here (this should just be used for drag/drop ops), but just in case we do...
 	FBlueprintActionMenuItem BlueprintActionItem(GetSampleAction());
@@ -49,7 +49,7 @@ UEdGraphNode* FBlueprintDragDropMenuItem::PerformAction(UEdGraph* ParentGraph, U
 }
 
 //------------------------------------------------------------------------------
-UEdGraphNode* FBlueprintDragDropMenuItem::PerformAction(UEdGraph* ParentGraph, TArray<UEdGraphPin*>& FromPins, FVector2D const Location, bool bSelectNewNode/* = true*/)
+UEdGraphNode* FBlueprintDragDropMenuItem::PerformAction(UEdGraph* ParentGraph, TArray<UEdGraphPin*>& FromPins, const FVector2f& Location, bool bSelectNewNode/* = true*/)
 {
 	UEdGraphPin* FromPin = nullptr;
 	if (FromPins.Num() > 0)
diff --git a/Engine/Source/Editor/Kismet/Private/BlueprintEditor.cpp b/Engine/Source/Editor/Kismet/Private/BlueprintEditor.cpp
--- a/Engine/Source/Editor/Kismet/Private/BlueprintEditor.cpp
+++ b/Engine/Source/Editor/Kismet/Private/BlueprintEditor.cpp
@@ -1752,7 +1752,7 @@ TSharedRef<SGraphEditor> FBlueprintEditor::CreateGraphEditorWidget(TSharedRef<FT
 
 	OnSetPinVisibility.AddSP(&Editor.Get(), &SGraphEditor::SetPinVisibility);
 
-	FVector2D ViewOffset = FVector2D::ZeroVector;
+	FVector2f ViewOffset = FVector2f::ZeroVector;
 	float ZoomAmount = INDEX_NONE;
 
 	TSharedPtr<SDockTab> ActiveTab = DocumentManager->GetActiveTab();
@@ -4708,7 +4708,17 @@ void FBlueprintEditor::DumpMessagesToCompilerLog(const TArray<TSharedRef<FTokeni
 	}
 }
 
-void FBlueprintEditor::DoPromoteToVariable( UBlueprint* InBlueprint, UEdGraphPin* InTargetPin, bool bInToMemberVariable, const FVector2D* InOptionalLocation /* = nullptr */)
+void FBlueprintEditor::DoPromoteToVariable(UBlueprint* InBlueprint, UEdGraphPin* InTargetPin, bool bInToMemberVariable, const FVector2D* InOptionalLocation /* = nullptr */)
+{
+	TOptional<FVector2f> Location;
+	if (InOptionalLocation)
+	{
+		Location = UE::Slate::CastToVector2f(*InOptionalLocation);
+	}
+	DoPromoteToVariable2f(InBlueprint, InTargetPin, bInToMemberVariable, Location);
+}
+
+void FBlueprintEditor::DoPromoteToVariable2f(UBlueprint* InBlueprint, UEdGraphPin* InTargetPin, bool bInToMemberVariable, TOptional<FVector2f> InOptionalLocation)
 {
 	FName PinName = InTargetPin->PinName;
 	UEdGraphNode* PinNode = InTargetPin->GetOwningNode();
@@ -4853,11 +4863,11 @@ void FBlueprintEditor::DoPromoteToVariable( UBlueprint* InBlueprint, UEdGraphPin
 		NodeInfo.NodeTemplate = TemplateNode;
 
 		// Set position of new node to be close to node we clicked on
-		FVector2D NewNodePos;
+		FVector2f NewNodePos;
 
-		if (InOptionalLocation)
+		if (InOptionalLocation.IsSet())
 		{
-			NewNodePos = *InOptionalLocation;
+			NewNodePos = InOptionalLocation.GetValue();
 		}
 		else
 		{
@@ -4882,7 +4892,7 @@ void FBlueprintEditor::OnPromoteToVariable(bool bInToMemberVariable)
 		check(GetBlueprintObj()->SkeletonGeneratedClass);
 		check(TargetPin);
 
-		DoPromoteToVariable( GetBlueprintObj(), TargetPin, bInToMemberVariable );
+		DoPromoteToVariable2f( GetBlueprintObj(), TargetPin, bInToMemberVariable );
 	}
 }
 
@@ -7273,7 +7283,7 @@ void FBlueprintEditor::PasteNodes()
 		return;
 	}
 
-	PasteNodesHere(FocusedGraphEd->GetCurrentGraph(), FocusedGraphEd->GetPasteLocation());
+	PasteNodesHere(FocusedGraphEd->GetCurrentGraph(), FocusedGraphEd->GetPasteLocation2f());
 
 	// Dump any temporary pre-compile warnings to the compiler log.
 	UBlueprint* BlueprintObj = GetBlueprintObj();
@@ -7477,7 +7487,7 @@ private:
 	}
 };
 
-void FBlueprintEditor::PasteNodesHere(class UEdGraph* DestinationGraph, const FVector2D& GraphLocation)
+void FBlueprintEditor::PasteNodesHere(class UEdGraph* DestinationGraph, const FVector2f& GraphLocation)
 {
 	// Find the graph editor with focus
 	TSharedPtr<SGraphEditor> FocusedGraphEd = FocusedGraphEdPtr.Pin();
@@ -7545,7 +7555,7 @@ void FBlueprintEditor::PasteNodesHere(class UEdGraph* DestinationGraph, const FV
 		}
 
 		//Average position of nodes so we can move them while still maintaining relative distances to each other
-		FVector2D AvgNodePosition(0.0f, 0.0f);
+		FVector2f AvgNodePosition(0.0f, 0.0f);
 
 		for (TSet<UEdGraphNode*>::TIterator It(PastedNodes); It; ++It)
 		{
@@ -9204,7 +9214,9 @@ void FBlueprintEditor::GetViewLocation(FVector2D& Location, float& ZoomAmount)
 
 	if (FocusedGraphEdPtr.IsValid())
 	{
-		FocusedGraphEdPtr.Pin()->GetViewLocation(Location, ZoomAmount);
+		FVector2f TempLocation;
+		FocusedGraphEdPtr.Pin()->GetViewLocation(TempLocation, ZoomAmount);
+		Location = FVector2D(TempLocation);
 	}
 }
 
@@ -9212,7 +9224,7 @@ void FBlueprintEditor::SetViewLocation(const FVector2D& Location, float ZoomAmou
 {
 	if (FocusedGraphEdPtr.IsValid())
 	{
-		FocusedGraphEdPtr.Pin()->SetViewLocation(Location, ZoomAmount, BookmarkId);
+		FocusedGraphEdPtr.Pin()->SetViewLocation(UE::Slate::CastToVector2f(Location), ZoomAmount, BookmarkId);
 	}
 }
 
@@ -10181,7 +10193,7 @@ void FBlueprintEditor::RestoreEditedObjectState()
 				if (TabWithGraph.IsValid())
 				{
 					TSharedRef<SGraphEditor> GraphEditor = StaticCastSharedRef<SGraphEditor>(TabWithGraph->GetContent());
-					GraphEditor->SetViewLocation(Blueprint->LastEditedDocuments[i].SavedViewOffset, Blueprint->LastEditedDocuments[i].SavedZoomAmount);
+					GraphEditor->SetViewLocation(UE::Slate::CastToVector2f(Blueprint->LastEditedDocuments[i].SavedViewOffset), Blueprint->LastEditedDocuments[i].SavedZoomAmount);
 				}
 			}
 			else
@@ -10224,7 +10236,7 @@ void FBlueprintEditor::OnCreateComment()
 				if (Schema->IsA(UEdGraphSchema_K2::StaticClass()))
 				{
 					FEdGraphSchemaAction_K2AddComment CommentAction;
-					CommentAction.PerformAction(Graph, nullptr, GraphEditor->GetPasteLocation());
+					CommentAction.PerformAction(Graph, nullptr, GraphEditor->GetPasteLocation2f());
 				}
 			}
 		}
@@ -10245,7 +10257,7 @@ void FBlueprintEditor::OnCreateCustomEvent()
 					// BlueprintEventNodeSpawner seems better but we'll use FEdGraphSchemaAction_K2AddCustomEvent
 					FEdGraphSchemaAction_K2AddCustomEvent EventAction;
 					EventAction.NodeTemplate = NewObject<UK2Node_CustomEvent>();
-					EventAction.PerformAction(Graph, nullptr, GraphEditor->GetPasteLocation());
+					EventAction.PerformAction(Graph, nullptr, GraphEditor->GetPasteLocation2f());
 				}
 			}
 		}
@@ -10836,7 +10848,9 @@ void FBlueprintEditor::SetGraphEditorQuickJump(int32 QuickJumpIndex)
 			FEditedDocumentInfo& QuickJumpInfo = LocalSettings->GraphEditorQuickJumps.FindOrAdd(QuickJumpIndex);
 
 			QuickJumpInfo.EditedObjectPath = GraphObject;
-			FocusedGraphEd->GetViewLocation(QuickJumpInfo.SavedViewOffset, QuickJumpInfo.SavedZoomAmount);
+			FVector2f TempViewOffset = UE::Slate::CastToVector2f(QuickJumpInfo.SavedViewOffset);
+			FocusedGraphEd->GetViewLocation(TempViewOffset, QuickJumpInfo.SavedZoomAmount);
+			QuickJumpInfo.SavedViewOffset = FDeprecateSlateVector2D(TempViewOffset);
 
 			LocalSettings->SaveConfig();
 		}
@@ -10864,7 +10878,7 @@ void FBlueprintEditor::OnGraphEditorQuickJump(int32 QuickJumpIndex)
 				TSharedPtr<SGraphEditor> GraphEditorPtr = IBlueprintEditorPtr->OpenGraphAndBringToFront(Cast<UEdGraph>(EditedObject));
 				if (GraphEditorPtr.IsValid())
 				{
-					GraphEditorPtr->SetViewLocation(QuickJumpInfo->SavedViewOffset, QuickJumpInfo->SavedZoomAmount);
+					GraphEditorPtr->SetViewLocation(UE::Slate::CastToVector2f(QuickJumpInfo->SavedViewOffset), QuickJumpInfo->SavedZoomAmount);
 				}
 			}
 		}
diff --git a/Engine/Source/Editor/Kismet/Private/BlueprintEditorTabFactories.h b/Engine/Source/Editor/Kismet/Private/BlueprintEditorTabFactories.h
--- a/Engine/Source/Editor/Kismet/Private/BlueprintEditorTabFactories.h
+++ b/Engine/Source/Editor/Kismet/Private/BlueprintEditorTabFactories.h
@@ -98,7 +98,7 @@ public:
 	 */
 	FGraphTabHistory(TSharedPtr<FDocumentTabFactory> InFactory, TSharedPtr<FTabPayload> InPayload)
 		: FGenericTabHistory(InFactory, InPayload)
-		, SavedLocation(FVector2D::ZeroVector)
+		, SavedLocation(FVector2f::ZeroVector)
 		, SavedZoomAmount(INDEX_NONE)
 	{
 
@@ -114,7 +114,7 @@ private:
 	/** The graph editor represented by this history node. While this node is inactive, the graph editor is invalid */
 	TWeakPtr< class SGraphEditor > GraphEditor;
 	/** Saved location the graph editor was at when this history node was last visited */
-	FVector2D SavedLocation;
+	FDeprecateSlateVector2D SavedLocation;
 	/** Saved zoom the graph editor was at when this history node was last visited */
 	float SavedZoomAmount;
 	/** Saved bookmark ID the graph editor was at when this history node was last visited */
diff --git a/Engine/Source/Editor/Kismet/Private/SBlueprintActionMenu.cpp b/Engine/Source/Editor/Kismet/Private/SBlueprintActionMenu.cpp
--- a/Engine/Source/Editor/Kismet/Private/SBlueprintActionMenu.cpp
+++ b/Engine/Source/Editor/Kismet/Private/SBlueprintActionMenu.cpp
@@ -30,6 +30,7 @@
 #include "Layout/Children.h"
 #include "Misc/AssertionMacros.h"
 #include "Misc/ConfigCacheIni.h"
+#include "Misc/Optional.h"
 #include "SBlueprintContextTargetMenu.h"
 #include "SBlueprintNamespaceEntry.h"
 #include "SBlueprintPalette.h"
@@ -83,14 +84,14 @@ struct FBlueprintAction_PromoteVariable : public FEdGraphSchemaAction
 	}
 
 	// FEdGraphSchemaAction interface
-	virtual UEdGraphNode* PerformAction( class UEdGraph* ParentGraph, UEdGraphPin* FromPin, const FVector2D Location, bool bSelectNewNode = true) override
+	virtual UEdGraphNode* PerformAction( class UEdGraph* ParentGraph, UEdGraphPin* FromPin, const FVector2f& Location, bool bSelectNewNode = true) override
 	{
 		if( ( ParentGraph != NULL ) && ( FromPin != NULL ) )
 		{
 			UBlueprint* Blueprint = FBlueprintEditorUtils::FindBlueprintForGraph(ParentGraph);
 			if( ( MyBlueprintEditor.IsValid() == true ) && ( Blueprint != NULL ) )
 			{
-				MyBlueprintEditor.Pin()->DoPromoteToVariable( Blueprint, FromPin, bToMemberVariable, &Location );
+				MyBlueprintEditor.Pin()->DoPromoteToVariable2f( Blueprint, FromPin, bToMemberVariable, Location);
 			}
 		}
 		return NULL;		
diff --git a/Engine/Source/Editor/Kismet/Private/SGraphTitleBarAddNewBookmark.cpp b/Engine/Source/Editor/Kismet/Private/SGraphTitleBarAddNewBookmark.cpp
--- a/Engine/Source/Editor/Kismet/Private/SGraphTitleBarAddNewBookmark.cpp
+++ b/Engine/Source/Editor/Kismet/Private/SGraphTitleBarAddNewBookmark.cpp
@@ -292,7 +292,9 @@ FReply SGraphTitleBarAddNewBookmark::OnAddButtonClicked()
 		{
 			FEditedDocumentInfo NewBookmarkInfo;
 			NewBookmarkInfo.EditedObjectPath = FocusedGraph;
-			EditorContext->GetViewLocation(NewBookmarkInfo.SavedViewOffset, NewBookmarkInfo.SavedZoomAmount);
+			FVector2D Location;
+			EditorContext->GetViewLocation(Location, NewBookmarkInfo.SavedZoomAmount);
+			NewBookmarkInfo.SavedViewOffset = Location;
 			if (FBPEditorBookmarkNode* NewNode = EditorContext->AddBookmark(CurrentNameText, NewBookmarkInfo))
 			{
 				EditorContext->SetViewLocation(NewBookmarkInfo.SavedViewOffset, NewBookmarkInfo.SavedZoomAmount, NewNode->NodeGuid);
diff --git a/Engine/Source/Editor/Kismet/Private/Tests/BlueprintEditorTests.cpp b/Engine/Source/Editor/Kismet/Private/Tests/BlueprintEditorTests.cpp
--- a/Engine/Source/Editor/Kismet/Private/Tests/BlueprintEditorTests.cpp
+++ b/Engine/Source/Editor/Kismet/Private/Tests/BlueprintEditorTests.cpp
@@ -291,7 +291,7 @@ namespace BlueprintEditorPromotionUtils
 	* @param GraphLocation - The location to place the node
 	* @param ConnectPin - The pin to connect the node to
 	*/
-	static UEdGraphNode* CreateNewGraphNodeFromTemplate(UK2Node* NodeTemplate, UEdGraph* InGraph, const FVector2D& GraphLocation, UEdGraphPin* ConnectPin = NULL)
+	static UEdGraphNode* CreateNewGraphNodeFromTemplate(UK2Node* NodeTemplate, UEdGraph* InGraph, const FVector2f& GraphLocation, UEdGraphPin* ConnectPin = NULL)
 	{
 		TSharedPtr<FEdGraphSchemaAction_K2NewNode> Action = TSharedPtr<FEdGraphSchemaAction_K2NewNode>(new FEdGraphSchemaAction_K2NewNode(FText::GetEmpty(), FText::GetEmpty(), FText::GetEmpty(), 0));
 		Action->NodeTemplate = NodeTemplate;
@@ -319,7 +319,7 @@ namespace BlueprintEditorPromotionUtils
 		UFunction* AddComponentFn = FindFieldChecked<UFunction>(AActor::StaticClass(), UK2Node_AddComponent::GetAddComponentFunctionName());
 		CallFuncNode->FunctionReference.SetFromField<UFunction>(AddComponentFn, FBlueprintEditorUtils::IsActorBased(InBlueprint));
 
-		UEdGraphNode* NewNode = CreateNewGraphNodeFromTemplate(CallFuncNode, InGraph, FVector2D(200, 0));
+		UEdGraphNode* NewNode = CreateNewGraphNodeFromTemplate(CallFuncNode, InGraph, FVector2f(200.0f, 0.0f));
 
 		TSubclassOf<UActorComponent> ComponentClass = InAsset ? FComponentAssetBrokerage::GetPrimaryComponentForAsset(InAsset->GetClass()) : NULL;
 		if ((NewNode != NULL) && (InBlueprint != NULL))
@@ -375,7 +375,7 @@ namespace BlueprintEditorPromotionUtils
 		UFunction* DelayFn = FindFieldChecked<UFunction>(UStaticMeshComponent::StaticClass(), PrintStringFunctionName);
 		CallFuncNode->FunctionReference.SetFromField<UFunction>(DelayFn, false);
 
-		return CreateNewGraphNodeFromTemplate(CallFuncNode, InGraph, FVector2D(850, 0));
+		return CreateNewGraphNodeFromTemplate(CallFuncNode, InGraph, FVector2f(850.0f, 0.0f));
 	}
 
 	/**
@@ -417,7 +417,7 @@ namespace BlueprintEditorPromotionUtils
 		FBlueprintEditor* CurrentBlueprintEditor = (FBlueprintEditor*)OpenEditor;
 
 		UEdGraphPin* PinToPromote = Node->FindPin(PinName);
-		CurrentBlueprintEditor->DoPromoteToVariable(InBlueprint, PinToPromote, true);
+		CurrentBlueprintEditor->DoPromoteToVariable2f(InBlueprint, PinToPromote, true);
 	}
 
 	/**
@@ -441,7 +441,7 @@ namespace BlueprintEditorPromotionUtils
 
 		if (!ExistingEvent)
 		{
-			return CreateNewGraphNodeFromTemplate(NewEventNode, InGraph, FVector2D(200, 0));
+			return CreateNewGraphNodeFromTemplate(NewEventNode, InGraph, FVector2f(200.0f, 0.0f));
 		}
 		return ExistingEvent;
 	}
@@ -462,7 +462,7 @@ namespace BlueprintEditorPromotionUtils
 		UK2Node_CustomEvent* NewEventNode = NewObject<UK2Node_CustomEvent>(TempOuter);
 		NewEventNode->CustomFunctionName = "EventName";
 
-		return CreateNewGraphNodeFromTemplate(NewEventNode, InGraph, FVector2D(1200, 0));
+		return CreateNewGraphNodeFromTemplate(NewEventNode, InGraph, FVector2f(1200.0f, 0.0f));
 	}
 
 	/**
@@ -500,7 +500,7 @@ namespace BlueprintEditorPromotionUtils
 		UK2Node* CallFuncNode = CreateKismetFunctionTemplate(TempOuter, DelayFunctionName);
 
 		//Create the node
-		return CreateNewGraphNodeFromTemplate(CallFuncNode, InGraph, FVector2D(400, 0), ConnectPin);
+		return CreateNewGraphNodeFromTemplate(CallFuncNode, InGraph, FVector2f(400.0f, 0.0f), ConnectPin);
 	}
 
 	/**
@@ -519,7 +519,7 @@ namespace BlueprintEditorPromotionUtils
 		static FName PrintStringFunctionName(TEXT("PrintString"));
 		UK2Node* CallFuncNode = CreateKismetFunctionTemplate(TempOuter, PrintStringFunctionName);
 
-		return CreateNewGraphNodeFromTemplate(CallFuncNode, InGraph, FVector2D(680, 0), ConnectPin);
+		return CreateNewGraphNodeFromTemplate(CallFuncNode, InGraph, FVector2f(680.0f, 0.0f), ConnectPin);
 	}
 
 	/**
@@ -539,7 +539,7 @@ namespace BlueprintEditorPromotionUtils
 		UK2Node_CallFunction* CallFuncNode = NewObject<UK2Node_CallFunction>(TempOuter);
 		CallFuncNode->FunctionReference.SetSelfMember(FunctionName);
 
-		return CreateNewGraphNodeFromTemplate(CallFuncNode, InGraph, FVector2D(1200, 0), ConnectPin);
+		return CreateNewGraphNodeFromTemplate(CallFuncNode, InGraph, FVector2f(1200.0f, 0.0f), ConnectPin);
 	}
 
 	/**
@@ -571,7 +571,7 @@ namespace BlueprintEditorPromotionUtils
 		TemplateNode->VariableReference.SetSelfMember(FName(*VarName));
 		NodeInfo.NodeTemplate = TemplateNode;
 
-		return NodeInfo.PerformAction(InGraph, NULL, FVector2D(XOffset, 130), true);
+		return NodeInfo.PerformAction(InGraph, NULL, FVector2f(XOffset, 130.0f), true);
 	}
 
 	/**
diff --git a/Engine/Source/Editor/Kismet/Public/BlueprintActionMenuItem.h b/Engine/Source/Editor/Kismet/Public/BlueprintActionMenuItem.h
--- a/Engine/Source/Editor/Kismet/Public/BlueprintActionMenuItem.h
+++ b/Engine/Source/Editor/Kismet/Public/BlueprintActionMenuItem.h
@@ -46,8 +46,9 @@ public:
 	
 	// FEdGraphSchemaAction interface
 	virtual FName         GetTypeId() const final { return StaticGetTypeId(); }
-	virtual UEdGraphNode* PerformAction(UEdGraph* ParentGraph, UEdGraphPin* FromPin, FVector2D const Location, bool bSelectNewNode = true) final;
-	virtual UEdGraphNode* PerformAction(UEdGraph* ParentGraph, TArray<UEdGraphPin*>& FromPins, FVector2D const Location, bool bSelectNewNode = true) final;
+	using FEdGraphSchemaAction::PerformAction; // Prevent hiding of deprecated base class function with FVector2D
+	virtual UEdGraphNode* PerformAction(UEdGraph* ParentGraph, UEdGraphPin* FromPin, const FVector2f& Location, bool bSelectNewNode = true) final;
+	virtual UEdGraphNode* PerformAction(UEdGraph* ParentGraph, TArray<UEdGraphPin*>& FromPins, const FVector2f& Location, bool bSelectNewNode = true) final;
 	virtual void          AddReferencedObjects(FReferenceCollector& Collector) final;
 	// End FEdGraphSchemaAction interface
 
diff --git a/Engine/Source/Editor/Kismet/Public/BlueprintDragDropMenuItem.h b/Engine/Source/Editor/Kismet/Public/BlueprintDragDropMenuItem.h
--- a/Engine/Source/Editor/Kismet/Public/BlueprintDragDropMenuItem.h
+++ b/Engine/Source/Editor/Kismet/Public/BlueprintDragDropMenuItem.h
@@ -62,8 +62,8 @@ public:
 	
 	// FEdGraphSchemaAction interface
 	virtual FName GetTypeId() const final { return StaticGetTypeId(); }
-	virtual UEdGraphNode* PerformAction(UEdGraph* ParentGraph, UEdGraphPin* FromPin, FVector2D const Location, bool bSelectNewNode = true) override;
-	virtual UEdGraphNode* PerformAction(UEdGraph* ParentGraph, TArray<UEdGraphPin*>& FromPins, FVector2D const Location, bool bSelectNewNode = true) override;
+	virtual UEdGraphNode* PerformAction(UEdGraph* ParentGraph, UEdGraphPin* FromPin, const FVector2f& Location, bool bSelectNewNode = true) override;
+	virtual UEdGraphNode* PerformAction(UEdGraph* ParentGraph, TArray<UEdGraphPin*>& FromPins, const FVector2f& Location, bool bSelectNewNode = true) override;
 	virtual void  AddReferencedObjects(FReferenceCollector& Collector) override;
 	// End FEdGraphSchemaAction interface
 
diff --git a/Engine/Source/Editor/Kismet/Public/BlueprintEditor.h b/Engine/Source/Editor/Kismet/Public/BlueprintEditor.h
--- a/Engine/Source/Editor/Kismet/Public/BlueprintEditor.h
+++ b/Engine/Source/Editor/Kismet/Public/BlueprintEditor.h
@@ -29,6 +29,7 @@
 #include "Misc/Guid.h"
 #include "Misc/NotifyHook.h"
 #include "Misc/Optional.h"
+#include "Math/Vector2D.h"
 #include "PreviewScene.h"
 #include "SKismetInspector.h"
 #include "SourceControlOperations.h"
@@ -665,7 +666,9 @@ public:
 	 * @param	bInToMemberVariable		TRUE if attempting to create a member variable, FALSE if the variable should be local
 	 * @param   InOptionalLocation		Where the new node should be placed. If null, a fixed offset from the parent node will be used.
 	 */
+	UE_DEPRECATED(5.6, "Use the version of this function accepting a TOptional")
 	void DoPromoteToVariable(UBlueprint* InBlueprint, UEdGraphPin* InTargetPin, bool bInToMemberVariable, const FVector2D* InOptionalLocation = nullptr);
+	void DoPromoteToVariable2f(UBlueprint* InBlueprint, UEdGraphPin* InTargetPin, bool bInToMemberVariable, TOptional<FVector2f> InOptionalLocation = TOptional<FVector2f>());
 
 	/** Called when node is spawned by keymap */
 	void OnNodeSpawnedByKeymap();
@@ -1100,7 +1103,7 @@ protected:
 	virtual bool CanCopyNodes() const;
 
 	/** Paste on graph at specific location */
-	virtual void PasteNodesHere(class UEdGraph* DestinationGraph, const FVector2D& GraphLocation) override;
+	virtual void PasteNodesHere(class UEdGraph* DestinationGraph, const FVector2f& GraphLocation) override;
 
 	/** Paste Variable Definition or Nodes */
 	virtual void PasteGeneric();
diff --git a/Engine/Source/Editor/Kismet/Public/BlueprintEditorModule.h b/Engine/Source/Editor/Kismet/Public/BlueprintEditorModule.h
--- a/Engine/Source/Editor/Kismet/Public/BlueprintEditorModule.h
+++ b/Engine/Source/Editor/Kismet/Public/BlueprintEditorModule.h
@@ -115,7 +115,9 @@ public:
 
 	virtual bool CanPasteNodes() const= 0;
 
-	virtual void PasteNodesHere(class UEdGraph* Graph, const FVector2D& Location) = 0;
+	UE_DEPRECATED(5.6, "Slate positions are represented in FVector2f. Please update to use the function acceping FVector2f.")
+	virtual void PasteNodesHere(class UEdGraph* Graph, const FVector2D& Location) UE_SLATE_DEPRECATED_VECTOR_VIRTUAL_FUNCTION { PasteNodesHere(Graph, UE::Slate::CastToVector2f(Location)); };
+	virtual void PasteNodesHere(class UEdGraph* Graph, const FVector2f& Location) = 0;
 
 	virtual bool GetBoundsForSelectedNodes(class FSlateRect& Rect, float Padding ) = 0;
 
diff --git a/Engine/Source/Editor/MaterialEditor/Private/MaterialEditor.cpp b/Engine/Source/Editor/MaterialEditor/Private/MaterialEditor.cpp
--- a/Engine/Source/Editor/MaterialEditor/Private/MaterialEditor.cpp
+++ b/Engine/Source/Editor/MaterialEditor/Private/MaterialEditor.cpp
@@ -2571,10 +2571,10 @@ void FMaterialEditor::RecenterEditor()
 	else
 	{
 		// Get current view location so that we don't change the zoom amount
-		FVector2D CurrLocation;
+		FVector2f CurrLocation;
 		float CurrZoomLevel;
 		FocusedGraphEd->GetViewLocation(CurrLocation, CurrZoomLevel);
-		FocusedGraphEd->SetViewLocation(FVector2D::ZeroVector, CurrZoomLevel);
+		FocusedGraphEd->SetViewLocation(FVector2f::ZeroVector, CurrZoomLevel);
 	}
 }
 
@@ -5074,7 +5074,7 @@ void FMaterialEditor::OnCreateComment()
 {
 	if (TSharedPtr<SGraphEditor> FocusedGraphEd = FocusedGraphEdPtr.Pin())
 	{
-		CreateNewMaterialExpressionComment(FocusedGraphEd->GetPasteLocation());
+		CreateNewMaterialExpressionComment(FocusedGraphEd->GetPasteLocation2f());
 	}
 }
 
@@ -5082,7 +5082,7 @@ void FMaterialEditor::OnCreateComponentMaskNode()
 {
 	if (TSharedPtr<SGraphEditor> FocusedGraphEd = FocusedGraphEdPtr.Pin())
 	{
-		CreateNewMaterialExpression(UMaterialExpressionComponentMask::StaticClass(), FocusedGraphEd->GetPasteLocation(), true, false);
+		CreateNewMaterialExpression(UMaterialExpressionComponentMask::StaticClass(), FocusedGraphEd->GetPasteLocation2f(), true, false);
 	}
 }
 
@@ -5222,8 +5222,8 @@ void FMaterialEditor::OnPromoteToParameter(const FToolMenuContext& InMenuContext
 		GraphObj->Modify();
 
 		// Set position of new node to be close to node we clicked on
-		FVector2D NewNodePos;
-		NewNodePos.X = PinNode->NodePosX - 100;
+		FVector2f NewNodePos;
+		NewNodePos.X = PinNode->NodePosX - 100.0f;
 		NewNodePos.Y = PinNode->NodePosY;
 
 		UMaterialGraphNode* MaterialNode = Cast<UMaterialGraphNode>(Action.PerformAction(GraphObj, const_cast<UEdGraphPin*>(TargetPin), NewNodePos));
@@ -5484,7 +5484,7 @@ void FMaterialEditor::OnCreateSubstrateNodeForPin(const FToolMenuContext& InMenu
 	GraphObj->Modify();
 
 	// Set position of new node to be close to node we clicked on
-	FVector2D NewNodePos;
+	FVector2f NewNodePos;
 	NewNodePos.X = PinNode->NodePosX + (bTargetPinIsInput ? -300 : 300);
 	NewNodePos.Y = PinNode->NodePosY;
 
@@ -6557,7 +6557,7 @@ void FMaterialEditor::PasteNodes()
 {
 	if (TSharedPtr<SGraphEditor> FocusedGraphEd = FocusedGraphEdPtr.Pin())
 	{
-		PasteNodesHere(FocusedGraphEd->GetPasteLocation(), FocusedGraphEd->GetCurrentGraph());
+		PasteNodesHere(FocusedGraphEd->GetPasteLocation2f(), FocusedGraphEd->GetCurrentGraph());
 	}
 }
 
@@ -8589,7 +8589,7 @@ FReply FMaterialEditor::OnSpawnGraphNodeByShortcut(FInputChord InChord, const FV
 		if (Action.IsValid())
 		{
 			TArray<UEdGraphPin*> DummyPins;
-			Action->PerformAction(Graph, DummyPins, InPosition);
+			Action->PerformAction(Graph, DummyPins, UE::Slate::FDeprecateVector2DParameter(InPosition));
 			return FReply::Handled();
 		}
 	}
diff --git a/Engine/Source/Editor/MaterialEditor/Private/MaterialEditorUtilities.cpp b/Engine/Source/Editor/MaterialEditor/Private/MaterialEditorUtilities.cpp
--- a/Engine/Source/Editor/MaterialEditor/Private/MaterialEditorUtilities.cpp
+++ b/Engine/Source/Editor/MaterialEditor/Private/MaterialEditorUtilities.cpp
@@ -49,12 +49,12 @@
 
 DEFINE_LOG_CATEGORY_STATIC(LogMaterialEditorUtilities, Log, All);
 
-UMaterialExpression* FMaterialEditorUtilities::CreateNewMaterialExpression(const class UEdGraph* Graph, UClass* NewExpressionClass, const FVector2D& NodePos, bool bAutoSelect, bool bAutoAssignResource)
+UMaterialExpression* FMaterialEditorUtilities::CreateNewMaterialExpression(const class UEdGraph* Graph, UClass* NewExpressionClass, const UE::Slate::FDeprecateVector2DParameter& NodePos, bool bAutoSelect, bool bAutoAssignResource)
 {
 	TSharedPtr<class IMaterialEditor> MaterialEditor = GetIMaterialEditorForObject(Graph);
 	if (MaterialEditor.IsValid())
 	{
-		return MaterialEditor->CreateNewMaterialExpression(NewExpressionClass, NodePos, bAutoSelect, bAutoAssignResource, Graph);
+		return MaterialEditor->CreateNewMaterialExpression(NewExpressionClass, FDeprecateSlateVector2D(NodePos), bAutoSelect, bAutoAssignResource, Graph);
 	}
 	return nullptr;
 }
diff --git a/Engine/Source/Editor/MaterialEditor/Public/MaterialEditorUtilities.h b/Engine/Source/Editor/MaterialEditor/Public/MaterialEditorUtilities.h
--- a/Engine/Source/Editor/MaterialEditor/Public/MaterialEditorUtilities.h
+++ b/Engine/Source/Editor/MaterialEditor/Public/MaterialEditorUtilities.h
@@ -52,7 +52,7 @@ public:
 	 * @param	bAutoSelect			If true, deselect all expressions and select the newly created one.
 	 * @param	bAutoAssignResource	If true, assign resources to new expression.
 	 */
-	static UMaterialExpression* CreateNewMaterialExpression(const class UEdGraph* Graph, UClass* NewExpressionClass, const FVector2D& NodePos, bool bAutoSelect, bool bAutoAssignResource);
+	static UMaterialExpression* CreateNewMaterialExpression(const class UEdGraph* Graph, UClass* NewExpressionClass, const UE::Slate::FDeprecateVector2DParameter& NodePos, bool bAutoSelect, bool bAutoAssignResource);
 
 	/**
 	 * Creates a new material expression of the specified class on the material represented by a graph.
@@ -63,7 +63,7 @@ public:
 	 * @param	bAutoAssignResource		If true, assign resources to new expression.
 	 */
 	template<typename MaterialExpressionClass>
-	static MaterialExpressionClass* CreateNewMaterialExpression(const UEdGraph* Graph, const FVector2D& NodePos, bool bAutoSelect, bool bAutoAssignResource = false)
+	static MaterialExpressionClass* CreateNewMaterialExpression(const UEdGraph* Graph, const UE::Slate::FDeprecateVector2DParameter& NodePos, bool bAutoSelect, bool bAutoAssignResource = false)
 	{
 		static_assert(TIsDerivedFrom<MaterialExpressionClass, UMaterialExpression>::Value, "MaterialExpressionClass needs to inherit from UMaterialExpression");
 		return CastChecked<MaterialExpressionClass>(CreateNewMaterialExpression(Graph, MaterialExpressionClass::StaticClass(), NodePos, bAutoSelect, bAutoAssignResource), ECastCheckedType::NullAllowed);
diff --git a/Engine/Source/Editor/MovieSceneTools/Private/MovieSceneDirectorBlueprintEndpointCustomization.cpp b/Engine/Source/Editor/MovieSceneTools/Private/MovieSceneDirectorBlueprintEndpointCustomization.cpp
--- a/Engine/Source/Editor/MovieSceneTools/Private/MovieSceneDirectorBlueprintEndpointCustomization.cpp
+++ b/Engine/Source/Editor/MovieSceneTools/Private/MovieSceneDirectorBlueprintEndpointCustomization.cpp
@@ -1014,7 +1014,7 @@ void FMovieSceneDirectorBlueprintEndpointCustomization::HandleQuickBindActionSel
 			UEdGraphPin* EndpointThenPin = NewEndpoint->FindPin(UEdGraphSchema_K2::PN_Then, EGPD_Output);
 			UEdGraphPin* CallTargetPin = FMovieSceneDirectorBlueprintUtils::FindCallTargetPin(NewEndpoint, EndpointDefinition.PossibleCallTargetClass);
 
-			FVector2D NodePosition(NewEndpoint->NodePosX + 400.f, NewEndpoint->NodePosY + 100.f);
+			FVector2f NodePosition(NewEndpoint->NodePosX + 400.f, NewEndpoint->NodePosY + 100.f);
 			UEdGraphNode* NewNode = Action->PerformAction(NewEndpoint->GetGraph(), CallTargetPin ? CallTargetPin : EndpointThenPin, NodePosition);
 
 			if (NewNode)
diff --git a/Engine/Source/Editor/UMGEditor/Private/Customizations/UMGDetailCustomizations.cpp b/Engine/Source/Editor/UMGEditor/Private/Customizations/UMGDetailCustomizations.cpp
--- a/Engine/Source/Editor/UMGEditor/Private/Customizations/UMGDetailCustomizations.cpp
+++ b/Engine/Source/Editor/UMGEditor/Private/Customizations/UMGDetailCustomizations.cpp
@@ -83,7 +83,7 @@ private:
 			Editor.Pin()->SetCurrentMode(FWidgetBlueprintApplicationModes::GraphMode);
 
 			// Figure out a decent place to stick the node
-			const FVector2D NewNodePos = TargetGraph->GetGoodPlaceForNewNode();
+			const FVector2f NewNodePos = TargetGraph->GetGoodPlaceForNewNode();
 
 			Action->PerformAction(TargetGraph, nullptr, NewNodePos);
 		}
diff --git a/Engine/Source/Editor/UnrealEd/Classes/MaterialGraph/MaterialGraphSchema.h b/Engine/Source/Editor/UnrealEd/Classes/MaterialGraph/MaterialGraphSchema.h
--- a/Engine/Source/Editor/UnrealEd/Classes/MaterialGraph/MaterialGraphSchema.h
+++ b/Engine/Source/Editor/UnrealEd/Classes/MaterialGraph/MaterialGraphSchema.h
@@ -40,7 +40,8 @@ struct FMaterialGraphSchemaAction_NewNode : public FEdGraphSchemaAction
 	{}
 
 	//~ Begin FEdGraphSchemaAction Interface
-	UNREALED_API virtual UEdGraphNode* PerformAction(class UEdGraph* ParentGraph, UEdGraphPin* FromPin, const FVector2D Location, bool bSelectNewNode = true) override;
+	using FEdGraphSchemaAction::PerformAction; // Prevent hiding of deprecated base class function with FVector2D
+	UNREALED_API virtual UEdGraphNode* PerformAction(class UEdGraph* ParentGraph, UEdGraphPin* FromPin, const FVector2f& Location, bool bSelectNewNode = true) override;
 	//~ End FEdGraphSchemaAction Interface
 
 	/**
@@ -75,7 +76,8 @@ struct FMaterialGraphSchemaAction_NewFunctionCall : public FEdGraphSchemaAction
 	{}
 
 	//~ Begin FEdGraphSchemaAction Interface
-	UNREALED_API virtual UEdGraphNode* PerformAction(class UEdGraph* ParentGraph, UEdGraphPin* FromPin, const FVector2D Location, bool bSelectNewNode = true) override;
+	using FEdGraphSchemaAction::PerformAction; // Prevent hiding of deprecated base class function with FVector2D
+	UNREALED_API virtual UEdGraphNode* PerformAction(class UEdGraph* ParentGraph, UEdGraphPin* FromPin, const FVector2f& Location, bool bSelectNewNode = true) override;
 	//~ End FEdGraphSchemaAction Interface
 };
 
@@ -98,10 +100,11 @@ struct FMaterialGraphSchemaAction_NewComposite : public FEdGraphSchemaAction
 	{}
 
 	//~ Begin FEdGraphSchemaAction Interface
-	UNREALED_API virtual UEdGraphNode* PerformAction(class UEdGraph* ParentGraph, UEdGraphPin* FromPin, const FVector2D Location, bool bSelectNewNode = true) override;
+	using FEdGraphSchemaAction::PerformAction; // Prevent hiding of deprecated base class function with FVector2D
+	UNREALED_API virtual UEdGraphNode* PerformAction(class UEdGraph* ParentGraph, UEdGraphPin* FromPin, const FVector2f& Location, bool bSelectNewNode = true) override;
 	//~ End FEdGraphSchemaAction Interface
 
-	static UNREALED_API UEdGraphNode* SpawnNode(class UEdGraph* ParentGraph, const FVector2D Location);
+	static UNREALED_API UEdGraphNode* SpawnNode(class UEdGraph* ParentGraph, const UE::Slate::FDeprecateVector2DParameter& Location);
 };
 
 /** Action to add a comment node to the graph */
@@ -123,7 +126,8 @@ struct FMaterialGraphSchemaAction_NewComment : public FEdGraphSchemaAction
 	{}
 
 	//~ Begin FEdGraphSchemaAction Interface
-	UNREALED_API virtual UEdGraphNode* PerformAction(class UEdGraph* ParentGraph, UEdGraphPin* FromPin, const FVector2D Location, bool bSelectNewNode = true) override;
+	using FEdGraphSchemaAction::PerformAction; // Prevent hiding of deprecated base class function with FVector2D
+	UNREALED_API virtual UEdGraphNode* PerformAction(class UEdGraph* ParentGraph, UEdGraphPin* FromPin, const FVector2f& Location, bool bSelectNewNode = true) override;
 	//~ End FEdGraphSchemaAction Interface
 };
 
@@ -150,7 +154,8 @@ struct FMaterialGraphSchemaAction_NewNamedRerouteUsage : public FEdGraphSchemaAc
 	{}
 
 	//~ Begin FEdGraphSchemaAction Interface
-	UNREALED_API virtual UEdGraphNode* PerformAction(class UEdGraph* ParentGraph, UEdGraphPin* FromPin, const FVector2D Location, bool bSelectNewNode = true) override;
+	using FEdGraphSchemaAction::PerformAction; // Prevent hiding of deprecated base class function with FVector2D
+	UNREALED_API virtual UEdGraphNode* PerformAction(class UEdGraph* ParentGraph, UEdGraphPin* FromPin, const FVector2f& Location, bool bSelectNewNode = true) override;
 	//~ End FEdGraphSchemaAction Interface
 };
 
@@ -173,7 +178,8 @@ struct FMaterialGraphSchemaAction_Paste : public FEdGraphSchemaAction
 	{}
 
 	//~ Begin FEdGraphSchemaAction Interface
-	UNREALED_API virtual UEdGraphNode* PerformAction(class UEdGraph* ParentGraph, UEdGraphPin* FromPin, const FVector2D Location, bool bSelectNewNode = true) override;
+	using FEdGraphSchemaAction::PerformAction; // Prevent hiding of deprecated base class function with FVector2D
+	UNREALED_API virtual UEdGraphNode* PerformAction(class UEdGraph* ParentGraph, UEdGraphPin* FromPin, const FVector2f& Location, bool bSelectNewNode = true) override;
 	//~ End FEdGraphSchemaAction Interface
 };
 
@@ -278,13 +284,13 @@ class UMaterialGraphSchema : public UEdGraphSchema
 	virtual void BreakPinLinks(UEdGraphPin& TargetPin, bool bSendsNodeNotifcation) const override;
 	virtual void BreakSinglePinLink(UEdGraphPin* SourcePin, UEdGraphPin* TargetPin) const override;
 	virtual bool CanEncapuslateNode(UEdGraphNode const& TestNode) const override;
-	virtual void DroppedAssetsOnGraph(const TArray<struct FAssetData>& Assets, const FVector2D& GraphPosition, UEdGraph* Graph) const override;
+	virtual void DroppedAssetsOnGraph(const TArray<struct FAssetData>& Assets, const FVector2f& GraphPosition, UEdGraph* Graph) const override;
 	virtual int32 GetNodeSelectionCount(const UEdGraph* Graph) const override;
 	virtual TSharedPtr<FEdGraphSchemaAction> GetCreateCommentAction() const override;
 	virtual bool IsCacheVisualizationOutOfDate(int32 InVisualizationCacheID) const override;
 	virtual int32 GetCurrentVisualizationCacheID() const override;
 	virtual void ForceVisualizationCacheClear() const override;
-	virtual void OnPinConnectionDoubleCicked(UEdGraphPin* PinA, UEdGraphPin* PinB, const FVector2D& GraphPosition) const override;
+	virtual void OnPinConnectionDoubleCicked(UEdGraphPin* PinA, UEdGraphPin* PinB, const FVector2f& GraphPosition) const override;
 	virtual bool SafeDeleteNodeFromGraph(UEdGraph* Graph, UEdGraphNode* NodeToDelete) const override;
 	virtual void GetAssetsGraphHoverMessage(const TArray<FAssetData>& Assets, const UEdGraph* HoverGraph, FString& OutTooltipText, bool& OutOkIcon) const;
 #if WITH_EDITORONLY_DATA
diff --git a/Engine/Source/Editor/UnrealEd/Private/Kismet2/Kismet2.cpp b/Engine/Source/Editor/UnrealEd/Private/Kismet2/Kismet2.cpp
--- a/Engine/Source/Editor/UnrealEd/Private/Kismet2/Kismet2.cpp
+++ b/Engine/Source/Editor/UnrealEd/Private/Kismet2/Kismet2.cpp
@@ -2410,7 +2410,7 @@ TSharedPtr<class IBlueprintEditor> FKismetEditorUtilities::GetIBlueprintEditorFo
 	return BlueprintEditor;
 }
 
-void FKismetEditorUtilities::PasteNodesHere( class UEdGraph* Graph, const FVector2D& Location )
+void FKismetEditorUtilities::PasteNodesHere( class UEdGraph* Graph, const UE::Slate::FDeprecateVector2DParameter& Location )
 {
 	TSharedPtr<class IBlueprintEditor> Kismet = GetIBlueprintEditorForObject(Graph,false);
 	if(Kismet.IsValid())
diff --git a/Engine/Source/Editor/UnrealEd/Private/MaterialGraphSchema.cpp b/Engine/Source/Editor/UnrealEd/Private/MaterialGraphSchema.cpp
--- a/Engine/Source/Editor/UnrealEd/Private/MaterialGraphSchema.cpp
+++ b/Engine/Source/Editor/UnrealEd/Private/MaterialGraphSchema.cpp
@@ -58,7 +58,7 @@ int32 UMaterialGraphSchema::CurrentCacheRefreshID = 0;
 ////////////////////////////////////////
 // FMaterialGraphSchemaAction_NewNode //
 
-UEdGraphNode* FMaterialGraphSchemaAction_NewNode::PerformAction(class UEdGraph* ParentGraph, UEdGraphPin* FromPin, const FVector2D Location, bool bSelectNewNode/* = true*/)
+UEdGraphNode* FMaterialGraphSchemaAction_NewNode::PerformAction(class UEdGraph* ParentGraph, UEdGraphPin* FromPin, const FVector2f& Location, bool bSelectNewNode/* = true*/)
 {
 	check(MaterialExpressionClass);
 
@@ -84,7 +84,7 @@ UEdGraphNode* FMaterialGraphSchemaAction_NewNode::PerformAction(class UEdGraph*
 		}
 	}
 
-	UMaterialExpression* NewExpression = FMaterialEditorUtilities::CreateNewMaterialExpression(ParentGraph, MaterialExpressionClass, Location, bSelectNewNode, /*bAutoAssignResource*/true);
+	UMaterialExpression* NewExpression = FMaterialEditorUtilities::CreateNewMaterialExpression(ParentGraph, MaterialExpressionClass, FDeprecateSlateVector2D(Location), bSelectNewNode, /*bAutoAssignResource*/true);
 	PostCreationDelegate.ExecuteIfBound(NewExpression);
 
 	if (NewExpression)
@@ -156,13 +156,13 @@ void FMaterialGraphSchemaAction_NewNode::SetFunctionInputType(UMaterialExpressio
 ////////////////////////////////////////////////
 // FMaterialGraphSchemaAction_NewFunctionCall //
 
-UEdGraphNode* FMaterialGraphSchemaAction_NewFunctionCall::PerformAction(class UEdGraph* ParentGraph, UEdGraphPin* FromPin, const FVector2D Location, bool bSelectNewNode/* = true*/)
+UEdGraphNode* FMaterialGraphSchemaAction_NewFunctionCall::PerformAction(class UEdGraph* ParentGraph, UEdGraphPin* FromPin, const FVector2f& Location, bool bSelectNewNode/* = true*/)
 {
 	const FScopedTransaction Transaction(NSLOCTEXT("UnrealEd", "MaterialEditorNewFunctionCall", "Material Editor: New Function Call"));
 
 	UMaterialExpressionMaterialFunctionCall* FunctionNode = CastChecked<UMaterialExpressionMaterialFunctionCall>(
 																FMaterialEditorUtilities::CreateNewMaterialExpression(
-																	ParentGraph, UMaterialExpressionMaterialFunctionCall::StaticClass(), Location, bSelectNewNode, /*bAutoAssignResource*/false));
+																	ParentGraph, UMaterialExpressionMaterialFunctionCall::StaticClass(), FDeprecateSlateVector2D(Location), bSelectNewNode, /*bAutoAssignResource*/false));
 
 	if (!FunctionNode->MaterialFunction)
 	{
@@ -188,16 +188,16 @@ UEdGraphNode* FMaterialGraphSchemaAction_NewFunctionCall::PerformAction(class UE
 /////////////////////////////////////////////
 // FMaterialGraphSchemaAction_NewComposite //
 
-UEdGraphNode* FMaterialGraphSchemaAction_NewComposite::PerformAction(UEdGraph* ParentGraph, UEdGraphPin* FromPin, const FVector2D Location, bool bSelectNewNode)
+UEdGraphNode* FMaterialGraphSchemaAction_NewComposite::PerformAction(UEdGraph* ParentGraph, UEdGraphPin* FromPin, const FVector2f& Location, bool bSelectNewNode)
 {
 	return SpawnNode(ParentGraph, Location);
 }
 
-UEdGraphNode* FMaterialGraphSchemaAction_NewComposite::SpawnNode(UEdGraph* ParentGraph, const FVector2D Location)
+UEdGraphNode* FMaterialGraphSchemaAction_NewComposite::SpawnNode(UEdGraph* ParentGraph, const UE::Slate::FDeprecateVector2DParameter& Location)
 {
 	const FScopedTransaction Transaction(NSLOCTEXT("UnrealEd", "MaterialEditorNewComposite", "Material Editor: New Composite"));
 
-	UMaterialExpressionComposite* NewComposite = FMaterialEditorUtilities::CreateNewMaterialExpressionComposite(ParentGraph, Location);
+	UMaterialExpressionComposite* NewComposite = FMaterialEditorUtilities::CreateNewMaterialExpressionComposite(ParentGraph, FDeprecateSlateVector2D(Location));
 
 	if (NewComposite)
 	{
@@ -210,11 +210,11 @@ UEdGraphNode* FMaterialGraphSchemaAction_NewComposite::SpawnNode(UEdGraph* Paren
 ///////////////////////////////////////////
 // FMaterialGraphSchemaAction_NewComment //
 
-UEdGraphNode* FMaterialGraphSchemaAction_NewComment::PerformAction(class UEdGraph* ParentGraph, UEdGraphPin* FromPin, const FVector2D Location, bool bSelectNewNode/* = true*/)
+UEdGraphNode* FMaterialGraphSchemaAction_NewComment::PerformAction(class UEdGraph* ParentGraph, UEdGraphPin* FromPin, const FVector2f& Location, bool bSelectNewNode/* = true*/)
 {
 	const FScopedTransaction Transaction(NSLOCTEXT("UnrealEd", "MaterialEditorNewComment", "Material Editor: New Comment") );
 
-	UMaterialExpressionComment* NewComment = FMaterialEditorUtilities::CreateNewMaterialExpressionComment(ParentGraph, Location);
+	UMaterialExpressionComment* NewComment = FMaterialEditorUtilities::CreateNewMaterialExpressionComment(ParentGraph, FDeprecateSlateVector2D(Location));
 
 	if (NewComment)
 	{
@@ -227,13 +227,13 @@ UEdGraphNode* FMaterialGraphSchemaAction_NewComment::PerformAction(class UEdGrap
 //////////////////////////////////////////////////////
 // FMaterialGraphSchemaAction_NewNamedRerouteUsage //
 
-UEdGraphNode* FMaterialGraphSchemaAction_NewNamedRerouteUsage::PerformAction(class UEdGraph* ParentGraph, UEdGraphPin* FromPin, const FVector2D Location, bool bSelectNewNode /*= true*/)
+UEdGraphNode* FMaterialGraphSchemaAction_NewNamedRerouteUsage::PerformAction(class UEdGraph* ParentGraph, UEdGraphPin* FromPin, const FVector2f& Location, bool bSelectNewNode /*= true*/)
 {
 	check(Declaration);
 
 	const FScopedTransaction Transaction(NSLOCTEXT("UnrealEd", "MaterialEditorNewNamedRerouteUsage", "Material Editor: New Named Reroute Usage") );
 	
-	UMaterialExpression* NewExpression = FMaterialEditorUtilities::CreateNewMaterialExpression(ParentGraph, UMaterialExpressionNamedRerouteUsage::StaticClass(), Location, bSelectNewNode, /*bAutoAssignResource*/true);
+	UMaterialExpression* NewExpression = FMaterialEditorUtilities::CreateNewMaterialExpression(ParentGraph, UMaterialExpressionNamedRerouteUsage::StaticClass(), FDeprecateSlateVector2D(Location), bSelectNewNode, /*bAutoAssignResource*/true);
 
 	if (NewExpression)
 	{
@@ -250,9 +250,9 @@ UEdGraphNode* FMaterialGraphSchemaAction_NewNamedRerouteUsage::PerformAction(cla
 //////////////////////////////////////
 // FMaterialGraphSchemaAction_Paste //
 
-UEdGraphNode* FMaterialGraphSchemaAction_Paste::PerformAction(class UEdGraph* ParentGraph, UEdGraphPin* FromPin, const FVector2D Location, bool bSelectNewNode/* = true*/)
+UEdGraphNode* FMaterialGraphSchemaAction_Paste::PerformAction(class UEdGraph* ParentGraph, UEdGraphPin* FromPin, const FVector2f& Location, bool bSelectNewNode/* = true*/)
 {
-	FMaterialEditorUtilities::PasteNodesHere(ParentGraph, Location);
+	FMaterialEditorUtilities::PasteNodesHere(ParentGraph, FDeprecateSlateVector2D(Location));
 	return NULL;
 }
 
@@ -854,12 +854,12 @@ bool UMaterialGraphSchema::CanEncapuslateNode(UEdGraphNode const& TestNode) cons
 		&& !TestNode.IsA(UMaterialGraphNode_Root::StaticClass());
 }
 
-void UMaterialGraphSchema::DroppedAssetsOnGraph(const TArray<struct FAssetData>& Assets, const FVector2D& GraphPosition, UEdGraph* Graph) const
+void UMaterialGraphSchema::DroppedAssetsOnGraph(const TArray<struct FAssetData>& Assets, const FVector2f& GraphPosition, UEdGraph* Graph) const
 {
 	UMaterialGraph* MaterialGraph = CastChecked<UMaterialGraph>(Graph);
 	const int32 LocOffsetBetweenNodes = 32;
 
-	FVector2D ExpressionPosition = GraphPosition;
+	FVector2f ExpressionPosition = GraphPosition;
 
 	for (int32 AssetIdx = 0; AssetIdx < Assets.Num(); ++AssetIdx)
 	{
@@ -1197,18 +1197,18 @@ void UMaterialGraphSchema::ForceVisualizationCacheClear() const
 	++CurrentCacheRefreshID;
 }
 
-void UMaterialGraphSchema::OnPinConnectionDoubleCicked(UEdGraphPin* PinA, UEdGraphPin* PinB, const FVector2D& GraphPosition) const
+void UMaterialGraphSchema::OnPinConnectionDoubleCicked(UEdGraphPin* PinA, UEdGraphPin* PinB, const FVector2f& GraphPosition) const
 {
 	const FScopedTransaction Transaction(LOCTEXT("CreateRerouteNodeOnWire", "Create Reroute Node"));
 
 	//@TODO: This constant is duplicated from inside of SGraphNodeKnot
-	const FVector2D NodeSpacerSize(42.0f, 24.0f);
-	const FVector2D KnotTopLeft = GraphPosition - (NodeSpacerSize * 0.5f);
+	const FVector2f NodeSpacerSize(42.0f, 24.0f);
+	const FVector2f KnotTopLeft = GraphPosition - (NodeSpacerSize * 0.5f);
 
 	// Create a new knot
 	UEdGraph* ParentGraph = PinA->GetOwningNode()->GetGraph();
 
-	if (UMaterialExpression* Expression = FMaterialEditorUtilities::CreateNewMaterialExpression(ParentGraph, UMaterialExpressionReroute::StaticClass(), KnotTopLeft, true, true))
+	if (UMaterialExpression* Expression = FMaterialEditorUtilities::CreateNewMaterialExpression(ParentGraph, UMaterialExpressionReroute::StaticClass(), FDeprecateSlateVector2D(KnotTopLeft), true, true))
 	{
 		// Move the connections across (only notifying the knot, as the other two didn't really change)
 		PinA->BreakLinkTo(PinB);
diff --git a/Engine/Source/Editor/UnrealEd/Public/GraphEditor.h b/Engine/Source/Editor/UnrealEd/Public/GraphEditor.h
--- a/Engine/Source/Editor/UnrealEd/Public/GraphEditor.h
+++ b/Engine/Source/Editor/UnrealEd/Public/GraphEditor.h
@@ -213,26 +213,40 @@ public:
 		return EdGraphObj;
 	}
 
-	virtual FVector2D GetPasteLocation() const
+	UE_DEPRECATED(5.6, "Slate positions are represented in floats. Please use the function returning FVector2f.")
+	virtual FVector2D GetPasteLocation() const UE_SLATE_DEPRECATED_VECTOR_VIRTUAL_FUNCTION
 	{
 		if (Implementation.IsValid())
 		{
-			return Implementation->GetPasteLocation();
+			return FVector2D(Implementation->GetPasteLocation2f());
 		}
 		else
 		{
 			return FVector2D::ZeroVector;
 		}
 	}
+	virtual UE::Slate::FDeprecateVector2DResult GetPasteLocation2f() const
+	{
+	PRAGMA_DISABLE_DEPRECATION_WARNINGS
+		return UE::Slate::CastToVector2f(GetPasteLocation());
+	PRAGMA_ENABLE_DEPRECATION_WARNINGS
+	}
 
 	/* Set new viewer location and optionally set the current bookmark */
-	virtual void SetViewLocation(const FVector2D& Location, float ZoomAmount, const FGuid& BookmarkId = FGuid())
+	UE_DEPRECATED(5.6, "Slate positions are represented in floats. Please use the function returning FVector2f.")
+	virtual void SetViewLocation(const FVector2D& Location, float ZoomAmount, const FGuid& BookmarkId = FGuid()) UE_SLATE_DEPRECATED_VECTOR_VIRTUAL_FUNCTION
 	{
 		if (Implementation.IsValid())
 		{
-			Implementation->SetViewLocation(Location, ZoomAmount, BookmarkId);
+			Implementation->SetViewLocation(UE::Slate::CastToVector2f(Location), ZoomAmount, BookmarkId);
 		}
 	}
+	virtual void SetViewLocation(const FVector2f& Location, float ZoomAmount, const FGuid& BookmarkId = FGuid())
+	{
+	PRAGMA_DISABLE_DEPRECATION_WARNINGS
+		SetViewLocation(FVector2D(Location), ZoomAmount, BookmarkId);
+	PRAGMA_ENABLE_DEPRECATION_WARNINGS
+	}
 
 	/**
 	 * Gets the view location of the graph
@@ -240,13 +254,24 @@ public:
 	 * @param OutLocation		Will have the current view location
 	 * @param OutZoomAmount		Will have the current zoom amount
 	 */
+	UE_DEPRECATED(5.6, "Slate positions are represented in floats. Please use the function returning FVector2f.")
 	virtual void GetViewLocation(FVector2D& OutLocation, float& OutZoomAmount)
 	{
 		if (Implementation.IsValid())
 		{
-			Implementation->GetViewLocation(OutLocation, OutZoomAmount);
+			FVector2f TempLocation;
+			Implementation->GetViewLocation(TempLocation, OutZoomAmount);
+			OutLocation = FVector2D(TempLocation);
 		}
 	}
+	virtual void GetViewLocation(FVector2f& OutLocation, float& OutZoomAmount)
+	{
+		FVector2D TempLocation;
+	PRAGMA_DISABLE_DEPRECATION_WARNINGS
+		GetViewLocation(TempLocation, OutZoomAmount);
+	PRAGMA_ENABLE_DEPRECATION_WARNINGS
+		OutLocation = UE::Slate::CastToVector2f(TempLocation);
+	}
 
 	/**
 	 * Gets the current graph view bookmark
diff --git a/Engine/Source/Editor/UnrealEd/Public/Kismet2/KismetEditorUtilities.h b/Engine/Source/Editor/UnrealEd/Public/Kismet2/KismetEditorUtilities.h
--- a/Engine/Source/Editor/UnrealEd/Public/Kismet2/KismetEditorUtilities.h
+++ b/Engine/Source/Editor/UnrealEd/Public/Kismet2/KismetEditorUtilities.h
@@ -464,7 +464,7 @@ public:
 	static UNREALED_API bool CanPasteNodes(const class UEdGraph* Graph);
 
 	/** Perform paste on graph, at location */
-	static UNREALED_API void  PasteNodesHere( class UEdGraph* Graph, const FVector2D& Location);
+	static UNREALED_API void  PasteNodesHere( class UEdGraph* Graph, const UE::Slate::FDeprecateVector2DParameter& Location);
 
 	/** Attempt to get the bounds for currently selected nodes
 		@returns false if no nodes are selected */
diff --git a/Engine/Source/Runtime/Engine/Classes/EdGraph/EdGraphSchema.h b/Engine/Source/Runtime/Engine/Classes/EdGraph/EdGraphSchema.h
--- a/Engine/Source/Runtime/Engine/Classes/EdGraph/EdGraphSchema.h
+++ b/Engine/Source/Runtime/Engine/Classes/EdGraph/EdGraphSchema.h
@@ -177,23 +177,40 @@ public:
 	virtual bool IsParentable() const { return false; }
 
 	/** Execute this action, given the graph and schema, and possibly a pin that we were dragged from. Returns a node that was created by this action (if any). */
-	virtual UEdGraphNode* PerformAction(class UEdGraph* ParentGraph, UEdGraphPin* FromPin, const FVector2D Location, bool bSelectNewNode = true) { return nullptr; }
+	UE_DEPRECATED(5.6, "Use the version of the function accepting FVector2f; this API no longer interfaces directly with double-precision scalars and vectors.")
+	virtual UEdGraphNode* PerformAction(class UEdGraph* ParentGraph, UEdGraphPin* FromPin, const FVector2D Location, bool bSelectNewNode = true) UE_SLATE_DEPRECATED_VECTOR_VIRTUAL_FUNCTION
+	{
+		return nullptr;
+	}
+	virtual UEdGraphNode* PerformAction(class UEdGraph* ParentGraph, UEdGraphPin* FromPin, const FVector2f& Location, bool bSelectNewNode = true)
+	{
+	PRAGMA_DISABLE_DEPRECATION_WARNINGS
+		return PerformAction(ParentGraph, FromPin, FVector2D(Location), bSelectNewNode);
+	PRAGMA_ENABLE_DEPRECATION_WARNINGS
+	}
 
 	/** Execute this action, given the graph and schema, and possibly a pin that we were dragged from. Returns a node that was created by this action (if any). */
-	virtual UEdGraphNode* PerformAction(class UEdGraph* ParentGraph, TArray<UEdGraphPin*>& FromPins, const FVector2D Location, bool bSelectNewNode = true)
+	UE_DEPRECATED(5.6, "Use the version of the function accepting FVector2f; this API no longer interfaces directly with double-precision scalars and vectors.")
+	virtual UEdGraphNode* PerformAction(class UEdGraph* ParentGraph, TArray<UEdGraphPin*>& FromPins, const FVector2D Location, bool bSelectNewNode = true) UE_SLATE_DEPRECATED_VECTOR_VIRTUAL_FUNCTION
 	{
 		UEdGraphNode* NewNode = nullptr;
 		if (FromPins.Num() > 0)
 		{
-			NewNode = PerformAction(ParentGraph, FromPins[0], Location, bSelectNewNode);
+			NewNode = PerformAction(ParentGraph, FromPins[0], UE::Slate::CastToVector2f(Location), bSelectNewNode);
 		}
 		else
 		{
-			NewNode = PerformAction(ParentGraph, nullptr, Location, bSelectNewNode);
+			NewNode = PerformAction(ParentGraph, nullptr, UE::Slate::CastToVector2f(Location), bSelectNewNode);
 		}
 
 		return NewNode;
 	}
+	virtual UEdGraphNode* PerformAction(class UEdGraph* ParentGraph, TArray<UEdGraphPin*>& FromPins, const FVector2f& Location, bool bSelectNewNode = true)
+	{
+	PRAGMA_DISABLE_DEPRECATION_WARNINGS
+		return PerformAction(ParentGraph, FromPins, FVector2D(Location), bSelectNewNode);
+	PRAGMA_ENABLE_DEPRECATION_WARNINGS
+	}
 
 	/** Performs a double click on the action */
 	virtual FReply OnDoubleClick(UBlueprint* InBlueprint) { return FReply::Unhandled(); }
@@ -361,13 +378,13 @@ struct FEdGraphSchemaAction_NewNode : public FEdGraphSchemaAction
 	{}
 
 	// FEdGraphSchemaAction interface
-	ENGINE_API virtual UEdGraphNode* PerformAction(class UEdGraph* ParentGraph, UEdGraphPin* FromPin, const FVector2D Location, bool bSelectNewNode = true) override;
-	ENGINE_API virtual UEdGraphNode* PerformAction(class UEdGraph* ParentGraph, TArray<UEdGraphPin*>& FromPins, const FVector2D Location, bool bSelectNewNode = true) override;
+	ENGINE_API virtual UEdGraphNode* PerformAction(class UEdGraph* ParentGraph, UEdGraphPin* FromPin, const FVector2f& Location, bool bSelectNewNode = true) override;
+	ENGINE_API virtual UEdGraphNode* PerformAction(class UEdGraph* ParentGraph, TArray<UEdGraphPin*>& FromPins, const FVector2f& Location, bool bSelectNewNode = true) override;
 	ENGINE_API virtual void AddReferencedObjects(FReferenceCollector& Collector) override;
 	// End of FEdGraphSchemaAction interface
 
 	template <typename NodeType>
-	static NodeType* SpawnNodeFromTemplate(class UEdGraph* ParentGraph, NodeType* InTemplateNode, const FVector2D Location, bool bSelectNewNode = true)
+	static NodeType* SpawnNodeFromTemplate(class UEdGraph* ParentGraph, NodeType* InTemplateNode, const UE::Slate::FDeprecateVector2DParameter& Location, bool bSelectNewNode = true)
 	{
 		FEdGraphSchemaAction_NewNode Action;
 		Action.NodeTemplate = InTemplateNode;
@@ -375,7 +392,7 @@ struct FEdGraphSchemaAction_NewNode : public FEdGraphSchemaAction
 		return Cast<NodeType>(Action.PerformAction(ParentGraph, nullptr, Location, bSelectNewNode));
 	}
 
-	static ENGINE_API UEdGraphNode* CreateNode(class UEdGraph* ParentGraph, UEdGraphPin* FromPin, const FVector2D Location, class UEdGraphNode* InNodeTemplate);
+	static ENGINE_API UEdGraphNode* CreateNode(class UEdGraph* ParentGraph, UEdGraphPin* FromPin, const UE::Slate::FDeprecateVector2DParameter& Location, class UEdGraphNode* InNodeTemplate);
 };
 
 /** Dummy action, useful for putting messages in the menu */
@@ -499,7 +516,7 @@ public:
 		 * @param FromPins Optional pins that the action was dragged from.
 		 * @param Location The position on the graph to place new nodes.
 		 */
-		ENGINE_API void PerformAction( class UEdGraph* ParentGraph, TArray<UEdGraphPin*>& FromPins, const FVector2D Location );
+		ENGINE_API void PerformAction( class UEdGraph* ParentGraph, TArray<UEdGraphPin*>& FromPins, const UE::Slate::FDeprecateVector2DParameter& Location );
 		
 		/**
 		 * Returns a the string that should be used when searching for matching actions. Looks only at the first action.
@@ -1027,7 +1044,14 @@ class UEdGraphSchema : public UObject
 	virtual void RecombinePin(UEdGraphPin* Pin) const { };
 
 	/** Handles double-clicking on a pin<->pin connection */
-	virtual void OnPinConnectionDoubleCicked(UEdGraphPin* PinA, UEdGraphPin* PinB, const FVector2D& GraphPosition) const { }
+	UE_DEPRECATED(5.6, "Use the version of the function accepting FVector2f; this API no longer interfaces directly with double-precision scalars and vectors.")
+	virtual void OnPinConnectionDoubleCicked(UEdGraphPin* PinA, UEdGraphPin* PinB, const FVector2D& GraphPosition) const UE_SLATE_DEPRECATED_VECTOR_VIRTUAL_FUNCTION { }
+	virtual void OnPinConnectionDoubleCicked(UEdGraphPin* PinA, UEdGraphPin* PinB, const FVector2f& GraphPosition) const
+	{
+	PRAGMA_DISABLE_DEPRECATION_WARNINGS
+		return OnPinConnectionDoubleCicked(PinA, PinB, FVector2D(GraphPosition));
+	PRAGMA_ENABLE_DEPRECATION_WARNINGS
+	}
 
 	/**
 	 * Break links on this pin and create links instead on MoveToPin
@@ -1081,7 +1105,9 @@ class UEdGraphSchema : public UObject
 	 * @param	Node			The node to set
 	 * @param	Position		The target position
 	 */
-	ENGINE_API virtual void SetNodePosition(UEdGraphNode* Node, const FVector2D& Position) const;
+	UE_DEPRECATED(5.6, "Use the version of the function accepting FVector2f; this API no longer interfaces directly with double-precision scalars and vectors.")
+	ENGINE_API virtual void SetNodePosition(UEdGraphNode* Node, const FVector2D& Position) const UE_SLATE_DEPRECATED_VECTOR_VIRTUAL_FUNCTION;
+	ENGINE_API virtual void SetNodePosition(UEdGraphNode* Node, const FVector2f& Position) const;
 
 	/**
 	 * Returns the currently selected graph node count
@@ -1163,13 +1189,34 @@ class UEdGraphSchema : public UObject
 	virtual FReply TrySetGraphCategory(const UEdGraph* InGraph, const FText& InCategory) const { return FReply::Unhandled(); }
 
 	/** Called when asset(s) are dropped onto a graph background. */
-	virtual void DroppedAssetsOnGraph(const TArray<struct FAssetData>& Assets, const FVector2D& GraphPosition, UEdGraph* Graph) const {}
+	UE_DEPRECATED(5.6, "Use the version of the function accepting FVector2f; this API no longer interfaces directly with double-precision scalars and vectors.")
+	virtual void DroppedAssetsOnGraph(const TArray<struct FAssetData>& Assets, const FVector2D& GraphPosition, UEdGraph* Graph) const UE_SLATE_DEPRECATED_VECTOR_VIRTUAL_FUNCTION {}
+	virtual void DroppedAssetsOnGraph(const TArray<struct FAssetData>& Assets, const FVector2f& GraphPosition, UEdGraph* Graph) const
+	{
+	PRAGMA_DISABLE_DEPRECATION_WARNINGS
+		return DroppedAssetsOnGraph(Assets, FVector2D(GraphPosition), Graph);
+	PRAGMA_ENABLE_DEPRECATION_WARNINGS
+	}
 
 	/** Called when asset(s) are dropped onto the specified node */
-	virtual void DroppedAssetsOnNode(const TArray<struct FAssetData>& Assets, const FVector2D& GraphPosition, UEdGraphNode* Node) const {}
+	UE_DEPRECATED(5.6, "Use the version of the function accepting FVector2f; this API no longer interfaces directly with double-precision scalars and vectors.")
+	virtual void DroppedAssetsOnNode(const TArray<struct FAssetData>& Assets, const FVector2D& GraphPosition, UEdGraphNode* Node) const UE_SLATE_DEPRECATED_VECTOR_VIRTUAL_FUNCTION {}
+	virtual void DroppedAssetsOnNode(const TArray<struct FAssetData>& Assets, const FVector2f& GraphPosition, UEdGraphNode* Node) const
+	{
+	PRAGMA_DISABLE_DEPRECATION_WARNINGS
+		return DroppedAssetsOnNode(Assets, FVector2D(GraphPosition), Node);
+	PRAGMA_ENABLE_DEPRECATION_WARNINGS
+	}
 
 	/** Called when asset(s) are dropped onto the specified pin */
-	virtual void DroppedAssetsOnPin(const TArray<struct FAssetData>& Assets, const FVector2D& GraphPosition, UEdGraphPin* Pin) const {}
+	UE_DEPRECATED(5.6, "Use the version of the function accepting FVector2f; this API no longer interfaces directly with double-precision scalars and vectors.")
+	virtual void DroppedAssetsOnPin(const TArray<struct FAssetData>& Assets, const FVector2D& GraphPosition, UEdGraphPin* Pin) const UE_SLATE_DEPRECATED_VECTOR_VIRTUAL_FUNCTION {}
+	virtual void DroppedAssetsOnPin(const TArray<struct FAssetData>& Assets, const FVector2f& GraphPosition, UEdGraphPin* Pin) const
+	{
+	PRAGMA_DISABLE_DEPRECATION_WARNINGS
+		return DroppedAssetsOnPin(Assets, FVector2D(GraphPosition), Pin);
+	PRAGMA_ENABLE_DEPRECATION_WARNINGS
+	}
 
 	/** Allows schema to generate a tooltip (icon & message) when the specified asset(s) are dragged over the specified node */
 	virtual void GetAssetsNodeHoverMessage(const TArray<struct FAssetData>& Assets, const UEdGraphNode* HoverNode, FString& OutTooltipText, bool& OutOkIcon) const 
@@ -1286,7 +1333,17 @@ class UEdGraphSchema : public UObject
 	 * @param	InScreenPosition The position inside of the screen
 	 * @return the response to rejection or performing the drop
 	 */
-	virtual bool RequestVariableDropOnPanel(UEdGraph* InGraph, FProperty* InVariableToDrop, const FVector2D& InDropPosition, const FVector2D& InScreenPosition) const { return false; }
+	UE_DEPRECATED(5.6, "Use the version of the function accepting FVector2f; this API no longer interfaces directly with double-precision scalars and vectors.")
+	virtual bool RequestVariableDropOnPanel(UEdGraph* InGraph, FProperty* InVariableToDrop, const FVector2D& InDropPosition, const FVector2D& InScreenPosition) const UE_SLATE_DEPRECATED_VECTOR_VIRTUAL_FUNCTION
+	{
+		return false;
+	}
+	virtual bool RequestVariableDropOnPanel(UEdGraph* InGraph, FProperty* InVariableToDrop, const FVector2f& InDropPosition, const FVector2f& InScreenPosition) const
+	{
+	PRAGMA_DISABLE_DEPRECATION_WARNINGS
+		return RequestVariableDropOnPanel(InGraph, InVariableToDrop, FVector2D(InDropPosition), FVector2D(InScreenPosition));
+	PRAGMA_ENABLE_DEPRECATION_WARNINGS
+	}
 
 	/**
 	 * Request to drop a variable on a node
@@ -1297,7 +1354,17 @@ class UEdGraphSchema : public UObject
 	 * @param	InScreenPosition The position inside of the screen
 	 * @return the response to rejection or performing the drop
 	 */
-	virtual bool RequestVariableDropOnNode(UEdGraph* InGraph, FProperty* InVariableToDrop, UEdGraphNode* InNode, const FVector2D& InDropPosition, const FVector2D& InScreenPosition) const { return false; }
+	UE_DEPRECATED(5.6, "Use the version of the function accepting FVector2f; this API no longer interfaces directly with double-precision scalars and vectors.")
+	virtual bool RequestVariableDropOnNode(UEdGraph* InGraph, FProperty* InVariableToDrop, UEdGraphNode* InNode, const FVector2D& InDropPosition, const FVector2D& InScreenPosition) const UE_SLATE_DEPRECATED_VECTOR_VIRTUAL_FUNCTION
+	{
+		return false;
+	}
+	virtual bool RequestVariableDropOnNode(UEdGraph* InGraph, FProperty* InVariableToDrop, UEdGraphNode* InNode, const FVector2f& InDropPosition, const FVector2f& InScreenPosition) const
+	{
+	PRAGMA_DISABLE_DEPRECATION_WARNINGS
+		return RequestVariableDropOnNode(InGraph, InVariableToDrop, InNode, FVector2D(InDropPosition), FVector2D(InScreenPosition));
+	PRAGMA_ENABLE_DEPRECATION_WARNINGS
+	}
 
 	/**
 	 * Request to drop a variable on a pin
@@ -1308,7 +1375,17 @@ class UEdGraphSchema : public UObject
 	 * @param	InScreenPosition The position inside of the screen
 	 * @return the response to rejection or performing the drop
 	 */
-	virtual bool RequestVariableDropOnPin(UEdGraph* InGraph, FProperty* InVariableToDrop, UEdGraphPin* InPin, const FVector2D& InDropPosition, const FVector2D& InScreenPosition) const { return false; }
+	UE_DEPRECATED(5.6, "Use the version of the function accepting FVector2f; this API no longer interfaces directly with double-precision scalars and vectors.")
+	virtual bool RequestVariableDropOnPin(UEdGraph* InGraph, FProperty* InVariableToDrop, UEdGraphPin* InPin, const FVector2D& InDropPosition, const FVector2D& InScreenPosition) const UE_SLATE_DEPRECATED_VECTOR_VIRTUAL_FUNCTION
+	{
+		return false;
+	}
+	virtual bool RequestVariableDropOnPin(UEdGraph* InGraph, FProperty* InVariableToDrop, UEdGraphPin* InPin, const FVector2f& InDropPosition, const FVector2f& InScreenPosition) const
+	{
+	PRAGMA_DISABLE_DEPRECATION_WARNINGS
+		return RequestVariableDropOnPin(InGraph, InVariableToDrop, InPin, FVector2D(InDropPosition), FVector2D(InScreenPosition));
+	PRAGMA_ENABLE_DEPRECATION_WARNINGS
+	}
 
 	/**
 	 * Returns true if the types and directions of two pins are schema compatible. Handles
diff --git a/Engine/Source/Runtime/Engine/Classes/Engine/Blueprint.h b/Engine/Source/Runtime/Engine/Classes/Engine/Blueprint.h
--- a/Engine/Source/Runtime/Engine/Classes/Engine/Blueprint.h
+++ b/Engine/Source/Runtime/Engine/Classes/Engine/Blueprint.h
@@ -288,7 +288,7 @@ struct FEditedDocumentInfo
 
 	/** Saved view position */
 	UPROPERTY()
-	FVector2D SavedViewOffset;
+	FDeprecateSlateVector2D SavedViewOffset;
 
 	/** Saved zoom amount */
 	UPROPERTY()
@@ -307,7 +307,7 @@ struct FEditedDocumentInfo
 		, EditedObject_DEPRECATED(nullptr)
 	{ }
 
-	FEditedDocumentInfo(UObject* InEditedObject, FVector2D& InSavedViewOffset, float InSavedZoomAmount)
+	FEditedDocumentInfo(UObject* InEditedObject, const UE::Slate::FDeprecateVector2DParameter& InSavedViewOffset, float InSavedZoomAmount)
 		: EditedObjectPath(InEditedObject)
 		, SavedViewOffset(InSavedViewOffset)
 		, SavedZoomAmount(InSavedZoomAmount)
diff --git a/Engine/Source/Runtime/Engine/Private/EdGraph/EdGraphSchema.cpp b/Engine/Source/Runtime/Engine/Private/EdGraph/EdGraphSchema.cpp
--- a/Engine/Source/Runtime/Engine/Private/EdGraph/EdGraphSchema.cpp
+++ b/Engine/Source/Runtime/Engine/Private/EdGraph/EdGraphSchema.cpp
@@ -249,7 +249,7 @@ const TArray<FString>& FGraphActionListBuilderBase::ActionGroup::GetCategoryChai
 #endif
 }
 
-void FGraphActionListBuilderBase::ActionGroup::PerformAction( class UEdGraph* ParentGraph, TArray<UEdGraphPin*>& FromPins, const FVector2D Location )
+void FGraphActionListBuilderBase::ActionGroup::PerformAction( class UEdGraph* ParentGraph, TArray<UEdGraphPin*>& FromPins, const UE::Slate::FDeprecateVector2DParameter& Location )
 {
 	for ( int32 ActionIndex = 0; ActionIndex < Actions.Num(); ActionIndex++ )
 	{
@@ -361,7 +361,7 @@ namespace
 	const int32 NodeDistance = 60;
 }
 
-UEdGraphNode* FEdGraphSchemaAction_NewNode::CreateNode(class UEdGraph* ParentGraph, UEdGraphPin* FromPin, const FVector2D Location, class UEdGraphNode* InNodeTemplate)
+UEdGraphNode* FEdGraphSchemaAction_NewNode::CreateNode(class UEdGraph* ParentGraph, UEdGraphPin* FromPin, const UE::Slate::FDeprecateVector2DParameter& Location, class UEdGraphNode* InNodeTemplate)
 {
 	// Duplicate template node to create new node
 	UEdGraphNode* ResultNode = nullptr;
@@ -401,7 +401,7 @@ UEdGraphNode* FEdGraphSchemaAction_NewNode::CreateNode(class UEdGraph* ParentGra
 	return ResultNode;
 }
 
-UEdGraphNode* FEdGraphSchemaAction_NewNode::PerformAction(class UEdGraph* ParentGraph, UEdGraphPin* FromPin, const FVector2D Location, bool bSelectNewNode/* = true*/)
+UEdGraphNode* FEdGraphSchemaAction_NewNode::PerformAction(class UEdGraph* ParentGraph, UEdGraphPin* FromPin, const FVector2f& Location, bool bSelectNewNode/* = true*/)
 {
 	UEdGraphNode* ResultNode = nullptr;
 
@@ -423,7 +423,7 @@ UEdGraphNode* FEdGraphSchemaAction_NewNode::PerformAction(class UEdGraph* Parent
 	return ResultNode;
 }
 
-UEdGraphNode* FEdGraphSchemaAction_NewNode::PerformAction(class UEdGraph* ParentGraph, TArray<UEdGraphPin*>& FromPins, const FVector2D Location, bool bSelectNewNode/* = true*/) 
+UEdGraphNode* FEdGraphSchemaAction_NewNode::PerformAction(class UEdGraph* ParentGraph, TArray<UEdGraphPin*>& FromPins, const FVector2f& Location, bool bSelectNewNode/* = true*/) 
 {
 	UEdGraphNode* ResultNode = nullptr;
 
@@ -1164,6 +1164,13 @@ void UEdGraphSchema::SetNodePosition(UEdGraphNode* Node, const FVector2D& Positi
 	Node->NodePosY = Position.Y;
 }
 
+void UEdGraphSchema::SetNodePosition(UEdGraphNode* Node, const FVector2f& Position) const
+{
+PRAGMA_DISABLE_DEPRECATION_WARNINGS
+	SetNodePosition(Node, FVector2D(Position));
+PRAGMA_ENABLE_DEPRECATION_WARNINGS
+}
+
 void UEdGraphSchema::GetGraphDisplayInformation(const UEdGraph& Graph, /*out*/ FGraphDisplayInfo& DisplayInfo) const
 {
 	DisplayInfo.PlainName = FText::FromString( Graph.GetName() );
