[HLOD] Use first spatially loaded HLOD partition ancestor as duplication operation source when creating spatially loaded HLOD runtime partitions, instead of always using parent partition, which could be persistent partition

diff --git a/Engine/Source/Runtime/Engine/Private/WorldPartition/RuntimeHashSet/WorldPartitionRuntimeHashSet.cpp b/Engine/Source/Runtime/Engine/Private/WorldPartition/RuntimeHashSet/WorldPartitionRuntimeHashSet.cpp
--- a/Engine/Source/Runtime/Engine/Private/WorldPartition/RuntimeHashSet/WorldPartitionRuntimeHashSet.cpp
+++ b/Engine/Source/Runtime/Engine/Private/WorldPartition/RuntimeHashSet/WorldPartitionRuntimeHashSet.cpp
@@ -18,6 +18,22 @@
 
 #include "WorldPartition/RuntimeHashSet/StaticSpatialIndex.h"
 
+#if WITH_EDITOR
+TObjectPtr<URuntimePartition> FRuntimePartitionDesc::GetFirstSpatiallyLoadedHLODPartitionAncestor(int32 HLODSetupsIndex)
+{
+	check(HLODSetupsIndex < HLODSetups.Num());
+
+	for (int32 Index = HLODSetupsIndex - 1; Index >= 0; Index--)
+	{
+		if (HLODSetups[Index].bIsSpatiallyLoaded)
+		{
+			return HLODSetups[Index].PartitionLayer;
+		}
+	}
+	return MainLayer;
+}
+#endif
+
 void FRuntimePartitionStreamingData::CreatePartitionsSpatialIndex() const
 {
 	if (!SpatialIndex)
@@ -1016,7 +1032,7 @@ void UWorldPartitionRuntimeHashSet::PostEditChangeChainProperty(FPropertyChanged
 			if (RuntimePartitionDesc.HLODSetups.IsValidIndex(HLODSetupsIndex))
 			{
 				FRuntimePartitionHLODSetup& RuntimePartitionHLODSetup = RuntimePartitionDesc.HLODSetups[HLODSetupsIndex];
-				URuntimePartition* ParentRuntimePartition = HLODSetupsIndex ? RuntimePartitionDesc.HLODSetups[HLODSetupsIndex - 1].PartitionLayer : RuntimePartitionDesc.MainLayer;
+				URuntimePartition* ParentRuntimePartition = RuntimePartitionDesc.GetFirstSpatiallyLoadedHLODPartitionAncestor(HLODSetupsIndex);
 				RuntimePartitionHLODSetup.Name = *FString::Printf(TEXT("HLOD_%d"), HLODSetupsIndex);
 				RuntimePartitionHLODSetup.bIsSpatiallyLoaded = true;
 				RuntimePartitionHLODSetup.PartitionLayer = ParentRuntimePartition->CreateHLODRuntimePartition(HLODSetupsIndex);
@@ -1073,7 +1089,7 @@ void UWorldPartitionRuntimeHashSet::PostEditChangeChainProperty(FPropertyChanged
 
 				if (RuntimePartitionHLODSetup.bIsSpatiallyLoaded)
 				{
-					URuntimePartition* ParentRuntimePartition = HLODSetupsIndex ? RuntimePartitionDesc.HLODSetups[HLODSetupsIndex - 1].PartitionLayer : RuntimePartitionDesc.MainLayer;
+					URuntimePartition* ParentRuntimePartition = RuntimePartitionDesc.GetFirstSpatiallyLoadedHLODPartitionAncestor(HLODSetupsIndex);
 					RuntimePartitionHLODSetup.PartitionLayer = ParentRuntimePartition->CreateHLODRuntimePartition(HLODSetupsIndex);
 					RuntimePartitionHLODSetup.PartitionLayer->Name = RuntimePartitionHLODSetup.Name;
 				}
@@ -1147,7 +1163,7 @@ void UWorldPartitionRuntimeHashSet::FixupHLODSetup(FRuntimePartitionDesc& Runtim
 				{
 					if (HLODSetup.bIsSpatiallyLoaded)
 					{
-						URuntimePartition* ParentRuntimePartition = HLODSetupIndex != 0 ? RuntimePartition.HLODSetups[HLODSetupIndex - 1].PartitionLayer : RuntimePartition.MainLayer;
+						URuntimePartition* ParentRuntimePartition = RuntimePartition.GetFirstSpatiallyLoadedHLODPartitionAncestor(HLODSetupIndex);
 						HLODSetup.PartitionLayer = ParentRuntimePartition->CreateHLODRuntimePartition(HLODSetupIndex);
 						HLODSetup.PartitionLayer->Name = HLODSetup.Name;
 					}
diff --git a/Engine/Source/Runtime/Engine/Public/WorldPartition/RuntimeHashSet/WorldPartitionRuntimeHashSet.h b/Engine/Source/Runtime/Engine/Public/WorldPartition/RuntimeHashSet/WorldPartitionRuntimeHashSet.h
--- a/Engine/Source/Runtime/Engine/Public/WorldPartition/RuntimeHashSet/WorldPartitionRuntimeHashSet.h
+++ b/Engine/Source/Runtime/Engine/Public/WorldPartition/RuntimeHashSet/WorldPartitionRuntimeHashSet.h
@@ -69,7 +69,7 @@ struct FRuntimePartitionDesc
 	TArray<FRuntimePartitionHLODSetup> HLODSetups;
 
 #if WITH_EDITOR
-	void UpdateHLODPartitionLayers();
+	TObjectPtr<URuntimePartition> GetFirstSpatiallyLoadedHLODPartitionAncestor(int32 HLODSetupsIndex);
 #endif
 };
 
