Fixed async compute skin cache due to duplicate updates not happening on the async pipe.

diff --git a/Engine/Source/Runtime/Engine/Private/SkeletalRenderGPUSkin.cpp b/Engine/Source/Runtime/Engine/Private/SkeletalRenderGPUSkin.cpp
--- a/Engine/Source/Runtime/Engine/Private/SkeletalRenderGPUSkin.cpp
+++ b/Engine/Source/Runtime/Engine/Private/SkeletalRenderGPUSkin.cpp
@@ -345,6 +345,7 @@ REGISTER_SKELETAL_MESH_UPDATE_BACKEND(FSkeletalMeshUpdatePacketGPUSkin);
 
 void FSkeletalMeshUpdatePacketGPUSkin::UpdateImmediate(FRHICommandList& RHICmdList, FSkeletalMeshObjectGPUSkin* MeshObject, FDynamicSkelMeshObjectDataGPUSkin* DynamicData)
 {
+	FRHICommandListScopedPipeline SkinCacheScope(RHICmdList, DynamicData->GPUSkinTechnique == ESkeletalMeshGPUSkinTechnique::GPUSkinCache ? GPUSkinCachePipeline : ERHIPipeline::Graphics);
 	MeshObject->UpdateDynamicData_RenderThread(RHICmdList, GPUSkinCache, DynamicData);
 }
 
