RemoteObjects:  Removing a hack fix from TickTaskManager that is causing issues during Migration.  The state machine wasn't right when setting a prerequisite during the tickgroup that is executing (GetLastTickGameTime is defaulted to -1.0f even when valid -- you can't determine invalid vs. valid in this case).

#rnx

diff --git a/Engine/Source/Runtime/Engine/Private/TickTaskManager.cpp b/Engine/Source/Runtime/Engine/Private/TickTaskManager.cpp
--- a/Engine/Source/Runtime/Engine/Private/TickTaskManager.cpp
+++ b/Engine/Source/Runtime/Engine/Private/TickTaskManager.cpp
@@ -344,13 +344,7 @@ public:
 
 			LIGHTWEIGHT_TIME_GUARD_END(FTickFunctionTask, Target->DiagnosticMessage());
 		}
-#if UE_WITH_REMOTE_OBJECT_HANDLE
-		const bool bTargetIsValid = (Target->GetLastTickGameTime() != -1.0f); // it's a hack to check for null InternalData
-		if (bTargetIsValid)
-#endif
-		{
-			Target->ClearTaskInformation();  // This is stale and a good time to clear it for safety
-		}
+		Target->ClearTaskInformation();  // This is stale and a good time to clear it for safety
 	}
 };
 
