Disabling ActorDestroy on Multi-Server Proxy which is happening during Push Migrations.  Without this fix, it would cause multiple check()s every migration.

#rnx

diff --git a/Engine/Source/Runtime/Engine/Classes/Engine/World.h b/Engine/Source/Runtime/Engine/Classes/Engine/World.h
--- a/Engine/Source/Runtime/Engine/Classes/Engine/World.h
+++ b/Engine/Source/Runtime/Engine/Classes/Engine/World.h
@@ -1072,6 +1072,7 @@ public:
 	 *		EndPlay()
 	 *		BeginComponentOverlap()
 	 *		EndComponentOverlap()
+	 * 		Actor::Destroyed() and related ones found in the flow of UWorld::DestroyActor().
 	 */
 	void EnableActorTickAndUserCallbacks(bool bEnable);
 	bool IsActorTickAndUserCallbacksEnabled() const { return bEnableActorTickAndUserCallbacks; };
diff --git a/Engine/Source/Runtime/Engine/Private/LevelActor.cpp b/Engine/Source/Runtime/Engine/Private/LevelActor.cpp
--- a/Engine/Source/Runtime/Engine/Private/LevelActor.cpp
+++ b/Engine/Source/Runtime/Engine/Private/LevelActor.cpp
@@ -910,6 +910,12 @@ bool UWorld::DestroyActor( AActor* ThisActor, bool bNetForce, bool bShouldModify
 	// Prevent recursion
 	FMarkActorIsBeingDestroyed MarkActorIsBeingDestroyed(ThisActor);
 
+#if UE_SUPPORT_FOR_ACTOR_TICK_DISABLE
+	// Skip over all of the logic below (this won't cause an indent, so look for the ifdef)
+	if (IsActorTickAndUserCallbacksEnabled())
+	{
+#endif
+
 	// Notify the texture streaming manager about the destruction of this actor.
 	IStreamingManager::Get().NotifyActorDestroyed( ThisActor );
 
@@ -967,6 +973,10 @@ bool UWorld::DestroyActor( AActor* ThisActor, bool bNetForce, bool bShouldModify
 #endif
 	}
 
+#if UE_SUPPORT_FOR_ACTOR_TICK_DISABLE
+	} // if (bEnableActorTickAndUserCallbacks)
+#endif
+
 	ThisActor->ClearComponentOverlaps();
 
 	// If this actor has an owner, notify it that it has lost a child.
