Allow testing dynamic resolution scaling paths via r.DynamicRes.TestScreenPercentage even when dynamic resolution is not supported by the RHI, and on stereo devices.

Dynamic resolution being "unsupported" means the platform does not provide an accurate heuristic for what screen percentage to use, not that we can't actually perform the scaling once a percentage has been provided. Additionally, some platforms do not natively support dynamic resolution, but may support it via a plugin. It should be possible to test dynamic res scaling pathways on those platforms, with or without the plugin.

diff --git a/Engine/Source/Runtime/Engine/Private/UnrealEngine.cpp b/Engine/Source/Runtime/Engine/Private/UnrealEngine.cpp
--- a/Engine/Source/Runtime/Engine/Private/UnrealEngine.cpp
+++ b/Engine/Source/Runtime/Engine/Private/UnrealEngine.cpp
@@ -13716,7 +13716,7 @@ void UEngine::GetDynamicResolutionCurrentStateInfos(FDynamicResolutionStateInfos
 			ForceResolutionFraction = CVarForceDynamicResScreenPercentage.GetValueOnGameThread() / 100.0f;
 		#endif
 
-		if (!DynamicResolutionState.IsValid() || !DynamicResolutionState->IsSupported())
+		if (!DynamicResolutionState.IsValid())
 		{
 			return;
 		}
@@ -13729,6 +13729,10 @@ void UEngine::GetDynamicResolutionCurrentStateInfos(FDynamicResolutionStateInfos
 			OutInfos.ResolutionFractionUpperBounds = DynamicResolutionState->GetResolutionFractionsUpperBoundBudgetValue();
 		}
 		#endif
+		else if (!DynamicResolutionState->IsSupported())
+		{
+			return;
+		}
 		else if (DynamicResolutionState->IsEnabled())
 		{
 			OutInfos.Status = EDynamicResolutionStatus::Enabled;
