[Engine]
* BodyInstance - changed how dllexported extern template is being done. New approach is compatible with header units

diff --git a/Engine/Source/Runtime/Engine/Classes/PhysicsEngine/BodyInstance.h b/Engine/Source/Runtime/Engine/Classes/PhysicsEngine/BodyInstance.h
--- a/Engine/Source/Runtime/Engine/Classes/PhysicsEngine/BodyInstance.h
+++ b/Engine/Source/Runtime/Engine/Classes/PhysicsEngine/BodyInstance.h
@@ -1456,7 +1456,7 @@ private:
 	static ENGINE_API bool IsValidCollisionProfileName(FName InCollisionProfileName);
 
 	template<typename AllocatorType>
-	bool OverlapTestForBodiesImpl(const FVector& Position, const FQuat& Rotation, const TArray<FBodyInstance*, AllocatorType>& Bodies, bool bTraceComplex = false) const;
+	ENGINE_API bool OverlapTestForBodiesImpl(const FVector& Position, const FQuat& Rotation, const TArray<FBodyInstance*, AllocatorType>& Bodies, bool bTraceComplex = false) const;
 
 	friend class UPhysicsAsset;
 	friend class UCollisionProfile;
@@ -1536,6 +1536,3 @@ FORCEINLINE_DEBUGGABLE bool FBodyInstance::IsInstanceSimulatingPhysics() const
 {
 	return ShouldInstanceSimulatingPhysics() && IsValidBodyInstance();
 }
-
-extern template ENGINE_API bool FBodyInstance::OverlapTestForBodiesImpl(const FVector& Position, const FQuat& Rotation, const TArray<FBodyInstance*>& Bodies, bool bTraceComplex) const;
-extern template ENGINE_API bool FBodyInstance::OverlapTestForBodiesImpl(const FVector& Position, const FQuat& Rotation, const TArray<FBodyInstance*, TInlineAllocator<1>>& Bodies, bool bTraceComplex) const;
diff --git a/Engine/Source/Runtime/Engine/Private/PhysicsEngine/BodyInstance.cpp b/Engine/Source/Runtime/Engine/Private/PhysicsEngine/BodyInstance.cpp
--- a/Engine/Source/Runtime/Engine/Private/PhysicsEngine/BodyInstance.cpp
+++ b/Engine/Source/Runtime/Engine/Private/PhysicsEngine/BodyInstance.cpp
@@ -4195,8 +4195,8 @@ bool FBodyInstance::OverlapTestForBodiesImpl(const FVector& Pos, const FQuat& Ro
 }
 
 // Explicit template instantiation for the above.
-template bool FBodyInstance::OverlapTestForBodiesImpl(const FVector& Pos, const FQuat& Rot, const TArray<FBodyInstance*>& Bodies, bool bTraceComplex) const;
-template bool FBodyInstance::OverlapTestForBodiesImpl(const FVector& Pos, const FQuat& Rot, const TArray<FBodyInstance*, TInlineAllocator<1>>& Bodies, bool bTraceComplex) const;
+template ENGINE_API bool FBodyInstance::OverlapTestForBodiesImpl(const FVector& Pos, const FQuat& Rot, const TArray<FBodyInstance*>& Bodies, bool bTraceComplex) const;
+template ENGINE_API bool FBodyInstance::OverlapTestForBodiesImpl(const FVector& Pos, const FQuat& Rot, const TArray<FBodyInstance*, TInlineAllocator<1>>& Bodies, bool bTraceComplex) const;
 
 
 bool FBodyInstance::OverlapTest(const FVector& Position, const FQuat& Rotation, const struct FCollisionShape& CollisionShape, FMTDResult* OutMTD, bool bTraceComplex) const
