Moved visuallogger UnnamedCategory global variable to be extern in header and implemented in cpp

diff --git a/Engine/Source/Runtime/Engine/Private/VisualLogger/VisualLoggerTypes.cpp b/Engine/Source/Runtime/Engine/Private/VisualLogger/VisualLoggerTypes.cpp
--- a/Engine/Source/Runtime/Engine/Private/VisualLogger/VisualLoggerTypes.cpp
+++ b/Engine/Source/Runtime/Engine/Private/VisualLogger/VisualLoggerTypes.cpp
@@ -12,6 +12,10 @@
 #include "VisualLogger/VisualLogger.h"
 #endif
 
+namespace UE::VisualLogger
+{
+	const FName NAME_UnnamedCategory = TEXT("UnnamedCategory");
+}
 
 //----------------------------------------------------------------------//
 // UVisualLoggerDebugSnapshotInterface
diff --git a/Engine/Source/Runtime/Engine/Public/VisualLogger/VisualLoggerTypes.h b/Engine/Source/Runtime/Engine/Public/VisualLogger/VisualLoggerTypes.h
--- a/Engine/Source/Runtime/Engine/Public/VisualLogger/VisualLoggerTypes.h
+++ b/Engine/Source/Runtime/Engine/Public/VisualLogger/VisualLoggerTypes.h
@@ -18,7 +18,7 @@ struct FLogEntryItem;
 
 namespace UE::VisualLogger
 {
-	const FName NAME_UnnamedCategory = TEXT("UnnamedCategory");
+	ENGINE_API extern const FName NAME_UnnamedCategory;
 }
 
 enum class ECreateIfNeeded : int8
