Make GetMeshDeformerMaxLOD() virtual to continue experimenting with using deformers
#tests experimented with subclassed component overriding different deformer lod strategies

diff --git a/Engine/Source/Runtime/Engine/Classes/Components/SkinnedMeshComponent.h b/Engine/Source/Runtime/Engine/Classes/Components/SkinnedMeshComponent.h
--- a/Engine/Source/Runtime/Engine/Classes/Components/SkinnedMeshComponent.h
+++ b/Engine/Source/Runtime/Engine/Classes/Components/SkinnedMeshComponent.h
@@ -355,7 +355,7 @@ public:
 	ENGINE_API UMeshDeformerInstance* GetMeshDeformerInstanceForLOD(int32 LODIndex) const;
 
 	/** Max LOD at which to update or apply the MeshDeformer. */
-	ENGINE_API int32 GetMeshDeformerMaxLOD() const;
+	ENGINE_API virtual int32 GetMeshDeformerMaxLOD() const;
 
 #if WITH_EDITORONLY_DATA
 private:
