Fixed so hidden friends are actually hidden friends (they must be inlined for hidden friend optimization to work)

diff --git a/Engine/Source/Runtime/CoreUObject/Private/Misc/WorldCompositionUtility.cpp b/Engine/Source/Runtime/CoreUObject/Private/Misc/WorldCompositionUtility.cpp
--- a/Engine/Source/Runtime/CoreUObject/Private/Misc/WorldCompositionUtility.cpp
+++ b/Engine/Source/Runtime/CoreUObject/Private/Misc/WorldCompositionUtility.cpp
@@ -71,8 +71,9 @@ void operator<<(FStructuredArchive::FSlot Slot, FWorldTileLODInfo& D)
 		<< SA_VALUE(TEXT("Reserved3"), D.Reserved3);
 }
 
-FArchive& operator<<( FArchive& Ar, FWorldTileInfo& D )
+void FWorldTileInfo::StreamArchive(FArchive& Ar)
 {
+	FWorldTileInfo& D = *this;
 	// Serialized with FPackageFileSummary
 	Ar.UsingCustomVersion(FFortniteMainBranchObjectVersion::GUID);
 
@@ -109,12 +110,11 @@ FArchive& operator<<( FArchive& Ar, FWorldTileInfo& D )
 	{
 		Ar << D.AbsolutePosition;
 	}
-
-	return Ar;
 }
 
-void operator<<(FStructuredArchive::FSlot Slot, FWorldTileInfo& D)
+void FWorldTileInfo::StreamStructuredArchive(FStructuredArchive::FSlot Slot)
 {
+	FWorldTileInfo& D = *this;
 	FStructuredArchive::FRecord Record = Slot.EnterRecord();
 	const FPackageFileVersion ArchiveVersion = Slot.GetUnderlyingArchive().UEVer();
 
diff --git a/Engine/Source/Runtime/CoreUObject/Public/Misc/WorldCompositionUtility.h b/Engine/Source/Runtime/CoreUObject/Public/Misc/WorldCompositionUtility.h
--- a/Engine/Source/Runtime/CoreUObject/Public/Misc/WorldCompositionUtility.h
+++ b/Engine/Source/Runtime/CoreUObject/Public/Misc/WorldCompositionUtility.h
@@ -106,8 +106,8 @@ public:
 	{
 	}
 	
-	COREUOBJECT_API friend FArchive& operator<<(FArchive& Ar, FWorldTileInfo& D);
-	COREUOBJECT_API friend void operator<<(FStructuredArchive::FSlot Slot, FWorldTileInfo& D);
+	friend FArchive& operator<<(FArchive& Ar, FWorldTileInfo& D) { D.StreamArchive(Ar); return Ar; }
+	friend void operator<<(FStructuredArchive::FSlot Slot, FWorldTileInfo& D) { D.StreamStructuredArchive(Slot); }
 
 	bool operator==(const FWorldTileInfo& OtherInfo) const
 	{
@@ -152,5 +152,8 @@ public:
 	TArray<FWorldTileLODInfo> LODList;
 	/** Sorting order */
 	int32				ZOrder;
-};
 
+private:
+	COREUOBJECT_API void StreamArchive(FArchive& Ar);
+	COREUOBJECT_API void StreamStructuredArchive(FStructuredArchive::FSlot Slot);
+};
diff --git a/Engine/Source/Runtime/Engine/Private/PerQualityLevelProperties.cpp b/Engine/Source/Runtime/Engine/Private/PerQualityLevelProperties.cpp
--- a/Engine/Source/Runtime/Engine/Private/PerQualityLevelProperties.cpp
+++ b/Engine/Source/Runtime/Engine/Private/PerQualityLevelProperties.cpp
@@ -435,10 +435,10 @@ void FPerQualityLevelProperty<StructType, ValueType, _BasePropertyName>::Convert
 
 /** Serializer to cook out the most appropriate platform override */
 template<typename StructType, typename ValueType, EName _BasePropertyName>
-ENGINE_API FArchive& operator<<(FArchive& Ar, FPerQualityLevelProperty<StructType, ValueType, _BasePropertyName>& Property)
+void FPerQualityLevelProperty<StructType, ValueType, _BasePropertyName>::StreamArchive(FArchive& Ar)
 {
 	bool bCooked = false;
-	StructType* This = StaticCast<StructType*>(&Property);
+	StructType* This = StaticCast<StructType*>(this);
 #if WITH_EDITOR
 	if (Ar.IsCooking())
 	{
@@ -473,18 +473,17 @@ ENGINE_API FArchive& operator<<(FArchive& Ar, FPerQualityLevelProperty<StructTyp
 		Ar << This->Default;
 		Ar << This->PerQuality;
 	}
-	return Ar;
 }
 
 /** Serializer to cook out the most appropriate platform override */
 template<typename StructType, typename ValueType, EName _BasePropertyName>
-ENGINE_API void operator<<(FStructuredArchive::FSlot Slot, FPerQualityLevelProperty<StructType, ValueType, _BasePropertyName>& Property)
+void FPerQualityLevelProperty<StructType, ValueType, _BasePropertyName>::StreamStructuredArchive(FStructuredArchive::FSlot Slot)
 {
 	FArchive& UnderlyingArchive = Slot.GetUnderlyingArchive();
 	FStructuredArchive::FRecord Record = Slot.EnterRecord();
 
 	bool bCooked = false;
-	StructType* This = StaticCast<StructType*>(&Property);
+	StructType* This = StaticCast<StructType*>(this);
 
 #if WITH_EDITOR
 	if (UnderlyingArchive.IsCooking())
@@ -500,8 +499,8 @@ ENGINE_API void operator<<(FStructuredArchive::FSlot Slot, FPerQualityLevelPrope
 	}
 }
 // 
-template ENGINE_API FArchive& operator<<(FArchive&, FPerQualityLevelProperty<FPerQualityLevelInt, int32, NAME_IntProperty>&);
-template ENGINE_API void operator<<(FStructuredArchive::FSlot Slot, FPerQualityLevelProperty<FPerQualityLevelInt, int32, NAME_IntProperty>&);
+template ENGINE_API void FPerQualityLevelProperty<FPerQualityLevelInt, int32, NAME_IntProperty>::StreamArchive(FArchive& Ar);
+template ENGINE_API void FPerQualityLevelProperty<FPerQualityLevelInt, int32, NAME_IntProperty>::StreamStructuredArchive(FStructuredArchive::FSlot Slot);
 
 #if WITH_EDITOR
 template ENGINE_API int32 FPerQualityLevelProperty<FPerQualityLevelInt, int32, NAME_IntProperty>::GetValueForPlatform(const ITargetPlatform* TargetPlatform) const;
@@ -534,8 +533,8 @@ FString FPerQualityLevelInt::ToString() const
 	return Result;
 }
 
-template ENGINE_API FArchive& operator<<(FArchive&, FPerQualityLevelProperty<FPerQualityLevelFloat, float, NAME_FloatProperty>&);
-template ENGINE_API void operator<<(FStructuredArchive::FSlot Slot, FPerQualityLevelProperty<FPerQualityLevelFloat, float, NAME_FloatProperty>&);
+template ENGINE_API void FPerQualityLevelProperty<FPerQualityLevelFloat, float, NAME_FloatProperty>::StreamArchive(FArchive& Ar);
+template ENGINE_API void FPerQualityLevelProperty<FPerQualityLevelFloat, float, NAME_FloatProperty>::StreamStructuredArchive(FStructuredArchive::FSlot Slot);
 
 #if WITH_EDITOR
 template int32 FPerQualityLevelProperty<FPerQualityLevelFloat, float, NAME_FloatProperty>::GetValueForPlatform(const ITargetPlatform* TargetPlatform) const;
diff --git a/Engine/Source/Runtime/Engine/Public/PerQualityLevelProperties.h b/Engine/Source/Runtime/Engine/Public/PerQualityLevelProperties.h
--- a/Engine/Source/Runtime/Engine/Public/PerQualityLevelProperties.h
+++ b/Engine/Source/Runtime/Engine/Public/PerQualityLevelProperties.h
@@ -206,12 +206,14 @@ struct FPerQualityLevelProperty
 
 private:
 
-#if WITH_EDITOR
-	
 	/* Make this operator friend so it can access private member SavedValue. */
-	template<typename __StructType, typename __ValueType, EName __BasePropertyName>
-	friend ENGINE_API FArchive& operator<<(FArchive& Ar, FPerQualityLevelProperty<__StructType, __ValueType, __BasePropertyName>& Property);
+	friend FArchive& operator<<(FArchive& Ar, FPerQualityLevelProperty& Property) { Property.StreamArchive(Ar); return Ar; } 
+	friend void operator<<(FStructuredArchive::FSlot Slot, FPerQualityLevelProperty& Property) { Property.StreamStructuredArchive(Slot); }
+
+	ENGINE_API void StreamArchive(FArchive& Ar);
+	ENGINE_API void StreamStructuredArchive(FStructuredArchive::FSlot Slot);
 
+#if WITH_EDITOR
 	/** 
 	* Store values here during harvesting so they can be restored during PostSave. 
 	* Why is this a TPimplPtr<...>. The idea is to only restore the state if it was first saved in SavedValue.
@@ -224,12 +226,6 @@ private:
 
 };
 
-template< typename _StructType, typename _ValueType, EName _BasePropertyName>
-ENGINE_API FArchive& operator<<(FArchive& Ar, FPerQualityLevelProperty<_StructType, _ValueType, _BasePropertyName>& Property);
-
-template<typename _StructType, typename _ValueType, EName _BasePropertyName>
-ENGINE_API void operator<<(FStructuredArchive::FSlot Slot, FPerQualityLevelProperty<_StructType, _ValueType, _BasePropertyName>& Property);
-
 USTRUCT(BlueprintType)
 struct FPerQualityLevelInt 
 #if CPP
@@ -258,9 +254,6 @@ struct FPerQualityLevelInt
 	int32 MaxType() const { return MAX_int32; }
 };
 
-extern template ENGINE_API FArchive& operator<<(FArchive&, FPerQualityLevelProperty<FPerQualityLevelInt, int32, NAME_IntProperty>&);
-extern template ENGINE_API void operator<<(FStructuredArchive::FSlot Slot, FPerQualityLevelProperty<FPerQualityLevelInt, int32, NAME_IntProperty>&);
-
 template<>
 struct TStructOpsTypeTraits<FPerQualityLevelInt>
 	: public TStructOpsTypeTraitsBase2<FPerQualityLevelInt>
@@ -301,9 +294,6 @@ struct FPerQualityLevelFloat
 	float MaxType() const { return UE_MAX_FLT; }
 };
 
-extern template ENGINE_API FArchive& operator<<(FArchive&, FPerQualityLevelProperty<FPerQualityLevelFloat, float, NAME_FloatProperty>&);
-extern template ENGINE_API void operator<<(FStructuredArchive::FSlot Slot, FPerQualityLevelProperty<FPerQualityLevelFloat, float, NAME_FloatProperty>&);
-
 template<>
 struct TStructOpsTypeTraits<FPerQualityLevelFloat>
 	: public TStructOpsTypeTraitsBase2<FPerQualityLevelFloat>
diff --git a/Engine/Source/Runtime/RenderCore/Private/ShaderCore.cpp b/Engine/Source/Runtime/RenderCore/Private/ShaderCore.cpp
--- a/Engine/Source/Runtime/RenderCore/Private/ShaderCore.cpp
+++ b/Engine/Source/Runtime/RenderCore/Private/ShaderCore.cpp
@@ -85,8 +85,9 @@ bool operator==(const FShaderStatVariant LHS, const FShaderStatVariant RHS)
 	return false;
 }
 
-FArchive& operator<<(FArchive& Ar, FGenericShaderStat& Stat)
+void FGenericShaderStat::StreamArchive(FArchive& Ar)
 {
+	FGenericShaderStat& Stat = *this;
 	if (Ar.IsSaving())
 	{
 		FString StatNameString = Stat.StatName.ToString();
@@ -113,7 +114,6 @@ FArchive& operator<<(FArchive& Ar, FGenericShaderStat& Stat)
 
 	Ar << Stat.Value;
 	Ar << Stat.Flags;
-	return Ar;
 }
 
 bool FGenericShaderStat::operator==(const FGenericShaderStat& RHS) const
diff --git a/Engine/Source/Runtime/RenderCore/Public/ShaderCore.h b/Engine/Source/Runtime/RenderCore/Public/ShaderCore.h
--- a/Engine/Source/Runtime/RenderCore/Public/ShaderCore.h
+++ b/Engine/Source/Runtime/RenderCore/Public/ShaderCore.h
@@ -80,8 +80,10 @@ public:
 
 	FGenericShaderStat() = default;
 
-	friend FArchive& operator<<(FArchive& Ar, FGenericShaderStat& Stat);
+	friend FArchive& operator<<(FArchive& Ar, FGenericShaderStat& Stat) { Stat.StreamArchive(Ar); return Ar; }
 	bool operator==(const FGenericShaderStat& RHS) const;
+
+	RENDERCORE_API void StreamArchive(FArchive& Ar);
 };
 
 CSV_DECLARE_CATEGORY_MODULE_EXTERN(RENDERCORE_API, Shaders);
