[Engine]
* Fixed so STAT_ClothTotalTime and STAT_ClothWriteback are instantiated in one place

diff --git a/Engine/Source/Runtime/Engine/Classes/Engine/SkinnedAssetCommon.h b/Engine/Source/Runtime/Engine/Classes/Engine/SkinnedAssetCommon.h
--- a/Engine/Source/Runtime/Engine/Classes/Engine/SkinnedAssetCommon.h
+++ b/Engine/Source/Runtime/Engine/Classes/Engine/SkinnedAssetCommon.h
@@ -25,8 +25,8 @@ struct FSkeletalMeshLODGroupSettings;
 
 
 //This is the total cloth time split up among multiple computation (updating gpu, updating sim, etc...)
-DECLARE_CYCLE_STAT(TEXT("Cloth Total"), STAT_ClothTotalTime, STATGROUP_Physics);
-DECLARE_CYCLE_STAT(TEXT("Cloth Writeback"), STAT_ClothWriteback, STATGROUP_Physics);
+DECLARE_CYCLE_STAT_EXTERN(TEXT("Cloth Total"), STAT_ClothTotalTime, STATGROUP_Physics, ENGINE_API);
+DECLARE_CYCLE_STAT_EXTERN(TEXT("Cloth Writeback"), STAT_ClothWriteback, STATGROUP_Physics, ENGINE_API);
 
 UENUM()
 enum class ESkinCacheUsage : uint8
diff --git a/Engine/Source/Runtime/Engine/Private/SkinnedAssetCommon.cpp b/Engine/Source/Runtime/Engine/Private/SkinnedAssetCommon.cpp
--- a/Engine/Source/Runtime/Engine/Private/SkinnedAssetCommon.cpp
+++ b/Engine/Source/Runtime/Engine/Private/SkinnedAssetCommon.cpp
@@ -21,6 +21,9 @@
 #include UE_INLINE_GENERATED_CPP_BY_NAME(SkinnedAssetCommon)
 
 
+DEFINE_STAT(STAT_ClothTotalTime);
+DEFINE_STAT(STAT_ClothWriteback);
+
 /*-----------------------------------------------------------------------------
 	FSkeletalMeshLODInfo
 -----------------------------------------------------------------------------*/
