Fixed various minor header unit compile errors
* Changed static to inline
* Forward declarations
* Removed extern inside inlined public functions
* Moved code from header to source

diff --git a/Engine/Plugins/Online/OnlineServices/Source/OnlineServicesInterface/Public/Online/OnlineErrorCode.h b/Engine/Plugins/Online/OnlineServices/Source/OnlineServicesInterface/Public/Online/OnlineErrorCode.h
--- a/Engine/Plugins/Online/OnlineServices/Source/OnlineServicesInterface/Public/Online/OnlineErrorCode.h
+++ b/Engine/Plugins/Online/OnlineServices/Source/OnlineServicesInterface/Public/Online/OnlineErrorCode.h
@@ -28,9 +28,9 @@ namespace ErrorCode
 
 	namespace System
 	{
-		const ErrorCodeType Engine = 1;
-		const ErrorCodeType Game = 2;
-		const ErrorCodeType ThirdPartyPlugin = 3;
+		inline constexpr ErrorCodeType Engine = 1;
+		inline constexpr ErrorCodeType Game = 2;
+		inline constexpr ErrorCodeType ThirdPartyPlugin = 3;
 	}
 }
 
diff --git a/Engine/Plugins/Runtime/GameplayAbilities/Source/GameplayAbilities/Private/GameplayAbilityTypes.cpp b/Engine/Plugins/Runtime/GameplayAbilities/Source/GameplayAbilities/Private/GameplayAbilityTypes.cpp
--- a/Engine/Plugins/Runtime/GameplayAbilities/Source/GameplayAbilities/Private/GameplayAbilityTypes.cpp
+++ b/Engine/Plugins/Runtime/GameplayAbilities/Source/GameplayAbilities/Private/GameplayAbilityTypes.cpp
@@ -306,6 +306,11 @@ void FGameplayAbilitySpecContainer::RegisterWithOwner(UAbilitySystemComponent* I
 	Owner = InOwner;
 }
 
+bool FGameplayAbilitySpecContainer::NetDeltaSerialize(FNetDeltaSerializeInfo & DeltaParms)
+{
+	return FFastArraySerializer::FastArrayDeltaSerialize<FGameplayAbilitySpec, FGameplayAbilitySpecContainer>(Items, DeltaParms, *this);
+}
+
 // ----------------------------------------------------
 
 FGameplayAbilitySpec::FGameplayAbilitySpec(UGameplayAbility* InAbility, int32 InLevel, int32 InInputID, UObject* InSourceObject)
diff --git a/Engine/Plugins/Runtime/GameplayAbilities/Source/GameplayAbilities/Public/GameplayAbilitySpec.h b/Engine/Plugins/Runtime/GameplayAbilities/Source/GameplayAbilities/Public/GameplayAbilitySpec.h
--- a/Engine/Plugins/Runtime/GameplayAbilities/Source/GameplayAbilities/Public/GameplayAbilitySpec.h
+++ b/Engine/Plugins/Runtime/GameplayAbilities/Source/GameplayAbilities/Public/GameplayAbilitySpec.h
@@ -291,7 +291,7 @@ PRAGMA_ENABLE_DEPRECATION_WARNINGS
 
 /** Fast serializer wrapper for above struct */
 USTRUCT(BlueprintType)
-struct GAMEPLAYABILITIES_API FGameplayAbilitySpecContainer : public FFastArraySerializer
+struct FGameplayAbilitySpecContainer : public FFastArraySerializer
 {
 	GENERATED_USTRUCT_BODY()
 
@@ -309,12 +309,9 @@ struct GAMEPLAYABILITIES_API FGameplayAbilitySpecContainer : public FFastArraySe
 	TObjectPtr<UAbilitySystemComponent> Owner;
 
 	/** Initializes Owner variable */
-	void RegisterWithOwner(UAbilitySystemComponent* Owner);
+	GAMEPLAYABILITIES_API void RegisterWithOwner(UAbilitySystemComponent* Owner);
 
-	bool NetDeltaSerialize(FNetDeltaSerializeInfo & DeltaParms)
-	{
-		return FFastArraySerializer::FastArrayDeltaSerialize<FGameplayAbilitySpec, FGameplayAbilitySpecContainer>(Items, DeltaParms, *this);
-	}
+	GAMEPLAYABILITIES_API bool NetDeltaSerialize(FNetDeltaSerializeInfo & DeltaParms);
 
 	template< typename Type, typename SerializerType >
 	bool ShouldWriteFastArrayItem(const Type& Item, const bool bIsWritingOnClient)
diff --git a/Engine/Source/Runtime/CoreUObject/Public/UObject/CoreNet.h b/Engine/Source/Runtime/CoreUObject/Public/UObject/CoreNet.h
--- a/Engine/Source/Runtime/CoreUObject/Public/UObject/CoreNet.h
+++ b/Engine/Source/Runtime/CoreUObject/Public/UObject/CoreNet.h
@@ -30,7 +30,7 @@
 class FArchive;
 class FName;
 class FNetTraceCollector;
-// Forward declarations
+class FGuidReferences;
 class FOutBunch;
 class FOutputDevice;
 class FProperty;
@@ -515,7 +515,7 @@ class INetSerializeCB
 {
 protected:
 
-	using FGuidReferencesMap = TMap<int32, class FGuidReferences>;
+	using FGuidReferencesMap = TMap<int32, FGuidReferences>;
 
 public:
 
diff --git a/Engine/Source/Runtime/Engine/Classes/Sound/SoundEffectSource.h b/Engine/Source/Runtime/Engine/Classes/Sound/SoundEffectSource.h
--- a/Engine/Source/Runtime/Engine/Classes/Sound/SoundEffectSource.h
+++ b/Engine/Source/Runtime/Engine/Classes/Sound/SoundEffectSource.h
@@ -87,10 +87,10 @@ struct FSoundEffectSourceInitData
 	double AudioClock = 0.0;
 
 	// The object id of the parent preset
-	uint32 ParentPresetUniqueId = INDEX_NONE;
+	uint32 ParentPresetUniqueId = uint32(INDEX_NONE);
 
 	// The audio device ID of the audio device instance this source instance was created from
-	uint32 AudioDeviceId = INDEX_NONE;
+	uint32 AudioDeviceId = uint32(INDEX_NONE);
 };
 
 /** Data required to update the source effect. */
diff --git a/Engine/Source/Runtime/Engine/Public/LightmapUniformShaderParameters.h b/Engine/Source/Runtime/Engine/Public/LightmapUniformShaderParameters.h
--- a/Engine/Source/Runtime/Engine/Public/LightmapUniformShaderParameters.h
+++ b/Engine/Source/Runtime/Engine/Public/LightmapUniformShaderParameters.h
@@ -12,7 +12,7 @@
 class FLightCacheInterface;
 
 /** The maximum value between NUM_LQ_LIGHTMAP_COEF and NUM_HQ_LIGHTMAP_COEF. */
-static const int32 MAX_NUM_LIGHTMAP_COEF = 2;
+inline constexpr int32 MAX_NUM_LIGHTMAP_COEF = 2;
 
 BEGIN_GLOBAL_SHADER_PARAMETER_STRUCT(FPrecomputedLightingUniformParameters, ENGINE_API)
 SHADER_PARAMETER(FVector4f, StaticShadowMapMasks) // TDistanceFieldShadowsAndLightMapPolicy
diff --git a/Engine/Source/Runtime/Engine/Public/MaterialTypes.h b/Engine/Source/Runtime/Engine/Public/MaterialTypes.h
--- a/Engine/Source/Runtime/Engine/Public/MaterialTypes.h
+++ b/Engine/Source/Runtime/Engine/Public/MaterialTypes.h
@@ -204,11 +204,11 @@ enum class EMaterialParameterType : uint8
 };
 DECLARE_INTRINSIC_TYPE_LAYOUT(EMaterialParameterType);
 
-static const int32 NumMaterialParameterTypes = (int32)EMaterialParameterType::Num;
-static const int32 NumMaterialRuntimeParameterTypes = (int32)EMaterialParameterType::NumRuntime;
-static const int32 NumMaterialEditorOnlyParameterTypes = NumMaterialParameterTypes - NumMaterialRuntimeParameterTypes;
+inline constexpr int32 NumMaterialParameterTypes = (int32)EMaterialParameterType::Num;
+inline constexpr int32 NumMaterialRuntimeParameterTypes = (int32)EMaterialParameterType::NumRuntime;
+inline constexpr int32 NumMaterialEditorOnlyParameterTypes = NumMaterialParameterTypes - NumMaterialRuntimeParameterTypes;
 
-static inline bool IsNumericMaterialParameter(EMaterialParameterType InType)
+inline bool IsNumericMaterialParameter(EMaterialParameterType InType)
 {
 	switch (InType)
 	{
@@ -221,7 +221,7 @@ static inline bool IsNumericMaterialParameter(EMaterialParameterType InType)
 	}
 }
 
-static inline bool IsStaticMaterialParameter(EMaterialParameterType InType)
+inline bool IsStaticMaterialParameter(EMaterialParameterType InType)
 {
 	switch (InType)
 	{
@@ -235,20 +235,20 @@ static inline bool IsStaticMaterialParameter(EMaterialParameterType InType)
 
 extern ENGINE_API UE::Shader::FType GetShaderValueType(EMaterialParameterType Type);
 
-static inline const UEnum* GetStaticMaterialParameterTypeEnum()
+inline const UEnum* GetStaticMaterialParameterTypeEnum()
 {
-	static const UEnum* Enum = StaticEnum<EMaterialParameterType>();
+	const UEnum* Enum = StaticEnum<EMaterialParameterType>();
 	check(Enum);
 	return Enum;
 }
 
-static inline FString MaterialParameterTypeToString(EMaterialParameterType Type)
+inline FString MaterialParameterTypeToString(EMaterialParameterType Type)
 {
 	const UEnum* Enum = GetStaticMaterialParameterTypeEnum();
 	return Enum->GetNameStringByValue(int64(Type));
 }
 
-static inline EMaterialParameterType GetMaterialParameterTypeFromString(FString TypeName)
+inline EMaterialParameterType GetMaterialParameterTypeFromString(FString TypeName)
 {
 	const UEnum* Enum = GetStaticMaterialParameterTypeEnum();
 	return static_cast<EMaterialParameterType>(Enum->GetValueByName(FName(TypeName)));
diff --git a/Engine/Source/Runtime/MovieScene/Public/MovieSceneSequenceID.h b/Engine/Source/Runtime/MovieScene/Public/MovieSceneSequenceID.h
--- a/Engine/Source/Runtime/MovieScene/Public/MovieSceneSequenceID.h
+++ b/Engine/Source/Runtime/MovieScene/Public/MovieSceneSequenceID.h
@@ -13,11 +13,11 @@ struct FMovieSceneSequenceID
 {
 	GENERATED_BODY()
 
-	FMovieSceneSequenceID()
+	constexpr FMovieSceneSequenceID()
 		: Value(-1)
 	{}
 
-	explicit FMovieSceneSequenceID(uint32 InValue)
+	constexpr explicit FMovieSceneSequenceID(uint32 InValue)
 		: Value(InValue)
 	{}
 
@@ -94,8 +94,8 @@ typedef TCallTraits<FMovieSceneSequenceID>::ParamType FMovieSceneSequenceIDRef;
 
 namespace MovieSceneSequenceID
 {
-	static const FMovieSceneSequenceID Invalid(-1);
-	static const FMovieSceneSequenceID Root(0);
+	inline constexpr FMovieSceneSequenceID Invalid = FMovieSceneSequenceID(-1);
+	inline constexpr FMovieSceneSequenceID Root = FMovieSceneSequenceID(0);
 }
 
 inline FString LexToString(const FMovieSceneSequenceID& SequenceID)
diff --git a/Engine/Source/Runtime/Net/Core/Classes/Net/Serialization/FastArraySerializer.h b/Engine/Source/Runtime/Net/Core/Classes/Net/Serialization/FastArraySerializer.h
--- a/Engine/Source/Runtime/Net/Core/Classes/Net/Serialization/FastArraySerializer.h
+++ b/Engine/Source/Runtime/Net/Core/Classes/Net/Serialization/FastArraySerializer.h
@@ -399,21 +399,8 @@ struct FFastArraySerializer
 {
 	GENERATED_USTRUCT_BODY()
 
-	FFastArraySerializer()
-		: IDCounter(0)
-		, ArrayReplicationKey(0)
-#if WITH_PUSH_MODEL
-		, OwningObject(nullptr)
-		, RepIndex(INDEX_NONE)
-#endif // WITH_PUSH_MODEL
-		, CachedNumItems(INDEX_NONE)
-		, CachedNumItemsToConsiderForWriting(INDEX_NONE)
-		, DeltaFlags(EFastArraySerializerDeltaFlags::None)
-	{
-		SetDeltaSerializationEnabled(true);
-	}
-
-	~FFastArraySerializer() {}
+	NETCORE_API FFastArraySerializer();
+	NETCORE_API ~FFastArraySerializer();
 
 	/** Maps Element ReplicationID to Array Index.*/
 	TMap<int32, int32> ItemMap;
diff --git a/Engine/Source/Runtime/Net/Core/Private/Net/Serialization/FastArraySerializer.cpp b/Engine/Source/Runtime/Net/Core/Private/Net/Serialization/FastArraySerializer.cpp
--- a/Engine/Source/Runtime/Net/Core/Private/Net/Serialization/FastArraySerializer.cpp
+++ b/Engine/Source/Runtime/Net/Core/Private/Net/Serialization/FastArraySerializer.cpp
@@ -21,3 +21,18 @@ FAutoConsoleVariableRef FFastArraySerializer::CVarMaxNumberOfAllowedDeletionsPer
 	TEXT("")
 );
 
+FFastArraySerializer::FFastArraySerializer()
+		: IDCounter(0)
+		, ArrayReplicationKey(0)
+#if WITH_PUSH_MODEL
+		, OwningObject(nullptr)
+		, RepIndex(INDEX_NONE)
+#endif // WITH_PUSH_MODEL
+		, CachedNumItems(INDEX_NONE)
+		, CachedNumItemsToConsiderForWriting(INDEX_NONE)
+		, DeltaFlags(EFastArraySerializerDeltaFlags::None)
+	{
+		SetDeltaSerializationEnabled(true);
+	}
+
+FFastArraySerializer::~FFastArraySerializer() = default;
diff --git a/Engine/Source/Runtime/SignalProcessing/Private/BufferVectorOperations.cpp b/Engine/Source/Runtime/SignalProcessing/Private/BufferVectorOperations.cpp
--- a/Engine/Source/Runtime/SignalProcessing/Private/BufferVectorOperations.cpp
+++ b/Engine/Source/Runtime/SignalProcessing/Private/BufferVectorOperations.cpp
@@ -1144,12 +1144,12 @@ namespace Audio
 	}
 
 
-	SIGNALPROCESSING_API void DownmixAndSumIntoBuffer(int32 NumSourceChannels, int32 NumDestinationChannels, const FAlignedFloatBuffer& SourceBuffer, FAlignedFloatBuffer& BufferToSumTo, const float* RESTRICT Gains)
+	void DownmixAndSumIntoBuffer(int32 NumSourceChannels, int32 NumDestinationChannels, const FAlignedFloatBuffer& SourceBuffer, FAlignedFloatBuffer& BufferToSumTo, const float* RESTRICT Gains)
 	{
 		DownmixAndSumIntoBuffer(NumSourceChannels, NumDestinationChannels, SourceBuffer.GetData(), BufferToSumTo.GetData(), SourceBuffer.Num() / NumSourceChannels, Gains);
 	}
 
-	SIGNALPROCESSING_API void DownmixAndSumIntoBuffer(int32 NumSourceChannels, int32 NumDestinationChannels, const float* RESTRICT SourceBuffer, float* RESTRICT BufferToSumTo, int32 NumFrames, const float* RESTRICT Gains)
+	void DownmixAndSumIntoBuffer(int32 NumSourceChannels, int32 NumDestinationChannels, const float* RESTRICT SourceBuffer, float* RESTRICT BufferToSumTo, int32 NumFrames, const float* RESTRICT Gains)
 	{
 		for (int32 FrameIndex = 0; FrameIndex < NumFrames; FrameIndex++)
 		{
diff --git a/Engine/Source/Runtime/UMG/Public/Components/SlateWrapperTypes.h b/Engine/Source/Runtime/UMG/Public/Components/SlateWrapperTypes.h
--- a/Engine/Source/Runtime/UMG/Public/Components/SlateWrapperTypes.h
+++ b/Engine/Source/Runtime/UMG/Public/Components/SlateWrapperTypes.h
@@ -183,7 +183,7 @@ namespace EVirtualKeyboardType
 
 namespace EVirtualKeyboardType
 {
-	static EKeyboardType AsKeyboardType( Type InType )
+	inline EKeyboardType AsKeyboardType( Type InType )
 	{
 		switch ( InType )
 		{
