Avoid container usage in loop that can modify its content

diff --git a/Engine/Source/Runtime/Engine/Private/Curves/IndexedCurve.cpp b/Engine/Source/Runtime/Engine/Private/Curves/IndexedCurve.cpp
--- a/Engine/Source/Runtime/Engine/Private/Curves/IndexedCurve.cpp
+++ b/Engine/Source/Runtime/Engine/Private/Curves/IndexedCurve.cpp
@@ -89,7 +89,10 @@ FKeyHandle FIndexedCurve::GetKeyHandle(int32 KeyIndex) const
 void FIndexedCurve::ShiftCurve(float DeltaTime)
 {
 	EnsureAllIndicesHaveHandles();
+
 	TSet<FKeyHandle> KeyHandles;
+	KeyHandles.Reserve(KeyHandlesToIndices.Num());
+
 	for (auto It = KeyHandlesToIndices.CreateConstIterator(); It; ++It)
 	{
 		KeyHandles.Add(*It);
@@ -102,10 +105,10 @@ void FIndexedCurve::ShiftCurve(float DeltaTime, const TSet<FKeyHandle>& KeyHandl
 {
 	if (KeyHandles.Num() != 0)
 	{
-		for (auto It = KeyHandlesToIndices.CreateConstIterator(); It; ++It)
+		for (auto It = KeyHandles.CreateConstIterator(); It; ++It)
 		{
 			const FKeyHandle& KeyHandle = *It;
-			if (KeyHandles.Contains(KeyHandle))
+			if (KeyHandlesToIndices.Find(KeyHandle) != nullptr)
 			{
 				SetKeyTime(KeyHandle, GetKeyTime(KeyHandle) + DeltaTime);
 			}
@@ -116,7 +119,10 @@ void FIndexedCurve::ShiftCurve(float DeltaTime, const TSet<FKeyHandle>& KeyHandl
 void FIndexedCurve::ScaleCurve(float ScaleOrigin, float ScaleFactor)
 {
 	EnsureAllIndicesHaveHandles();
+
 	TSet<FKeyHandle> KeyHandles;
+	KeyHandles.Reserve(KeyHandlesToIndices.Num());
+
 	for (auto It = KeyHandlesToIndices.CreateConstIterator(); It; ++It)
 	{
 		KeyHandles.Add(*It);
@@ -129,10 +135,10 @@ void FIndexedCurve::ScaleCurve(float ScaleOrigin, float ScaleFactor, const TSet<
 {
 	if (KeyHandles.Num() != 0)
 	{
-		for (auto It = KeyHandlesToIndices.CreateConstIterator(); It; ++It)
+		for (auto It = KeyHandles.CreateConstIterator(); It; ++It)
 		{
 			const FKeyHandle& KeyHandle = *It;
-			if (KeyHandles.Contains(KeyHandle))
+			if (KeyHandlesToIndices.Find(KeyHandle) != nullptr)
 			{
 				SetKeyTime(KeyHandle, (GetKeyTime(KeyHandle) - ScaleOrigin) * ScaleFactor + ScaleOrigin);
 			}
