Enable sRGB flag for GBufferC (diffuse color) on mobile platforms. This reduces color banding

diff --git a/Engine/Source/Runtime/Engine/Private/SceneTexturesConfig.cpp b/Engine/Source/Runtime/Engine/Private/SceneTexturesConfig.cpp
--- a/Engine/Source/Runtime/Engine/Private/SceneTexturesConfig.cpp
+++ b/Engine/Source/Runtime/Engine/Private/SceneTexturesConfig.cpp
@@ -224,9 +224,7 @@ uint32 FSceneTexturesConfig::GetEditorPrimitiveNumSamples(ERHIFeatureLevel::Type
 static void SetupMobileGBufferFlags(FGBufferBindings GBufferBindings[GBL_Num], bool bRequiresMultiPass)
 {
 	ETextureCreateFlags AddFlags = TexCreate_InputAttachmentRead;
-	
-	// Mobile uses FBF/subpassLoad to fetch data from GBuffer, and FBF does not always work with sRGB targets 
-	ETextureCreateFlags RemoveFlags = TexCreate_SRGB;
+	ETextureCreateFlags RemoveFlags = TexCreate_None;
 
 	if (!bRequiresMultiPass)
 	{
