Modify some virtual texture default settings in line with the work to enable by default.
Reduced default value for throttling runtime VT updates from 8 to 2.

diff --git a/Engine/Source/Runtime/Engine/Private/VT/VirtualTextureScalability.cpp b/Engine/Source/Runtime/Engine/Private/VT/VirtualTextureScalability.cpp
--- a/Engine/Source/Runtime/Engine/Private/VT/VirtualTextureScalability.cpp
+++ b/Engine/Source/Runtime/Engine/Private/VT/VirtualTextureScalability.cpp
@@ -19,7 +19,7 @@ namespace VirtualTextureScalability
 
 	static TAutoConsoleVariable<int32> CVarVTMaxUploadsPerFrame(
 		TEXT("r.VT.MaxUploadsPerFrame"),
-		8,
+		2,
 		TEXT("Max number of page uploads per frame in game"),
 		ECVF_RenderThreadSafe | ECVF_Scalability
 	);
