[Shaders] Fix parsing ECompilerFlags entries for "recompileshaders" command since CFLAG_ prefix is no longer included in the string representations.

#rnx

diff --git a/Engine/Source/Runtime/Engine/Private/ShaderCompiler/ShaderCompiler.cpp b/Engine/Source/Runtime/Engine/Private/ShaderCompiler/ShaderCompiler.cpp
--- a/Engine/Source/Runtime/Engine/Private/ShaderCompiler/ShaderCompiler.cpp
+++ b/Engine/Source/Runtime/Engine/Private/ShaderCompiler/ShaderCompiler.cpp
@@ -4147,11 +4147,11 @@ static bool ParseShaderCompilerFlags(const TCHAR* InFlagsString, FShaderCompiler
 		ECompilerFlags NextFlag = CFLAG_Max;
 		if (NextFlagArg.StartsWith(TEXT("CFLAG_")))
 		{
-			LexFromString(NextFlag, *NextFlagArg);
+			LexFromString(NextFlag, *NextFlagArg.Mid(6));
 		}
 		else
 		{
-			LexFromString(NextFlag, *FString::Printf(TEXT("CFLAG_%s"), *NextFlagArg));
+			LexFromString(NextFlag, *NextFlagArg);
 		}
 
 		if (NextFlag != CFLAG_Max)
