Fix missing InstancedStaticMesh allocations that were ending up in LoadMapMisc

#tests Verified tracking in memory Insights

diff --git a/Engine/Source/Runtime/Engine/Private/InstancedStaticMesh.cpp b/Engine/Source/Runtime/Engine/Private/InstancedStaticMesh.cpp
--- a/Engine/Source/Runtime/Engine/Private/InstancedStaticMesh.cpp
+++ b/Engine/Source/Runtime/Engine/Private/InstancedStaticMesh.cpp
@@ -1190,6 +1190,8 @@ public:
 
 	FInstancedStaticMeshVFLooseUniformShaderParametersRef CreateUniformBuffer(const FInstancedStaticMeshVFLooseUniformShaderParameters& Params, EUniformBufferUsage UniformBufferUsage)
 	{
+		LLM_SCOPE(ELLMTag::InstancedMesh);
+
 		check(IsInRenderingThread());
 
 		FInstancedStaticMeshVFLooseUniformShaderParametersRef* CachedUB = CachedParams.Find(Params);
@@ -2386,6 +2388,8 @@ FPrimitiveSceneProxy* UInstancedStaticMeshComponent::CreateStaticMeshSceneProxy(
 
 FPrimitiveSceneProxy* UInstancedStaticMeshComponent::CreateSceneProxy()
 {
+	LLM_SCOPE(ELLMTag::InstancedMesh);
+
 	ProxySize = 0;
 
 	PrimitiveInstanceDataManager.ResetComponentDirtyTracking();
