Postload more archetype objects, specifically to make sure components are patched before compilation

#rnx

diff --git a/Engine/Source/Editor/Kismet/Private/BlueprintCompilationManager.cpp b/Engine/Source/Editor/Kismet/Private/BlueprintCompilationManager.cpp
--- a/Engine/Source/Editor/Kismet/Private/BlueprintCompilationManager.cpp
+++ b/Engine/Source/Editor/Kismet/Private/BlueprintCompilationManager.cpp
@@ -2270,6 +2270,17 @@ void FBlueprintCompilationManagerImpl::GatherArchetypesRequiringPostload(TConstA
 			RF_NewerVersionExists|RF_Transient, 
 			EInternalObjectFlags::Garbage);
 			
+		// intentional append, inner objects will form a dag:
+		for(int32 I = 0; I < ArchetypeObjects.Num(); ++I)
+		{
+			GetObjectsWithOuter(
+				ArchetypeObjects[I], 
+				ArchetypeObjects, 
+				true, 
+				RF_NewerVersionExists|RF_Transient, 
+				EInternalObjectFlags::Garbage);
+		}
+
 		// filter out non-archetype instances, note that WidgetTrees and some component
 		// archetypes do not have RF_ArchetypeObject or RF_InheritableComponentTemplate so
 		// we simply detect that they are outered to a UClass or UBlueprint and assume that 
diff --git a/Engine/Source/Runtime/Engine/Private/Blueprint.cpp b/Engine/Source/Runtime/Engine/Private/Blueprint.cpp
--- a/Engine/Source/Runtime/Engine/Private/Blueprint.cpp
+++ b/Engine/Source/Runtime/Engine/Private/Blueprint.cpp
@@ -879,6 +879,63 @@ void UBlueprint::DebuggingWorldRegistrationHelper(UObject* ObjectProvidingWorld,
 	}
 }
 
+void UBlueprint::ForceLoadGraph(UObject* Object)
+{
+	// find every directly referenced object:
+	TArray<UObject*> ReachableInnerObjects;
+	ReachableInnerObjects.Add(Object);
+	FReferenceFinder InnerCollector(ReachableInnerObjects, Object, false, true, true, true);
+	InnerCollector.FindReferences(Object);
+
+	// preload until we stabilize
+	TSet<UObject*> LoadedObjects;
+	TArray<UObject*> NewObjects = ReachableInnerObjects;
+	while(NewObjects.Num())
+	{
+		LoadedObjects.Append(NewObjects);
+		TArray<UObject*> ExternalObjects;
+		for(UObject* InnerObject : NewObjects)
+		{
+			ForceLoad(InnerObject);
+		}
+		TArray<UObject*> ReachableObjects;
+		for(UObject* InnerObject : NewObjects)
+		{
+			FReferenceFinder ComponentCollector(ReachableObjects, nullptr, false, true, true);
+			ComponentCollector.FindReferences(InnerObject);
+		}
+
+		NewObjects.Reset();
+		for(UObject* ReachableObject : ReachableObjects)
+		{
+			if(LoadedObjects.Contains(ReachableObject))
+			{
+				continue;
+			}
+
+			if(ReachableObject->HasAnyFlags(RF_WasLoaded) && !ReachableObject->HasAnyFlags(RF_LoadCompleted))
+			{
+				NewObjects.Add(ReachableObject);
+			}
+			else
+			{
+				// make sure this objects loaded inner objects are also loaded - this specifically handles 
+				// natively created objects like BodySetup that are not reachable by FReferenceFinder:
+				TArray<UObject*> SubObjs;
+				const bool bIncludeNested = true;
+				GetObjectsWithOuter(ReachableObject, SubObjs, bIncludeNested, RF_Transient, EInternalObjectFlags::Garbage);
+				for(UObject* SubObj : SubObjs)
+				{
+					if(SubObj->HasAnyFlags(RF_WasLoaded) && !SubObj->HasAnyFlags(RF_LoadCompleted))
+					{
+						NewObjects.Add(SubObj);
+					}
+				}
+			}
+		}
+	}
+}
+
 #if WITH_EDITOR
 void UBlueprint::OnCookEvent(UE::Cook::ECookEvent CookEvent, UE::Cook::FCookEventContext& CookContext)
 {
@@ -1337,27 +1394,6 @@ void UBlueprint::ForceLoadMembers(UObject* InObject, const UBlueprint* InBluepri
 	}
 }
 
-void UBlueprint::ForceLoadGraph(UObject* Object)
-{
-	// find every directly referenced object:
-	TArray<UObject*> ReachableInnerObjects;
-	ReachableInnerObjects.Add(Object);
-	FReferenceFinder InnerCollector(ReachableInnerObjects, Object, false, true, true, true);
-	InnerCollector.FindReferences(Object);
-
-	TArray<UObject*> ExternalObjects;
-	for(UObject* InnerObject : ReachableInnerObjects)
-	{
-		FReferenceFinder ComponentCollector(ExternalObjects, nullptr, false, true, true);
-		ComponentCollector.FindReferences(InnerObject);
-	}
-
-	for(UObject* Obj : ExternalObjects) //-V1078
-	{
-		ForceLoad(Obj);
-	}
-}
-
 bool UBlueprint::ValidateGeneratedClass(const UClass* InClass)
 {
 	const UBlueprintGeneratedClass* GeneratedClass = Cast<const UBlueprintGeneratedClass>(InClass);
