[Engine]
* Minor code cleanup in BonePose

diff --git a/Engine/Source/Runtime/Engine/Public/BonePose.h b/Engine/Source/Runtime/Engine/Public/BonePose.h
--- a/Engine/Source/Runtime/Engine/Public/BonePose.h
+++ b/Engine/Source/Runtime/Engine/Public/BonePose.h
@@ -617,31 +617,31 @@ void FCSPose<PoseType>::CalculateComponentSpaceTransform(BoneIndexType BoneIndex
 		// current Bones(Index) should contain LocalPoses.
 		FTransform& Bone = Pose[BoneIndex];
 		FTransform& ParentBone = Pose[ParentIndex];
-		check(!Pose[BoneIndex].ContainsNaN());
-		check(!Pose[ParentIndex].ContainsNaN());
+		check(!Bone.ContainsNaN());
+		check(!ParentBone.ContainsNaN());
 
-		FTransform ComponentTransform = Pose[BoneIndex] * Pose[ParentIndex];
+		FTransform ComponentTransform = Bone * ParentBone;
 		if (ComponentTransform.ContainsNaN())
 		{
 			// We've failed, output as much info as we can....
 			// Added for Jira UE-55511
 			auto BoolToStr = [](const bool& bValue) { return bValue ? TEXT("true") : TEXT("false"); };
 
-			const TCHAR* BoneHasNaN = BoolToStr(Pose[BoneIndex].ContainsNaN());
-			const TCHAR* ParentHasNaN = BoolToStr(Pose[ParentIndex].ContainsNaN());
+			const TCHAR* BoneHasNaN = BoolToStr(Bone.ContainsNaN());
+			const TCHAR* ParentHasNaN = BoolToStr(ParentBone.ContainsNaN());
 			FString ErrorMsg = FString(TEXT("NaN created in during FTransform Multiplication\n"));
 			ErrorMsg += FString::Format(TEXT("\tBoneIndex {0} : ParentBoneIndex {1} BoneTransformNaN={2} : ParentTransformNaN={3}\n"), { BoneIndex.GetInt(), ParentIndex.GetInt(), BoneHasNaN, ParentHasNaN });
-			ErrorMsg += FString::Format(TEXT("\tBone {0}\n"), { Pose[BoneIndex].ToString() });
-			ErrorMsg += FString::Format(TEXT("\tParent {0}\n"), { Pose[ParentIndex].ToString() });
+			ErrorMsg += FString::Format(TEXT("\tBone {0}\n"), { Bone.ToString() });
+			ErrorMsg += FString::Format(TEXT("\tParent {0}\n"), { ParentBone.ToString() });
 			ErrorMsg += FString::Format(TEXT("\tResult {0}\n"), { ComponentTransform.ToString() });
-			ErrorMsg += FString::Format(TEXT("\tBone B64 {0}\n"), { FBase64::Encode((uint8*)&Pose[BoneIndex], sizeof(FTransform)) });
-			ErrorMsg += FString::Format(TEXT("\tParent B64 {0}\n"), { FBase64::Encode((uint8*)&Pose[ParentIndex], sizeof(FTransform)) });
+			ErrorMsg += FString::Format(TEXT("\tBone B64 {0}\n"), { FBase64::Encode((uint8*)&Bone, sizeof(FTransform)) });
+			ErrorMsg += FString::Format(TEXT("\tParent B64 {0}\n"), { FBase64::Encode((uint8*)&ParentBone, sizeof(FTransform)) });
 			ErrorMsg += FString::Format(TEXT("\tResult B64 {0}\n"), { FBase64::Encode((uint8*)&ComponentTransform, sizeof(FTransform)) });
 			checkf(false, TEXT("Error during CalculateComponentSpaceTransform\n%s"), *ErrorMsg); // Failed during multiplication
 		}
-		Pose[BoneIndex] = ComponentTransform;
-		Pose[BoneIndex].NormalizeRotation();
-		check(!Pose[BoneIndex].ContainsNaN());
+		Bone = ComponentTransform;
+		Bone.NormalizeRotation();
+		check(!Bone.ContainsNaN());
 		ComponentSpaceFlags[BoneIndex] = 1;
 
 		BoneIndex = BoneIndexStack.Pop(EAllowShrinking::No);
