[Engine[
* Added forward declarations to fix compile errors

diff --git a/Engine/Source/Runtime/Engine/Public/Animation/AnimInstanceProxy.h b/Engine/Source/Runtime/Engine/Public/Animation/AnimInstanceProxy.h
--- a/Engine/Source/Runtime/Engine/Public/Animation/AnimInstanceProxy.h
+++ b/Engine/Source/Runtime/Engine/Public/Animation/AnimInstanceProxy.h
@@ -15,6 +15,8 @@
 #include "AnimInstanceProxy.generated.h"
 
 class UAnimInstance;
+class UAnimNotify;
+class UAnimNotifyState;
 class UBlendProfile;
 class UBlendSpace;
 class FTokenizedMessage;
