Don't invalidate the CurveTable cache id when emptying an already empty table or a CompositeCurveTable that is still being loaded. We only want to invalidate when something could possibly be referencing the curve

diff --git a/Engine/Source/Runtime/Engine/Private/CurveTable.cpp b/Engine/Source/Runtime/Engine/Private/CurveTable.cpp
--- a/Engine/Source/Runtime/Engine/Private/CurveTable.cpp
+++ b/Engine/Source/Runtime/Engine/Private/CurveTable.cpp
@@ -535,19 +535,26 @@ void UCurveTable::EmptyTable()
 {
 	CURVETABLE_CHANGE_SCOPE();
 
-	for (const TPair<FName, FRealCurve*>& CurveRow : GetRowMap())
+	if (RowMap.Num() > 0)
 	{
-		delete CurveRow.Value;
-	}
+		for (const TPair<FName, FRealCurve*>& CurveRow : RowMap)
+		{
+			delete CurveRow.Value;
+		}
 
-	// Finally empty the map
-	RowMap.Empty();
+		// Finally empty the map
+		RowMap.Empty();
 
-	CurveTableMode = ECurveTableMode::Empty;
+		if (HasAnyFlags(RF_LoadCompleted) || !HasAnyFlags(RF_WasLoaded))
+		{
+			// AttributeSets can cache pointers to curves in this table, so we'll need to make sure they've
+			// all been invalidated properly, since we just blew them away.
+			// Ignore this for objects that are still loading as nothing can reference them yet.
+			UCurveTable::InvalidateAllCachedCurves();
+		}
+	}
 
-	// AttributeSets can cache pointers to curves in this table, so we'll need to make sure they've
-	// all been invalidated properly, since we just blew them away.
-	UCurveTable::InvalidateAllCachedCurves();
+	CurveTableMode = ECurveTableMode::Empty;
 }
 
 void UCurveTable::RemoveRow(FName RowName)
