Fix PhysicsConstraintComponent failing to reinitialize in editor when constraint is reregistered without calling InitializeComponent.

In PIE for constraints in child actors in world partition, the constraint component can be reregistered without calling InitializeComponent. Since OnRegister can't initialize the constraint by itself, the constraint will be left uninitialized. This change allows OnRegister to initialize the constraint if it otherwise hasn't been.

diff --git a/Engine/Source/Runtime/Engine/Private/PhysicsEngine/PhysicsConstraintComponent.cpp b/Engine/Source/Runtime/Engine/Private/PhysicsEngine/PhysicsConstraintComponent.cpp
--- a/Engine/Source/Runtime/Engine/Private/PhysicsEngine/PhysicsConstraintComponent.cpp
+++ b/Engine/Source/Runtime/Engine/Private/PhysicsEngine/PhysicsConstraintComponent.cpp
@@ -372,6 +372,12 @@ void UPhysicsConstraintComponent::OnRegister()
 {
 	Super::OnRegister();
 
+	// This reinitialization is needed in PIE because the editor can re-register a constraint without calling InitializeComponent.
+	if(ConstraintInstance.IsTerminated())
+	{
+		InitComponentConstraint();
+	}
+
 	if (SpriteComponent)
 	{
 		UpdateSpriteTexture();
@@ -390,8 +396,6 @@ void UPhysicsConstraintComponent::EndPlay(const EEndPlayReason::Type EndPlayReas
 void UPhysicsConstraintComponent::OnUnregister()
 {
 	Super::OnUnregister();
-
-	// #PHYS2 This is broken because if you re-register again, constraint will not be created, as constraint is init'd ini InitializeComponent not RegisterComponent
 	TermComponentConstraint();
 }
 
