[ActiveSound] - Add CVar for occlusion interval for perf tuning

#rnx

diff --git a/Engine/Source/Runtime/Engine/Private/ActiveSound.cpp b/Engine/Source/Runtime/Engine/Private/ActiveSound.cpp
--- a/Engine/Source/Runtime/Engine/Private/ActiveSound.cpp
+++ b/Engine/Source/Runtime/Engine/Private/ActiveSound.cpp
@@ -21,6 +21,13 @@ FAutoConsoleVariableRef CVarAudioOcclusionEnabled(
 	TEXT("Disables (1) or enables (0) audio occlusion.\n"),
 	ECVF_Default);
 
+static float AudioOcclusionCheckIntervalCvar = 0.f;
+FAutoConsoleVariableRef CVarAudioOcclusionCheckInterval(
+	TEXT("au.OcclusionCheckInterval"),
+	AudioOcclusionCheckIntervalCvar,
+	TEXT("Time interval to recheck occlusion from the closest listener (in seconds).\n"),
+	ECVF_Default);
+
 static bool bUseAsyncOcclusionTracesCVar = false;
 FAutoConsoleVariableRef CVarUseAsyncOcclusionTraces(
 	TEXT("au.UseAsyncOcclusionTraces"),
@@ -176,6 +183,8 @@ FActiveSound::FActiveSound()
 	{
 		ActiveSoundTraceDelegate.BindStatic(&OcclusionTraceDone);
 	}
+
+	OcclusionCheckInterval = AudioOcclusionCheckIntervalCvar;
 }
 
 FActiveSound::FActiveSound(const FActiveSound&) = default;
