Volume Actors: Use the correct icon for different volume actor types instead of forcing the brush icon which we inherit from ABrush

diff --git a/Engine/Source/Runtime/Engine/Classes/GameFramework/Volume.h b/Engine/Source/Runtime/Engine/Classes/GameFramework/Volume.h
--- a/Engine/Source/Runtime/Engine/Classes/GameFramework/Volume.h
+++ b/Engine/Source/Runtime/Engine/Classes/GameFramework/Volume.h
@@ -58,6 +58,7 @@ public:
 #if WITH_EDITOR
 	ENGINE_API virtual void PostEditImport() override;
 	ENGINE_API virtual void PostEditChangeProperty(FPropertyChangedEvent& PropertyChangedEvent) override;
+	ENGINE_API virtual FName GetCustomIconName() const override;
 #endif // WITH_EDITOR
 	//End UObject Interface
 
diff --git a/Engine/Source/Runtime/Engine/Private/Volume.cpp b/Engine/Source/Runtime/Engine/Private/Volume.cpp
--- a/Engine/Source/Runtime/Engine/Private/Volume.cpp
+++ b/Engine/Source/Runtime/Engine/Private/Volume.cpp
@@ -58,6 +58,13 @@ void AVolume::PostEditChangeProperty(FPropertyChangedEvent& PropertyChangedEvent
 	}
 }
 
+FName AVolume::GetCustomIconName() const
+{
+	// Our parent class ABrush provides a default for all brushes (which includes volumes), but we don't want to use those since all volume types
+	// have custom icons
+	return NAME_None;
+}
+
 #endif // WITH_EDITOR
 
 /** @returns the coarse bounds of this volume */
