Added MaterialNormal and MaterialAmbientOcclusion show flags. They can be useful to investigate issues with materials and find materials which have lighting baked into albedo or too strong AO removing indirect lighting.

diff --git a/Engine/Plugins/MovieScene/MovieRenderPipeline/Source/MovieRenderPipelineCore/Private/MoviePipelineUtils.cpp b/Engine/Plugins/MovieScene/MovieRenderPipeline/Source/MovieRenderPipelineCore/Private/MoviePipelineUtils.cpp
--- a/Engine/Plugins/MovieScene/MovieRenderPipeline/Source/MovieRenderPipelineCore/Private/MoviePipelineUtils.cpp
+++ b/Engine/Plugins/MovieScene/MovieRenderPipeline/Source/MovieRenderPipelineCore/Private/MoviePipelineUtils.cpp
@@ -187,7 +187,7 @@ namespace UE
 				View->DiffuseOverrideParameter = FVector4f(0.f, 0.f, 0.f, 0.f);
 				View->SpecularOverrideParameter = FVector4f(1, 1, 1, 0.0f);
 				View->NormalOverrideParameter = FVector4f(0, 0, 1, 0.0f);
-				View->RoughnessOverrideParameter = FVector2D(0.0f, 0.0f);
+				View->RoughnessOverrideParameter = FVector2f(0.0f, 0.0f);
 			}
 	
 			if (!View->Family->EngineShowFlags.Diffuse)
@@ -199,6 +199,16 @@ namespace UE
 			{
 				View->SpecularOverrideParameter = FVector4f(0.f, 0.f, 0.f, 0.f);
 			}
+
+			if (!View->Family->EngineShowFlags.MaterialNormal)
+			{
+				View->NormalOverrideParameter = FVector4f(0.0f, 0.0f, 1.0f, 0.0f);
+			}
+
+			if (!View->Family->EngineShowFlags.MaterialAmbientOcclusion)
+			{
+				View->AmbientOcclusionOverrideParameter = FVector2f(1.0f, 0.0f);
+			}
 	
 			static const FName BufferVisualizationMode(TEXT("WorldNormal"));
 			View->CurrentBufferVisualizationMode = BufferVisualizationMode;
diff --git a/Engine/Plugins/Runtime/nDisplay/Source/DisplayCluster/Private/Game/EngineClasses/Basics/DisplayClusterViewportClient.cpp b/Engine/Plugins/Runtime/nDisplay/Source/DisplayCluster/Private/Game/EngineClasses/Basics/DisplayClusterViewportClient.cpp
--- a/Engine/Plugins/Runtime/nDisplay/Source/DisplayCluster/Private/Game/EngineClasses/Basics/DisplayClusterViewportClient.cpp
+++ b/Engine/Plugins/Runtime/nDisplay/Source/DisplayCluster/Private/Game/EngineClasses/Basics/DisplayClusterViewportClient.cpp
@@ -598,7 +598,7 @@ void UDisplayClusterViewportClient::Draw(FViewport* InViewport, FCanvas* SceneCa
 						View->DiffuseOverrideParameter = FVector4f(0.f, 0.f, 0.f, 0.f);
 						View->SpecularOverrideParameter = FVector4f(1, 1, 1, 0.0f);
 						View->NormalOverrideParameter = FVector4f(0, 0, 1, 0.0f);
-						View->RoughnessOverrideParameter = FVector2D(0.0f, 0.0f);
+						View->RoughnessOverrideParameter = FVector2f(0.0f, 0.0f);
 					}
 
 					if (!View->Family->EngineShowFlags.Diffuse)
@@ -611,6 +611,16 @@ void UDisplayClusterViewportClient::Draw(FViewport* InViewport, FCanvas* SceneCa
 						View->SpecularOverrideParameter = FVector4f(0.f, 0.f, 0.f, 0.f);
 					}
 
+					if (!View->Family->EngineShowFlags.MaterialNormal)
+					{
+						View->NormalOverrideParameter = FVector4f(0.0f, 0.0f, 1.0f, 0.0f);
+					}
+
+					if (!View->Family->EngineShowFlags.MaterialAmbientOcclusion)
+					{
+						View->AmbientOcclusionOverrideParameter = FVector2f(1.0f, 0.0f);
+					}
+
 					View->CurrentBufferVisualizationMode = GetCurrentBufferVisualizationMode();
 
 					View->CameraConstrainedViewRect = View->UnscaledViewRect;
diff --git a/Engine/Shaders/Private/MaterialTemplate.ush b/Engine/Shaders/Private/MaterialTemplate.ush
--- a/Engine/Shaders/Private/MaterialTemplate.ush
+++ b/Engine/Shaders/Private/MaterialTemplate.ush
@@ -3979,7 +3979,16 @@ SHADER_POP_WARNINGS_STATE
 
 half GetMaterialAmbientOcclusionRaw(FPixelMaterialInputs PixelMaterialInputs)
 {
-	return PixelMaterialInputs.AmbientOcclusion;
+	half AmbientOcclusion = PixelMaterialInputs.AmbientOcclusion;
+
+	#if (USE_EDITOR_SHADERS && !ES3_1_PROFILE) || MOBILE_EMULATION
+	{
+		// this feature is only needed for development/editor - we can compile it out for a shipping build (see r.CompileShadersForDevelopment cvar help)
+		AmbientOcclusion = AmbientOcclusion * ResolvedView.AmbientOcclusionOverrideParameter.y + ResolvedView.AmbientOcclusionOverrideParameter.x;
+	}
+	#endif
+
+	return AmbientOcclusion;
 }
 
 half GetMaterialAmbientOcclusion(FPixelMaterialInputs PixelMaterialInputs)
diff --git a/Engine/Source/Editor/UnrealEd/Private/EditorViewportClient.cpp b/Engine/Source/Editor/UnrealEd/Private/EditorViewportClient.cpp
--- a/Engine/Source/Editor/UnrealEd/Private/EditorViewportClient.cpp
+++ b/Engine/Source/Editor/UnrealEd/Private/EditorViewportClient.cpp
@@ -4242,7 +4242,7 @@ void FEditorViewportClient::SetupViewForRendering(FSceneViewFamily& ViewFamily,
 		View.DiffuseOverrideParameter = FVector4f(0.f, 0.f, 0.f, 0.f);
 		View.SpecularOverrideParameter = FVector4f(1, 1, 1, 0.0f);
 		View.NormalOverrideParameter = FVector4f(0, 0, 1, 0.0f);
-		View.RoughnessOverrideParameter = FVector2D(0.0f, 0.0f);
+		View.RoughnessOverrideParameter = FVector2f(0.0f, 0.0f);
 	}
 
 	if (!ViewFamily.EngineShowFlags.Diffuse)
@@ -4255,6 +4255,16 @@ void FEditorViewportClient::SetupViewForRendering(FSceneViewFamily& ViewFamily,
 		View.SpecularOverrideParameter = FVector4f(0.f, 0.f, 0.f, 0.f);
 	}
 
+	if (!ViewFamily.EngineShowFlags.MaterialNormal)
+	{
+		View.NormalOverrideParameter = FVector4f(0.0f, 0.0f, 1.0f, 0.0f);
+	}
+
+	if (!ViewFamily.EngineShowFlags.MaterialAmbientOcclusion)
+	{
+		View.AmbientOcclusionOverrideParameter = FVector2f(1.0f, 0.0f);
+	}
+
 	View.CurrentBufferVisualizationMode = CurrentBufferVisualizationMode;
 	View.CurrentNaniteVisualizationMode = CurrentNaniteVisualizationMode;
 	View.CurrentLumenVisualizationMode = CurrentLumenVisualizationMode;
diff --git a/Engine/Source/Runtime/Engine/Private/GameViewportClient.cpp b/Engine/Source/Runtime/Engine/Private/GameViewportClient.cpp
--- a/Engine/Source/Runtime/Engine/Private/GameViewportClient.cpp
+++ b/Engine/Source/Runtime/Engine/Private/GameViewportClient.cpp
@@ -1669,7 +1669,7 @@ void UGameViewportClient::Draw(FViewport* InViewport, FCanvas* SceneCanvas)
 						View->DiffuseOverrideParameter = FVector4f(0.f, 0.f, 0.f, 0.f);
 						View->SpecularOverrideParameter = FVector4f(1, 1, 1, 0.0f);
 						View->NormalOverrideParameter = FVector4f(0, 0, 1, 0.0f);
-						View->RoughnessOverrideParameter = FVector2D(0.0f, 0.0f);
+						View->RoughnessOverrideParameter = FVector2f(0.0f, 0.0f);
 					}
 
 					if (!View->Family->EngineShowFlags.Diffuse)
@@ -1682,6 +1682,16 @@ void UGameViewportClient::Draw(FViewport* InViewport, FCanvas* SceneCanvas)
 						View->SpecularOverrideParameter = FVector4f(0.f, 0.f, 0.f, 0.f);
 					}
 
+					if (!ViewFamily.EngineShowFlags.MaterialNormal)
+					{
+						View->NormalOverrideParameter = FVector4f(0.0f, 0.0f, 1.0f, 0.0f);
+					}
+
+					if (!ViewFamily.EngineShowFlags.MaterialAmbientOcclusion)
+					{
+						View->AmbientOcclusionOverrideParameter = FVector2f(1.0f, 0.0f);
+					}
+
 					View->CurrentBufferVisualizationMode = CurrentBufferVisualizationMode;
 					View->CurrentNaniteVisualizationMode = CurrentNaniteVisualizationMode;
 					View->CurrentLumenVisualizationMode = CurrentLumenVisualizationMode;
diff --git a/Engine/Source/Runtime/Engine/Private/SceneView.cpp b/Engine/Source/Runtime/Engine/Private/SceneView.cpp
--- a/Engine/Source/Runtime/Engine/Private/SceneView.cpp
+++ b/Engine/Source/Runtime/Engine/Private/SceneView.cpp
@@ -885,7 +885,8 @@ FSceneView::FSceneView(const FSceneViewInitOptions& InitOptions)
 	, DiffuseOverrideParameter(FVector4(0,0,0,1))
 	, SpecularOverrideParameter(FVector4(0,0,0,1))
 	, NormalOverrideParameter(FVector4(0,0,0,1))
-	, RoughnessOverrideParameter(FVector2D(0,1))
+	, RoughnessOverrideParameter(FVector2f(0,1))
+	, AmbientOcclusionOverrideParameter(0.0f, 1.0f)
 	, MaterialTextureMipBias(0.f)
 	, HiddenPrimitives(InitOptions.HiddenPrimitives)
 	, ShowOnlyPrimitives(InitOptions.ShowOnlyPrimitives)
@@ -2621,7 +2622,7 @@ void SetupCommonViewUniformBufferParameters(
 {
 	QUICK_SCOPE_CYCLE_COUNTER(STAT_SetupCommonViewUniformBufferParameters);
 	FVector4f LocalDiffuseOverrideParameter = Inputs.DiffuseOverrideParameter;
-	FVector2D LocalRoughnessOverrideParameter = Inputs.RoughnessOverrideParameter;
+	FVector2f LocalRoughnessOverrideParameter = Inputs.RoughnessOverrideParameter;
 
 #if !(UE_BUILD_SHIPPING || UE_BUILD_TEST)
 	{
@@ -2727,7 +2728,8 @@ void SetupCommonViewUniformBufferParameters(
 	ViewUniformShaderParameters.DiffuseOverrideParameter = LocalDiffuseOverrideParameter;
 	ViewUniformShaderParameters.SpecularOverrideParameter = Inputs.SpecularOverrideParameter;
 	ViewUniformShaderParameters.NormalOverrideParameter = Inputs.NormalOverrideParameter;
-	ViewUniformShaderParameters.RoughnessOverrideParameter = FVector2f(LocalRoughnessOverrideParameter);	// LWC_TODO: Precision loss
+	ViewUniformShaderParameters.RoughnessOverrideParameter = LocalRoughnessOverrideParameter;
+	ViewUniformShaderParameters.AmbientOcclusionOverrideParameter = Inputs.AmbientOcclusionOverrideParameter;
 	ViewUniformShaderParameters.WorldCameraMovementSinceLastFrame = FVector3f(InViewMatrices.GetViewOrigin() - InPrevViewMatrices.GetViewOrigin());
 	ViewUniformShaderParameters.CullingSign = Inputs.bReverseCulling ? -1.0f : 1.0f;
 	ViewUniformShaderParameters.NearPlane = InViewMatrices.ComputeNearPlane();
diff --git a/Engine/Source/Runtime/Engine/Public/SceneView.h b/Engine/Source/Runtime/Engine/Public/SceneView.h
--- a/Engine/Source/Runtime/Engine/Public/SceneView.h
+++ b/Engine/Source/Runtime/Engine/Public/SceneView.h
@@ -921,6 +921,7 @@ enum ETranslucencyVolumeCascade
 	VIEW_UNIFORM_BUFFER_MEMBER_EX(FVector4f, SpecularOverrideParameter, EShaderPrecisionModifier::Half) \
 	VIEW_UNIFORM_BUFFER_MEMBER_EX(FVector4f, NormalOverrideParameter, EShaderPrecisionModifier::Half) \
 	VIEW_UNIFORM_BUFFER_MEMBER_EX(FVector2f, RoughnessOverrideParameter, EShaderPrecisionModifier::Half) \
+	VIEW_UNIFORM_BUFFER_MEMBER_EX(FVector2f, AmbientOcclusionOverrideParameter, EShaderPrecisionModifier::Half) \
 	VIEW_UNIFORM_BUFFER_MEMBER(float, PrevFrameGameTime) \
 	VIEW_UNIFORM_BUFFER_MEMBER(float, PrevFrameRealTime) \
 	VIEW_UNIFORM_BUFFER_MEMBER_EX(float, OutOfBoundsMask, EShaderPrecisionModifier::Half) \
@@ -1515,7 +1516,8 @@ public:
 	FVector4f DiffuseOverrideParameter;
 	FVector4f SpecularOverrideParameter;
 	FVector4f NormalOverrideParameter;
-	FVector2D RoughnessOverrideParameter;
+	FVector2f RoughnessOverrideParameter;
+	FVector2f AmbientOcclusionOverrideParameter;
 
 	/** Mip bias to apply in material's samplers. */
 	float MaterialTextureMipBias;
@@ -2597,6 +2599,7 @@ struct FSetupViewUniformParametersInputs
 			, .NormalOverrideParameter = View.NormalOverrideParameter
 			, .SpecularOverrideParameter = View.SpecularOverrideParameter
 			, .RoughnessOverrideParameter = View.RoughnessOverrideParameter
+			, .AmbientOcclusionOverrideParameter = View.AmbientOcclusionOverrideParameter
 			, .DebugViewShaderMode = View.Family->GetDebugViewShaderMode()
 			, .FrameCounter = View.Family->FrameCounter
 			, .FrameNumber = View.Family->FrameNumber
@@ -2620,7 +2623,8 @@ struct FSetupViewUniformParametersInputs
 	FVector4f DiffuseOverrideParameter       = FVector4f(0, 0, 0, 1);
 	FVector4f NormalOverrideParameter        = FVector4f(0, 0, 0, 1);
 	FVector4f SpecularOverrideParameter      = FVector4f(0, 0, 0, 1);
-	FVector2D RoughnessOverrideParameter     = FVector2D(0, 1);
+	FVector2f RoughnessOverrideParameter     = FVector2f(0, 1);
+	FVector2f AmbientOcclusionOverrideParameter = FVector2f(0, 1);
 	EDebugViewShaderMode DebugViewShaderMode = EDebugViewShaderMode::DVSM_None;
 	uint64 FrameCounter                      = 0;
 	uint32 FrameNumber                       = 0;
diff --git a/Engine/Source/Runtime/Engine/Public/ShowFlags.h b/Engine/Source/Runtime/Engine/Public/ShowFlags.h
--- a/Engine/Source/Runtime/Engine/Public/ShowFlags.h
+++ b/Engine/Source/Runtime/Engine/Public/ShowFlags.h
@@ -402,6 +402,8 @@ private:
 		SetOverrideDiffuseAndSpecular(false);
 		SetLightingOnlyOverride(false);
 		SetReflectionOverride(false);
+		SetMaterialNormal(true);
+		SetMaterialAmbientOcclusion(true);
 		SetVisualizeBuffer(false);
 		SetVisualizeNanite(false);
 		SetVisualizeLumen(false);
diff --git a/Engine/Source/Runtime/Engine/Public/ShowFlagsValues.inl b/Engine/Source/Runtime/Engine/Public/ShowFlagsValues.inl
--- a/Engine/Source/Runtime/Engine/Public/ShowFlagsValues.inl
+++ b/Engine/Source/Runtime/Engine/Public/ShowFlagsValues.inl
@@ -66,6 +66,10 @@ SHOWFLAG_FIXED_IN_SHIPPING(0, OverrideDiffuseAndSpecular, SFG_Hidden, NSLOCTEXT(
 SHOWFLAG_FIXED_IN_SHIPPING(0, LightingOnlyOverride, SFG_Hidden, NSLOCTEXT("UnrealEd", "LightingOnlyOverrideSF", "Lighting Only"))
 /** needed for VMI_ReflectionOverride, Whether to override all materials to be smooth, mirror reflections. */
 SHOWFLAG_FIXED_IN_SHIPPING(0, ReflectionOverride, SFG_Hidden, NSLOCTEXT("UnrealEd", "ReflectionOverrideSF", "Reflections"))
+/** Whether to use material normal. */
+SHOWFLAG_FIXED_IN_SHIPPING(1, MaterialNormal, SFG_Hidden, NSLOCTEXT("UnrealEd", "MaterialNormalSF", "Material Normal"))
+/** Whether to use material ambient occlusion. */
+SHOWFLAG_FIXED_IN_SHIPPING(1, MaterialAmbientOcclusion, SFG_Hidden, NSLOCTEXT("UnrealEd", "MaterialAmbientOcclusionSF", "Material Ambient Occlusion"))
 /** needed for VMI_VisualizeBuffer, Whether to enable the buffer visualization mode. */
 SHOWFLAG_FIXED_IN_SHIPPING(0, VisualizeBuffer, SFG_Hidden, NSLOCTEXT("UnrealEd", "VisualizeBufferSF", "Buffer Visualization"))
 /** Needed for VMI_VisualizeNanite, Whether to enable the Nanite visualization mode. */
diff --git a/Engine/Source/Runtime/Renderer/Private/ReflectionEnvironmentCapture.cpp b/Engine/Source/Runtime/Renderer/Private/ReflectionEnvironmentCapture.cpp
--- a/Engine/Source/Runtime/Renderer/Private/ReflectionEnvironmentCapture.cpp
+++ b/Engine/Source/Runtime/Renderer/Private/ReflectionEnvironmentCapture.cpp
@@ -1361,7 +1361,7 @@ void CaptureSceneIntoScratchCubemap(
 		FSceneView* View = new FSceneView(ViewInitOptions);
 
 		// Force all surfaces diffuse
-		View->RoughnessOverrideParameter = FVector2D( 1.0f, 0.0f );
+		View->RoughnessOverrideParameter = FVector2f(1.0f, 0.0f);
 
 		if (bCaptureEmissiveOnly)
 		{
