Moved parameter default lambda value to named value stored in constexpr to work around compile errors with header units/c++modules

diff --git a/Engine/Source/Runtime/Engine/Classes/GameFramework/ActorPrimitiveColorHandler.h b/Engine/Source/Runtime/Engine/Classes/GameFramework/ActorPrimitiveColorHandler.h
--- a/Engine/Source/Runtime/Engine/Classes/GameFramework/ActorPrimitiveColorHandler.h
+++ b/Engine/Source/Runtime/Engine/Classes/GameFramework/ActorPrimitiveColorHandler.h
@@ -19,6 +19,7 @@ class ENGINE_API FActorPrimitiveColorHandler
 	using FGetColorFunc = TFunction<FLinearColor(const UPrimitiveComponent*)>;
 	using FActivateFunc = TFunction<void(void)>;
 	using FDeactivateFunc = TFunction<void(void)>;
+	constexpr static auto DefaultFunc = []() {};
 
 public:
 	struct FPrimitiveColorHandler
@@ -36,8 +37,8 @@ public:
 	static FActorPrimitiveColorHandler& Get();
 
 	void RegisterPrimitiveColorHandler(FPrimitiveColorHandler& PrimitiveColorHandler);
-	void RegisterPrimitiveColorHandler(FName InHandlerName, const FText& InHandlerText, const FGetColorFunc& InHandlerFunc, const FActivateFunc& InActivateFunc = []() {}, const FText& InHandlerToolTipText = FText());
-	void RegisterPrimitiveColorHandler(FName InHandlerName, const FText& InHandlerText, bool bInAvailalbleInEditor, const FGetColorFunc& InHandlerFunc, const FActivateFunc& InActivateFunc = []() {}, const FText& InHandlerToolTipText = FText());
+	void RegisterPrimitiveColorHandler(FName InHandlerName, const FText& InHandlerText, const FGetColorFunc& InHandlerFunc, const FActivateFunc& InActivateFunc = DefaultFunc, const FText& InHandlerToolTipText = FText());
+	void RegisterPrimitiveColorHandler(FName InHandlerName, const FText& InHandlerText, bool bInAvailalbleInEditor, const FGetColorFunc& InHandlerFunc, const FActivateFunc& InActivateFunc = DefaultFunc, const FText& InHandlerToolTipText = FText());
 	void UnregisterPrimitiveColorHandler(FName InHandlerName);
 	void GetRegisteredPrimitiveColorHandlers(TArray<FPrimitiveColorHandler>& OutPrimitiveColorHandlers) const;
 
