Add flag allowing to set the owner on a ChildActor

#rnx

diff --git a/Engine/Source/Runtime/Engine/Classes/Components/ChildActorComponent.h b/Engine/Source/Runtime/Engine/Classes/Components/ChildActorComponent.h
--- a/Engine/Source/Runtime/Engine/Classes/Components/ChildActorComponent.h
+++ b/Engine/Source/Runtime/Engine/Classes/Components/ChildActorComponent.h
@@ -148,6 +148,12 @@ private:
 	UPROPERTY(EditDefaultsOnly, Category=ChildActorComponent)
 	uint8 bChildActorIsTransient:1;
 
+	/**
+	 * Should the owner of this component be set as the owner of the child actor?
+	 */
+	UPROPERTY(EditDefaultsOnly, Category=ChildActorComponent)
+	uint8 bSetOwner:1;
+
 	/** Flag indicating that when the component is registered that the child actor should be recreated */
 	uint8 bNeedsRecreate:1;
 
diff --git a/Engine/Source/Runtime/Engine/Private/Components/ChildActorComponent.cpp b/Engine/Source/Runtime/Engine/Private/Components/ChildActorComponent.cpp
--- a/Engine/Source/Runtime/Engine/Private/Components/ChildActorComponent.cpp
+++ b/Engine/Source/Runtime/Engine/Private/Components/ChildActorComponent.cpp
@@ -44,6 +44,7 @@ UChildActorComponent::UChildActorComponent(const FObjectInitializer& ObjectIniti
 {
 	bAllowReregistration = false;
 	bChildActorIsTransient = false;
+	bSetOwner = false;
 
 #if WITH_EDITORONLY_DATA
 	EditorTreeViewVisualizationMode = EChildActorComponentTreeViewVisualizationMode::UseDefault;
@@ -789,6 +790,11 @@ void UChildActorComponent::CreateChildActor(TFunction<void(AActor*)> CustomizerF
 				}
 
 				Params.OverrideParentComponent = this;
+
+				if (bSetOwner)
+				{
+					Params.Owner = MyOwner;
+				}
 				
 				if (bChildActorIsTransient || HasAllFlags(RF_Transient) || (MyOwner && MyOwner->HasAllFlags(RF_Transient)))
 				{
