Follow-up change to CL 41617255 to improve clarity that the availability of the r.DynRes.TestScreenPercentage path does not imply DynRes is supported by the device.
Adds additional info to 'stat unit' messages and ensures that GPixelRenderCounters are populated even when using a non-TSR anti-aliasing method.

diff --git a/Engine/Source/Runtime/Engine/Private/UnrealClient.cpp b/Engine/Source/Runtime/Engine/Private/UnrealClient.cpp
--- a/Engine/Source/Runtime/Engine/Private/UnrealClient.cpp
+++ b/Engine/Source/Runtime/Engine/Private/UnrealClient.cpp
@@ -810,6 +810,17 @@ int32 FStatUnitData::DrawStat(FViewport* InViewport, FCanvas* InCanvas, int32 In
 				FColor Color = (ResolutionFraction < AlertResolutionFraction) ? StatRed : ((ResolutionFraction < FMath::Min(ResolutionFraction * 0.97f, 1.0f)) ? StatOrange : StatGreen);
 				DrawDefaultAvgCell(FString::Printf(TEXT("%3.1f%% (%dx%d)"), ScreenPercentage, RenderResolution.X, RenderResolution.Y), Color);
 			}
+			else if (DynamicResolutionStateInfos.Status == EDynamicResolutionStatus::DebugForceEnabled)
+			{
+				if (GEngine->GetDynamicResolutionState()->IsSupported())
+				{
+					DrawDefaultAvgCell(FString::Printf(TEXT("%3.1f%% (%dx%d) (Forced)"), ScreenPercentage, RenderResolution.X, RenderResolution.Y), StatMagenda);
+				}
+				else
+				{
+					DrawDefaultAvgCell(FString::Printf(TEXT("%3.1f%% (%dx%d) (Forced, Unsupported)"), ScreenPercentage, RenderResolution.X, RenderResolution.Y), StatRed);
+				}
+			}
 			else
 			{
 				DrawDefaultAvgCell(FString::Printf(TEXT("%3.1f%% (%dx%d)"), ScreenPercentage, RenderResolution.X, RenderResolution.Y), StatMagenda);
diff --git a/Engine/Source/Runtime/Renderer/Private/SceneRendering.cpp b/Engine/Source/Runtime/Renderer/Private/SceneRendering.cpp
--- a/Engine/Source/Runtime/Renderer/Private/SceneRendering.cpp
+++ b/Engine/Source/Runtime/Renderer/Private/SceneRendering.cpp
@@ -3328,7 +3328,7 @@ void FSceneRenderer::PrepareViewRectsForRendering()
 		check(View.VerifyMembersChecks());
 #endif
 
-		if (!ViewFamily.EngineShowFlags.HitProxies && View.AntiAliasingMethod == AAM_TSR)
+		if (!ViewFamily.EngineShowFlags.HitProxies)
 		{
 			FIntPoint SecondaryViewRect = View.GetSecondaryViewRectSize();
 			GPixelRenderCounters.AddViewStatistics(ViewResolutionFraction, View.ViewRect.Size(), SecondaryViewRect.X * SecondaryViewRect.Y);
