Fix PVS studio warnings from CL 41661592

diff --git a/Engine/Source/Runtime/Engine/Private/UnrealClient.cpp b/Engine/Source/Runtime/Engine/Private/UnrealClient.cpp
--- a/Engine/Source/Runtime/Engine/Private/UnrealClient.cpp
+++ b/Engine/Source/Runtime/Engine/Private/UnrealClient.cpp
@@ -810,9 +810,9 @@ int32 FStatUnitData::DrawStat(FViewport* InViewport, FCanvas* InCanvas, int32 In
 				FColor Color = (ResolutionFraction < AlertResolutionFraction) ? StatRed : ((ResolutionFraction < FMath::Min(ResolutionFraction * 0.97f, 1.0f)) ? StatOrange : StatGreen);
 				DrawDefaultAvgCell(FString::Printf(TEXT("%3.1f%% (%dx%d)"), ScreenPercentage, RenderResolution.X, RenderResolution.Y), Color);
 			}
-			else if (DynamicResolutionStateInfos.Status == EDynamicResolutionStatus::DebugForceEnabled)
+			else
 			{
-				if (GEngine->GetDynamicResolutionState()->IsSupported())
+				if (GEngine->GetDynamicResolutionState() && GEngine->GetDynamicResolutionState()->IsSupported())
 				{
 					DrawDefaultAvgCell(FString::Printf(TEXT("%3.1f%% (%dx%d) (Forced)"), ScreenPercentage, RenderResolution.X, RenderResolution.Y), StatMagenda);
 				}
@@ -821,10 +821,7 @@ int32 FStatUnitData::DrawStat(FViewport* InViewport, FCanvas* InCanvas, int32 In
 					DrawDefaultAvgCell(FString::Printf(TEXT("%3.1f%% (%dx%d) (Forced, Unsupported)"), ScreenPercentage, RenderResolution.X, RenderResolution.Y), StatRed);
 				}
 			}
-			else
-			{
-				DrawDefaultAvgCell(FString::Printf(TEXT("%3.1f%% (%dx%d)"), ScreenPercentage, RenderResolution.X, RenderResolution.Y), StatMagenda);
-			}
+
 			PushRows(/* RowsCount = */ 1);
 
 			// Other dynamic render scalings
