Disable false postive SA warning

#rnx

diff --git a/Engine/Source/Runtime/Engine/Private/Blueprint.cpp b/Engine/Source/Runtime/Engine/Private/Blueprint.cpp
--- a/Engine/Source/Runtime/Engine/Private/Blueprint.cpp
+++ b/Engine/Source/Runtime/Engine/Private/Blueprint.cpp
@@ -906,7 +906,7 @@ void UBlueprint::ForceLoadGraph(UObject* Object)
 		}
 
 		NewObjects.Reset();
-		for(UObject* ReachableObject : ReachableObjects)
+		for(UObject* ReachableObject : ReachableObjects) //-V1078
 		{
 			if(LoadedObjects.Contains(ReachableObject))
 			{
