New material translator:
- Introduced remaining mathematical operators supported by HLSL.
- Added FMaterialExpressionOperator.
- Fixed a bug with the Cast instruction.

#rnx

diff --git a/Engine/Config/BaseMaterialExpressions.ini b/Engine/Config/BaseMaterialExpressions.ini
--- a/Engine/Config/BaseMaterialExpressions.ini
+++ b/Engine/Config/BaseMaterialExpressions.ini
@@ -995,6 +995,9 @@ bShaderInputData=true
 [/Script/Engine.MaterialExpressionRecordTextureStreamingInfo]
 +MenuCategories=NSLOCTEXT("MaterialExpression", "Texture", "Texture")
 
+[/Script/Engine.MaterialExpressionOperator]
++MenuCategories=NSLOCTEXT("MaterialExpression", "Math", "Math")
+
 [/Script/Engine.MaterialExternalCodeCollection]
 +ExternalCodeDeclarations=(ReturnType=Float3,			Name="AtmosphericLightVector",			Definition="MaterialExpressionAtmosphericLightVector(Parameters)")
 +ExternalCodeDeclarations=(ReturnType=Float3,			Name="AtmosphericLightColor",			Definition="MaterialExpressionAtmosphericLightColor(Parameters)")
diff --git a/Engine/Source/Runtime/Engine/Private/Materials/MaterialExpressions.cpp b/Engine/Source/Runtime/Engine/Private/Materials/MaterialExpressions.cpp
--- a/Engine/Source/Runtime/Engine/Private/Materials/MaterialExpressions.cpp
+++ b/Engine/Source/Runtime/Engine/Private/Materials/MaterialExpressions.cpp
@@ -301,6 +301,7 @@
 #include "Materials/MaterialExpressionIsFirstPerson.h"
 #include "Materials/MaterialExternalCodeRegistry.h"
 #include "Materials/MaterialExpressionRecordTextureStreamingInfo.h"
+#include "Materials/MaterialExpressionOperator.h"
 #include "EditorSupportDelegates.h"
 #if WITH_EDITOR
 #include "MaterialGraph/MaterialGraphNode_Comment.h"
@@ -23929,6 +23930,59 @@ int32 UMaterialExpressionRecordTextureStreamingInfo::Compile(class FMaterialComp
 	return TextureCodeIndex;
 }
 
+///////////////////////////////////////////////////////////////////////////////
+// MaterialExpressionOperator
+///////////////////////////////////////////////////////////////////////////////
+
+// defined in MaterialExpressionsIR.cpp
+uint32 GetMaterialExpressionOperatorArity(EMaterialExpressionOperator Operator);
+
+void UMaterialExpressionOperator::GetCaption(TArray<FString>& OutCaptions) const
+{
+	OutCaptions.Add(UEnum::GetDisplayValueAsText(Operator).ToString());
+}
+
+FText UMaterialExpressionOperator::GetKeywords() const
+{
+	return FText::FromString(TEXT("operator"));
+}
+
+FText UMaterialExpressionOperator::GetCreationName() const
+{
+	return FText::FromString(TEXT("Operator"));
+}
+
+FExpressionInput* UMaterialExpressionOperator::GetInput(int32 InputIndex) 
+{
+	switch (InputIndex)
+	{
+		case 0: return &A;
+		case 1: return Arity >= 2 ? &B : nullptr;
+		case 2: return Arity >= 3 ? &C : nullptr;
+		default: return nullptr;
+	}
+}
+
+void UMaterialExpressionOperator::PostEditChangeProperty(FPropertyChangedEvent& PropertyChangedEvent)
+{
+	if (PropertyChangedEvent.GetPropertyName() == TEXT("Operator"))
+	{
+		uint32 OldArity = Arity;
+		Arity = GetMaterialExpressionOperatorArity(Operator);
+		if (Arity != OldArity)
+		{
+			GraphNode->ReconstructNode();
+		}
+	}
+	UMaterialExpression::PostEditChangeProperty(PropertyChangedEvent);
+}
+
+void UMaterialExpressionOperator::PostLoad()
+{
+	Arity = GetMaterialExpressionOperatorArity(Operator);
+	UMaterialExpression::PostLoad();
+}
+
 #endif // WITH_EDITOR
 
 #undef LOCTEXT_NAMESPACE
diff --git a/Engine/Source/Runtime/Engine/Private/Materials/MaterialExpressionsIR.cpp b/Engine/Source/Runtime/Engine/Private/Materials/MaterialExpressionsIR.cpp
--- a/Engine/Source/Runtime/Engine/Private/Materials/MaterialExpressionsIR.cpp
+++ b/Engine/Source/Runtime/Engine/Private/Materials/MaterialExpressionsIR.cpp
@@ -236,6 +236,7 @@
 #include "Materials/MaterialExternalCodeRegistry.h"
 #include "Materials/MaterialFunctionInterface.h"
 #include "Materials/MaterialParameterCollection.h"
+#include "Materials/MaterialExpressionOperator.h"
 #include "MaterialShared.h"
 #include "Misc/MemStackUtility.h"
 #include "RenderUtils.h"
@@ -1186,4 +1187,32 @@ void UMaterialExpressionDDY::Build(MIR::FEmitter& Em)
 	Em.Output(0, Em.PartialDerivative(Em.Input(&Value), MIR::EDerivativeAxis::Y));
 }
 
+static constexpr MIR::EOperator MaterialExpressionOperatorToMIR(EMaterialExpressionOperator Operator)
+{
+	return static_cast<MIR::EOperator>(static_cast<uint32>(Operator) + 1);
+}
+
+// Checks to make sure the two enums are aligned.
+static_assert(MaterialExpressionOperatorToMIR(EMaterialExpressionOperator::BitwiseNot) == MIR::UO_BitwiseNot);
+static_assert(MaterialExpressionOperatorToMIR(EMaterialExpressionOperator::Sign) == MIR::UO_Sign);
+static_assert(MaterialExpressionOperatorToMIR(EMaterialExpressionOperator::BitwiseAnd) == MIR::BO_BitwiseAnd);
+static_assert(MaterialExpressionOperatorToMIR(EMaterialExpressionOperator::Smoothstep) == MIR::TO_Smoothstep);
+
+uint32 GetMaterialExpressionOperatorArity(EMaterialExpressionOperator Operator)
+{
+	return MIR::GetOperatorArity(MaterialExpressionOperatorToMIR(Operator));
+}
+
+void UMaterialExpressionOperator::Build(MIR::FEmitter& Em)
+{
+	MIR::EOperator OpMIR = MaterialExpressionOperatorToMIR(Operator);
+	int32 OperatorArity = MIR::GetOperatorArity(OpMIR);
+
+	FValueRef AValue = Em.InputDefaultFloat(&A, ConstA);
+	FValueRef BValue = OperatorArity >= 2 ? Em.InputDefaultFloat(&B, ConstB) : FValueRef{};
+	FValueRef CValue = OperatorArity >= 3 ? Em.InputDefaultFloat(&C, ConstC) : FValueRef{};
+
+	Em.Output(0, Em.Operator(OpMIR, AValue, BValue, CValue));
+}
+
 #endif // WITH_EDITOR
diff --git a/Engine/Source/Runtime/Engine/Private/Materials/MaterialIR.cpp b/Engine/Source/Runtime/Engine/Private/Materials/MaterialIR.cpp
--- a/Engine/Source/Runtime/Engine/Private/Materials/MaterialIR.cpp
+++ b/Engine/Source/Runtime/Engine/Private/Materials/MaterialIR.cpp
@@ -185,6 +185,12 @@ TConstArrayView<FValue*> FValue::GetUses() const
 			return { &This->Arg, FSubscript::NumStaticUses };
 		}
 
+		case VK_Cast:
+		{
+			auto This = static_cast<const FCast*>(this);
+			return { &This->Arg, FCast::NumStaticUses };
+		}
+			
 		case VK_TextureRead:
 		{
 			auto This = static_cast<const FTextureRead*>(this);
@@ -559,12 +565,22 @@ FBlock* FInstruction::GetDesiredBlockForUse(EStage Stage, int32 UseIndex)
 
 bool IsComparisonOperator(EOperator Op)
 {
-	return Op >= BO_GreaterThan && Op <= BO_NotEquals;
-}
-
-bool IsBitOperator(EOperator Op)
-{
-	return Op == UO_Not || (Op >= BO_Or && Op <= BO_And);
+	switch (Op)
+	{
+		case UO_Not:
+		case UO_IsFinite:
+		case UO_IsInf:
+		case UO_IsNan:
+		case BO_Equals:
+		case BO_GreaterThan:
+		case BO_GreaterThanOrEquals:
+		case BO_LessThan:
+		case BO_LessThanOrEquals:
+		case BO_NotEquals:
+			return true;
+		default:
+			return false;
+	}
 }
 
 bool IsUnaryOperator(EOperator Op)
@@ -582,6 +598,13 @@ bool IsTernaryOperator(EOperator Op)
 	return Op >= TO_FirstTernaryOperator;
 }
 
+int GetOperatorArity(EOperator Op)
+{
+	return IsUnaryOperator(Op) ? 1
+		: IsBinaryOperator(Op) ? 2
+		: 3;
+}
+
 const TCHAR* LexToString(EOperator Op)
 {
 	// Note: sorted alphabetically
@@ -590,14 +613,22 @@ const TCHAR* LexToString(EOperator Op)
 		/* Unary operators */
 		case UO_Abs: return TEXT("Abs");
 		case UO_ACos: return TEXT("ACos");
+		case UO_ACosh: return TEXT("ACosh");
 		case UO_ASin: return TEXT("ASin");
+		case UO_ASinh: return TEXT("ASinh");
 		case UO_ATan: return TEXT("ATan");
+		case UO_ATanh: return TEXT("ATanh");
+		case UO_BitwiseNot: return TEXT("BitwiseNot");
 		case UO_Ceil: return TEXT("Ceil");
 		case UO_Cos: return TEXT("Cos");
+		case UO_Cosh: return TEXT("Cosh");
 		case UO_Exponential: return TEXT("Exponential");
 		case UO_Exponential2: return TEXT("Exponential2");
 		case UO_Floor: return TEXT("Floor");
 		case UO_Frac: return TEXT("Frac");
+		case UO_IsFinite: return TEXT("IsFinite");
+		case UO_IsInf: return TEXT("IsInf");
+		case UO_IsNan: return TEXT("IsNan");
 		case UO_Length: return TEXT("Length");
 		case UO_Logarithm: return TEXT("Logarithm");
 		case UO_Logarithm10: return TEXT("Logarithm10");
@@ -606,15 +637,23 @@ const TCHAR* LexToString(EOperator Op)
 		case UO_Not: return TEXT("Not");
 		case UO_Round: return TEXT("Round");
 		case UO_Saturate: return TEXT("Saturate");
+		case UO_Sign: return TEXT("Sign");
 		case UO_Sin: return TEXT("Sin");
+		case UO_Sinh: return TEXT("Sinh");
 		case UO_Sqrt: return TEXT("Sqrt");
 		case UO_Tan: return TEXT("Tan");
+		case UO_Tanh: return TEXT("Tanh");
 		case UO_Truncate: return TEXT("Truncate");
 
 		/* Binary operators */
 		case BO_Add: return TEXT("Add");
 		case BO_And: return TEXT("And");
+		case BO_BitShiftLeft: return TEXT("BitShiftLeft");
+		case BO_BitShiftRight: return TEXT("BitShiftRight");
+		case BO_BitwiseAnd: return TEXT("BitwiseAnd");
+		case BO_BitwiseOr: return TEXT("BitwiseOr");
 		case BO_Cross: return TEXT("Cross");
+		case BO_Distance: return TEXT("Distance");
 		case BO_Divide: return TEXT("Divide");
 		case BO_Dot: return TEXT("Dot");
 		case BO_Equals: return TEXT("Equals");
@@ -625,9 +664,12 @@ const TCHAR* LexToString(EOperator Op)
 		case BO_LessThanOrEquals: return TEXT("LessThanOrEquals");
 		case BO_Max: return TEXT("Max");
 		case BO_Min: return TEXT("Min");
+		case BO_Modulo: return TEXT("Module");
 		case BO_Multiply: return TEXT("Multiply");
 		case BO_NotEquals: return TEXT("NotEquals");
 		case BO_Or: return TEXT("Or");
+		case BO_Pow: return TEXT("Pow");
+		case BO_Step: return TEXT("Step");
 		case BO_Subtract: return TEXT("Subtract");
 
 		/* Ternary operators */
diff --git a/Engine/Source/Runtime/Engine/Private/Materials/MaterialIREmitter.cpp b/Engine/Source/Runtime/Engine/Private/Materials/MaterialIREmitter.cpp
--- a/Engine/Source/Runtime/Engine/Private/Materials/MaterialIREmitter.cpp
+++ b/Engine/Source/Runtime/Engine/Private/Materials/MaterialIREmitter.cpp
@@ -361,6 +361,16 @@ FValueRef FEmitter::CheckIsArithmetic(FValueRef Value)
 	return Value;
 }
 
+FValueRef FEmitter::CheckIsInteger(FValueRef Value)
+{
+	if (Value.IsValid() && !Value->Type->IsInteger())
+	{
+		Errorf(Value, TEXT("Expected an integer value, got a '%s' instead."), Value->Type->GetSpelling().GetData());
+		return Value.ToPoison();
+	}
+	return Value;
+}
+
 FValueRef FEmitter::CheckIsScalar(FValueRef Value)
 {
 	if (Value.IsValid() && (!Value->Type->AsPrimitive() || !Value->Type->AsPrimitive()->IsScalar()))
@@ -789,6 +799,17 @@ FValueRef FEmitter::Vector4(FValueRef InX, FValueRef InY, FValueRef InZ, FValueR
 template <typename T>
 static bool FoldComparisonOperatorScalar(EOperator Operator, T A, T B)
 {
+	if constexpr (std::is_floating_point_v<T>)
+	{
+		switch (Operator)
+		{
+			case UO_IsFinite: return FGenericPlatformMath::IsFinite(A);
+			case UO_IsInf: return !FGenericPlatformMath::IsFinite(A);
+			case UO_IsNan: return FGenericPlatformMath::IsNaN(A);
+			default: break;
+		}
+	}
+
 	switch (Operator)
 	{
 		case UO_Not: return !A;
@@ -802,47 +823,63 @@ static bool FoldComparisonOperatorScalar(EOperator Operator, T A, T B)
 	}
 }
 
-static TInteger FoldScalarBitOperator(EOperator Operator, TInteger A, TInteger B)
+template <typename T>
+static T ACosh(T x)
 {
-	switch (Operator)
-	{
-		case UO_Not: return !A;
-		case BO_And: return A & B;
-		case BO_Or: return A | B;
-		default: UE_MIR_UNREACHABLE();
-	}
+	static_assert(std::is_floating_point_v<T>);
+	check(x >= 1);
+	return FGenericPlatformMath::Loge(x + FGenericPlatformMath::Sqrt(x * x - 1));
+}
+
+template <typename T>
+static T ASinh(T x)
+{
+	static_assert(std::is_floating_point_v<T>);
+	return FGenericPlatformMath::Loge(x + FGenericPlatformMath::Sqrt(x * x + 1));
 }
 
 template <typename T>
-static TOptional<T> FoldScalarArithmeticOperator(FEmitter& Emitter, EOperator Operator, T A, T B, T C)
+static T ATanh(T x)
 {
-	if constexpr (std::is_same<T, float>::value || std::is_same<T, double>::value)
+	static_assert(std::is_floating_point_v<T>);
+    check(x > -1 && x < 1);
+    return T(0.5) * FGenericPlatformMath::Loge((1 + x) / (1 - x));
+}
+
+template <typename T>
+static T FoldScalarOperator(FEmitter& Emitter, EOperator Operator, T A, T B, T C)
+{
+	if constexpr (std::is_floating_point_v<T>)
 	{
 		switch (Operator)
 		{
-			case UO_ACos: return FMath::Acos(A);
-			case UO_ASin: return FMath::Asin(A);
-			case UO_ATan: return FMath::Atan(A);
+			case UO_ACos: return FGenericPlatformMath::Acos(A);
+			case UO_ACosh: return MIR::ACosh(A);
+			case UO_ASin: return FGenericPlatformMath::Asin(A);
+			case UO_ASinh: return MIR::ASinh(A);
+			case UO_ATan: return FGenericPlatformMath::Atan(A);
+			case UO_ATanh: return MIR::ATanh(A);
 			case UO_Ceil:
 				if constexpr (std::is_same<T, float>::value)
 				{
-					return FMath::CeilToFloat(A);
+					return FGenericPlatformMath::CeilToFloat(A);
 				}
 				else
 				{
-					return FMath::CeilToDouble(A);
+					return FGenericPlatformMath::CeilToDouble(A);
 				}
-			case UO_Cos: return FMath::Cos(A);
+			case UO_Cos: return FGenericPlatformMath::Cos(A);
+			case UO_Cosh: return FGenericPlatformMath::Cosh(A);
 			case UO_Exponential: return FGenericPlatformMath::Pow(UE_EULERS_NUMBER, A);
 			case UO_Exponential2: return FGenericPlatformMath::Pow(2.0f, A);
 			case UO_Floor:
 				if constexpr (std::is_same<T, float>::value)
 				{
-					return FMath::FloorToFloat(A);
+					return FGenericPlatformMath::FloorToFloat(A);
 				}
 				else
 				{
-					return FMath::FloorToDouble(A);
+					return FGenericPlatformMath::FloorToDouble(A);
 				}
 			case UO_Frac:
 			{
@@ -852,12 +889,7 @@ static TOptional<T> FoldScalarArithmeticOperator(FEmitter& Emitter, EOperator Op
 			case UO_Logarithm2:
 			case UO_Logarithm10:
 			{
-				if (A <= 0)
-				{
-					Emitter.Errorf(TEXT("Logarithm of 0 or a negative number is not valid."));
-					return NullOpt;
-				}
-				const T Base = Operator == UO_Logarithm ? UE_EULERS_NUMBER
+				const T Base = (Operator == UO_Logarithm) ? T(UE_EULERS_NUMBER)
 					: Operator == UO_Logarithm2 ? T(2)
 					: T(10);
 				return FGenericPlatformMath::LogX(Base, A);
@@ -865,31 +897,51 @@ static TOptional<T> FoldScalarArithmeticOperator(FEmitter& Emitter, EOperator Op
 			case UO_Round:
 				if constexpr (std::is_same<T, float>::value)
 				{
-					return FMath::RoundToFloat(A);
+					return FGenericPlatformMath::RoundToFloat(A);
 				}
 				else
 				{
-					return FMath::RoundToDouble(A);
+					return FGenericPlatformMath::RoundToDouble(A);
 				}
-			case UO_Sin: return FMath::Sin(A);
-			case UO_Sqrt: return FMath::Sqrt(A);
-			case UO_Tan: return FMath::Tan(A);
+			case UO_Sin: return FGenericPlatformMath::Sin(A);
+			case UO_Sinh: return FGenericPlatformMath::Sinh(A);
+			case UO_Sqrt: return FGenericPlatformMath::Sqrt(A);
+			case UO_Tan: return FGenericPlatformMath::Tan(A);
+			case UO_Tanh: return FGenericPlatformMath::Tanh(A);
 			case UO_Truncate:
 				if constexpr (std::is_same<T, float>::value)
 				{
-					return FMath::TruncToFloat(A);
+					return FGenericPlatformMath::TruncToFloat(A);
 				}
 				else
 				{
-					return FMath::TruncToDouble(A);
+					return FGenericPlatformMath::TruncToDouble(A);
 				}
 			case BO_Fmod: return FGenericPlatformMath::Fmod(A, B);
+			case BO_Pow: return FGenericPlatformMath::Pow(A, B);
 			case TO_Lerp: return FMath::Lerp<T>(A, B, C);
 			case TO_Smoothstep: return FMath::SmoothStep<T>(A, B, C);
 			default: break;
 		}
 	}
 
+	if constexpr (std::is_integral_v<T>)
+	{
+		switch (Operator)
+		{
+			case UO_Not: return !A;
+			case UO_BitwiseNot: return ~A;
+			case BO_And: return A & B;
+			case BO_Or: return A | B;
+			case BO_BitwiseAnd: return A & B;
+			case BO_BitwiseOr: return A | B;
+			case BO_BitShiftLeft: return A << B;
+			case BO_BitShiftRight: return A >> B;
+			case BO_Modulo: return A % B;
+			default: break;
+		}
+	}
+
 	switch (Operator)
 	{
 		case UO_Abs: return FGenericPlatformMath::Abs<T>(A);
@@ -901,6 +953,7 @@ static TOptional<T> FoldScalarArithmeticOperator(FEmitter& Emitter, EOperator Op
 		case BO_Divide: return A / B;
 		case BO_Min: return FMath::Min<T>(A, B);
 		case BO_Max: return FMath::Max<T>(A, B);
+		case BO_Step: return B >= A ? 1.0f : 0.0f;
 		case TO_Clamp: return FMath::Clamp(A, B, C);
 		default: UE_MIR_UNREACHABLE();
 	}
@@ -983,7 +1036,55 @@ static FValueRef TrySimplifyOperator(FEmitter& Emitter, EOperator& Op, FValueRef
 				return A;
 			}
 			break;
-		
+
+		case BO_Modulo:
+			if (A->AreAllNearlyZero() || B->AreAllNearlyOne())
+			{
+				return Emitter.ConstantZero(A->Type->AsPrimitive()->ScalarKind);
+			}
+			break;
+
+		case BO_BitwiseAnd:
+			if (A->AreAllExactlyZero())
+			{
+				return A;
+			}
+			else if (B->AreAllExactlyZero())
+			{
+				return B;
+			}
+			break;
+
+		case BO_BitwiseOr:
+			if (A->AreAllExactlyZero())
+			{
+				return B;
+			}
+			else if (B->AreAllExactlyZero())
+			{
+				return A;
+			}
+			break;
+
+		case BO_BitShiftLeft:
+		case BO_BitShiftRight:
+			if (A->AreAllExactlyZero() || B->AreAllExactlyZero())
+			{
+				return A;
+			}
+			break;
+
+		case BO_Pow:
+			if (A->AreAllNearlyZero())
+			{
+				return A;
+			}
+			else if (B->AreAllNearlyZero())
+			{
+				return Emitter.ConstantOne(A->Type->AsPrimitive()->ScalarKind);
+			}
+			break;
+
 		case TO_Clamp:
 			if (B->AreAllNearlyZero() && C->AreAllNearlyOne())
 			{
@@ -1068,40 +1169,26 @@ static FValueRef TryFoldOperatorScalar(FEmitter& Emitter, EOperator Op, FValueRe
 		}
 		return Emitter.ConstantBool(Result);
 	}
-	else if (IsBitOperator(Op))
+	else
 	{
 		switch (PrimitiveType->ScalarKind)
 		{
 			case ScalarKind_Bool:
 			{
-				bool Result = (bool)FoldScalarBitOperator(Op, AConstant->Boolean, BConstant ? BConstant->Boolean : 0);
+				bool Result = FoldScalarOperator<TInteger>(Emitter, Op, AConstant->Boolean, BConstant ? BConstant->Boolean : 0, 0) & 0x1;
 				return Emitter.ConstantBool(Result);
 			}
 
 			case ScalarKind_Int:
 			{
-				TInteger Result = FoldScalarBitOperator(Op, AConstant->Integer, BConstant ? BConstant->Integer : 0);
+				TInteger Result = FoldScalarOperator<TInteger>(Emitter, Op, AConstant->Integer, BConstant ? BConstant->Integer : 0, CConstant ? CConstant->Integer : 0);
 				return Emitter.ConstantInt(Result);
 			}
-			
-			default:
-				UE_MIR_UNREACHABLE();
-		}
-	}
-	else
-	{
-		switch (PrimitiveType->ScalarKind)
-		{
-			case ScalarKind_Int:
-			{
-				TOptional<TInteger> Result = FoldScalarArithmeticOperator<TInteger>(Emitter, Op, AConstant->Integer, BConstant ? BConstant->Integer : 0, CConstant ? CConstant->Integer : 0);
-				return Result ? Emitter.ConstantInt(*Result) : nullptr;
-			}
 
 			case ScalarKind_Float:
 			{
-				TOptional<TFloat> Result = FoldScalarArithmeticOperator<TFloat>(Emitter, Op, AConstant->Float, BConstant ? BConstant->Float : 0, CConstant ? CConstant->Float : 0);
-				return Result ? Emitter.ConstantFloat(*Result) : nullptr;
+				TFloat Result = FoldScalarOperator<TFloat>(Emitter, Op, AConstant->Float, BConstant ? BConstant->Float : 0, CConstant ? CConstant->Float : 0);
+				return Emitter.ConstantFloat(Result);
 			}
 
 			default:
@@ -1113,19 +1200,22 @@ static FValueRef TryFoldOperatorScalar(FEmitter& Emitter, EOperator Op, FValueRe
 /// Used to filter what parameter *primitive* types operators can take.
 enum EOperatorParameterTypeFilter
 {
-	OPTF_Unknown = 0xff,					///< Unspecified
-	OPTF_Any = 0,							///< Any primitive type (no filters)
-	OPTF_Boolean = 1 << 0,					///< Boolean primitive type of any dimension
-	OPTF_Arithmetic = 1 << 1,				///< Arithmetic primitive type of any dimension (i.e. that supports arithmetic operations)
-	OPTF_NotMatrix = 1 << 2,				///< Any primitive type except matrices
-	OPTF_Vector3 = 1 << 3,					///< 3D vector of any scalar type
-	OPTF_NonNegativeFloatConst = 1 << 5,		///< Validates that if the argument is constant, it is not zero
-	OPTF_NonZeroFloatConst = 1 << 6,		///< Validates that if the argument is constant, it is not zero
-	OPTF_CastToAnyFloat = 1 << 4,			///< Floating point primitive type of any dimension
-	OPTF_CastToFirstArgumentType = 1 << 7,	///< Must be the same type as the first argument
-	OPTF_CastToCommon = 1 << 8,	 			///< An argument with this flag will be automatically cast to the common type between arguments.
-	OPTF_CastToCommonArithmetic = OPTF_Arithmetic | OPTF_CastToCommon,
-	OPTF_CastToCommonFloat = OPTF_CastToAnyFloat | OPTF_CastToCommon,
+	OPTF_Unknown = 0xff,										///< Unspecified
+	OPTF_Any = 0,												///< Any primitive type
+	OPTF_CastToFirstArgumentType = 1 << 10,						///< Cast the argument to the first argument's type
+	OPTF_CastToAnyFloat = 1 << 9,								///< Cast the argument to the floating point primitive type of any dimension
+	OPTF_CheckIsBoolean = 1 << 0,								///< Check the type is boolean primitive of any dimension
+	OPTF_CheckIsInteger = 1 << 1,								///< Check the type is integer primitive of any dimension
+	OPTF_CheckIsArithmetic = 1 << 2,							///< Check the type is arithmetic primitive of any dimension (i.e. that supports arithmetic operations)
+	OPTF_CheckIsNotMatrix = 1 << 3,								///< Check the type is any primitive type except matrices
+	OPTF_CheckIsVector3 = 1 << 4,								///< Check the type is a 3D vector of any scalar type
+	OPTF_CheckIsNonNegativeFloatConst = 1 << 5,					///< Check that if the argument is a constant float, it is not negative (x >= 0)
+	OPTF_CheckIsNonZeroFloatConst = 1 << 6,						///< Check that if the argument is a constant float, it is not zero	(x != 0)
+	OPTF_CheckIsOneOrGreaterFloatConst = 1 << 7,				///< Check that if the argument is a constant float, it is 1 or greater (xFloat >= 1)
+	OPTF_CheckIsBetweenMinusOneAndPlusOneFloatConst = 1 << 8,	///< Check that if the argument is a constant float, it is between -1 and 1 (-1 < x < 1)
+	OPTF_CastToCommonType = 1 << 11,	 						///< Cast the argument to the common arguments type
+	OPTF_CastToCommonArithmeticType = OPTF_CheckIsArithmetic | OPTF_CastToCommonType,
+	OPTF_CastToCommonFloatType = OPTF_CastToAnyFloat | OPTF_CastToCommonType,
 };
 
 inline EOperatorParameterTypeFilter operator|(EOperatorParameterTypeFilter A, EOperatorParameterTypeFilter B)
@@ -1155,57 +1245,86 @@ static const FOperatorSignature* GetOperatorSignature(EOperator Op)
 {
 	static const FOperatorSignature* Signatures = [] ()
 	{
-		const FOperatorSignature UnaryFloat = { { OPTF_Arithmetic | OPTF_CastToAnyFloat }, ORT_FirstArgumentType };
-		const FOperatorSignature BinaryArithmetic = { { OPTF_CastToCommonArithmetic, OPTF_CastToCommonArithmetic }, ORT_FirstArgumentType };
-		const FOperatorSignature TernaryArithmetic = { { OPTF_CastToCommonArithmetic, OPTF_CastToCommonArithmetic, OPTF_CastToCommonArithmetic }, ORT_FirstArgumentType };
-		const FOperatorSignature TernaryFloat = { { OPTF_CastToCommonArithmetic | OPTF_CastToAnyFloat, OPTF_CastToCommonArithmetic, OPTF_CastToCommonArithmetic }, ORT_FirstArgumentType };
-		const FOperatorSignature BinaryArithmeticComparison = { { OPTF_CastToCommonArithmetic, OPTF_CastToCommonArithmetic }, ORT_BooleanWithFirstArgumentDimensions };
-		const FOperatorSignature BinaryLogical = { { OPTF_Boolean | OPTF_CastToCommon, OPTF_Boolean | OPTF_CastToCommon }, ORT_FirstArgumentType };
+		const FOperatorSignature UnaryFloat = { { OPTF_CheckIsArithmetic | OPTF_CastToAnyFloat }, ORT_FirstArgumentType };
+		const FOperatorSignature UnaryFloatToBoolean = { { OPTF_CheckIsArithmetic | OPTF_CastToAnyFloat }, ORT_BooleanWithFirstArgumentDimensions };
+		const FOperatorSignature BinaryArithmetic = { { OPTF_CastToCommonArithmeticType, OPTF_CastToCommonArithmeticType }, ORT_FirstArgumentType };
+		const FOperatorSignature BinaryInteger = { { OPTF_CheckIsInteger | OPTF_CastToCommonArithmeticType, OPTF_CheckIsInteger | OPTF_CastToCommonArithmeticType }, ORT_FirstArgumentType };
+		const FOperatorSignature BinaryFloat = { { OPTF_CastToCommonFloatType , OPTF_CastToCommonFloatType }, ORT_FirstArgumentType };
+		const FOperatorSignature BinaryArithmeticComparison = { { OPTF_CastToCommonArithmeticType, OPTF_CastToCommonArithmeticType }, ORT_BooleanWithFirstArgumentDimensions };
+		const FOperatorSignature BinaryLogical = { { OPTF_CheckIsBoolean | OPTF_CastToCommonType, OPTF_CheckIsBoolean | OPTF_CastToCommonType }, ORT_FirstArgumentType };
+		const FOperatorSignature TernaryArithmetic = { { OPTF_CastToCommonArithmeticType, OPTF_CastToCommonArithmeticType, OPTF_CastToCommonArithmeticType }, ORT_FirstArgumentType };
+		const FOperatorSignature TernaryFloat = { { OPTF_CastToCommonArithmeticType | OPTF_CastToAnyFloat, OPTF_CastToCommonArithmeticType, OPTF_CastToCommonArithmeticType }, ORT_FirstArgumentType };
 
 		static FOperatorSignature S[MIR::OperatorCount];
-		S[UO_Negate] 				= { { OPTF_Arithmetic }, ORT_FirstArgumentType };
-		S[UO_Not] 					= { { OPTF_Boolean }, ORT_FirstArgumentType };
-		S[UO_Abs] 					= UnaryFloat;
-		S[UO_ACos] 					= UnaryFloat;
-		S[UO_ASin] 					= UnaryFloat;
-		S[UO_ATan] 					= UnaryFloat;
-		S[UO_Ceil] 					= UnaryFloat;
-		S[UO_Cos] 					= UnaryFloat;
-		S[UO_Exponential] 			= UnaryFloat;
-		S[UO_Exponential2] 			= UnaryFloat;
-		S[UO_Floor] 				= UnaryFloat;
-		S[UO_Frac]					= UnaryFloat;
-		S[UO_Length]				= { { OPTF_Arithmetic | OPTF_NotMatrix | OPTF_CastToAnyFloat }, ORT_FirstArgumentTypeToScalar };
-		S[UO_Logarithm] 			= { { OPTF_Arithmetic | OPTF_NonZeroFloatConst | OPTF_NonNegativeFloatConst | OPTF_CastToAnyFloat }, ORT_FirstArgumentType };;
-		S[UO_Logarithm2] 			= { { OPTF_Arithmetic | OPTF_NonZeroFloatConst | OPTF_NonNegativeFloatConst | OPTF_CastToAnyFloat }, ORT_FirstArgumentType };;
-		S[UO_Logarithm10] 			= { { OPTF_Arithmetic | OPTF_NonZeroFloatConst | OPTF_NonNegativeFloatConst | OPTF_CastToAnyFloat }, ORT_FirstArgumentType };;
-		S[UO_Round] 				= UnaryFloat;
-		S[UO_Saturate]				= UnaryFloat;
-		S[UO_Sin] 					= UnaryFloat;
-		S[UO_Sqrt]					= { { OPTF_Arithmetic | OPTF_NonNegativeFloatConst | OPTF_CastToAnyFloat }, ORT_FirstArgumentType };;
-		S[UO_Tan] 					= UnaryFloat;
-		S[UO_Truncate] 				= UnaryFloat;
-		S[BO_GreaterThan] 			= BinaryArithmeticComparison;
-		S[BO_GreaterThanOrEquals]	= BinaryArithmeticComparison;
-		S[BO_LessThan]				= BinaryArithmeticComparison;
-		S[BO_LessThanOrEquals]		= BinaryArithmeticComparison;
-		S[BO_Equals]				= { { OPTF_CastToCommon, OPTF_CastToCommon }, ORT_BooleanWithFirstArgumentDimensions };
-		S[BO_NotEquals]				= { { OPTF_CastToCommon, OPTF_CastToCommon }, ORT_BooleanWithFirstArgumentDimensions };
-		S[BO_Add] 					= BinaryArithmetic;
-		S[BO_Subtract] 				= BinaryArithmetic;
-		S[BO_Multiply] 				= BinaryArithmetic;
-		S[BO_Divide] 				= BinaryArithmetic;
-		S[BO_Fmod] 					= BinaryArithmetic;
-		S[BO_Max] 					= BinaryArithmetic;
-		S[BO_Min] 					= BinaryArithmetic;
-		S[BO_Or] 					= BinaryLogical;
-		S[BO_And] 					= BinaryLogical;
-		S[BO_Dot] 					= { { OPTF_Arithmetic | OPTF_NotMatrix, OPTF_CastToFirstArgumentType }, ORT_FirstArgumentTypeToScalar };
-		S[BO_Cross] 				= { { OPTF_Arithmetic | OPTF_Vector3, OPTF_CastToFirstArgumentType }, ORT_FirstArgumentType };
-		S[TO_Clamp]					= TernaryArithmetic;
-		S[TO_Lerp] 					= TernaryFloat;
-		S[TO_Select]				= { { OPTF_Boolean | OPTF_NotMatrix, OPTF_NotMatrix, OPTF_NotMatrix}, ORT_SecondArgumentType }; // Note: this is a special operator, which is handled manually in the Validate function
-		S[TO_Smoothstep]			= TernaryFloat;
+
+		/* unary operators */
+		S[UO_BitwiseNot] 				= { { OPTF_CheckIsInteger }, ORT_FirstArgumentType };
+		S[UO_Negate] 					= { { OPTF_CheckIsArithmetic }, ORT_FirstArgumentType };
+		S[UO_Not] 						= { { OPTF_CheckIsBoolean }, ORT_FirstArgumentType };
+
+		S[UO_Abs] 						= UnaryFloat;
+		S[UO_ACos] 						= UnaryFloat;
+		S[UO_ACosh] 					= { { OPTF_CheckIsArithmetic | OPTF_CastToAnyFloat | OPTF_CheckIsOneOrGreaterFloatConst }, ORT_FirstArgumentType };
+		S[UO_ASin] 						= UnaryFloat;
+		S[UO_ASinh] 					= UnaryFloat;
+		S[UO_ATan] 						= UnaryFloat;
+		S[UO_ATanh] 					= { { OPTF_CheckIsArithmetic | OPTF_CastToAnyFloat | OPTF_CheckIsBetweenMinusOneAndPlusOneFloatConst }, ORT_FirstArgumentType };
+		S[UO_Ceil] 						= UnaryFloat;
+		S[UO_Cos] 						= UnaryFloat;
+		S[UO_Exponential]				= UnaryFloat;
+		S[UO_Exponential2]				= UnaryFloat;
+		S[UO_Floor] 					= UnaryFloat;
+		S[UO_Frac]						= UnaryFloat;
+		S[UO_IsFinite]					= UnaryFloatToBoolean;
+		S[UO_IsInf]						= UnaryFloatToBoolean;
+		S[UO_IsNan]						= UnaryFloatToBoolean;
+		S[UO_Length]					= { { OPTF_CheckIsArithmetic | OPTF_CheckIsNotMatrix | OPTF_CastToAnyFloat }, ORT_FirstArgumentTypeToScalar };
+		S[UO_Logarithm] 				= { { OPTF_CheckIsArithmetic | OPTF_CheckIsNonZeroFloatConst | OPTF_CheckIsNonNegativeFloatConst | OPTF_CastToAnyFloat }, ORT_FirstArgumentType };
+		S[UO_Logarithm10]				= { { OPTF_CheckIsArithmetic | OPTF_CheckIsNonZeroFloatConst | OPTF_CheckIsNonNegativeFloatConst | OPTF_CastToAnyFloat }, ORT_FirstArgumentType };
+		S[UO_Logarithm2] 				= { { OPTF_CheckIsArithmetic | OPTF_CheckIsNonZeroFloatConst | OPTF_CheckIsNonNegativeFloatConst | OPTF_CastToAnyFloat }, ORT_FirstArgumentType };
+		S[UO_Round] 					= UnaryFloat;
+		S[UO_Saturate]					= UnaryFloat;
+		S[UO_Sign]						= UnaryFloat;
+		S[UO_Sin] 						= UnaryFloat;
+		S[UO_Sqrt]						= { { OPTF_CheckIsArithmetic | OPTF_CheckIsNonNegativeFloatConst | OPTF_CastToAnyFloat }, ORT_FirstArgumentType };
+		S[UO_Tan] 						= UnaryFloat;
+		S[UO_Tanh] 						= UnaryFloat;
+		S[UO_Truncate]					= UnaryFloat;
+
+		/* binary operators */
+		S[BO_Equals]					= { { OPTF_CastToCommonType, OPTF_CastToCommonType }, ORT_BooleanWithFirstArgumentDimensions };
+		S[BO_GreaterThan]				= BinaryArithmeticComparison;
+		S[BO_GreaterThanOrEquals]		= BinaryArithmeticComparison;
+		S[BO_LessThan]					= BinaryArithmeticComparison;
+		S[BO_LessThanOrEquals]			= BinaryArithmeticComparison;
+		S[BO_NotEquals]					= { { OPTF_CastToCommonType, OPTF_CastToCommonType }, ORT_BooleanWithFirstArgumentDimensions };
+		
+		S[BO_And] 						= BinaryLogical;
+		S[BO_Or] 						= BinaryLogical;
+		S[BO_Add] 						= BinaryArithmetic;
+		S[BO_Subtract] 					= BinaryArithmetic;
+		S[BO_Multiply] 					= BinaryArithmetic;
+		S[BO_Divide] 					= BinaryArithmetic;
+		S[BO_Modulo] 					= BinaryInteger;
+		S[BO_BitwiseAnd]				= BinaryInteger;
+		S[BO_BitwiseOr]					= BinaryInteger;
+		S[BO_BitShiftLeft]				= BinaryInteger;
+		S[BO_BitShiftRight]				= BinaryInteger;
+
+		S[BO_Cross] 					= { { OPTF_CheckIsArithmetic | OPTF_CheckIsVector3, OPTF_CastToFirstArgumentType }, ORT_FirstArgumentType };
+		S[BO_Distance]					= BinaryFloat;
+		S[BO_Dot] 						= { { OPTF_CheckIsArithmetic | OPTF_CheckIsNotMatrix, OPTF_CastToFirstArgumentType }, ORT_FirstArgumentTypeToScalar };
+		S[BO_Fmod] 						= BinaryFloat;
+		S[BO_Max] 						= BinaryArithmetic;
+		S[BO_Min] 						= BinaryArithmetic;
+		S[BO_Pow] 						= BinaryFloat;
+		S[BO_Step] 						= BinaryArithmetic;
+
+		/* ternary operators */
+		S[TO_Clamp]						= TernaryArithmetic;
+		S[TO_Lerp] 						= TernaryFloat;
+		S[TO_Select]					= { { OPTF_CheckIsBoolean | OPTF_CheckIsNotMatrix, OPTF_CheckIsNotMatrix, OPTF_CheckIsNotMatrix}, ORT_SecondArgumentType }; // Note: this is a special operator, which is handled manually in the Validate function
+		S[TO_Smoothstep]				= TernaryFloat;
 		return S;
 	} ();
 	return &Signatures[Op];
@@ -1267,7 +1386,7 @@ static const FPrimitiveType* ValidateOperatorAndGetResultType(FEmitter& Emitter,
 			}
 		}
 		
-		if (Filter & OPTF_Boolean)
+		if (Filter & OPTF_CheckIsBoolean)
 		{
 			if (ArgumentPrimitiveType->ScalarKind != ScalarKind_Bool)
 			{
@@ -1276,17 +1395,22 @@ static const FPrimitiveType* ValidateOperatorAndGetResultType(FEmitter& Emitter,
 			}
 		}
 		
-		if (Filter & OPTF_Arithmetic)
+		if (Filter & OPTF_CheckIsArithmetic)
 		{
 			bValid &= Emitter.CheckIsArithmetic(Arguments[i]).IsValid(); 
 		}
+
+		if (Filter & OPTF_CheckIsInteger)
+		{
+			bValid &= Emitter.CheckIsInteger(Arguments[i]).IsValid(); 
+		}
 		
-		if (Filter & OPTF_NotMatrix)
+		if (Filter & OPTF_CheckIsNotMatrix)
 		{
 			bValid &= Emitter.CheckIsScalarOrVector(Arguments[i]).IsValid(); 
 		}
 
-		if (Filter & OPTF_Vector3)
+		if (Filter & OPTF_CheckIsVector3)
 		{
 			if (!ArgumentPrimitiveType->IsVector() || ArgumentPrimitiveType->GetNumComponents() != 3)
 			{
@@ -1295,25 +1419,46 @@ static const FPrimitiveType* ValidateOperatorAndGetResultType(FEmitter& Emitter,
 			}
 		}
 
-		if (Filter & OPTF_NonZeroFloatConst)
+		if (FConstant* Constant = Arguments[i]->As<FConstant>())
 		{
-			check((Filter & OPTF_CastToAnyFloat) || (Filter & OPTF_CastToCommonFloat));
-			FConstant* Constant = Arguments[i]->As<FConstant>();
-			if (Constant && Constant->Float == 0)
+			if (Filter & OPTF_CheckIsNonZeroFloatConst)
 			{
-				Emitter.Errorf(Arguments[i], TEXT("Expected non-zero value."));
-				bValid = false;
+				check((Filter & OPTF_CastToAnyFloat) || (Filter & OPTF_CastToCommonFloatType));
+				if (Constant->Float == 0)
+				{
+					Emitter.Errorf(Arguments[i], TEXT("Expected non-zero value."));
+					bValid = false;
+				}
 			}
-		}
 
-		if (Filter & OPTF_NonNegativeFloatConst)
-		{
-			check((Filter & OPTF_CastToAnyFloat) || (Filter & OPTF_CastToCommonFloat));
-			FConstant* Constant = Arguments[i]->As<FConstant>();
-			if (Constant && Constant->Float < 0)
+			if (Filter & OPTF_CheckIsNonNegativeFloatConst)
 			{
-				Emitter.Errorf(Arguments[i], TEXT("Expected non-negative value."));
-				bValid = false;
+				check((Filter & OPTF_CastToAnyFloat) || (Filter & OPTF_CastToCommonFloatType));
+				if (Constant->Float < 0)
+				{
+					Emitter.Errorf(Arguments[i], TEXT("Expected non-negative value."));
+					bValid = false;
+				}
+			}
+
+			if (Filter & OPTF_CheckIsOneOrGreaterFloatConst)
+			{
+				check((Filter & OPTF_CastToAnyFloat) || (Filter & OPTF_CastToCommonFloatType));
+				if (Constant->Float < 1)
+				{
+					Emitter.Errorf(Arguments[i], TEXT("Expected a value equal or greater than 1."));
+					bValid = false;
+				}
+			}
+
+			if (Filter & OPTF_CheckIsBetweenMinusOneAndPlusOneFloatConst)
+			{
+				check((Filter & OPTF_CastToAnyFloat) || (Filter & OPTF_CastToCommonFloatType));
+				if (Constant->Float < -1 || Constant->Float > 1)
+				{
+					Emitter.Errorf(Arguments[i], TEXT("Expected a value greater than -1 and lower than 1."));
+					bValid = false;
+				}
 			}
 		}
 		
@@ -1358,7 +1503,7 @@ static const FPrimitiveType* ValidateOperatorAndGetResultType(FEmitter& Emitter,
 		for (int32 i = 0; Arguments[i]; ++i)
 		{
 			EOperatorParameterTypeFilter Filter = Signature->ParameterTypes[i];
-			if (Filter & OPTF_CastToCommon)
+			if (Filter & OPTF_CastToCommonType)
 			{
 				check(ArgumentsCommonType && ArgumentsCommonType->AsPrimitive());
 				Arguments[i] = Emitter.Cast(Arguments[i], ArgumentsCommonType);
@@ -2270,19 +2415,35 @@ static FValue* DifferentiateOperator(FEmitter& E, FOperator* Op, EDerivativeAxis
 		// d/dx arccos(f(x)) = -1 / sqrt(1 - f(x)^2) * f'(x)
 		case UO_ACos:
 			return E.Negate(E.Divide(dF, E.Sqrt(E.Subtract(One(), E.Multiply(F, F)))));
-
+			
+		// d/dx acosh(f(x)) = 1 / sqrt(f(x)^2 - 1) * f'(x)
+		case UO_ACosh:
+			return E.Divide(dF, E.Sqrt(E.Subtract(E.Multiply(F, F), One())));
+			
 		// d/dx arcsin(f(x)) = 1 / sqrt(1 - f(x)^2) * f'(x)
 		case UO_ASin:
 			return E.Divide(dF, E.Sqrt(E.Subtract(One(), E.Multiply(F, F))));
+		
+		// d/dx asinh(f(x)) = 1 / sqrt(f(x)^2 + 1) * f'(x)
+		case UO_ASinh:
+			return E.Divide(dF, E.Sqrt(E.Add(E.Multiply(F, F), One())));
 
 		// d/dx arctan(f(x)) = 1 / (1 + f(x)^2) * f'(x)
 		case UO_ATan:
 			return E.Divide(dF, E.Add(One(), E.Multiply(F, F)));
 
+    	// d/dx atanh(f(x)) = f'(x) / (1 - f(x)^2)
+		case UO_ATanh:
+    		return E.Divide(dF, E.Subtract(One(), E.Multiply(F, F)));
+
 		// d/dx cos(f(x)) = -sin(f(x)) * f'(x)
 		case UO_Cos:
 			return E.Negate(E.Multiply(E.Sin(F), dF));
 
+		// d/dx cosh(f(x)) = sinh(f(x)) * f'(x)
+		case UO_Cosh:
+			return E.Multiply(E.Sinh(F), dF);
+
 		// d/dx e^f(x) = e^f(x) * f'(x)
 		case UO_Exponential:
 			return E.Multiply(Op, dF);
@@ -2321,6 +2482,10 @@ static FValue* DifferentiateOperator(FEmitter& E, FOperator* Op, EDerivativeAxis
 		// d/dx sin(f(x)) = cos(f(x)) * f'(x)
 		case UO_Sin:
 			return E.Multiply(E.Cos(F), dF);
+		
+		// d/dx sinh(f(x)) = cosh(f(x)) * f'(x)
+		case UO_Sinh:
+			return E.Multiply(E.Cosh(F), dF);
 
 		// d/dx sqrt(f(x)) = 1 / (2 * sqrt(f(x))) * f'(x)
 		case UO_Sqrt:
@@ -2332,6 +2497,10 @@ static FValue* DifferentiateOperator(FEmitter& E, FOperator* Op, EDerivativeAxis
 			FValue* CosVal = E.Cos(F);
 			return E.Divide(dF, E.Multiply(CosVal, CosVal));
 		}
+		
+		// d/dx tanh(f(x)) = (1 - tanh(f(x))^2) * f'(x)
+		case UO_Tanh:
+			return E.Multiply(E.Subtract(One(), E.Multiply(Op, Op)), dF);
 
 		// These functions are piecewise constant, that is mostly constant with some
 		// discontinuities. We assume they're always constant, as they're not differentiable
@@ -2373,6 +2542,14 @@ static FValue* DifferentiateOperator(FEmitter& E, FOperator* Op, EDerivativeAxis
 		case BO_Min:
 			return E.Select(E.LessThan( F, G), dF, dG);
 
+		// d/dx pow(f(x), g(x)) = f(x)^g(x) * (g'(x) * ln(f(x)) + g(x) * f'(x) / f(x))
+		case BO_Pow:
+		{
+			FValueRef Term1 = E.Multiply(dG, E.Logarithm(F)); // g'(x) * ln(f(x))
+			FValueRef Term2 = E.Divide(E.Multiply(G, dF), F); // g(x) * f'(x) / f(x)
+			return E.Multiply(Op, E.Add(Term1, Term2));
+		}
+
 		// The multiplication rule applies for the dot product too.
 		// d/dx (f(x) • g(x)) = f'(x) • g(x) + f(x) • g'(x)
 		case BO_Dot:
@@ -2413,6 +2590,20 @@ static FValue* DifferentiateOperator(FEmitter& E, FOperator* Op, EDerivativeAxis
 			return E.Multiply(dZ, E.Multiply(Constant(6), E.Subtract(Z, E.Multiply(Z, Z))));
 		}
 
+		// these are either invalid or constant
+		case UO_BitwiseNot:
+		case UO_IsFinite:
+		case UO_IsInf:
+		case UO_IsNan:
+		case UO_Sign:
+		case BO_Modulo:
+		case BO_BitwiseAnd:
+		case BO_BitwiseOr:
+		case BO_BitShiftLeft:
+		case BO_BitShiftRight:
+		case BO_Step:
+			return Zero();
+
 		default:
 			UE_MIR_UNREACHABLE();
 	}
diff --git a/Engine/Source/Runtime/Engine/Private/Materials/MaterialIRToHLSLTranslator.cpp b/Engine/Source/Runtime/Engine/Private/Materials/MaterialIRToHLSLTranslator.cpp
--- a/Engine/Source/Runtime/Engine/Private/Materials/MaterialIRToHLSLTranslator.cpp
+++ b/Engine/Source/Runtime/Engine/Private/Materials/MaterialIRToHLSLTranslator.cpp
@@ -388,9 +388,28 @@ struct FTranslator : FMaterialIRToHLSLTranslation
 
 		switch (checkNoEntry->ScalarKind)
 		{
-			case MIR::ScalarKind_Bool:  Printer.Buffer.Append(Constant->Boolean ? TEXT("true") : TEXT("false")); break;
-			case MIR::ScalarKind_Int:   Printer.Buffer.Appendf(TEXT("%") PRId64, Constant->Integer); break;
-			case MIR::ScalarKind_Float: Printer.Buffer.Appendf(TEXT("%.8ff"), Constant->Float); break;
+			case MIR::ScalarKind_Bool:
+				Printer.Buffer.Append(Constant->Boolean ? TEXT("true") : TEXT("false"));
+				break;
+				
+			case MIR::ScalarKind_Int:
+				Printer.Buffer.Appendf(TEXT("%") PRId64, Constant->Integer);
+				break;
+			
+			case MIR::ScalarKind_Float:
+				if (FGenericPlatformMath::IsNaN(Constant->Float))
+				{
+					Printer.Buffer.Append(TEXT("(0.0f / 0.0f)"));
+				}
+				else if (!FGenericPlatformMath::IsFinite(Constant->Float))
+				{
+					Printer.Buffer.Append(TEXT("INFINITE_FLOAT"));
+				}
+				else
+				{
+					Printer.Buffer.Appendf(TEXT("%.8f"), Constant->Float);
+				}
+				break;
 		}
 	}
 
@@ -671,6 +690,7 @@ struct FTranslator : FMaterialIRToHLSLTranslation
 			{
 				case MIR::BO_Add: OpString = TEXT("+"); break;
 				case MIR::BO_Divide: OpString = TEXT("/"); break;
+				case MIR::BO_Modulo: OpString = TEXT("%"); break;
 				case MIR::BO_Equals: OpString = TEXT("=="); break;
 				case MIR::BO_GreaterThan: OpString = TEXT(">"); break;
 				case MIR::BO_GreaterThanOrEquals: OpString = TEXT(">="); break;
@@ -689,37 +709,53 @@ struct FTranslator : FMaterialIRToHLSLTranslation
 			const TCHAR* OpString;
 			switch (Operator->Op)
 			{
-				case MIR::BO_And: OpString = TEXT("and"); break;
-				case MIR::BO_Cross: OpString = TEXT("cross"); break;
-				case MIR::BO_Dot: OpString = TEXT("dot"); break;
-				case MIR::BO_Fmod: OpString = TEXT("fmod"); break;
-				case MIR::BO_Max: OpString = TEXT("max"); break;
-				case MIR::BO_Min: OpString = TEXT("min"); break;
-				case MIR::BO_Or: OpString = TEXT("or"); break;
-				case MIR::TO_Clamp: OpString = TEXT("clamp"); break;
-				case MIR::TO_Lerp: OpString = TEXT("lerp"); break;
-				case MIR::TO_Select: OpString = TEXT("select"); break;
-				case MIR::TO_Smoothstep: OpString = TEXT("smoothstep"); break;
 				case MIR::UO_Abs: OpString = TEXT("abs"); break;
 				case MIR::UO_ACos: OpString = TEXT("acos"); break;
+				case MIR::UO_ACosh: OpString = TEXT("acosh"); break;
 				case MIR::UO_ASin: OpString = TEXT("asin"); break;
+				case MIR::UO_ASinh: OpString = TEXT("asinh"); break;
 				case MIR::UO_ATan: OpString = TEXT("atan"); break;
+				case MIR::UO_ATanh: OpString = TEXT("atanh"); break;
 				case MIR::UO_Ceil: OpString = TEXT("ceil"); break;
 				case MIR::UO_Cos: OpString = TEXT("cos"); break;
+				case MIR::UO_Cosh: OpString = TEXT("cosh"); break;
 				case MIR::UO_Exponential: OpString = TEXT("exp"); break;
 				case MIR::UO_Exponential2: OpString = TEXT("exp2"); break;
 				case MIR::UO_Floor: OpString = TEXT("floor"); break;
 				case MIR::UO_Frac: OpString = TEXT("frac"); break;
+				case MIR::UO_IsFinite: OpString = TEXT("isfinite"); break;
+				case MIR::UO_IsInf: OpString = TEXT("isinf"); break;
+				case MIR::UO_IsNan: OpString = TEXT("isnan"); break;
 				case MIR::UO_Length: OpString = TEXT("length"); break;
 				case MIR::UO_Logarithm: OpString = TEXT("log"); break;
 				case MIR::UO_Logarithm10: OpString = TEXT("log10"); break;
 				case MIR::UO_Logarithm2: OpString = TEXT("log2"); break;
 				case MIR::UO_Round: OpString = TEXT("round"); break;
 				case MIR::UO_Saturate: OpString = TEXT("saturate"); break;
+				case MIR::UO_Sign: OpString = TEXT("sign"); break;
 				case MIR::UO_Sin: OpString = TEXT("sin"); break;
+				case MIR::UO_Sinh: OpString = TEXT("sinh"); break;
 				case MIR::UO_Sqrt: OpString = TEXT("sqrt"); break;
 				case MIR::UO_Tan: OpString = TEXT("tan"); break;
+				case MIR::UO_Tanh: OpString = TEXT("tanh"); break;
 				case MIR::UO_Truncate: OpString = TEXT("trunc"); break;
+
+				case MIR::BO_And: OpString = TEXT("and"); break;
+				case MIR::BO_Cross: OpString = TEXT("cross"); break;
+				case MIR::BO_Distance: OpString = TEXT("distance"); break;
+				case MIR::BO_Dot: OpString = TEXT("dot"); break;
+				case MIR::BO_Fmod: OpString = TEXT("fmod"); break;
+				case MIR::BO_Max: OpString = TEXT("max"); break;
+				case MIR::BO_Min: OpString = TEXT("min"); break;
+				case MIR::BO_Or: OpString = TEXT("or"); break;
+				case MIR::BO_Pow: OpString = TEXT("pow"); break;
+				case MIR::BO_Step: OpString = TEXT("step"); break;
+
+				case MIR::TO_Clamp: OpString = TEXT("clamp"); break;
+				case MIR::TO_Lerp: OpString = TEXT("lerp"); break;
+				case MIR::TO_Select: OpString = TEXT("select"); break;
+				case MIR::TO_Smoothstep: OpString = TEXT("smoothstep"); break;
+			
 				default: UE_MIR_UNREACHABLE();
 			}
 
diff --git a/Engine/Source/Runtime/Engine/Public/Materials/MaterialExpressionOperator.h b/Engine/Source/Runtime/Engine/Public/Materials/MaterialExpressionOperator.h
new file mode 100644
--- /dev/null
+++ b/Engine/Source/Runtime/Engine/Public/Materials/MaterialExpressionOperator.h
@@ -0,0 +1,133 @@
+// Copyright Epic Games, Inc. All Rights Reserved.
+
+
+#pragma once
+
+#include "CoreMinimal.h"
+#include "UObject/ObjectMacros.h"
+#include "MaterialExpressionIO.h"
+#include "Materials/MaterialExpression.h"
+#include "MaterialExpressionOperator.generated.h"
+
+UENUM()
+enum class EMaterialExpressionOperator
+{
+	// Unary
+	BitwiseNot,
+	Negate,
+	Not,
+	Abs,
+	ACos,
+	ACosh,
+	ASin,
+	ASinh,
+	ATan,
+	ATanh,
+	Ceil,
+	Cos,
+	Cosh,
+	Exponential,
+	Exponential2,
+	Floor,
+	Frac,
+	IsFinite,
+	IsInf,
+	IsNan,
+	Length,
+	Logarithm,
+	Logarithm10,
+	Logarithm2,
+	Round,
+	Saturate,
+	Sign,
+	Sin,
+	Sinh,
+	Sqrt,
+	Tan,
+	Tanh,
+	Truncate,
+
+	// Binary
+	Equals,
+	GreaterThan,
+	GreaterThanOrEquals,
+	LessThan,
+	LessThanOrEquals,
+	NotEquals,
+	And,
+	Or,
+	Add,
+	Subtract,
+	Multiply,
+	Divide,
+	Modulo,
+	BitwiseAnd,
+	BitwiseOr,
+	BitShiftLeft,
+	BitShiftRight,
+	Cross,
+	Distance,
+	Dot,
+	Fmod,
+	Max,
+	Min,
+	Pow,
+	Step,
+
+	// Ternary
+	Clamp,
+	Lerp,
+	Select,
+	Smoothstep,
+};
+
+UCLASS(MinimalAPI)
+class UMaterialExpressionOperator : public UMaterialExpression
+{
+	GENERATED_BODY()
+
+public:
+
+	UPROPERTY(EditAnywhere, Category=MaterialExpressionOperator)
+	EMaterialExpressionOperator Operator = EMaterialExpressionOperator::Add;
+
+	UPROPERTY(meta = (RequiredInput = "false", ToolTip = "Defaults to 'ConstA' if not specified"))
+	FExpressionInput A;
+
+	UPROPERTY(meta = (RequiredInput = "false", ToolTip = "Defaults to 'ConstB' if not specified"))
+	FExpressionInput B;
+
+	UPROPERTY(meta = (RequiredInput = "false", ToolTip = "Defaults to 'ConstB' if not specified"))
+	FExpressionInput C;
+
+	// only used if A is not hooked up
+	UPROPERTY(EditAnywhere, Category=MaterialExpressionOperator, meta=(OverridingInputProperty = "A"))
+	float ConstA = 0.0f;
+
+	// only used if B is not hooked up
+	UPROPERTY(EditAnywhere, Category=MaterialExpressionOperator, meta=(OverridingInputProperty = "B"))
+	float ConstB = 1.0f;
+	
+	// only used if B is not hooked up
+	UPROPERTY(EditAnywhere, Category=MaterialExpressionOperator, meta=(OverridingInputProperty = "C"))
+	float ConstC = 1.0f;
+
+	//
+	uint32 Arity = 2;
+
+	//~ Begin UMaterialExpression Interface
+#if WITH_EDITOR
+	virtual void GetCaption(TArray<FString>& OutCaptions) const override;
+	virtual FText GetKeywords() const override;
+	virtual FText GetCreationName() const override;
+	virtual FExpressionInput* GetInput(int32 InputIndex) override;
+	virtual void PostEditChangeProperty(FPropertyChangedEvent& PropertyChangedEvent) override;
+	virtual void PostLoad() override;
+	virtual void Build(MIR::FEmitter& Emitter) override;
+
+#endif // WITH_EDITOR
+	//~ End UMaterialExpression Interface
+};
+
+
+
diff --git a/Engine/Source/Runtime/Engine/Public/Materials/MaterialIR.h b/Engine/Source/Runtime/Engine/Public/Materials/MaterialIR.h
--- a/Engine/Source/Runtime/Engine/Public/Materials/MaterialIR.h
+++ b/Engine/Source/Runtime/Engine/Public/Materials/MaterialIR.h
@@ -556,61 +556,81 @@ enum EOperator
 	// Unary
 	UO_FirstUnaryOperator,
 
-	UO_Negate = UO_FirstUnaryOperator, 	// Arithmetic negation: negate(5) -> -5
-	UO_Not,								// Logical negation: not(true) -> false
+	// Unary operators
+	UO_BitwiseNot = UO_FirstUnaryOperator, // ~(x)
+	UO_Negate, // Arithmetic negation: negate(5) -> -5
+	UO_Not, // Logical negation: not(true) -> false
 
-	// Unary Intrinsic Operators
+	// Unary intrinsics
 	UO_Abs,
 	UO_ACos,
+	UO_ACosh,
 	UO_ASin,
+	UO_ASinh,
 	UO_ATan,
+	UO_ATanh,
 	UO_Ceil,
 	UO_Cos,
+	UO_Cosh,
 	UO_Exponential,
 	UO_Exponential2,
 	UO_Floor,
 	UO_Frac,
+	UO_IsFinite,
+	UO_IsInf,
+	UO_IsNan,
 	UO_Length,
 	UO_Logarithm,
-	UO_Logarithm2,
 	UO_Logarithm10,
+	UO_Logarithm2,
 	UO_Round,
 	UO_Saturate,
+	UO_Sign,
 	UO_Sin,
+	UO_Sinh,
 	UO_Sqrt,
 	UO_Tan,
+	UO_Tanh,
 	UO_Truncate,
 
-	// Binary Comparison Operators
 	BO_FirstBinaryOperator,
 
-	BO_GreaterThan = BO_FirstBinaryOperator,
+	// Binary comparisons
+	BO_Equals = BO_FirstBinaryOperator,
+	BO_GreaterThan,
 	BO_GreaterThanOrEquals,
 	BO_LessThan,
 	BO_LessThanOrEquals,
-	BO_Equals,
 	BO_NotEquals,
 
-	// Binary Arithmetic Operators
+	// Binary logical
+	BO_And,
+	BO_Or,
+
+	// Binary arithmetic
 	BO_Add,
 	BO_Subtract,
 	BO_Multiply,
 	BO_Divide,
+	BO_Modulo,
+	BO_BitwiseAnd,
+	BO_BitwiseOr,
+	BO_BitShiftLeft,
+	BO_BitShiftRight,
+	
+	// Binary intrinsics
+	BO_Cross,
+	BO_Distance,
+	BO_Dot,
 	BO_Fmod,
 	BO_Max,
 	BO_Min,
+	BO_Pow,
+	BO_Step,
 
-	// Binary Logical Operators
-	BO_And,
-	BO_Or,
-
-	// Binary Intrinsic Operators
-	BO_Dot,
-	BO_Cross,
-
-	// Ternary Arithmetic Operators
 	TO_FirstTernaryOperator,
 
+	// Ternary intrinsics
 	TO_Clamp = TO_FirstTernaryOperator,
 	TO_Lerp,
 	TO_Select,
@@ -622,9 +642,6 @@ enum EOperator
 // Whether the given operator identifies a comparison operation (e.g., ">=", "==").
 bool IsComparisonOperator(EOperator Op);
 
-// Whether the given operator identifies a bit operation (e.g., "&", "|").
-bool IsBitOperator(EOperator Op);
-
 // Whether the given operator identifies a unary operator.
 bool IsUnaryOperator(EOperator Op);
 
@@ -634,6 +651,9 @@ bool IsBinaryOperator(EOperator Op);
 // Whether the given operator identifies a ternary operator.
 bool IsTernaryOperator(EOperator Op);
 
+// Returns the arity of the operator (the number of arguments it take, 1, 2 or 3).
+int GetOperatorArity(EOperator Op);
+
 // Returns the string representation of the given operator.
 const TCHAR* LexToString(EOperator Op);
 
diff --git a/Engine/Source/Runtime/Engine/Public/Materials/MaterialIREmitter.h b/Engine/Source/Runtime/Engine/Public/Materials/MaterialIREmitter.h
--- a/Engine/Source/Runtime/Engine/Public/Materials/MaterialIREmitter.h
+++ b/Engine/Source/Runtime/Engine/Public/Materials/MaterialIREmitter.h
@@ -212,10 +212,13 @@ public:
 	
 	// Validates that the value is of a primitive type. If not, reports an error and returns poison.
 	FValueRef CheckIsPrimitive(FValueRef Value);
-
+	
 	// Validates that the value is of an arithmetic type (int or float). Reports error and returns poison otherwise.
 	FValueRef CheckIsArithmetic(FValueRef Value);
 
+	// Validates that the value is of a integral type. If not, reports an error and returns poison.
+	FValueRef CheckIsInteger(FValueRef Value);
+
 	// Validates that the value is a scalar. If not, reports error and returns poison.
 	FValueRef CheckIsScalar(FValueRef Value);
 
@@ -440,14 +443,19 @@ public:
 
 	/*--------------------------- Unary operators shortcuts ----------------------------*/
 
+	FValueRef BitwiseNot(FValueRef A) { return Operator(UO_BitwiseNot, A); }
 	FValueRef Negate(FValueRef A) { return Operator(UO_Negate, A); }
 	FValueRef Not(FValueRef A) { return Operator(UO_Not, A); }
 	FValueRef Abs(FValueRef A) { return Operator(UO_Abs, A); }
 	FValueRef ACos(FValueRef A) { return Operator(UO_ACos, A); }
+	FValueRef ACosh(FValueRef A) { return Operator(UO_ACosh, A); }
 	FValueRef ASin(FValueRef A) { return Operator(UO_ASin, A); }
+	FValueRef ASinh(FValueRef A) { return Operator(UO_ASinh, A); }
 	FValueRef ATan(FValueRef A) { return Operator(UO_ATan, A); }
+	FValueRef ATanh(FValueRef A) { return Operator(UO_ATan, A); }
 	FValueRef Ceil(FValueRef A) { return Operator(UO_Ceil, A); }
 	FValueRef Cos(FValueRef A) { return Operator(UO_Cos, A); }
+	FValueRef Cosh(FValueRef A) { return Operator(UO_Cosh, A); }
 	FValueRef Exponential(FValueRef A) { return Operator(UO_Exponential, A); }
 	FValueRef Exponential2(FValueRef A) { return Operator(UO_Exponential2, A); }
 	FValueRef Floor(FValueRef A) { return Operator(UO_Floor, A); }
@@ -459,6 +467,7 @@ public:
 	FValueRef Round(FValueRef A) { return Operator(UO_Round, A); }
 	FValueRef Saturate(FValueRef A) { return Operator(UO_Saturate, A); }
 	FValueRef Sin(FValueRef A) { return Operator(UO_Sin, A); }
+	FValueRef Sinh(FValueRef A) { return Operator(UO_Sinh, A); }
 	FValueRef Sqrt(FValueRef A) { return Operator(UO_Sqrt, A); }
 	FValueRef Tan(FValueRef A) { return Operator(UO_Tan, A); }
 	FValueRef Truncate(FValueRef A) { return Operator(UO_Truncate, A); }
@@ -471,15 +480,19 @@ public:
 	FValueRef LessThanOrEquals(FValueRef A, FValueRef B) { return Operator(BO_LessThanOrEquals, A, B); }
 	FValueRef Equals(FValueRef A, FValueRef B) { return Operator(BO_Equals, A, B); }
 	FValueRef NotEquals(FValueRef A, FValueRef B) { return Operator(BO_NotEquals, A, B); }
+	FValueRef And(FValueRef A, FValueRef B) { return Operator(BO_And, A, B); }
+	FValueRef Or(FValueRef A, FValueRef B) { return Operator(BO_Or, A, B); }
 	FValueRef Add(FValueRef A, FValueRef B) { return Operator(BO_Add, A, B); }
 	FValueRef Subtract(FValueRef A, FValueRef B) { return Operator(BO_Subtract, A, B); }
 	FValueRef Multiply(FValueRef A, FValueRef B) { return Operator(BO_Multiply, A, B); }
 	FValueRef Divide(FValueRef A, FValueRef B) { return Operator(BO_Divide, A, B); }
+	FValueRef BitwiseAnd(FValueRef A, FValueRef B) { return Operator(BO_BitwiseAnd, A, B); }
+	FValueRef BitwiseOr(FValueRef A, FValueRef B) { return Operator(BO_BitwiseOr, A, B); }
+	FValueRef BitShiftLeft(FValueRef A, FValueRef B) { return Operator(BO_BitShiftLeft, A, B); }
+	FValueRef BitShiftRight(FValueRef A, FValueRef B) { return Operator(BO_BitShiftRight, A, B); }
 	FValueRef Fmod(FValueRef A, FValueRef B) { return Operator(BO_Fmod, A, B); }
 	FValueRef Max(FValueRef A, FValueRef B) { return Operator(BO_Max, A, B); }
 	FValueRef Min(FValueRef A, FValueRef B) { return Operator(BO_Min, A, B); }
-	FValueRef And(FValueRef A, FValueRef B) { return Operator(BO_And, A, B); }
-	FValueRef Or(FValueRef A, FValueRef B) { return Operator(BO_Or, A, B); }
 	FValueRef Dot(FValueRef A, FValueRef B) { return Operator(BO_Dot, A, B); }
 	FValueRef Cross(FValueRef A, FValueRef B) { return Operator(BO_Cross, A, B); }
 
