Renamed EMaterialExpressionOperator to EMaterialExpressionOperatorKind to avoid Python script naming conflict with UMaterialExpressionOperator.

diff --git a/Engine/Source/Runtime/Engine/Private/Materials/MaterialExpressions.cpp b/Engine/Source/Runtime/Engine/Private/Materials/MaterialExpressions.cpp
--- a/Engine/Source/Runtime/Engine/Private/Materials/MaterialExpressions.cpp
+++ b/Engine/Source/Runtime/Engine/Private/Materials/MaterialExpressions.cpp
@@ -23935,7 +23935,7 @@ int32 UMaterialExpressionRecordTextureStreamingInfo::Compile(class FMaterialComp
 ///////////////////////////////////////////////////////////////////////////////
 
 // defined in MaterialExpressionsIR.cpp
-uint32 GetMaterialExpressionOperatorArity(EMaterialExpressionOperator Operator);
+uint32 GetMaterialExpressionOperatorArity(EMaterialExpressionOperatorKind Operator);
 
 void UMaterialExpressionOperator::GetCaption(TArray<FString>& OutCaptions) const
 {
diff --git a/Engine/Source/Runtime/Engine/Private/Materials/MaterialExpressionsIR.cpp b/Engine/Source/Runtime/Engine/Private/Materials/MaterialExpressionsIR.cpp
--- a/Engine/Source/Runtime/Engine/Private/Materials/MaterialExpressionsIR.cpp
+++ b/Engine/Source/Runtime/Engine/Private/Materials/MaterialExpressionsIR.cpp
@@ -1187,18 +1187,18 @@ void UMaterialExpressionDDY::Build(MIR::FEmitter& Em)
 	Em.Output(0, Em.PartialDerivative(Em.Input(&Value), MIR::EDerivativeAxis::Y));
 }
 
-static constexpr MIR::EOperator MaterialExpressionOperatorToMIR(EMaterialExpressionOperator Operator)
+static constexpr MIR::EOperator MaterialExpressionOperatorToMIR(EMaterialExpressionOperatorKind Operator)
 {
 	return static_cast<MIR::EOperator>(static_cast<uint32>(Operator) + 1);
 }
 
 // Checks to make sure the two enums are aligned.
-static_assert(MaterialExpressionOperatorToMIR(EMaterialExpressionOperator::BitwiseNot) == MIR::UO_BitwiseNot);
-static_assert(MaterialExpressionOperatorToMIR(EMaterialExpressionOperator::Sign) == MIR::UO_Sign);
-static_assert(MaterialExpressionOperatorToMIR(EMaterialExpressionOperator::BitwiseAnd) == MIR::BO_BitwiseAnd);
-static_assert(MaterialExpressionOperatorToMIR(EMaterialExpressionOperator::Smoothstep) == MIR::TO_Smoothstep);
+static_assert(MaterialExpressionOperatorToMIR(EMaterialExpressionOperatorKind::BitwiseNot) == MIR::UO_BitwiseNot);
+static_assert(MaterialExpressionOperatorToMIR(EMaterialExpressionOperatorKind::Sign) == MIR::UO_Sign);
+static_assert(MaterialExpressionOperatorToMIR(EMaterialExpressionOperatorKind::BitwiseAnd) == MIR::BO_BitwiseAnd);
+static_assert(MaterialExpressionOperatorToMIR(EMaterialExpressionOperatorKind::Smoothstep) == MIR::TO_Smoothstep);
 
-uint32 GetMaterialExpressionOperatorArity(EMaterialExpressionOperator Operator)
+uint32 GetMaterialExpressionOperatorArity(EMaterialExpressionOperatorKind Operator)
 {
 	return MIR::GetOperatorArity(MaterialExpressionOperatorToMIR(Operator));
 }
diff --git a/Engine/Source/Runtime/Engine/Public/Materials/MaterialExpressionOperator.h b/Engine/Source/Runtime/Engine/Public/Materials/MaterialExpressionOperator.h
--- a/Engine/Source/Runtime/Engine/Public/Materials/MaterialExpressionOperator.h
+++ b/Engine/Source/Runtime/Engine/Public/Materials/MaterialExpressionOperator.h
@@ -10,7 +10,7 @@
 #include "MaterialExpressionOperator.generated.h"
 
 UENUM()
-enum class EMaterialExpressionOperator
+enum class EMaterialExpressionOperatorKind
 {
 	// Unary
 	BitwiseNot,
@@ -89,7 +89,7 @@ class UMaterialExpressionOperator : public UMaterialExpression
 public:
 
 	UPROPERTY(EditAnywhere, Category=MaterialExpressionOperator)
-	EMaterialExpressionOperator Operator = EMaterialExpressionOperator::Add;
+	EMaterialExpressionOperatorKind Operator = EMaterialExpressionOperatorKind::Add;
 
 	UPROPERTY(meta = (RequiredInput = "false", ToolTip = "Defaults to 'ConstA' if not specified"))
 	FExpressionInput A;
