Incremental cook support for WorldPartition maps

Added a way to report cook dependencies in IWorldPartitionCookPackageContext, used for generated package to provide their dependencies (redirects to FPopulateContext)
Added code to generate an hash for each cell from the output of the WorldPartition streaming generation and feed into the generation hash of generated packages. This hash is used in conjunction with package dependencies to ensure we recook all cell packages that have changed when doing incremental cooking.
Added a few private helpers used compute this hash.
Added ActorContainerPackageDependencies to UWorldPartitionRuntimeLevelStreamingCell to accumulate LevelInstance actor dependencies during AddActorToCell (we lose that info afterwards) so embedded level instances actors resaves cause recooks. Also used to preserve package dependencies for AlwaysLoaded actors that are store in the main world persistent level.
Add World package as a dependency for streaming cells, takes care of cell transformers setup changes and also misc dependencies any actor can have on the world package.

Small fixes in cook generation helper to properly keep previously cooked packages hashes and to use the provided GenerationHash to in the generated package hash.

diff --git a/Engine/Source/Editor/UnrealEd/Private/Cooker/CookGenerationHelper.cpp b/Engine/Source/Editor/UnrealEd/Private/Cooker/CookGenerationHelper.cpp
--- a/Engine/Source/Editor/UnrealEd/Private/Cooker/CookGenerationHelper.cpp
+++ b/Engine/Source/Editor/UnrealEd/Private/Cooker/CookGenerationHelper.cpp
@@ -1081,6 +1081,9 @@ bool FGenerationHelper::TryGenerateList()
 			PackageData->SetWorkerAssignmentConstraint(FWorkerId::Local());
 		}
 
+		// Copy hash from package splitter so it can partcipate in the package hash
+		GeneratedInfo->GenerationHash = SplitterData.GenerationHash;
+
 		// Create the Hash from the GenerationHash and Dependencies
 		GeneratedInfo->CreatePackageHash();
 
@@ -1842,7 +1845,9 @@ void FGenerationHelper::ResetSaveState(FCookGenerationInfo& Info, UPackage* Pack
 				// Now that we've finished saving, we know that we will not call QueueGeneratedPackages again, so we can
 				// teardown iterative results as well
 				PlatformData.ClearKeepForIterative(*this, TargetPlatform);
-				PlatformData.PreviousGeneratedPackages.Empty();
+
+				// We still need access to those to compare hashes
+				//PlatformData.PreviousGeneratedPackages.Empty();
 
 				// And also teardown data needed during save
 				PlatformData.ClearKeepForGeneratorSave(*this, TargetPlatform);
diff --git a/Engine/Source/Runtime/Engine/Private/WorldPartition/Cook/WorldPartitionCook.cpp b/Engine/Source/Runtime/Engine/Private/WorldPartition/Cook/WorldPartitionCook.cpp
--- a/Engine/Source/Runtime/Engine/Private/WorldPartition/Cook/WorldPartitionCook.cpp
+++ b/Engine/Source/Runtime/Engine/Private/WorldPartition/Cook/WorldPartitionCook.cpp
@@ -72,7 +72,7 @@ bool UWorldPartition::GatherPackagesToCook(IWorldPartitionCookPackageContext& Co
 	// Add level packages to the CookContext
 	for (const FString& PackageName : LevelPackagesToCook)
 	{
-		IWorldPartitionCookPackageObject* CookPackageObject = GetCookPackageObject(FWorldPartitionCookPackage(TEXT("<Root>"), PackageName, FWorldPartitionCookPackage::EType::Level));
+		IWorldPartitionCookPackageObject* CookPackageObject = GetCookPackageObject(FWorldPartitionCookPackage(TEXT("<Root>"), PackageName, FWorldPartitionCookPackage::EType::Level, FWorldPartitionPackageHash()));
 		AddPackageToCookContext(CookPackageObject, PackageName);
 	}
 
diff --git a/Engine/Source/Runtime/Engine/Private/WorldPartition/Cook/WorldPartitionCookPackageContext.cpp b/Engine/Source/Runtime/Engine/Private/WorldPartition/Cook/WorldPartitionCookPackageContext.cpp
--- a/Engine/Source/Runtime/Engine/Private/WorldPartition/Cook/WorldPartitionCookPackageContext.cpp
+++ b/Engine/Source/Runtime/Engine/Private/WorldPartition/Cook/WorldPartitionCookPackageContext.cpp
@@ -6,6 +6,7 @@
 #include "WorldPartition/WorldPartitionLog.h"
 #include "WorldPartition/Cook/WorldPartitionCookPackageGenerator.h"
 #include "WorldPartition/WorldPartitionRuntimeCell.h"
+#include "Cooker/CookDependency.h"
 
 FWorldPartitionCookPackageContext::FWorldPartitionCookPackageContext()
 {
@@ -29,11 +30,11 @@ const FWorldPartitionCookPackage* FWorldPartitionCookPackageContext::AddPackageT
 
 	if (InCookPackageObject->IsLevelPackage())
 	{
-		WorldPartitionCookPackage = AddLevelStreamingPackageToGenerate(Generator, Root, RelativePath);
+		WorldPartitionCookPackage = AddPackageToGenerateInternal(Generator, Root, RelativePath, FWorldPartitionCookPackage::EType::Level,  InCookPackageObject->GetGenerationHash());
 	}
 	else
 	{
-		WorldPartitionCookPackage = AddGenericPackageToGenerate(Generator, Root, RelativePath);
+		WorldPartitionCookPackage = AddPackageToGenerateInternal(Generator, Root, RelativePath, FWorldPartitionCookPackage::EType::Generic, InCookPackageObject->GetGenerationHash());
 	}
 
 	if (WorldPartitionCookPackage)
@@ -58,12 +59,14 @@ FString FWorldPartitionCookPackageContext::GetGeneratedPackagePath(IWorldPartiti
 
 const FWorldPartitionCookPackage* FWorldPartitionCookPackageContext::AddLevelStreamingPackageToGenerate(IWorldPartitionCookPackageGenerator* CookPackageGenerator, const FString& Root, const FString& RelativePath)
 {
-	return AddPackageToGenerateInternal(CookPackageGenerator, Root, RelativePath, FWorldPartitionCookPackage::EType::Level);
+	FWorldPartitionPackageHash GenerationHash;
+	return AddPackageToGenerateInternal(CookPackageGenerator, Root, RelativePath, FWorldPartitionCookPackage::EType::Level, GenerationHash);
 }
 
 const FWorldPartitionCookPackage* FWorldPartitionCookPackageContext::AddGenericPackageToGenerate(IWorldPartitionCookPackageGenerator* CookPackageGenerator, const FString& Root, const FString& RelativePath)
 {
-	return AddPackageToGenerateInternal(CookPackageGenerator, Root, RelativePath, FWorldPartitionCookPackage::EType::Generic);
+	FWorldPartitionPackageHash GenerationHash;
+	return AddPackageToGenerateInternal(CookPackageGenerator, Root, RelativePath, FWorldPartitionCookPackage::EType::Generic, GenerationHash);
 }
 
 bool FWorldPartitionCookPackageContext::GetCookPackageGeneratorAndPackage(const FString& PackageRoot, const FString& PackageRelativePath, IWorldPartitionCookPackageGenerator*& CookPackageGenerator, FWorldPartitionCookPackage*& CookPackage)
@@ -84,7 +87,7 @@ bool FWorldPartitionCookPackageContext::GetCookPackageGeneratorAndPackage(const
 	return false;
 }
 
-const FWorldPartitionCookPackage* FWorldPartitionCookPackageContext::AddPackageToGenerateInternal(IWorldPartitionCookPackageGenerator* CookPackageGenerator, const FString& Root, const FString& RelativePath, FWorldPartitionCookPackage::EType Type)
+const FWorldPartitionCookPackage* FWorldPartitionCookPackageContext::AddPackageToGenerateInternal(IWorldPartitionCookPackageGenerator* CookPackageGenerator, const FString& Root, const FString& RelativePath, FWorldPartitionCookPackage::EType Type, const FWorldPartitionPackageHash& GenerationHash)
 {
 	if (CookPackageGenerators.Contains(CookPackageGenerator))
 	{
@@ -92,7 +95,7 @@ const FWorldPartitionCookPackage* FWorldPartitionCookPackageContext::AddPackageT
 		TUniquePtr<FWorldPartitionCookPackage>* ExistingPackage = PackagesToCookById.Find(PackageId);
 		if (ExistingPackage == nullptr)
 		{
-			TUniquePtr<FWorldPartitionCookPackage>& CookPackage = PackagesToCookById.Emplace(PackageId, MakeUnique<FWorldPartitionCookPackage>(FWorldPartitionCookPackage::SanitizePathComponent(Root), FWorldPartitionCookPackage::SanitizePathComponent(RelativePath), Type));
+			TUniquePtr<FWorldPartitionCookPackage>& CookPackage = PackagesToCookById.Emplace(PackageId, MakeUnique<FWorldPartitionCookPackage>(FWorldPartitionCookPackage::SanitizePathComponent(Root), FWorldPartitionCookPackage::SanitizePathComponent(RelativePath), Type, GenerationHash));
 			check(PackageId == CookPackage->PackageId);
 
 			CookGeneratorByPackageId.Add(PackageId, CookPackageGenerator);
@@ -152,4 +155,12 @@ const FWorldPartitionCookPackageContextParams& FWorldPartitionCookPackageContext
 	return Params;
 }
 
+void FWorldPartitionCookPackageContext::ReportSaveDependency(UE::Cook::FCookDependency CookDependency) const
+{
+	if (PopulateContext)
+	{
+		PopulateContext->ReportSaveDependency(CookDependency);
+	}
+}
+
 #endif
\ No newline at end of file
diff --git a/Engine/Source/Runtime/Engine/Private/WorldPartition/Cook/WorldPartitionCookPackageContext.h b/Engine/Source/Runtime/Engine/Private/WorldPartition/Cook/WorldPartitionCookPackageContext.h
--- a/Engine/Source/Runtime/Engine/Private/WorldPartition/Cook/WorldPartitionCookPackageContext.h
+++ b/Engine/Source/Runtime/Engine/Private/WorldPartition/Cook/WorldPartitionCookPackageContext.h
@@ -8,6 +8,7 @@
 #include "UObject/ObjectKey.h"
 #include "WorldPartition/Cook/WorldPartitionCookPackageContextInterface.h"
 #include "WorldPartition/Cook/WorldPartitionCookPackage.h"
+#include "CookPackageSplitter.h"
 
 class IWorldPartitionCookPackageGenerator;
 
@@ -35,10 +36,13 @@ public:
 	TArray<IWorldPartitionCookPackageGenerator*>& GetCookPackageGenerators() { return CookPackageGenerators; }
 	const TArray<IWorldPartitionCookPackageGenerator*>& GetCookPackageGenerators() const { return CookPackageGenerators; }
 
+	virtual void ReportSaveDependency(UE::Cook::FCookDependency CookDependency) const override;
+	ICookPackageSplitter::FPopulateContext* PopulateContext = nullptr;
+
 private:
 	const FWorldPartitionCookPackage* AddLevelStreamingPackageToGenerate(IWorldPartitionCookPackageGenerator* CookPackageGenerator, const FString& Root, const FString& RelativePath);
 	const FWorldPartitionCookPackage* AddGenericPackageToGenerate(IWorldPartitionCookPackageGenerator* CookPackageGenerator, const FString& Root, const FString& RelativePath);
-	const FWorldPartitionCookPackage* AddPackageToGenerateInternal(IWorldPartitionCookPackageGenerator* CookPackageGenerator, const FString& Root, const FString& RelativePath, FWorldPartitionCookPackage::EType Type);
+	const FWorldPartitionCookPackage* AddPackageToGenerateInternal(IWorldPartitionCookPackageGenerator* CookPackageGenerator, const FString& Root, const FString& RelativePath, FWorldPartitionCookPackage::EType Type, const FWorldPartitionPackageHash& GenerationHash);
 
 	TArray<IWorldPartitionCookPackageGenerator*> CookPackageGenerators;
 	TMap<FWorldPartitionCookPackage::IDType, TUniquePtr<FWorldPartitionCookPackage>> PackagesToCookById;
diff --git a/Engine/Source/Runtime/Engine/Private/WorldPartition/WorldPartitionCookPackageSplitter.cpp b/Engine/Source/Runtime/Engine/Private/WorldPartition/WorldPartitionCookPackageSplitter.cpp
--- a/Engine/Source/Runtime/Engine/Private/WorldPartition/WorldPartitionCookPackageSplitter.cpp
+++ b/Engine/Source/Runtime/Engine/Private/WorldPartition/WorldPartitionCookPackageSplitter.cpp
@@ -112,6 +112,7 @@ TArray<ICookPackageSplitter::FGeneratedPackage> FWorldPartitionCookPackageSplitt
 bool FWorldPartitionCookPackageSplitter::PopulateGeneratedPackage(FPopulateContext& PopulateContext)
 {
 	const FGeneratedPackageForPopulate& GeneratedPackage = *PopulateContext.GetTargetGeneratedPackage();
+	TGuardValue<FPopulateContext*> Guard(CookContext.PopulateContext, &PopulateContext);
 
 	UE_LOG(LogWorldPartition, Verbose, TEXT("[Cook][PopulateGeneratedPackage] Processing %s"), *FWorldPartitionCookPackage::MakeGeneratedFullPath(GeneratedPackage.GeneratedRootPath, GeneratedPackage.RelativePath));
 
@@ -139,6 +140,7 @@ bool FWorldPartitionCookPackageSplitter::PopulateGeneratedPackage(FPopulateConte
 bool FWorldPartitionCookPackageSplitter::PopulateGeneratorPackage(FPopulateContext& PopulateContext)
 {
 	TConstArrayView<ICookPackageSplitter::FGeneratedPackageForPopulate> GeneratedPackages = PopulateContext.GetGeneratedPackages();
+	TGuardValue<FPopulateContext*> Guard(CookContext.PopulateContext, &PopulateContext);
 
 	UE_LOG(LogWorldPartition, Display, TEXT("[Cook][PopulateGeneratorPackage] Processing %u packages"), GeneratedPackages.Num());
 
@@ -184,7 +186,9 @@ void FWorldPartitionCookPackageSplitter::BuildPackagesToGenerateList(TArray<ICoo
 				ICookPackageSplitter::FGeneratedPackage& GeneratedPackage = PackagesToGenerate.Emplace_GetRef();
 				GeneratedPackage.GeneratedRootPath = CookPackage->Root;
 				GeneratedPackage.RelativePath = CookPackage->RelativePath;
-				// GenerationHash is left as empty. All dependencies for the package's bytes are specified by PackageDependencies
+				GeneratedPackage.GenerationHash = CookPackage->GenerationHash;
+
+				UE_LOG(LogWorldPartition, Log, TEXT("Adding Cell %s with GenerationHash: %s to packages to generate"), *CookPackage->GetFullGeneratedPath(), *LexToString(CookPackage->GenerationHash))
 
 				CookPackage->Type == FWorldPartitionCookPackage::EType::Level ? GeneratedPackage.SetCreateAsMap(true) : GeneratedPackage.SetCreateAsMap(false);
 
@@ -204,4 +208,4 @@ void FWorldPartitionCookPackageSplitter::BuildPackagesToGenerateList(TArray<ICoo
 	}
 }
 
-#endif
+#endif
\ No newline at end of file
diff --git a/Engine/Source/Runtime/Engine/Private/WorldPartition/WorldPartitionRuntimeCell.cpp b/Engine/Source/Runtime/Engine/Private/WorldPartition/WorldPartitionRuntimeCell.cpp
--- a/Engine/Source/Runtime/Engine/Private/WorldPartition/WorldPartitionRuntimeCell.cpp
+++ b/Engine/Source/Runtime/Engine/Private/WorldPartition/WorldPartitionRuntimeCell.cpp
@@ -5,6 +5,7 @@
 #include "Engine/Level.h"
 #include "Algo/Transform.h"
 #include "Algo/Count.h"
+#include "Cooker/CookEvents.h"
 #include "Misc/HierarchicalLogArchive.h"
 #include "WorldPartition/WorldPartitionDebugHelper.h"
 #include "WorldPartition/DataLayer/DataLayerManager.h"
@@ -14,6 +15,7 @@
 #include "WorldPartition/DataLayer/WorldDataLayers.h"
 #include "WorldPartition/WorldPartitionStreamingSource.h"
 #include "WorldPartition/WorldPartitionRuntimeHash.h"
+#include "WorldPartition/WorldPartitionUtils.h"
 
 #include UE_INLINE_GENERATED_CPP_BY_NAME(WorldPartitionRuntimeCell)
 
@@ -201,3 +203,39 @@ FString UWorldPartitionRuntimeCell::GetDebugName() const
 {
 	return RuntimeCellData->GetDebugName();
 }
+
+#if WITH_EDITOR
+
+using namespace UE::Private::WorldPartition;
+
+void FWorldPartitionRuntimeCellObjectMapping::UpdateHash(FWorldPartitionPackageHashBuilder& Builder) const
+{	
+	::UpdateHash(Builder, Package);
+	::UpdateHash(Builder, Path);
+	::UpdateHash(Builder, BaseClass);
+	::UpdateHash(Builder, NativeClass);
+	::UpdateHash(Builder, ContainerID);
+	::UpdateHash(Builder, ContainerTransform);
+	::UpdateHash(Builder, EditorOnlyParentTransform);
+	::UpdateHash(Builder, ContainerPackage);
+	::UpdateHash(Builder, WorldPackage);
+	::UpdateHash(Builder, ActorInstanceGuid);
+	::UpdateHash(Builder, LoadedPath);
+	::UpdateHash(Builder, bIsEditorOnly);
+
+	for (const FWorldPartitionRuntimeCellPropertyOverride& Override : PropertyOverrides)
+	{
+		Override.UpdateHash(Builder);
+	}	
+}
+
+void FWorldPartitionRuntimeCellPropertyOverride::UpdateHash(FWorldPartitionPackageHashBuilder& Builder) const
+{
+	::UpdateHash(Builder, OwnerContainerID);
+	::UpdateHash(Builder, AssetPath);
+	::UpdateHash(Builder, PackageName);
+	::UpdateHash(Builder, ContainerPath.ContainerGuids);
+
+}
+#endif
+
diff --git a/Engine/Source/Runtime/Engine/Private/WorldPartition/WorldPartitionRuntimeHash.cpp b/Engine/Source/Runtime/Engine/Private/WorldPartition/WorldPartitionRuntimeHash.cpp
--- a/Engine/Source/Runtime/Engine/Private/WorldPartition/WorldPartitionRuntimeHash.cpp
+++ b/Engine/Source/Runtime/Engine/Private/WorldPartition/WorldPartitionRuntimeHash.cpp
@@ -13,6 +13,7 @@
 #include "WorldPartition/DataLayer/DataLayerInstanceNames.h"
 #include "WorldPartition/DataLayer/DataLayerManager.h"
 #include "Misc/ArchiveMD5.h"
+#include "Cooker/CookEvents.h"
 #if WITH_EDITOR
 #include "WorldPartition/Cook/WorldPartitionCookPackage.h"
 #include "WorldPartition/Cook/WorldPartitionCookPackageContextInterface.h"
@@ -863,4 +864,12 @@ EDataLayerRuntimeState FWorldPartitionStreamingContext::ResolveDataLayerRuntimeS
 	return Result;
 }
 
+#if WITH_EDITOR
+FWorldPartitionPackageHash URuntimeHashExternalStreamingObjectBase::GetGenerationHash() const
+{
+	// Dependencies URuntimeHashExternalStreamingObjectBase are correctly handled the standard discovery through UObjects mechanisms
+	return FWorldPartitionPackageHash();
+}
+#endif
+
 #undef LOCTEXT_NAMESPACE
diff --git a/Engine/Source/Runtime/Engine/Private/WorldPartition/WorldPartitionRuntimeLevelStreamingCell.cpp b/Engine/Source/Runtime/Engine/Private/WorldPartition/WorldPartitionRuntimeLevelStreamingCell.cpp
--- a/Engine/Source/Runtime/Engine/Private/WorldPartition/WorldPartitionRuntimeLevelStreamingCell.cpp
+++ b/Engine/Source/Runtime/Engine/Private/WorldPartition/WorldPartitionRuntimeLevelStreamingCell.cpp
@@ -7,6 +7,7 @@
 #include "WorldPartition/WorldPartitionStreamingGenerationContext.h"
 #include "WorldPartition/WorldPartitionRuntimeHash.h"
 #include "WorldPartition/WorldPartitionStreamingPolicy.h"
+#include "WorldPartition/WorldPartitionUtils.h"
 #include "Engine/LevelStreaming.h"
 #include "Engine/Level.h"
 #include "Misc/HierarchicalLogArchive.h"
@@ -14,6 +15,7 @@
 #include "AssetRegistry/IAssetRegistry.h"
 #include "UObject/Package.h"
 #include "UObject/AssetRegistryTagsContext.h"
+#include "Cooker/CookEvents.h"
 
 #if WITH_EDITOR
 #include "WorldPartition/Cook/WorldPartitionCookPackageContextInterface.h"
@@ -300,6 +302,22 @@ void UWorldPartitionRuntimeLevelStreamingCell::AddActorToCell(const FStreamingGe
 	const FTransform ContainerTransform = ContainerInstance->GetTransform();
 	const FName ContainerPackage = ContainerInstance->GetContainerPackage();
 
+	// Add all parent LevelInstance actors to the dependency list for this cell (for incremental cooks)
+	const UActorDescContainerInstance* CurrentContainerInstance = ContainerInstance;
+	const UActorDescContainerInstance* ParentContainerInstance = ContainerInstance->GetParentContainerInstance();
+	while (ParentContainerInstance)
+	{
+		FGuid ContainerActorGuid = CurrentContainerInstance->GetContainerActorGuid();
+		if (FWorldPartitionActorDescInstance* ContainerActorDescInstance = ParentContainerInstance->GetActorDescInstance(ContainerActorGuid))
+		{
+			FName ContainerActorPackage = ContainerActorDescInstance->GetActorPackage();
+			ActorContainerPackageDependencies.Add(ContainerActorPackage);
+		}
+		
+		CurrentContainerInstance = ParentContainerInstance;
+		ParentContainerInstance = CurrentContainerInstance->GetParentContainerInstance();
+	}
+
 	for (const FGuid& EditorReferenceGuid : ActorDescView.GetEditorReferences())
 	{
 		FName ReferencePackage;
@@ -380,6 +398,18 @@ void UWorldPartitionRuntimeLevelStreamingCell::Fixup()
 	}
 }
 
+FWorldPartitionPackageHash UWorldPartitionRuntimeLevelStreamingCell::GetGenerationHash() const
+{
+	FWorldPartitionPackageHashBuilder Builder;
+
+	for (const FWorldPartitionRuntimeCellObjectMapping& ActorMapping : Packages)
+	{
+		ActorMapping.UpdateHash(Builder);
+	}
+
+	return Builder.Finalize();
+}
+
 FString UWorldPartitionRuntimeLevelStreamingCell::GetPackageNameToCreate() const
 {
 	return UWorldPartitionLevelStreamingPolicy::GetCellPackagePath(GetFName(), GetOuterWorld());
@@ -415,6 +445,12 @@ bool UWorldPartitionRuntimeLevelStreamingCell::OnPrepareGeneratorPackageForCook(
 		// ex: AFunctionalTest actors need to be part of the Worlds asset tags once they are not longer external so they can be discovered at runtime
 		IAssetRegistry::Get()->AssetUpdateTags(OuterWorld, EAssetRegistryTagsCaller::Fast);
 
+		// Preserve actor package list for dependencies, we'll need them later during OnCookEvent
+		for (const FName& ActorName : GetActorPackageNames())
+		{
+			ActorContainerPackageDependencies.Add(ActorName);
+		}
+
 		// Empty cell's package list (this ensures that no one can rely on cell's content).
 		Packages.Empty();
 	}
@@ -482,6 +518,24 @@ private:
 	friend class UWorldPartitionRuntimeLevelStreamingCell;
 };
 
+void UWorldPartitionRuntimeLevelStreamingCell::OnCookEvent(UE::Cook::ECookEvent CookEvent, UE::Cook::FCookEventContext& CookContext)
+{
+	// These dependencies will be added on the main world (WorldPartition outer world) package
+	Super::OnCookEvent(CookEvent, CookContext);
+	if (CookEvent == UE::Cook::ECookEvent::PlatformCookDependencies && CookContext.IsCooking())
+	{
+		for (const FName& ActorName : GetActorPackageNames())
+		{
+			CookContext.AddSaveBuildDependency(UE::Cook::FCookDependency::Package(ActorName));
+		}
+
+		for (const FName& DependencyName : ActorContainerPackageDependencies)
+		{
+			CookContext.AddSaveBuildDependency(UE::Cook::FCookDependency::Package(DependencyName));
+		}
+	}
+}
+
 bool UWorldPartitionRuntimeLevelStreamingCell::OnPopulateGeneratedPackageForCook(const IWorldPartitionCookPackageContext& InCookContext, UPackage* InPackage, TArray<UPackage*>& OutModifiedPackages)
 {
 	check(!IsAlwaysLoaded());
@@ -498,7 +552,22 @@ bool UWorldPartitionRuntimeLevelStreamingCell::OnPopulateGeneratedPackageForCook
 			return false;
 		}
 
+		// These dependencies will be added on the generated packages (the cells Level used in runtime)
+		for (const FName& ActorName : GetActorPackageNames())
+		{
+			InCookContext.ReportSaveDependency(UE::Cook::FCookDependency::Package(ActorName));
+		}
+
+		for (const FName& DependencyName : ActorContainerPackageDependencies)
+		{
+			InCookContext.ReportSaveDependency(UE::Cook::FCookDependency::Package(DependencyName));
+		}
+
 		UWorld* OuterWorld = GetOuterWorld();
+
+		// Until we also hash WorldPartition settings we add the world package (it contains all settings) to our dependencies
+		InCookContext.ReportSaveDependency(UE::Cook::FCookDependency::Package(OuterWorld->GetPackage()->GetFName()));
+
 		UWorldPartition* WorldPartition = OuterWorld->GetWorldPartition();
 
 		// Load cell Actors
diff --git a/Engine/Source/Runtime/Engine/Public/WorldPartition/Cook/WorldPartitionCookPackage.h b/Engine/Source/Runtime/Engine/Public/WorldPartition/Cook/WorldPartitionCookPackage.h
--- a/Engine/Source/Runtime/Engine/Public/WorldPartition/Cook/WorldPartitionCookPackage.h
+++ b/Engine/Source/Runtime/Engine/Public/WorldPartition/Cook/WorldPartitionCookPackage.h
@@ -7,6 +7,11 @@
 #include "Misc/Paths.h"
 #include "WorldPartition/WorldPartitionLog.h"
 #include "UObject/Package.h"
+#include "Hash/Blake3.h"
+
+using FWorldPartitionPackageHash = FBlake3Hash;
+using FWorldPartitionPackageHashBuilder = FBlake3;
+
 #else
 #include "Containers/StringView.h"
 #endif
@@ -62,11 +67,12 @@ struct FWorldPartitionCookPackage
 		return FPaths::RemoveDuplicateSlashes(*FullPath);
 	}
 
-	FWorldPartitionCookPackage(const FString& InRoot, const FString& InRelativePath, EType InType)
+	FWorldPartitionCookPackage(const FString& InRoot, const FString& InRelativePath, EType InType, const FWorldPartitionPackageHash& InGenerationHash)
 		: Root(SanitizePathComponent(InRoot)),
 		RelativePath(SanitizePathComponent(InRelativePath)),
 		PackageId(MakeCookPackageID(Root, RelativePath)),
-		Type(InType)
+		Type(InType),
+		GenerationHash(InGenerationHash)
 	{
 	}
 
@@ -78,6 +84,7 @@ struct FWorldPartitionCookPackage
 	const FString RelativePath;
 	const IDType PackageId;
 	const EType Type;
+	const FWorldPartitionPackageHash GenerationHash;
 #endif
 
 	/** Returns "_Generated_" */
diff --git a/Engine/Source/Runtime/Engine/Public/WorldPartition/Cook/WorldPartitionCookPackageContextInterface.h b/Engine/Source/Runtime/Engine/Public/WorldPartition/Cook/WorldPartitionCookPackageContextInterface.h
--- a/Engine/Source/Runtime/Engine/Public/WorldPartition/Cook/WorldPartitionCookPackageContextInterface.h
+++ b/Engine/Source/Runtime/Engine/Public/WorldPartition/Cook/WorldPartitionCookPackageContextInterface.h
@@ -5,6 +5,7 @@
 
 #include "CoreMinimal.h"
 #include "Templates/SubclassOf.h"
+#include "Cooker/CookDependency.h"
 
 class IWorldPartitionCookPackageGenerator;
 class IWorldPartitionCookPackageObject;
@@ -34,6 +35,8 @@ public:
 
 	virtual bool GatherPackagesToCook(const FWorldPartitionCookPackageContextParams& Params = FWorldPartitionCookPackageContextParams()) = 0;
 	virtual const FWorldPartitionCookPackageContextParams& GetParams() const = 0;
+
+	virtual void ReportSaveDependency(UE::Cook::FCookDependency CookDependency) const = 0;
 };
 
 #endif
\ No newline at end of file
diff --git a/Engine/Source/Runtime/Engine/Public/WorldPartition/Cook/WorldPartitionCookPackageInterface.h b/Engine/Source/Runtime/Engine/Public/WorldPartition/Cook/WorldPartitionCookPackageInterface.h
--- a/Engine/Source/Runtime/Engine/Public/WorldPartition/Cook/WorldPartitionCookPackageInterface.h
+++ b/Engine/Source/Runtime/Engine/Public/WorldPartition/Cook/WorldPartitionCookPackageInterface.h
@@ -3,6 +3,7 @@
 
 #include "CoreFwd.h"
 #include "UObject/Interface.h"
+#include "WorldPartitionCookPackage.h"
 #include "WorldPartitionCookPackageInterface.generated.h"
 
 class UExternalDataLayerAsset;
@@ -23,6 +24,7 @@ public:
 	virtual bool IsLevelPackage() const = 0;
 	virtual const UExternalDataLayerAsset* GetExternalDataLayerAsset() const = 0;
 	virtual FString GetPackageNameToCreate() const = 0;
+	virtual FWorldPartitionPackageHash GetGenerationHash() const = 0;
 	virtual bool OnPrepareGeneratorPackageForCook(TArray<UPackage*>& OutModifiedPackages) = 0;
 	virtual bool OnPopulateGeneratorPackageForCook(const IWorldPartitionCookPackageContext& InCookContext, UPackage* InPackage) = 0;
 	virtual bool OnPopulateGeneratedPackageForCook(const IWorldPartitionCookPackageContext& InCookContext, UPackage* InPackage, TArray<UPackage*>& OutModifiedPackages) = 0;
diff --git a/Engine/Source/Runtime/Engine/Public/WorldPartition/WorldPartitionRuntimeCell.h b/Engine/Source/Runtime/Engine/Public/WorldPartition/WorldPartitionRuntimeCell.h
--- a/Engine/Source/Runtime/Engine/Public/WorldPartition/WorldPartitionRuntimeCell.h
+++ b/Engine/Source/Runtime/Engine/Public/WorldPartition/WorldPartitionRuntimeCell.h
@@ -64,6 +64,8 @@ struct FWorldPartitionRuntimeCellPropertyOverride
 
 	UPROPERTY()
 	FActorContainerPath ContainerPath;
+
+	void UpdateHash(FWorldPartitionPackageHashBuilder& Builder) const;
 #endif
 };
 
@@ -185,6 +187,8 @@ struct FWorldPartitionRuntimeCellObjectMapping
 		
 	UPROPERTY()
 	TArray<FWorldPartitionRuntimeCellPropertyOverride> PropertyOverrides;
+
+	void UpdateHash(FWorldPartitionPackageHashBuilder& Builder) const;
 #endif
 };
 
@@ -338,6 +342,7 @@ class UWorldPartitionRuntimeCell : public UObject, public IWorldPartitionCell, p
 	ENGINE_API virtual bool OnPrepareGeneratorPackageForCook(TArray<UPackage*>& OutModifiedPackages) override { return true; }
 	ENGINE_API virtual bool OnPopulateGeneratorPackageForCook(const IWorldPartitionCookPackageContext& InCookContext, UPackage* InPackage) override { return true; }
 	ENGINE_API virtual bool OnPopulateGeneratedPackageForCook(const IWorldPartitionCookPackageContext& InCookContext, UPackage* InPackage, TArray<UPackage*>& OutModifiedPackages) override { return true; }
+	ENGINE_API virtual FWorldPartitionPackageHash GetGenerationHash() const override PURE_VIRTUAL(UWorldPartitionRuntimeCell::GetGenerationHash, return FWorldPartitionPackageHash(););
 	//~End IWorldPartitionCookPackageObject
 
 	ENGINE_API virtual void DumpStateLog(FHierarchicalLogArchive& Ar) const;
diff --git a/Engine/Source/Runtime/Engine/Public/WorldPartition/WorldPartitionRuntimeHash.h b/Engine/Source/Runtime/Engine/Public/WorldPartition/WorldPartitionRuntimeHash.h
--- a/Engine/Source/Runtime/Engine/Public/WorldPartition/WorldPartitionRuntimeHash.h
+++ b/Engine/Source/Runtime/Engine/Public/WorldPartition/WorldPartitionRuntimeHash.h
@@ -97,6 +97,7 @@ public:
 	ENGINE_API virtual bool OnPrepareGeneratorPackageForCook(TArray<UPackage*>& OutModifiedPackages) override { return true; }
 	ENGINE_API virtual bool OnPopulateGeneratorPackageForCook(const IWorldPartitionCookPackageContext& InCookContext, UPackage* InPackage) override;
 	ENGINE_API virtual bool OnPopulateGeneratedPackageForCook(const IWorldPartitionCookPackageContext& InCookContext, UPackage* InPackage, TArray<UPackage*>& OutModifiedPackages) override;
+	ENGINE_API virtual FWorldPartitionPackageHash GetGenerationHash() const override;
 	//~End IWorldPartitionCookPackageObject interface
 
 	ENGINE_API const static TCHAR* GetCookedExternalStreamingObjectName() { return TEXT("RuntimeHashExternalStreamingObjectBase"); }
diff --git a/Engine/Source/Runtime/Engine/Public/WorldPartition/WorldPartitionRuntimeLevelStreamingCell.h b/Engine/Source/Runtime/Engine/Public/WorldPartition/WorldPartitionRuntimeLevelStreamingCell.h
--- a/Engine/Source/Runtime/Engine/Public/WorldPartition/WorldPartitionRuntimeLevelStreamingCell.h
+++ b/Engine/Source/Runtime/Engine/Public/WorldPartition/WorldPartitionRuntimeLevelStreamingCell.h
@@ -54,11 +54,14 @@ class UWorldPartitionRuntimeLevelStreamingCell : public UWorldPartitionRuntimeCe
 	ENGINE_API virtual bool PrepareCellForCook(const IWorldPartitionCookPackageContext& InCookContext, UPackage* InGeneratedPackage) override;
 	//~End UWorldPartitionRuntimeCell interface
 
+	virtual void OnCookEvent(UE::Cook::ECookEvent CookEvent, UE::Cook::FCookEventContext& CookContext) override;
+
 	//~Begin IWorldPartitionCookPackageObject interface
 	ENGINE_API virtual FString GetPackageNameToCreate() const override;
 	ENGINE_API virtual bool OnPrepareGeneratorPackageForCook(TArray<UPackage*>& OutModifiedPackages) override;
 	ENGINE_API virtual bool OnPopulateGeneratorPackageForCook(const IWorldPartitionCookPackageContext& InCookContext, UPackage* InPackage) override;
 	ENGINE_API virtual bool OnPopulateGeneratedPackageForCook(const IWorldPartitionCookPackageContext& InCookContext, UPackage* InPackage, TArray<UPackage*>& OutModifiedPackages) override;
+	ENGINE_API virtual FWorldPartitionPackageHash GetGenerationHash() const override;
 	//~End IWorldPartitionCookPackageObject interface
 
 	//~Begin IWorldPartitionCell Interface
@@ -86,6 +89,8 @@ private:
 #if WITH_EDITORONLY_DATA
 	UPROPERTY()
 	TArray<FWorldPartitionRuntimeCellObjectMapping> Packages;
+
+	TSet<FName> ActorContainerPackageDependencies;
 #endif
 
 	UPROPERTY()
diff --git a/Engine/Source/Runtime/Engine/Public/WorldPartition/WorldPartitionUtils.h b/Engine/Source/Runtime/Engine/Public/WorldPartition/WorldPartitionUtils.h
--- a/Engine/Source/Runtime/Engine/Public/WorldPartition/WorldPartitionUtils.h
+++ b/Engine/Source/Runtime/Engine/Public/WorldPartition/WorldPartitionUtils.h
@@ -6,6 +6,9 @@
 
 #include "UObject/WeakObjectPtr.h"
 #include "Templates/SubclassOf.h"
+#include "Misc/ArchiveMD5.h"
+#include "WorldPartition/WorldPartitionActorContainerID.h"
+#include "Memory/MemoryView.h"
 
 class UWorld;
 class UWorldPartition;
@@ -40,4 +43,60 @@ struct FWorldPartitionUtils
 	};
 };
 
+namespace UE::Private::WorldPartition
+{
+	// Hash helpers
+	template<class T, class HashBuilder> void UpdateHash(HashBuilder& Builder, const T& Value);
+
+	template<class HashBuilder> void UpdateHash(HashBuilder& Builder, const FGuid& Value)
+	{
+		static_assert(sizeof(FGuid) == 4 * sizeof(uint32));
+		Builder.Update(&Value, sizeof(FGuid));
+	}
+
+	template<class HashBuilder> void UpdateHash(HashBuilder& Builder, const FName& Value)
+	{
+		TCHAR NameBuffer[FName::StringBufferSize];
+		uint32 Length = Value.ToStringTruncate(NameBuffer);
+		Builder.Update(MakeMemoryView(NameBuffer, Length * sizeof(TCHAR)));
+	}
+
+	template<class HashBuilder> void UpdateHash(HashBuilder& Builder, const FString& Value)
+	{
+		Builder.Update(*Value, Value.Len() * sizeof(FString::ElementType));
+	}
+
+
+	template<class HashBuilder> void UpdateHash(HashBuilder& Builder, const FTopLevelAssetPath& Value)
+	{
+		UpdateHash(Builder, Value.GetPackageName());
+		UpdateHash(Builder, Value.GetAssetName());
+	}
+
+	template<class HashBuilder> void UpdateHash(HashBuilder& Builder, const FTransform& Value)
+	{
+		Builder.Update(&Value, sizeof(FTransform));
+	}
+
+	template<class HashBuilder> void UpdateHash(HashBuilder& Builder, bool Value)
+	{
+		Builder.Update(&Value, sizeof(bool));
+	}
+
+
+	template<class HashBuilder> void UpdateHash(HashBuilder& Builder, const FActorContainerID& Value)
+	{
+		static_assert(sizeof(FActorContainerID) == sizeof(FGuid));
+		Builder.Update(&Value, sizeof(FGuid));
+	}
+
+	template<class T, class HashBuilder> void UpdateHash(HashBuilder& Builder, const TArray<T>& Values)
+	{
+		for (const T& Val : Values)
+		{
+			UpdateHash(Builder, Val);
+		}
+	}
+}
+
 #endif
\ No newline at end of file
