Fix build issue with unity stress test config

diff --git a/Engine/Source/Runtime/Engine/Private/WorldPartition/WorldPartitionRuntimeCell.cpp b/Engine/Source/Runtime/Engine/Private/WorldPartition/WorldPartitionRuntimeCell.cpp
--- a/Engine/Source/Runtime/Engine/Private/WorldPartition/WorldPartitionRuntimeCell.cpp
+++ b/Engine/Source/Runtime/Engine/Private/WorldPartition/WorldPartitionRuntimeCell.cpp
@@ -206,22 +206,20 @@ FString UWorldPartitionRuntimeCell::GetDebugName() const
 
 #if WITH_EDITOR
 
-using namespace UE::Private::WorldPartition;
-
 void FWorldPartitionRuntimeCellObjectMapping::UpdateHash(FWorldPartitionPackageHashBuilder& Builder) const
 {	
-	::UpdateHash(Builder, Package);
-	::UpdateHash(Builder, Path);
-	::UpdateHash(Builder, BaseClass);
-	::UpdateHash(Builder, NativeClass);
-	::UpdateHash(Builder, ContainerID);
-	::UpdateHash(Builder, ContainerTransform);
-	::UpdateHash(Builder, EditorOnlyParentTransform);
-	::UpdateHash(Builder, ContainerPackage);
-	::UpdateHash(Builder, WorldPackage);
-	::UpdateHash(Builder, ActorInstanceGuid);
-	::UpdateHash(Builder, LoadedPath);
-	::UpdateHash(Builder, bIsEditorOnly);
+	UE::Private::WorldPartition::UpdateHash(Builder, Package);
+	UE::Private::WorldPartition::UpdateHash(Builder, Path);
+	UE::Private::WorldPartition::UpdateHash(Builder, BaseClass);
+	UE::Private::WorldPartition::UpdateHash(Builder, NativeClass);
+	UE::Private::WorldPartition::UpdateHash(Builder, ContainerID);
+	UE::Private::WorldPartition::UpdateHash(Builder, ContainerTransform);
+	UE::Private::WorldPartition::UpdateHash(Builder, EditorOnlyParentTransform);
+	UE::Private::WorldPartition::UpdateHash(Builder, ContainerPackage);
+	UE::Private::WorldPartition::UpdateHash(Builder, WorldPackage);
+	UE::Private::WorldPartition::UpdateHash(Builder, ActorInstanceGuid);
+	UE::Private::WorldPartition::UpdateHash(Builder, LoadedPath);
+	UE::Private::WorldPartition::UpdateHash(Builder, bIsEditorOnly);
 
 	for (const FWorldPartitionRuntimeCellPropertyOverride& Override : PropertyOverrides)
 	{
@@ -231,10 +229,10 @@ void FWorldPartitionRuntimeCellObjectMapping::UpdateHash(FWorldPartitionPackageH
 
 void FWorldPartitionRuntimeCellPropertyOverride::UpdateHash(FWorldPartitionPackageHashBuilder& Builder) const
 {
-	::UpdateHash(Builder, OwnerContainerID);
-	::UpdateHash(Builder, AssetPath);
-	::UpdateHash(Builder, PackageName);
-	::UpdateHash(Builder, ContainerPath.ContainerGuids);
+	UE::Private::WorldPartition::UpdateHash(Builder, OwnerContainerID);
+	UE::Private::WorldPartition::UpdateHash(Builder, AssetPath);
+	UE::Private::WorldPartition::UpdateHash(Builder, PackageName);
+	UE::Private::WorldPartition::UpdateHash(Builder, ContainerPath.ContainerGuids);
 
 }
 #endif
