Fix NDC issue created by UMVVMK2Node_MakeBrushFromSoftResource's use of UK2Node_LoadAsset in ExpandNode. UK2Node_LoadAsset in turn creates a UK2Node_Event based on the loadasset's NodeGuid which creates a UFunction based on the name of the UK2Node_Event.

[RN] Fix instability in the generated class for blueprints that use UMVVMK2Node_MakeBrushFromSoftResource

diff --git a/Engine/Source/Editor/KismetCompiler/Public/KismetCompiler.h b/Engine/Source/Editor/KismetCompiler/Public/KismetCompiler.h
--- a/Engine/Source/Editor/KismetCompiler/Public/KismetCompiler.h
+++ b/Engine/Source/Editor/KismetCompiler/Public/KismetCompiler.h
@@ -235,9 +235,8 @@ public:
 		}
 
 		NodeType* Result = ParentGraph->CreateIntermediateNode<NodeType>();
-		//check (Cast<UK2Node_Event>(Result) == nullptr); -- Removed to avoid any fallout, will replace with care later
 		MessageLog.NotifyIntermediateObjectCreation(Result, SourceNode); // this might be useful to track back function entry nodes to events.
-		Result->CreateNewGuid();
+		Result->CreateDeterministicGuid();
 
 		AutoAssignNodePosition(Result);
 
diff --git a/Engine/Source/Runtime/Engine/Classes/EdGraph/EdGraphNode.h b/Engine/Source/Runtime/Engine/Classes/EdGraph/EdGraphNode.h
--- a/Engine/Source/Runtime/Engine/Classes/EdGraph/EdGraphNode.h
+++ b/Engine/Source/Runtime/Engine/Classes/EdGraph/EdGraphNode.h
@@ -367,10 +367,8 @@ public:
 	UPROPERTY()
 	uint8 bHasCompilerMessage:1;
 
-	/** Comment bubble pinned state */
-
-
 #if WITH_EDITORONLY_DATA
+	/** Comment bubble pinned state */
 	UPROPERTY()
 	uint8 bCommentBubblePinned : 1;
 
@@ -889,6 +887,10 @@ public:
 
 	/** Create a new unique Guid for this node */
 	ENGINE_API void CreateNewGuid();
+#if WITH_EDITORONLY_DATA
+	/** Create a new unique Guid for this node deterministically, provided the node's pathname is deterministic */
+	ENGINE_API void CreateDeterministicGuid();
+#endif
 
 	/** Gets a list of actions that can be done to this particular node */
 	virtual void GetNodeContextMenuActions(class UToolMenu* Menu, class UGraphNodeContextMenuContext* Context) const {}
diff --git a/Engine/Source/Runtime/Engine/Private/EdGraph/EdGraphNode.cpp b/Engine/Source/Runtime/Engine/Private/EdGraph/EdGraphNode.cpp
--- a/Engine/Source/Runtime/Engine/Private/EdGraph/EdGraphNode.cpp
+++ b/Engine/Source/Runtime/Engine/Private/EdGraph/EdGraphNode.cpp
@@ -787,6 +787,13 @@ void UEdGraphNode::CreateNewGuid()
 	NodeGuid = FGuid::NewGuid();
 }
 
+#if WITH_EDITORONLY_DATA
+void UEdGraphNode::CreateDeterministicGuid()
+{
+	NodeGuid = FGuid::NewDeterministicGuid(GetPathName());
+}
+#endif//WITH_EDITORONLY_DATA
+
 void UEdGraphNode::FindDiffs(UEdGraphNode* OtherNode, struct FDiffResults& Results)
 {
 	if (OtherNode != nullptr)
