gaurd against trying to create soundwave proxy data when running with `-nosound`

#swarm https://p4-swarm.epicgames.net/reviews/41665003
#rnx

diff --git a/Engine/Source/Runtime/Engine/Private/SoundWave.cpp b/Engine/Source/Runtime/Engine/Private/SoundWave.cpp
--- a/Engine/Source/Runtime/Engine/Private/SoundWave.cpp
+++ b/Engine/Source/Runtime/Engine/Private/SoundWave.cpp
@@ -4032,6 +4032,11 @@ int32 USoundWave::GetCompressionQuality() const
 
 FSoundWaveProxyPtr USoundWave::CreateSoundWaveProxy()
 {
+	if (!FApp::CanEverRenderAudio())
+	{
+		return Proxy;
+	}
+	
 	LLM_SCOPE(ELLMTag::AudioSoundWaveProxies);
 
 #if WITH_EDITORONLY_DATA
@@ -4048,7 +4053,7 @@ TSharedPtr<Audio::IProxyData> USoundWave::CreateProxyData(const Audio::FProxyDat
 	check(SoundWaveDataPtr);
 	LLM_SCOPE(ELLMTag::AudioSoundWaveProxies);
 
-	if (!Proxy.IsValid())
+	if (!Proxy.IsValid() && FApp::CanEverRenderAudio())
 	{
 		Proxy = CreateSoundWaveProxy();
 	}
@@ -4616,12 +4621,12 @@ void USoundWave::OverrideLoadingBehavior(ESoundWaveLoadingBehavior InLoadingBeha
 				// In editor, just make sure that data is available in the local DDC cache for quick access.
 				CachePlatformData(true /* bAsyncCache */);
 
-				if (!Proxy.IsValid())
+				if (!Proxy.IsValid() && FApp::CanEverRenderAudio())
 				{
 					Proxy = CreateSoundWaveProxy();
 				}
 #else
-				if (GetNumChunks() > 1)
+				if (GetNumChunks() > 1 && FApp::CanEverRenderAudio())
 				{
 					if (!Proxy.IsValid())
 					{
