[Graph Editor] Fix an issue with the truncation warnings where X/Y values were transposed

diff --git a/Engine/Source/Runtime/Engine/Classes/EdGraph/EdGraphNode.h b/Engine/Source/Runtime/Engine/Classes/EdGraph/EdGraphNode.h
--- a/Engine/Source/Runtime/Engine/Classes/EdGraph/EdGraphNode.h
+++ b/Engine/Source/Runtime/Engine/Classes/EdGraph/EdGraphNode.h
@@ -501,7 +501,7 @@ public:
 	void SetNodePosX(float InX) { NodePosX = UE::LWC::FloatToIntCastChecked<int32>(InX); }
 	float GetNodePosY() const { return static_cast<float>(NodePosY); }
 	void SetNodePosY(float InY) { NodePosY = UE::LWC::FloatToIntCastChecked<int32>(InY); }
-	FVector2f GetSize() const { return FVector2f(GetHeight(), GetWidth()); }
+	FVector2f GetSize() const { return FVector2f(GetWidth(), GetHeight()); }
 	void SetSize(const FVector2f& InSize) { SetWidth(InSize.X); SetHeight(InSize.Y); }
 	FVector2f GetPosition() const { return FVector2f(GetNodePosX(), GetNodePosY()); }
 	void SetPosition(const FVector2f& InPosition) { SetNodePosX(InPosition.X); SetNodePosY(InPosition.Y); }
