Put EngineSubsystemCollection in a protected section, rather than private.

This will allow unit tests which subclass UEngine to access this member.

#jir SOL-8019

diff --git a/Engine/Source/Runtime/Engine/Classes/Engine/Engine.h b/Engine/Source/Runtime/Engine/Classes/Engine/Engine.h
--- a/Engine/Source/Runtime/Engine/Classes/Engine/Engine.h
+++ b/Engine/Source/Runtime/Engine/Classes/Engine/Engine.h
@@ -3744,7 +3744,7 @@ public:
 		}, TSubsystemClass::StaticClass());
 	}
 
-private:
+protected:
 	FObjectSubsystemCollection<UEngineSubsystem> EngineSubsystemCollection;
 
 public:
