Added a disabled by default object validation warning modified via net.WithObjectValidationCheck. The warning outputs if the property to be stored or stored property is non-null, but not valid.

diff --git a/Engine/Source/Runtime/Engine/Private/RepLayout.cpp b/Engine/Source/Runtime/Engine/Private/RepLayout.cpp
--- a/Engine/Source/Runtime/Engine/Private/RepLayout.cpp
+++ b/Engine/Source/Runtime/Engine/Private/RepLayout.cpp
@@ -92,6 +92,9 @@ static FAutoConsoleVariableRef CVarTrackNetSerializeObjectReferences(TEXT("net.T
 bool GbWithArrayOnRepFix = false;
 static FAutoConsoleVariableRef CVarWithArrayOnRepFix(TEXT("net.WithArrayOnRepFix"), GbWithArrayOnRepFix, TEXT("If true, attempt to prevent issues with Arrays not receiving OnRep calls until their size changes if their Archetypes have different values from instances in levels."));
 
+bool GbWithObjectValidationCheck = false;
+static FAutoConsoleVariableRef CVarWithObjectValidationCheck(TEXT("net.WithObjectValidationCheck"), GbWithObjectValidationCheck, TEXT("If true, will run object property validation on property storage."));
+
 #if WITH_PUSH_VALIDATION_SUPPORT
 
 static bool GbPushModelValidateProperties = false;
@@ -827,6 +830,30 @@ static FORCEINLINE bool PropertiesAreIdentical(
 static FORCEINLINE void StoreProperty(const FRepLayoutCmd& Cmd, void* A, const void* B)
 {
 	UE_LOG(LogRepCompares, VeryVerbose, TEXT("StoreProperty: %s"), *GetNameSafe(Cmd.Property));
+
+	if (GbWithObjectValidationCheck)
+	{
+		if (Cmd.Type == ERepLayoutCmdType::PropertyObject ||
+			Cmd.Type == ERepLayoutCmdType::PropertyWeakObject ||
+			Cmd.Type == ERepLayoutCmdType::PropertySoftObject)
+		{
+			if (const FObjectPropertyBase* ObjProperty = CastFieldChecked<FObjectPropertyBase>(Cmd.Property))
+			{
+				const UObject* AObjValue = ObjProperty->GetObjectPropertyValue(A);
+				if (AObjValue && !IsValid(AObjValue))
+				{
+					UE_LOG(LogRepCompares, Warning, TEXT("Destination Object %s was not valid when storing for replication. StoreProperty: %s"), *GetNameSafe(AObjValue), *GetFullNameSafe(Cmd.Property));
+				}
+
+				const UObject* BObjValue = ObjProperty->GetObjectPropertyValue(B);
+				if (BObjValue && !IsValid(BObjValue))
+				{
+					UE_LOG(LogRepCompares, Warning, TEXT("Source Object %s was not valid when storing for replication. StoreProperty: %s"), *GetNameSafe(BObjValue), *GetFullNameSafe(Cmd.Property));
+				}
+			}
+		}
+	}
+
 	Cmd.Property->CopySingleValue(A, B);
 }
 
