Fixed formatting issue when printing object addresses.

#rnx

diff --git a/Engine/Source/Runtime/Engine/Private/PrintObjectUtils.cpp b/Engine/Source/Runtime/Engine/Private/PrintObjectUtils.cpp
--- a/Engine/Source/Runtime/Engine/Private/PrintObjectUtils.cpp
+++ b/Engine/Source/Runtime/Engine/Private/PrintObjectUtils.cpp
@@ -55,7 +55,7 @@ static void LogObjectState(const UObject* Object, FOutputDevice* OutputDevice, b
 		EObjectFlags ObjectFlags = Object ? Object->GetFlags() : EObjectFlags::RF_NoFlags;
 		OutputDevice->Logf(TEXT("\tFlags: %s"), *LexToString(ObjectFlags));
 
-		OutputDevice->Logf(TEXT("\tAddress: 0x%x"), Object);
+		OutputDevice->Logf(TEXT("\tAddress: 0x%p"), Object);
 	}
 }
 
