TEMPFIX: hide ColorRamp node temporarily

diff --git a/Engine/Source/Runtime/Engine/Public/Materials/MaterialExpressionColorRamp.h b/Engine/Source/Runtime/Engine/Public/Materials/MaterialExpressionColorRamp.h
--- a/Engine/Source/Runtime/Engine/Public/Materials/MaterialExpressionColorRamp.h
+++ b/Engine/Source/Runtime/Engine/Public/Materials/MaterialExpressionColorRamp.h
@@ -10,7 +10,7 @@
 #include "Curves/CurveLinearColor.h"
 #include "MaterialExpressionColorRamp.generated.h"
 
-UCLASS(MinimalAPI)
+UCLASS(Abstract, MinimalAPI)
 class UMaterialExpressionColorRamp : public UMaterialExpression
 {
 	GENERATED_BODY()
