Fixed text import test when compiled for low-level tests.

Replaced AddError with multi-framework macro to correctly handle the error when there isn't an automation test in flight.

diff --git a/Engine/Source/Runtime/Engine/Private/Tests/TextImportTest.cpp b/Engine/Source/Runtime/Engine/Private/Tests/TextImportTest.cpp
--- a/Engine/Source/Runtime/Engine/Private/Tests/TextImportTest.cpp
+++ b/Engine/Source/Runtime/Engine/Private/Tests/TextImportTest.cpp
@@ -56,7 +56,7 @@ TEST_CASE_NAMED(FTextImportStructPropertyTest, "System::Engine::TextImport::Stru
 		if (!TestCase.bShouldError && bImportErrored)
 		{
 			FString ErrorString = FString::Printf(TEXT("Test case failed due to ImportError %s for Test Case: %s"), *ImportError, InBuffer);
-			AddError(ErrorString);
+			ADD_ERROR(ErrorString);
 		}
 		
 		CHECK_MESSAGE(FString::Printf(TEXT("Import Error Mismatch. Result: %s Expected: %s for TestCase: %s"), bImportErrored ? TEXT("Error") : TEXT("NoError"), TestCase.bShouldError ? TEXT("Error") : TEXT("NoError"), InBuffer), TestCase.bShouldError == bImportErrored);
