Rename partially resident texture cvars to avoid confusion with software virtual texturing cvars.
r.VirtualTexture -> r.PRT.Enable
r.VirtualTextureReducedMemory ->r.PRT.ReducedMemory
Done with cvar deprecation so should respect old cvars but fire an ensure() when it finds them.

diff --git a/Engine/Source/Runtime/Core/Private/HAL/ConsoleManager.cpp b/Engine/Source/Runtime/Core/Private/HAL/ConsoleManager.cpp
--- a/Engine/Source/Runtime/Core/Private/HAL/ConsoleManager.cpp
+++ b/Engine/Source/Runtime/Core/Private/HAL/ConsoleManager.cpp
@@ -4276,12 +4276,6 @@ static TAutoConsoleVariable<int32> CVarPreTileTextures(
 	TEXT("If set to 1, textures will be tiled during cook and are expected to be cooked at runtime"),
 	ECVF_RenderThreadSafe);
 
-static TAutoConsoleVariable<int32> CVarVirtualTextureReducedMemoryEnabled(
-	TEXT("r.VirtualTextureReducedMemory"),
-	0,
-	TEXT("If set to 1, the cost of virtual textures will be reduced by using a more packed layout."),
-	ECVF_RenderThreadSafe);
-
 static TAutoConsoleVariable<int32> CVarPrecomputedVisibilityWarning(
 	TEXT("r.PrecomputedVisibilityWarning"),
 	0,
diff --git a/Engine/Source/Runtime/Engine/Private/Rendering/StreamableTextureResource.cpp b/Engine/Source/Runtime/Engine/Private/Rendering/StreamableTextureResource.cpp
--- a/Engine/Source/Runtime/Engine/Private/Rendering/StreamableTextureResource.cpp
+++ b/Engine/Source/Runtime/Engine/Private/Rendering/StreamableTextureResource.cpp
@@ -19,12 +19,22 @@ int64 GUITextureMemory = 0;
 int64 GNeverStreamTextureMemory = 0;
 #endif
 
-static TAutoConsoleVariable<int32> CVarVirtualTextureEnabled(
-	TEXT("r.VirtualTexture"),
+static TAutoConsoleVariable<int32> CVarPartiallyResidentTextureEnabled(
+	TEXT("r.PRT.Enable"),
 	1,
-	TEXT("If set to 1, textures will use virtual memory so they can be partially resident."),
+	TEXT("If set to 1, textures will use virtual memory on supported platforms so that they can be partially resident."),
 	ECVF_RenderThreadSafe);
 
+static TAutoConsoleVariable<int32> CVarPartiallyResidentTextureReducedMemory(
+	TEXT("r.PRT.ReducedMemory"),
+	0,
+	TEXT("If set to 1, the cost of partially resident textures will be reduced by using a more packed layout."),
+	ECVF_RenderThreadSafe);
+
+static FAutoConsoleVariableDeprecated CVarPartiallyResidentTextureEnabled_Deprecated(TEXT("r.VirtualTexture"), TEXT("r.PRT.Enable"), TEXT("5.6"));
+static FAutoConsoleVariableDeprecated CVarPartiallyResidentTextureReducedMemory_Deprecated(TEXT("r.VirtualTextureReducedMemory"), TEXT("r.PRT.ReducedMemory"), TEXT("5.6"));
+
+
 bool CanCreateWithPartiallyResidentMips(ETextureCreateFlags TexCreateFlags)
 {
 #if PLATFORM_SUPPORTS_VIRTUAL_TEXTURES
@@ -38,7 +48,7 @@ bool CanCreateWithPartiallyResidentMips(ETextureCreateFlags TexCreateFlags)
 	const ETextureCreateFlags iRequiredFlags =
 		TexCreate_OfflineProcessed;
 
-	return ((TexCreateFlags & (iDisableFlags | iRequiredFlags)) == iRequiredFlags) && CVarVirtualTextureEnabled.GetValueOnAnyThread();
+	return ((TexCreateFlags & (iDisableFlags | iRequiredFlags)) == iRequiredFlags) && CVarPartiallyResidentTextureEnabled.GetValueOnAnyThread();
 	
 #else
 	return false;
@@ -183,7 +193,7 @@ void FStreamableTextureResource::InitRHI(FRHICommandListBase&)
 	RefreshSamplerStates();
 	
 	// Check if this is the initial creation of the texture, or if we're recreating a texture that was released by ReleaseRHI.
-	static TConsoleVariableData<int32>* CVarReducedMode = IConsoleManager::Get().FindTConsoleVariableDataInt(TEXT("r.VirtualTextureReducedMemory"));
+	static TConsoleVariableData<int32>* CVarReducedMode = IConsoleManager::Get().FindTConsoleVariableDataInt(TEXT("r.PRT.ReducedMemory"));
 	check(CVarReducedMode);
 	if (bUsePartiallyResidentMips && (!CVarReducedMode->GetValueOnRenderThread() || State.NumRequestedLODs > State.NumNonStreamingLODs))
 	{
diff --git a/Engine/Source/Runtime/Engine/Private/Rendering/Texture2DResource.cpp b/Engine/Source/Runtime/Engine/Private/Rendering/Texture2DResource.cpp
--- a/Engine/Source/Runtime/Engine/Private/Rendering/Texture2DResource.cpp
+++ b/Engine/Source/Runtime/Engine/Private/Rendering/Texture2DResource.cpp
@@ -202,7 +202,7 @@ uint64 FTexture2DResource::GetPlatformMipsSize(uint32 NumMips) const
 {
 	if (PlatformData && NumMips > 0)
 	{
-		static TConsoleVariableData<int32>* CVarReducedMode = IConsoleManager::Get().FindTConsoleVariableDataInt(TEXT("r.VirtualTextureReducedMemory"));
+		static TConsoleVariableData<int32>* CVarReducedMode = IConsoleManager::Get().FindTConsoleVariableDataInt(TEXT("r.PRT.ReducedMemory"));
 		check(CVarReducedMode);
 
 		FRHICalcTextureSizeResult Result{};
diff --git a/Engine/Source/Runtime/Engine/Private/Streaming/Texture2DStreamOut_Virtual.cpp b/Engine/Source/Runtime/Engine/Private/Streaming/Texture2DStreamOut_Virtual.cpp
--- a/Engine/Source/Runtime/Engine/Private/Streaming/Texture2DStreamOut_Virtual.cpp
+++ b/Engine/Source/Runtime/Engine/Private/Streaming/Texture2DStreamOut_Virtual.cpp
@@ -28,7 +28,7 @@ void FTexture2DStreamOut_Virtual::Finalize(const FContext& Context)
 	DECLARE_SCOPE_CYCLE_COUNTER(TEXT("FTexture2DStreamOut_Virtual::Finalize"), STAT_Texture2DStreamOutVirtual_Finalize, STATGROUP_StreamingDetails);
 	check(Context.CurrentThread == TT_Render);
 
-	static TConsoleVariableData<int32>* CVarReducedMode = IConsoleManager::Get().FindTConsoleVariableDataInt(TEXT("r.VirtualTextureReducedMemory"));
+	static TConsoleVariableData<int32>* CVarReducedMode = IConsoleManager::Get().FindTConsoleVariableDataInt(TEXT("r.PRT.ReducedMemory"));
 	check(CVarReducedMode);
 
 	if (!CVarReducedMode->GetValueOnRenderThread() || ResourceState.NumRequestedLODs > ResourceState.NumNonStreamingLODs)
diff --git a/Engine/Source/Runtime/Engine/Private/Texture2D.cpp b/Engine/Source/Runtime/Engine/Private/Texture2D.cpp
--- a/Engine/Source/Runtime/Engine/Private/Texture2D.cpp
+++ b/Engine/Source/Runtime/Engine/Private/Texture2D.cpp
@@ -834,7 +834,7 @@ int32 UTexture2D::CalcTextureMemorySize( int32 MipCount ) const
 	int32 Size = 0;
 	if (GetPlatformData())
 	{
-		static TConsoleVariableData<int32>* CVarReducedMode = IConsoleManager::Get().FindTConsoleVariableDataInt(TEXT("r.VirtualTextureReducedMemory"));
+		static TConsoleVariableData<int32>* CVarReducedMode = IConsoleManager::Get().FindTConsoleVariableDataInt(TEXT("r.PRT.ReducedMemory"));
 		check(CVarReducedMode);
 
 		ETextureCreateFlags TexCreateFlags = (SRGB ? TexCreate_SRGB : TexCreate_None) | (bNoTiling ? TexCreate_NoTiling : TexCreate_None) | (bNotOfflineProcessed ? TexCreate_None : TexCreate_OfflineProcessed) | TexCreate_Streamable;
