[InstancedStaticMesh] make sure that UpdateBounds is completed before UStaticMeshComponent::PostStaticMeshCompilation calls FNavigationSubsystem::UpdateComponentData, which relies on updated and cached bounds to create a valid FNavigationElement.

diff --git a/Engine/Source/Runtime/Engine/Private/InstancedStaticMesh.cpp b/Engine/Source/Runtime/Engine/Private/InstancedStaticMesh.cpp
--- a/Engine/Source/Runtime/Engine/Private/InstancedStaticMesh.cpp
+++ b/Engine/Source/Runtime/Engine/Private/InstancedStaticMesh.cpp
@@ -5397,7 +5397,8 @@ void UInstancedStaticMeshComponent::GetNavigationData(FNavigationRelevantData& D
 FBox UInstancedStaticMeshComponent::GetNavigationBounds() const
 {
 	const UWorld* const World = GetWorld();
-	const bool bCheckValid = World && World->IsGameWorld();
+	// Navigation bounds must be calculated before being requested
+	constexpr bool bCheckValid = true;
 	return CachedBounds[(int)EBoundsCacheType::Navigation].GetValue(bCheckValid).GetBox();
 }
 
@@ -5467,15 +5468,12 @@ void UInstancedStaticMeshComponent::SetLODDistanceScale(float InLODDistanceScale
 #if WITH_EDITOR
 void UInstancedStaticMeshComponent::PostStaticMeshCompilation()
 {
-	// When UpdateBounds calls CalcAndCacheNavigationBounds, it will call GetInstanceNavigationBounds which calls UStaticMesh::GetNavCollision.
-	// If the mesh was compiling, it will wait for compilation to finish. 
-	// We override PostStaticMeshCompilation to intercept the call and make sure that UpdateBounds completes before 
-	// UStaticMeshComponent::PostStaticMeshCompilation calls FNavigationSubsystem::UpdateComponentData, which will end up calling 
-	// FNavigationElement::InitializeFromInterface (which calls UInstancedStaticMeshComponent::GetNavigationBounds).
-	if (bIsUpdatingBounds)
-	{
-		UpdateBounds();
-	}
+	// Multiple operations (including 'UpdateBounds') can early out when mesh is compiling so we need to 
+	// override PostStaticMeshCompilation to make sure that UpdateBounds is completed before 
+	// UStaticMeshComponent::PostStaticMeshCompilation calls FNavigationSubsystem::UpdateComponentData, which relies
+	// on updated and cached bounds to create a valid FNavigationElement.
+
+	UpdateBounds();
 
 	Super::PostStaticMeshCompilation();
 }
