[shaders] unify and fix shader job retry mechanisms
- both the messagebox-based retries for global/default shaders and the retries for jobs with errors/warnings based on the value of r.DumpShaderDebugInfo now occur in the same code path
- fixed the issue of r.DumpShaderDebugInfo=2/3 not working correctly if a global/default shader error is encountered during startup (the previous reissue path occurred asynchronously and so chances were good the fatal error would cause the process to exit before the reissued job completed; this unified code blocks until the reissued shaders are finished compiling)
- fixed the issue of r.DumpShaderDebugInfo=3 not outputting all artifacts (it wasn't actually re-running the jobs due to hitting the job cache), required an explicit state bit in the job input to skip the cache query
- improved the interactive retry re-issue so it only re-runs failed jobs, not all jobs for a particular shader map

diff --git a/Engine/Source/Runtime/Engine/Private/ShaderCompiler/ShaderCompiler.cpp b/Engine/Source/Runtime/Engine/Private/ShaderCompiler/ShaderCompiler.cpp
--- a/Engine/Source/Runtime/Engine/Private/ShaderCompiler/ShaderCompiler.cpp
+++ b/Engine/Source/Runtime/Engine/Private/ShaderCompiler/ShaderCompiler.cpp
@@ -2393,21 +2393,20 @@ static void BuildErrorStringAndReport(const FShaderDiagnosticInfo& DiagInfo, FSt
 	}
 }
 
-bool FShaderCompilingManager::HandlePotentialRetryOnError(TMap<int32, FShaderMapFinalizeResults>& CompletedShaderMaps)
+bool FShaderCompilingManager::HandlePotentialRetry(TMap<int32, FShaderMapFinalizeResults>& CompletedShaderMaps)
 {
-	if (FApp::IsUnattended())
-	{
-		return false;
-	}
-
-	bool bRetryCompile = false;
-
 #if WITH_EDITORONLY_DATA
+	TArray<FShaderCommonCompileJobPtr> ReissueJobs;
+	TArray<int32> ReissueMapIds;
+
 	for (TMap<int32, FShaderMapFinalizeResults>::TIterator It(CompletedShaderMaps); It; ++It)
 	{
 		FShaderMapFinalizeResults& Results = It.Value();
 
-		if (!Results.bAllJobsSucceeded)
+		bool bRetryCompile = false;
+
+		// interactive retries - prompt for global/default shaders that have errors
+		if (!FApp::IsUnattended() && !Results.bAllJobsSucceeded)
 		{
 			bool bSpecialEngineMaterial = false;
 
@@ -2424,9 +2423,9 @@ bool FShaderCompilingManager::HandlePotentialRetryOnError(TMap<int32, FShaderMap
 				}
 			}
 
-			if (UE_LOG_ACTIVE(LogShaders, Log) 
+			if (UE_LOG_ACTIVE(LogShaders, Log)
 				// Always log detailed errors when a special engine material or global shader fails to compile, as those will be fatal errors
-				|| bSpecialEngineMaterial 
+				|| bSpecialEngineMaterial
 				|| It.Key() == GlobalShaderMapId)
 			{
 				TArray<FShaderCommonCompileJobPtr>& CompleteJobs = Results.FinishedJobs;
@@ -2467,87 +2466,95 @@ bool FShaderCompilingManager::HandlePotentialRetryOnError(TMap<int32, FShaderMap
 						}
 					}
 				}
+			}
+		}
 
-				if (bRetryCompile)
+		int32 OriginalFinishedJobCount = Results.FinishedJobs.Num();
+		bool bAnyReissued = false;
+
+		// retries based on r.DumpShaderDebugInfo=2/3
+		for (int32 FinishedJobIndex = 0; FinishedJobIndex < Results.FinishedJobs.Num(); ++FinishedJobIndex)
+		{
+			FShaderCommonCompileJobPtr Job = Results.FinishedJobs[FinishedJobIndex];
+			bool bReissue = false;
+			Job->ForEachSingleShaderJob([&bReissue, bRetryCompile](FShaderCompileJob& SingleJob)
 				{
-					break;
-				}
+					// this will set bReissue to true if any stage of a pipeline job failed, so the entire pipeline job will be 
+					// resubmitted, but only the failed stage(s) will have debug info enabled on the re-run.
+					bool bShouldRecompileForDebug = GShaderCompilingManager->ShouldRecompileToDumpShaderDebugInfo(SingleJob);
+					// reissue the job if either we want to dump debug info for it, or a retry was requested above for any failed jobs
+					bReissue = bShouldRecompileForDebug || (bRetryCompile && !SingleJob.bSucceeded);
+					if (bShouldRecompileForDebug)
+					{
+						SingleJob.Input.DumpDebugInfoPath = GShaderCompilingManager->CreateShaderDebugInfoPath(SingleJob.Input);
+						// Any reissued jobs due to this condition will dump debug information, so increment the dump count here
+						GShaderCompilingManager->IncrementNumDumpedShaderSources();
+					}
+				});
+			if (bReissue)
+			{
+				bAnyReissued = true;
+				ReissueJobs.Add(Job);
+				// remove the job we will reissue from the finished jobs list; no point in allowing shrinking as it will grow again below when the reissued jobs finish
+				Results.FinishedJobs.RemoveAt(FinishedJobIndex--, EAllowShrinking::No);
 			}
 		}
+
+		if (bAnyReissued)
+		{
+			ReissueMapIds.Add(It.Key());
+		}
 	}
 
-	if (bRetryCompile)
+	if (ReissueJobs.Num())
 	{
 		// Flush the shader file cache so that any changes will be propagated.
 		FlushShaderFileCache();
 
-		TArray<int32> MapsToRemove;
-
-		for (TMap<int32, FShaderMapFinalizeResults>::TIterator It(CompletedShaderMaps); It; ++It)
+		// Reset outputs
+		for (const FShaderCommonCompileJobPtr& CurrentJob : ReissueJobs)
 		{
-			FShaderMapFinalizeResults& Results = It.Value();
+			auto* SingleJob = CurrentJob->GetSingleShaderJob();
 
-			if (!Results.bAllJobsSucceeded)
+			// NOTE: Changes to MaterialTemplate.usf before retrying won't work, because the entry for Material.usf in CurrentJob.Environment.IncludeFileNameToContentsMap isn't reset
+			if (SingleJob)
 			{
-				MapsToRemove.Add(It.Key());
-
-				// Reset outputs
-				for (int32 JobIndex = 0; JobIndex < Results.FinishedJobs.Num(); JobIndex++)
+				SingleJob->Output = FShaderCompilerOutput();
+				SingleJob->PreprocessOutput = FShaderPreprocessOutput();
+				SingleJob->JobStatusPtr->Reset();
+			}
+			else
+			{
+				auto* PipelineJob = CurrentJob->GetShaderPipelineJob();
+				for (auto CommonJob : PipelineJob->StageJobs)
 				{
-					FShaderCommonCompileJob& CurrentJob = *Results.FinishedJobs[JobIndex];
-					auto* SingleJob = CurrentJob.GetSingleShaderJob();
-
-					// NOTE: Changes to MaterialTemplate.usf before retrying won't work, because the entry for Material.usf in CurrentJob.Environment.IncludeFileNameToContentsMap isn't reset
-					if (SingleJob)
-					{
-						if (GShaderCompilingManager->ShouldRecompileToDumpShaderDebugInfo(*SingleJob))
-						{
-							SingleJob->Input.DumpDebugInfoPath = GShaderCompilingManager->CreateShaderDebugInfoPath(SingleJob->Input);
-						}
-						SingleJob->Output = FShaderCompilerOutput();
-						SingleJob->PreprocessOutput = FShaderPreprocessOutput();
-						SingleJob->JobStatusPtr->Reset();
-					}
-					else
-					{
-						auto* PipelineJob = CurrentJob.GetShaderPipelineJob();
-						for (auto CommonJob : PipelineJob->StageJobs)
-						{
-							CommonJob->Output = FShaderCompilerOutput();
-							CommonJob->PreprocessOutput = FShaderPreprocessOutput();
-							CommonJob->JobStatusPtr->Reset();
-							CommonJob->bFinalized = false;
-						}
-					}
-					
-					// Reset DDC query request owner
-					CurrentJob.RequestOwner.Reset();
-					CurrentJob.bFinalized = false;
+					CommonJob->Output = FShaderCompilerOutput();
+					CommonJob->PreprocessOutput = FShaderPreprocessOutput();
+					CommonJob->JobStatusPtr->Reset();
+					CommonJob->bFinalized = false;
 				}
-
-				// Send all the shaders from this shader map through the compiler again
-				SubmitJobs(Results.FinishedJobs, FString(""), FString(""));
 			}
+			
+			// Reset DDC query request owner
+			CurrentJob->RequestOwner.Reset();
+			CurrentJob->bInputHashSet = false;
+			CurrentJob->bFinalized = false;
+			// Need to force reissued jobs to skip the cache queries otherwise jobs with warnings will just be cache hits and not actually recompile
+			// (since debug info being enabled for a job intentionally does not affect the cached key/cached results)
+			CurrentJob->bBypassCache = true;
 		}
 
-		const int32 OriginalNumShaderMaps = CompletedShaderMaps.Num();
-
-		// Remove the failed shader maps
-		for (int32 RemoveIndex = 0; RemoveIndex < MapsToRemove.Num(); RemoveIndex++)
-		{
-			CompletedShaderMaps.Remove(MapsToRemove[RemoveIndex]);
-		}
-
-		check(CompletedShaderMaps.Num() == OriginalNumShaderMaps - MapsToRemove.Num());
+		// Submit all the jobs which we want to recompile
+		SubmitJobs(ReissueJobs, FString(""), FString(""));
 
-		// Block until the failed shader maps have been compiled again
-		BlockOnShaderMapCompletion(MapsToRemove, CompletedShaderMaps);
+		// Block until the shader maps with reissued jobs have been compiled again (this may include new jobs since submission could
+		// be occurring in parallel on other threads)
+		BlockOnShaderMapCompletion(ReissueMapIds, CompletedShaderMaps);
 
-		check(CompletedShaderMaps.Num() == OriginalNumShaderMaps);
+		return true;
 	}
 #endif	//WITH_EDITORONLY_DATA
-
-	return bRetryCompile;
+	return false;
 }
 
 void FShaderMapCompileResults::CheckIfHung()
@@ -2663,7 +2670,7 @@ void FShaderCompilingManager::FinishCompilation(const TCHAR* MaterialName, const
 	bool bRetry = false;
 	do 
 	{
-		bRetry = HandlePotentialRetryOnError(CompiledShaderMaps);
+		bRetry = HandlePotentialRetry(CompiledShaderMaps);
 	} 
 	while (bRetry);
 
@@ -2695,7 +2702,7 @@ void FShaderCompilingManager::FinishAllCompilation()
 	bool bRetry = false;
 	do 
 	{
-		bRetry = HandlePotentialRetryOnError(CompiledShaderMaps);
+		bRetry = HandlePotentialRetry(CompiledShaderMaps);
 	} 
 	while (bRetry);
 
@@ -2785,7 +2792,7 @@ void FShaderCompilingManager::ProcessAsyncResults(float TimeSlice, bool bBlockOn
 		bool bRetry = false;
 		do 
 		{
-			bRetry = HandlePotentialRetryOnError(PendingFinalizeShaderMaps);
+			bRetry = HandlePotentialRetry(PendingFinalizeShaderMaps);
 		} 
 		while (bRetry);
 
diff --git a/Engine/Source/Runtime/Engine/Private/ShaderCompiler/ShaderCompilerUtil.cpp b/Engine/Source/Runtime/Engine/Private/ShaderCompiler/ShaderCompilerUtil.cpp
--- a/Engine/Source/Runtime/Engine/Private/ShaderCompiler/ShaderCompilerUtil.cpp
+++ b/Engine/Source/Runtime/Engine/Private/ShaderCompiler/ShaderCompilerUtil.cpp
@@ -76,43 +76,6 @@ static void ModalErrorOrLog(const FString& Title, const FString& Text, int64 Cur
 	}
 }
 
-static void ReissueShaderCompileJobs(const TArray<FShaderCommonCompileJob*>& SourceJobs)
-{
-	if (SourceJobs.Num())
-	{
-		TRACE_CPUPROFILER_EVENT_SCOPE(ReissueShaderCompileJobs);
-
-		TArray<FShaderCommonCompileJobPtr> ReissueJobs;
-		ReissueJobs.Reserve(SourceJobs.Num());
-		const uint32 JobId = FShaderCommonCompileJob::GetNextJobId();
-		for (const FShaderCommonCompileJob* SourceJob : SourceJobs)
-		{
-			if (const FShaderCompileJob* SingleSourceJob = SourceJob->GetSingleShaderJob())
-			{
-				if (FShaderCompileJob* ReissueJob = GShaderCompilingManager->PrepareShaderCompileJob(JobId, SingleSourceJob->Key, SingleSourceJob->Priority))
-				{
-					ReissueJob->Input = SingleSourceJob->Input;
-					ReissueJobs.Add(FShaderCommonCompileJobPtr(ReissueJob));
-				}
-			}
-			else if (const FShaderPipelineCompileJob* PipelineSourceJob = SourceJob->GetShaderPipelineJob())
-			{
-				if (FShaderPipelineCompileJob* ReissueJob = GShaderCompilingManager->PreparePipelineCompileJob(JobId, PipelineSourceJob->Key, PipelineSourceJob->Priority))
-				{
-					ReissueJob->StageJobs = PipelineSourceJob->StageJobs;
-					ReissueJobs.Add(FShaderCommonCompileJobPtr(ReissueJob));
-				}
-			}
-			else
-			{
-				checkf(0, TEXT("Reissued shader compile job is neither a single nor a pipeline job"));
-			}
-		}
-
-		GShaderCompilingManager->SubmitJobs(ReissueJobs, FString(""), FString(""));
-	}
-}
-
 // Make functions so the crash reporter can disambiguate the actual error because of the different callstacks
 namespace ShaderCompileWorkerError
 {
@@ -537,7 +500,7 @@ bool FShaderCompileUtilities::DoWriteTasks(const TArray<FShaderCommonCompileJobP
 	return DoWriteTasksInner(QueuedJobs, InTransferFile, BuildDistributionController, bUseRelativePaths, bCompressTaskFile);
 }
 
-static bool ReadSingleJob(FShaderCompileJob* CurrentJob, FArchive& WorkerOutputFileReader)
+static void ReadSingleJob(FShaderCompileJob* CurrentJob, FArchive& WorkerOutputFileReader)
 {
 	TRACE_CPUPROFILER_EVENT_SCOPE(ReadSingleJob);
 
@@ -549,23 +512,6 @@ static bool ReadSingleJob(FShaderCompileJob* CurrentJob, FArchive& WorkerOutputF
 
 	// The job should already have a non-zero output hash
 	checkf(CurrentJob->Output.OutputHash != FSHAHash() || !CurrentJob->bSucceeded, TEXT("OutputHash for a successful job was not set in the shader compile worker!"));
-
-
-	// Support dumping debug info for only failed compilations or those with warnings
-	if (GShaderCompilingManager->ShouldRecompileToDumpShaderDebugInfo(*CurrentJob))
-	{
-		// Build debug info path and create the directory if it doesn't already exist
-		CurrentJob->Input.DumpDebugInfoPath = GShaderCompilingManager->CreateShaderDebugInfoPath(CurrentJob->Input);
-		// Reset job status object to ensure both job state and "duplicate" flag are reset.
-		CurrentJob->JobStatusPtr->Reset();
-
-		// We failed to compile a shader, we will retry and dump the shader source. so increment the number of shader sources we've dumped so far.
-		GShaderCompilingManager->IncrementNumDumpedShaderSources();
-
-		return true;
-	}
-
-	return false;
 }
 
 static const TCHAR* GetCompileJobSuccessText(FShaderCompileJob* SingleJob)
@@ -826,7 +772,6 @@ FSCWErrorCode::ECode FShaderCompileUtilities::DoReadTaskResults(const TArray<FSh
 	TArray<FShaderCompileJob*> QueuedSingleJobs;
 	TArray<FShaderPipelineCompileJob*> QueuedPipelineJobs;
 	SplitJobsByType(QueuedJobs, QueuedSingleJobs, QueuedPipelineJobs);
-	TArray<FShaderCommonCompileJob*> ReissueSourceJobs;
 
 	// Read single jobs
 	{
@@ -848,10 +793,7 @@ FSCWErrorCode::ECode FShaderCompileUtilities::DoReadTaskResults(const TArray<FSh
 			for (int32 JobIndex = 0; JobIndex < NumJobs; JobIndex++)
 			{
 				auto* CurrentJob = QueuedSingleJobs[JobIndex];
-				if (ReadSingleJob(CurrentJob, OutputFile))
-				{
-					ReissueSourceJobs.Add(CurrentJob);
-				}
+				ReadSingleJob(CurrentJob, OutputFile);
 			}
 		}
 	}
@@ -902,7 +844,6 @@ FSCWErrorCode::ECode FShaderCompileUtilities::DoReadTaskResults(const TArray<FSh
 					for (int32 Index = 0; Index < NumStageJobs; Index++)
 					{
 						FShaderCompileJob* SingleJob = CurrentJob->StageJobs[Index];
-						// cannot reissue a single stage of a pipeline job
 						ReadSingleJob(SingleJob, OutputFile);
 					}
 				}
@@ -910,9 +851,6 @@ FSCWErrorCode::ECode FShaderCompileUtilities::DoReadTaskResults(const TArray<FSh
 		}
 	}
 	
-	// Requeue any jobs we wish to run again
-	ReissueShaderCompileJobs(ReissueSourceJobs);
-
 	return FSCWErrorCode::Success;
 }
 
diff --git a/Engine/Source/Runtime/Engine/Public/ShaderCompiler.h b/Engine/Source/Runtime/Engine/Public/ShaderCompiler.h
--- a/Engine/Source/Runtime/Engine/Public/ShaderCompiler.h
+++ b/Engine/Source/Runtime/Engine/Public/ShaderCompiler.h
@@ -1025,8 +1025,8 @@ private:
 	/** Propagate the completed compile to primitives that might be using the materials compiled. */
 	void PropagateMaterialChangesToPrimitives(TMap<TRefCountPtr<FMaterial>, TRefCountPtr<FMaterialShaderMap>>& MaterialsToUpdate);
 
-	/** Recompiles shader jobs with errors if requested, and returns true if a retry was needed. */
-	bool HandlePotentialRetryOnError(TMap<int32, FShaderMapFinalizeResults>& CompletedShaderMaps);
+	/** Recompiles shader jobs if requested (either interactively for failed jobs via msgbox in editor, or for jobs with errors/warnings due to r.DumpShaderDebugInfo=2/3), and returns true if a retry was needed. */
+	bool HandlePotentialRetry(TMap<int32, FShaderMapFinalizeResults>& CompletedShaderMaps);
 	
 	/** Checks if any target platform down't support remote shader compiling */
 	bool AllTargetPlatformSupportsRemoteShaderCompiling();
diff --git a/Engine/Source/Runtime/RenderCore/Public/ShaderCompilerJobTypes.h b/Engine/Source/Runtime/RenderCore/Public/ShaderCompilerJobTypes.h
--- a/Engine/Source/Runtime/RenderCore/Public/ShaderCompilerJobTypes.h
+++ b/Engine/Source/Runtime/RenderCore/Public/ShaderCompilerJobTypes.h
@@ -186,6 +186,8 @@ public:
 	uint8 bIsDefaultMaterial : 1;
 	/** Whether or not we are a global shader. */
 	uint8 bIsGlobalShader : 1;
+	/** Whether or not to bypass the job/ddc caches when executing this job */
+	uint8 bBypassCache : 1;
 	/** Hash of all the job inputs */
 	FShaderCompilerInputHash InputHash;
 
@@ -283,8 +285,8 @@ protected:
 		bReleased(false),
 		bInputHashSet(false),
 		bIsDefaultMaterial(false),
-		bIsGlobalShader(false)
-
+		bIsGlobalShader(false),
+		bBypassCache(false)
 	{
 		check(InPriority != EShaderCompileJobPriority::None);
 
