GC fix for bus sends on wave instances

[RN] Fixed an issue with GC and bus sends on wave instances

diff --git a/Engine/Source/Runtime/Engine/Private/Audio.cpp b/Engine/Source/Runtime/Engine/Private/Audio.cpp
--- a/Engine/Source/Runtime/Engine/Private/Audio.cpp
+++ b/Engine/Source/Runtime/Engine/Private/Audio.cpp
@@ -1203,12 +1203,29 @@ void FWaveInstance::AddReferencedObjects( FReferenceCollector& Collector )
 			Collector.AddReferencedObject(Info.SoundSubmix);
 		}	
 	};
-	
+
 	Algo::ForEach(SoundSubmixSends, AddSubmixSendRef);
 	Algo::ForEach(AttenuationSubmixSends, AddSubmixSendRef);
 	
 	Collector.AddReferencedObject( SoundClass );
 	NotifyBufferFinishedHooks.AddReferencedObjects( Collector );
+
+	auto AddBusSendRef = [&Collector](FSoundSourceBusSendInfo& Info)
+		{
+			if (Info.SoundSourceBus)
+			{
+				Collector.AddReferencedObject(Info.SoundSourceBus);
+			}
+			if (Info.AudioBus)
+			{
+				Collector.AddReferencedObject(Info.AudioBus);
+			}
+		};
+
+	for (int32 BusSendCount = 0; BusSendCount < (int32)EBusSendType::Count; ++BusSendCount)
+	{
+		Algo::ForEach(BusSends[BusSendCount], AddBusSendRef);
+	}
 }
 
 float FWaveInstance::GetActualVolume() const
