Fix race between AActor::InitializeDefaults and RecordWorldCountsToCSV because of missing usage of critical section

diff --git a/Engine/Source/Runtime/Engine/Private/LevelTick.cpp b/Engine/Source/Runtime/Engine/Private/LevelTick.cpp
--- a/Engine/Source/Runtime/Engine/Private/LevelTick.cpp
+++ b/Engine/Source/Runtime/Engine/Private/LevelTick.cpp
@@ -1186,6 +1186,7 @@ static TAutoConsoleVariable<int32> CVarRecordActorCountsToCSVThreshold(
 	5,
 	TEXT("Number of instances of an native Actor class required before recording to CSV stat"));
 
+extern FCriticalSection CSVActorClassNameToCountMapLock;
 extern TMap<FName, int32> CSVActorClassNameToCountMap;
 extern int32 CSVActorTotalCount;
 
@@ -1222,11 +1223,14 @@ static void RecordWorldCountsToCSV(UWorld* World, bool bDoingActorTicks)
 
 			const int32 Threshold = CVarRecordActorCountsToCSVThreshold.GetValueOnAnyThread();
 
-			for (auto It = CSVActorClassNameToCountMap.CreateConstIterator(); It; ++It)
 			{
-				if (It->Value > Threshold)
+				FScopeLock Lock(&CSVActorClassNameToCountMapLock);
+				for (auto It = CSVActorClassNameToCountMap.CreateConstIterator(); It; ++It)
 				{
-					FCsvProfiler::Get()->RecordCustomStat(It->Key, CSV_CATEGORY_INDEX(ActorCount), It->Value, ECsvCustomStatOp::Set);
+					if (It->Value > Threshold)
+					{
+						FCsvProfiler::Get()->RecordCustomStat(It->Key, CSV_CATEGORY_INDEX(ActorCount), It->Value, ECsvCustomStatOp::Set);
+					}
 				}
 			}
 
