Fix pose watch array OOB crash

diff --git a/Engine/Source/Runtime/Engine/Private/Engine/PoseWatchRenderData.cpp b/Engine/Source/Runtime/Engine/Private/Engine/PoseWatchRenderData.cpp
--- a/Engine/Source/Runtime/Engine/Private/Engine/PoseWatchRenderData.cpp
+++ b/Engine/Source/Runtime/Engine/Private/Engine/PoseWatchRenderData.cpp
@@ -44,18 +44,21 @@ void FAnimNodePoseWatch::CopyPoseWatchData(const FReferenceSkeleton& RefSkeleton
 	const UBlendProfile* Mask = PoseWatchPoseElement->ViewportMask;
 	for (const FBoneIndexType& BoneIndex : RequiredBones)
 	{
-		if (Mask)
+		if (RefSkeleton.IsValidIndex(BoneIndex))
 		{
-			const FName& BoneName = RefSkeleton.GetBoneName(BoneIndex);
-			const bool bBoneBlendScaleMeetsThreshold = Mask->GetBoneBlendScale(BoneName) > PoseWatchPoseElement->BlendScaleThreshold;
-			const bool bBoneRequired = (!PoseWatchPoseElement->bInvertViewportMask && bBoneBlendScaleMeetsThreshold) || (PoseWatchPoseElement->bInvertViewportMask && !bBoneBlendScaleMeetsThreshold);
-			if (!bBoneRequired)
+			if (Mask)
 			{
-				continue;
+				const FName& BoneName = RefSkeleton.GetBoneName(BoneIndex);
+				const bool bBoneBlendScaleMeetsThreshold = Mask->GetBoneBlendScale(BoneName) > PoseWatchPoseElement->BlendScaleThreshold;
+				const bool bBoneRequired = (!PoseWatchPoseElement->bInvertViewportMask && bBoneBlendScaleMeetsThreshold) || (PoseWatchPoseElement->bInvertViewportMask && !bBoneBlendScaleMeetsThreshold);
+				if (!bBoneRequired)
+				{
+					continue;
+				}
 			}
-		}
 
-		ViewportMaskAllowedList.Add(BoneIndex);
+			ViewportMaskAllowedList.Add(BoneIndex);
+		}
 	}
 }
 
