Fixed streamablemanager error not including softobjectpath.

diff --git a/Engine/Source/Runtime/Engine/Classes/Engine/Experimental/StreamableManagerError.h b/Engine/Source/Runtime/Engine/Classes/Engine/Experimental/StreamableManagerError.h
--- a/Engine/Source/Runtime/Engine/Classes/Engine/Experimental/StreamableManagerError.h
+++ b/Engine/Source/Runtime/Engine/Classes/Engine/Experimental/StreamableManagerError.h
@@ -57,7 +57,7 @@ namespace UE::UnifiedError::StreamableManager
 		{
 			const FText InnerFormatString = GetInnerErrorDetails()->GetErrorFormatString(Error);
 
-			FText MyFormatString = NSLOCTEXT("StreamableManager", "DetailsFormatString", "{FStreamableManagerAdditionalContext}");
+			FText MyFormatString = NSLOCTEXT("StreamableManager", "DetailsFormatString", "(SoftObjectPath:{SoftObjectPath})");
 
 			return FText::Join(NSLOCTEXT("StreamableManager", "DetailsFormatStringDelimiter", " "), InnerFormatString, MyFormatString);
 		}
