[Shaders] Small simplification and hardening assumptions about state changes in compile job status.

#rnx

diff --git a/Engine/Source/Runtime/Engine/Private/ShaderCompiler/ShaderCompiler.cpp b/Engine/Source/Runtime/Engine/Private/ShaderCompiler/ShaderCompiler.cpp
--- a/Engine/Source/Runtime/Engine/Private/ShaderCompiler/ShaderCompiler.cpp
+++ b/Engine/Source/Runtime/Engine/Private/ShaderCompiler/ShaderCompiler.cpp
@@ -2515,26 +2515,16 @@ bool FShaderCompilingManager::HandlePotentialRetry(TMap<int32, FShaderMapFinaliz
 		// Reset outputs
 		for (const FShaderCommonCompileJobPtr& CurrentJob : ReissueJobs)
 		{
-			auto* SingleJob = CurrentJob->GetSingleShaderJob();
-
 			// NOTE: Changes to MaterialTemplate.usf before retrying won't work, because the entry for Material.usf in CurrentJob.Environment.IncludeFileNameToContentsMap isn't reset
-			if (SingleJob)
-			{
-				SingleJob->Output = FShaderCompilerOutput();
-				SingleJob->PreprocessOutput = FShaderPreprocessOutput();
-				SingleJob->JobStatusPtr->Reset();
-			}
-			else
-			{
-				auto* PipelineJob = CurrentJob->GetShaderPipelineJob();
-				for (auto CommonJob : PipelineJob->StageJobs)
+			CurrentJob->ForEachSingleShaderJob(
+				[](FShaderCompileJob& SingleJob)
 				{
-					CommonJob->Output = FShaderCompilerOutput();
-					CommonJob->PreprocessOutput = FShaderPreprocessOutput();
-					CommonJob->JobStatusPtr->Reset();
-					CommonJob->bFinalized = false;
+					SingleJob.Output = FShaderCompilerOutput();
+					SingleJob.PreprocessOutput = FShaderPreprocessOutput();
+					SingleJob.JobStatusPtr->Reset();
+					SingleJob.bFinalized = false;
 				}
-			}
+			);
 			
 			// Reset DDC query request owner
 			CurrentJob->RequestOwner.Reset();
diff --git a/Engine/Source/Runtime/RenderCore/Public/ShaderCompilerJobTypes.h b/Engine/Source/Runtime/RenderCore/Public/ShaderCompilerJobTypes.h
--- a/Engine/Source/Runtime/RenderCore/Public/ShaderCompilerJobTypes.h
+++ b/Engine/Source/Runtime/RenderCore/Public/ShaderCompilerJobTypes.h
@@ -95,13 +95,15 @@ class FShaderCompileJobStatus
 public:
 	void Reset()
 	{
+		// Unset status before the input hash since GetInputHash() uses it as a substitute for lock guards
+		Status.store(EShaderCompileJobStatus::Unset);
 		InputHash = nullptr;
 		IsDuplicate.store(false);
-		Status.store(EShaderCompileJobStatus::Unset);
 	}
 
 	void SetStatus(EShaderCompileJobStatus NewStatus)
 	{
+		checkf(NewStatus == EShaderCompileJobStatus::Unset || InputHash != nullptr, TEXT("Unless the job status is reset back to 'Unset', the input hash must be provided"));
 		Status.store(NewStatus);
 	}
 
