Add LevelStreaming.VisibilityPrioritySort to change the order that it processes level streaming adds and removes.
The default value of 0 matches the old behavior where it would always process adds before removes, but 1 will use less memory.

diff --git a/Engine/Source/Runtime/Engine/Private/World.cpp b/Engine/Source/Runtime/Engine/Private/World.cpp
--- a/Engine/Source/Runtime/Engine/Private/World.cpp
+++ b/Engine/Source/Runtime/Engine/Private/World.cpp
@@ -178,6 +178,15 @@ static FAutoConsoleVariableRef CVarMaximumMakingInvisibleLevels(
 	TEXT("(see p.Chaos.EnableAsyncInitBody and LevelStreaming.AllowIncrementalPreUnregisterComponents)."),
 	ECVF_Default);
 
+static int32 GLevelVisibilityPrioritySort = 0;
+FAutoConsoleVariableRef CVarLevelVisibilityPrioritySort(TEXT("LevelStreaming.VisibilityPrioritySort"), GLevelVisibilityPrioritySort,
+	TEXT("Defines how to sort the priority of handling different streaming level requests\n")
+	TEXT("-1 - Do not sort at all\n")
+	TEXT("0 - Default behavior of sorting all add to world requests before removes\n")
+	TEXT("1 - Sort removes before adds, this is better for memory usage\n")
+	TEXT("2 - Interleave adds and removes, starting with add\n")
+	TEXT("3 - Interleave adds and removes, starting with remove\n"));
+
 static int32 bDisableRemapScriptActors = 0;
 FAutoConsoleVariableRef CVarDisableRemapScriptActors(TEXT("net.DisableRemapScriptActors"), bDisableRemapScriptActors, TEXT("When set, disables name remapping of compiled script actors (for networking)"));
 
@@ -4732,27 +4741,61 @@ void UWorld::UpdateLevelStreaming(const TOptional<const UE::FTimeout>& ExternalT
 	// Store current number of pending unload levels, it may change in loop bellow
 	const int32 NumLevelsPendingPurge = FLevelStreamingGCHelper::GetNumLevelsPendingPurge();
 
+	if (GLevelVisibilityPrioritySort >= 0)
 	{
 		// Make sure to respect the order of making visible and making invisible queues
 		TRACE_CPUPROFILER_EVENT_SCOPE(UpdateLevelVisibilityPriority);
 		static constexpr int32 MaxPriority = MAX_int32;
 		int32 Priority = MaxPriority;
-		for (ULevel* Level : MakingVisibleLevels)
+
+		// Use GLevelVisibilityPrioritySort to iterate the arrays (rules described above)
+		check(GLevelVisibilityPrioritySort <= 3);
+		const bool bAlternatePriority = GLevelVisibilityPrioritySort >= 2;
+		bool bVisibleNext = (GLevelVisibilityPrioritySort == 0 || GLevelVisibilityPrioritySort == 2);
+		int32 VisibleIndex = 0, InvisibleIndex = 0;
+		ULevel* NextLevel = nullptr;
+		
+		// Sets NextLevel and returns true if it was a valid index that was incremented (it could still be null)
+		auto SetNextLevel = [&]() -> bool
 		{
-			if (ULevelStreaming* StreamingLevel = ULevelStreaming::FindStreamingLevel(Level))
+			if (bVisibleNext)
 			{
-				check(StreamingLevelsToConsider.GetStreamingLevels().Contains(StreamingLevel));
-				StreamingLevel->SetPriority(Priority--);
+				if (MakingVisibleLevels.IsValidIndex(VisibleIndex))
+				{
+					NextLevel = MakingVisibleLevels[VisibleIndex++];
+					return true;
+				}
 			}
-		}
+			else if (MakingInvisibleLevels.IsValidIndex(InvisibleIndex))
+			{
+				NextLevel = MakingInvisibleLevels[InvisibleIndex++];
+				return true;
+			}
+			return false;
+		};
 
-		for (ULevel* Level : MakingInvisibleLevels)
+		while (true)
 		{
-			if (ULevelStreaming* StreamingLevel = ULevelStreaming::FindStreamingLevel(Level))
+			if (!SetNextLevel())
+			{
+				// Try other array, break if that is also empty
+				bVisibleNext = !bVisibleNext;
+				if (!SetNextLevel())
+				{
+					break;
+				}
+			}
+
+			if (ULevelStreaming* StreamingLevel = ULevelStreaming::FindStreamingLevel(NextLevel))
 			{
 				check(StreamingLevelsToConsider.GetStreamingLevels().Contains(StreamingLevel));
 				StreamingLevel->SetPriority(Priority--);
 			}
+
+			if (bAlternatePriority)
+			{
+				bVisibleNext = !bVisibleNext;
+			}
 		}
 	}
 
