TextureLODSettings : Add cvar r.GlobalTextureLODBias to apply a global LODBias adjustment to texture builds
Similar to existing cvar for UI textures, but affects all textures
Intended for marketing/cinematics/screencaps, not normal cook usage

diff --git a/Engine/Source/Runtime/Engine/Private/TextureLODSettings.cpp b/Engine/Source/Runtime/Engine/Private/TextureLODSettings.cpp
--- a/Engine/Source/Runtime/Engine/Private/TextureLODSettings.cpp
+++ b/Engine/Source/Runtime/Engine/Private/TextureLODSettings.cpp
@@ -20,6 +20,20 @@ FAutoConsoleVariableRef CVarUITextureLODBias(
 	ECVF_Scalability
 );
 
+// GlobalTextureLODBias similar to UI Texture LODBias but affects all textures types
+// Intended for marketing/cinematics/screencaps, not normal cook usage
+//	(eg. set to "-15" to get all possible mip levels in built textures)
+// !Note that changing this in an interactive editor may not affect already loaded textures!
+//	it is best to set this at startup, via command line or ini config
+// (?? could use ECVF_ReadOnly)
+int32 GlobalTextureLODBias = 0;
+FAutoConsoleVariableRef CVarGlobalTextureLODBias(
+	TEXT("r.GlobalTextureLODBias"),
+	GlobalTextureLODBias,
+	TEXT("Extra LOD bias to apply to all textures.  Can be negative. (default=0)"),
+	ECVF_Scalability
+);
+
 
 void FTextureLODGroup::SetupGroup()
 {
@@ -199,17 +213,17 @@ int32 UTextureLODSettings::CalculateLODBias(int32 Width, int32 Height, int32 Max
 		// Considering them again here would apply them twice.
 		UsedLODBias	+= LODBias;
 		UsedLODBias += bVirtualTexture ? LODGroupInfo.LODBias_VT : LODGroupInfo.LODBias;
+
+		// GlobalTextureLODBias from cvar r.GlobalTextureLODBias
+		//	only applied on non-cooked platforms (will apply when cooking)
+		UsedLODBias += GlobalTextureLODBias;
 	}
 	
 	if (LODGroup == TEXTUREGROUP_UI)
 	{
 		UsedLODBias += GUITextureLODBias;  
 		// @todo Oodle : GUITextureLODBias is applied at both cook time & run time , which is screwy
-		//	this is not inside the if (!FPlatformProperties::RequiresCookedData()) ?
-		//	it should either act like cinematic bias (runtime streaming)
-		//	or like the "drop mip" lod bias (cook time)
-		//	but it's not quite like either
-		// -> this looks broken, probably never used
+		// -> this looks broken, probably should be inside the !RequiresCookedData
 	}
 
 	// Min/Max LOD Size clamps the *LODBias*, not the actual texture LODs. Meaning:
