[Localization] Volume menu in the show flags level editor viewports will now pull the display name from the class data instead of looking for an entry in show flags

diff --git a/Engine/Source/Editor/LevelEditor/Private/LevelViewportActions.cpp b/Engine/Source/Editor/LevelEditor/Private/LevelViewportActions.cpp
--- a/Engine/Source/Editor/LevelEditor/Private/LevelViewportActions.cpp
+++ b/Engine/Source/Editor/LevelEditor/Private/LevelViewportActions.cpp
@@ -296,7 +296,10 @@ void FLevelViewportCommands::RegisterShowVolumeCommands()
 		if (!FoundVolumeCommand.IsValid())
 		{
 			FText DisplayName;
-			FEngineShowFlags::FindShowFlagDisplayName(VolumeName.ToString(), DisplayName);
+			if (!FEngineShowFlags::FindShowFlagDisplayName(VolumeName.ToString(), DisplayName))
+			{
+				DisplayName = VolumeClasses[VolumeClassIndex]->GetDisplayNameText();
+			}
 
 			FFormatNamedArguments Args;
 			Args.Add(TEXT("ShowFlagName"), DisplayName);
diff --git a/Engine/Source/Runtime/Engine/Public/ShowFlags.h b/Engine/Source/Runtime/Engine/Public/ShowFlags.h
--- a/Engine/Source/Runtime/Engine/Public/ShowFlags.h
+++ b/Engine/Source/Runtime/Engine/Public/ShowFlags.h
@@ -103,7 +103,7 @@ struct FEngineShowFlags
 	static ENGINE_API FSimpleMulticastDelegate OnCustomShowFlagRegistered;
 
 	/** Retrieve the localized display name for a named show flag */
-	static void FindShowFlagDisplayName(const FString& InName, FText& OutText)
+	static bool FindShowFlagDisplayName(const FString& InName, FText& OutText)
 	{
 		static TMap< FString, FText > LocNames;
 
@@ -142,14 +142,15 @@ struct FEngineShowFlags
 		if (OutName)
 		{
 			OutText = *OutName;
+			return true;
 		}
-		else if (FindCustomShowFlagDisplayName(InName, OutText))
-		{
-		}
-		else
+		if (FindCustomShowFlagDisplayName(InName, OutText))
 		{
-			OutText = FText::FromString(InName);
+			return true;
 		}
+
+		OutText = FText::FromString(InName);
+		return false;
 	}
 
 	/** this function defines the grouping in the editor. If a flag is not defined here it's assumed to be a postprocess flag */
