Fix assertion in FSceneView() when GRHISupportsArrayIndexFromAnyShader=true, but GSupportsMobileMultiView is not.

diff --git a/Engine/Source/Runtime/Engine/Private/GameViewportClient.cpp b/Engine/Source/Runtime/Engine/Private/GameViewportClient.cpp
--- a/Engine/Source/Runtime/Engine/Private/GameViewportClient.cpp
+++ b/Engine/Source/Runtime/Engine/Private/GameViewportClient.cpp
@@ -65,6 +65,7 @@
 #include "HDRHelper.h"
 #include "GlobalRenderResources.h"
 #include "ShaderCore.h"
+#include "StereoRenderUtils.h"
 
 #if WITH_EDITOR
 #include "Settings/LevelEditorPlaySettings.h"
@@ -1455,9 +1456,8 @@ void UGameViewportClient::Draw(FViewport* InViewport, FCanvas* SceneCanvas)
 	bool bRequireMultiView = false;
 	if (GEngine->IsStereoscopic3D())
 	{
-		static const auto MobileMultiViewCVar = IConsoleManager::Get().FindTConsoleVariableDataInt(TEXT("vr.MobileMultiView"));
-		const bool bUsingMobileRenderer = GetFeatureLevelShadingPath(MyWorld->Scene->GetFeatureLevel()) == EShadingPath::Mobile;
-		bRequireMultiView = (GSupportsMobileMultiView || GRHISupportsArrayIndexFromAnyShader) && bUsingMobileRenderer && (MobileMultiViewCVar && MobileMultiViewCVar->GetValueOnAnyThread() != 0);
+		const UE::StereoRenderUtils::FStereoShaderAspects Aspects(MyWorld->Scene->GetShaderPlatform());
+		bRequireMultiView = Aspects.IsMobileMultiViewEnabled();
 	}
 
 	FSceneViewFamilyContext ViewFamily(FSceneViewFamily::ConstructionValues(
