Enable virtual texture defaults to on (but still defaults off for mobile).

diff --git a/Engine/Source/Runtime/Engine/Private/Texture.cpp b/Engine/Source/Runtime/Engine/Private/Texture.cpp
--- a/Engine/Source/Runtime/Engine/Private/Texture.cpp
+++ b/Engine/Source/Runtime/Engine/Private/Texture.cpp
@@ -63,21 +63,21 @@
 
 static TAutoConsoleVariable<int32> CVarVirtualTextures(
 	TEXT("r.VirtualTextures"),
-	0,
-	TEXT("Is virtual texture streaming enabled?"),
+	1,
+	TEXT("Enable virtual texture support"),
 	ECVF_RenderThreadSafe | ECVF_ReadOnly);
 
 static TAutoConsoleVariable<int32> CVarMobileVirtualTextures(
 	TEXT("r.Mobile.VirtualTextures"),
 	0,
-	TEXT("Whether virtual texture streaming is enabled on mobile platforms. Requires r.VirtualTextures enabled as well. \n"),
+	TEXT("Enable virtual texture suppport on mobile platforms.\nAlso requires r.VirtualTextures enabled."),
 	ECVF_RenderThreadSafe | ECVF_ReadOnly
 );
 
 static TAutoConsoleVariable<int32> CVarVirtualTexturesAutoImport(
 	TEXT("r.VT.EnableAutoImport"),
-	1,
-	TEXT("Enable virtual texture on texture import"),
+	0,
+	TEXT("Import new textures as streaming virtual textures"),
 	ECVF_Default);
 
 static TAutoConsoleVariable<int32> CVarVirtualTexturesMenuRestricted(
