Remove some long deprectated functions from debug view mode helpers.
#rnx

diff --git a/Engine/Source/Runtime/Engine/Public/DebugViewModeHelpers.h b/Engine/Source/Runtime/Engine/Public/DebugViewModeHelpers.h
--- a/Engine/Source/Runtime/Engine/Public/DebugViewModeHelpers.h
+++ b/Engine/Source/Runtime/Engine/Public/DebugViewModeHelpers.h
@@ -10,12 +10,10 @@
 
 #define WITH_DEBUG_VIEW_MODES (WITH_EDITOR || !(UE_BUILD_SHIPPING || UE_BUILD_TEST))
 
-class UMaterialInterface;
 struct FSlowTask;
-struct FMeshMaterialShaderPermutationParameters;
+class UMaterialInterface;
 enum EShaderPlatform : uint16;
 namespace ERHIFeatureLevel { enum Type : int; }
-
 namespace EMaterialQualityLevel { enum Type : uint8; }
 
 /** 
@@ -42,7 +40,7 @@ enum EDebugViewShaderMode
 ENGINE_API const TCHAR* DebugViewShaderModeToString(EDebugViewShaderMode InShaderMode);
 
 #if WITH_DEBUG_VIEW_MODES
-/** Returns true debug view shaders are only ODSC compiled (in editor or cooked builds). Returns false if they are can be compiled for editor only shader maps as well as ODSC. */
+/** Returns true if debug view shaders are only ODSC compiled (in editor or cooked builds). Returns false if they are can also be compiled for editor only shader maps as well as ODSC. */
 ENGINE_API bool IsDebugViewShaderModeODSCOnly();
 /** Returns true if a specific debug view supported for a platform. Can be used as input for shader cooking. */
 ENGINE_API bool SupportDebugViewShaderMode(EDebugViewShaderMode ShaderMode, EShaderPlatform Platform);
@@ -58,15 +56,3 @@ ENGINE_API int32 GetNumActorsInWorld(UWorld* InWorld);
 ENGINE_API bool GetUsedMaterialsInWorld(UWorld* InWorld, OUT TSet<UMaterialInterface*>& OutMaterials, FSlowTask* Task);
 ENGINE_API bool CompileDebugViewModeShaders(EDebugViewShaderMode Mode, EMaterialQualityLevel::Type QualityLevel, ERHIFeatureLevel::Type FeatureLevel, TSet<UMaterialInterface*>& Materials, FSlowTask* ProgressTask);
 ENGINE_API bool WaitForShaderCompilation(const FText& Message, FSlowTask* ProgressTask);
-
-UE_DEPRECATED(4.26, "Parameters bFullRebuild and bWaitForPreviousShaders should no longer be used")
-inline bool CompileDebugViewModeShaders(EDebugViewShaderMode Mode, EMaterialQualityLevel::Type QualityLevel, ERHIFeatureLevel::Type FeatureLevel, bool bFullRebuild, bool bWaitForPreviousShaders, TSet<UMaterialInterface*>& Materials, FSlowTask* ProgressTask)
-{
-	return CompileDebugViewModeShaders(Mode, QualityLevel, FeatureLevel, Materials, ProgressTask);
-}
-
-UE_DEPRECATED(4.26, "ClearDebugViewMaterials() should no longer be called")
-inline void ClearDebugViewMaterials(UMaterialInterface*) {}
-
-UE_DEPRECATED(4.26, "UpdateDebugViewModeShaders() should no longer be called")
-inline void UpdateDebugViewModeShaders() {}
