Add UAssetUserData to USoundClass
- More efficient way of sharing data with hundreds of assets

[RN] minor, audio

diff --git a/Engine/Source/Runtime/Engine/Classes/Sound/SoundClass.h b/Engine/Source/Runtime/Engine/Classes/Sound/SoundClass.h
--- a/Engine/Source/Runtime/Engine/Classes/Sound/SoundClass.h
+++ b/Engine/Source/Runtime/Engine/Classes/Sound/SoundClass.h
@@ -7,6 +7,7 @@
 #include "AudioDynamicParameter.h"
 #include "IAudioExtensionPlugin.h"
 #include "IAudioModulation.h"
+#include "Interfaces/Interface_AssetUserData.h"
 #include "Sound/AudioOutputTarget.h"
 #include "Sound/SoundModulationDestination.h"
 #include "UObject/ObjectMacros.h"
@@ -20,8 +21,7 @@
 
 #include "SoundClass.generated.h"
 
-
-
+class UAssetUserData;
 
 USTRUCT()
 struct FSoundClassEditorData
@@ -205,7 +205,7 @@ public:
 
 
 UCLASS(config = Engine, hidecategories = Object, editinlinenew, BlueprintType, MinimalAPI)
-class USoundClass : public UObject
+class USoundClass : public UObject, public IInterface_AssetUserData
 {
 	GENERATED_UCLASS_BODY()
 
@@ -224,6 +224,10 @@ public:
 	UPROPERTY(BlueprintReadOnly, Category = General)
 	TObjectPtr<USoundClass> ParentClass;
 
+	/** Array of user data stored with the asset */
+	UPROPERTY(EditAnywhere, AdvancedDisplay, Instanced, Category = Advanced)
+	TArray<TObjectPtr<UAssetUserData>> AssetUserData;
+	
 #if WITH_EDITORONLY_DATA
 	/** EdGraph based representation of the SoundClass */
 	TObjectPtr<class UEdGraph> SoundClassGraph;
@@ -248,6 +252,13 @@ public:
 	 */
 	ENGINE_API void Interpolate( float InterpValue, FSoundClassProperties& Current, const FSoundClassProperties& Start, const FSoundClassProperties& End );
 
+	//~ Begin IInterface_AssetUserData Interface
+	ENGINE_API virtual void AddAssetUserData(UAssetUserData* InUserData) override;
+	ENGINE_API virtual void RemoveUserDataOfClass(TSubclassOf<UAssetUserData> InUserDataClass) override;
+	ENGINE_API virtual UAssetUserData* GetAssetUserDataOfClass(TSubclassOf<UAssetUserData> InUserDataClass) override;
+	ENGINE_API virtual const TArray<UAssetUserData*>* GetAssetUserDataArray() const override;
+	//~ End IInterface_AssetUserData Interface
+	
 	// Sound Class Editor functionality
 #if WITH_EDITOR
 	/** 
@@ -282,7 +293,7 @@ public:
 
 	/** Gets the sound cue graph editor implementation. */
 	static ENGINE_API TSharedPtr<ISoundClassAudioEditor> GetSoundClassAudioEditor();
-
+	
 private:
 
 	/** Ptr to interface to sound class editor operations. */
diff --git a/Engine/Source/Runtime/Engine/Private/SoundClass.cpp b/Engine/Source/Runtime/Engine/Private/SoundClass.cpp
--- a/Engine/Source/Runtime/Engine/Private/SoundClass.cpp
+++ b/Engine/Source/Runtime/Engine/Private/SoundClass.cpp
@@ -2,6 +2,7 @@
 
 
 #include "Sound/SoundClass.h"
+#include "Engine/AssetUserData.h"
 #include "Engine/Engine.h"
 #include "Sound/AudioOutputTarget.h"
 #include "UObject/UObjectIterator.h"
@@ -301,6 +302,50 @@ FString USoundClass::GetDesc( void )
 	return( FString::Printf( TEXT( "Children: %d" ), ChildClasses.Num() ) );
 }
 
+void USoundClass::AddAssetUserData(UAssetUserData* InUserData)
+{
+	if (InUserData != nullptr)
+	{
+		UAssetUserData* ExistingData = GetAssetUserDataOfClass(InUserData->GetClass());
+		if (ExistingData != nullptr)
+		{
+			AssetUserData.Remove(ExistingData);
+		}
+		AssetUserData.Add(InUserData);
+	}
+}
+
+UAssetUserData* USoundClass::GetAssetUserDataOfClass(TSubclassOf<UAssetUserData> InUserDataClass)
+{
+	for (int32 DataIdx = 0; DataIdx < AssetUserData.Num(); DataIdx++)
+	{
+		UAssetUserData* Datum = AssetUserData[DataIdx];
+		if (Datum != nullptr && Datum->IsA(InUserDataClass))
+		{
+			return Datum;
+		}
+	}
+	return nullptr;
+}
+
+void USoundClass::RemoveUserDataOfClass(TSubclassOf<UAssetUserData> InUserDataClass)
+{
+	for (int32 DataIdx = 0; DataIdx < AssetUserData.Num(); DataIdx++)
+	{
+		UAssetUserData* Datum = AssetUserData[DataIdx];
+		if (Datum != nullptr && Datum->IsA(InUserDataClass))
+		{
+			AssetUserData.RemoveAt(DataIdx);
+			return;
+		}
+	}
+}
+
+const TArray<UAssetUserData*>* USoundClass::GetAssetUserDataArray() const
+{
+	return &ToRawPtrTArrayUnsafe(AssetUserData);
+}
+
 #if WITH_EDITOR
 bool USoundClass::RecurseCheckChild( USoundClass* ChildSoundClass )
 {
@@ -377,6 +422,5 @@ TSharedPtr<ISoundClassAudioEditor> USoundClass::GetSoundClassAudioEditor()
 	return SoundClassAudioEditor;
 }
 
-
 #endif
 
