Console command sorting

When sorting the autocomplete list, give preference to commands that actually begin with what the user's typed so far.  For example, if the user enters "ai." (with the period), they'll see all of ai's "sub-commands" grouped together (e.g. "ai.debug", "ai.mass", "ai.smartobject", etc.).  This makes it consistent with the other command box, which already does this.

diff --git a/Engine/Source/Runtime/Engine/Private/UserInterface/Console.cpp b/Engine/Source/Runtime/Engine/Private/UserInterface/Console.cpp
--- a/Engine/Source/Runtime/Engine/Private/UserInterface/Console.cpp
+++ b/Engine/Source/Runtime/Engine/Private/UserInterface/Console.cpp
@@ -565,8 +565,27 @@ void UConsole::UpdateCompleteIndices()
 			}
 		}
 
-		AutoComplete.Sort();
-	}	
+		AutoComplete.Sort([this](const FAutoCompleteCommand& A, const FAutoCompleteCommand& B)
+			{
+				// Give preference to commands that actually begin with what the user's typed so far
+				if (A.Command.StartsWith(TypedStr))
+				{
+					if (!B.Command.StartsWith(TypedStr))
+					{
+						return true;
+					}
+				}
+				else
+				{
+					if (B.Command.StartsWith(TypedStr))
+					{
+						return false;
+					}
+				}
+
+				return A < B;
+			});
+	}
 }
 
 void UConsole::SetAutoCompleteFromHistory()
