[Lumen surface cache] MeshCardRepresentation::SetCardsFromBounds now make cards cover the entire mesh bounds by default. This gives better surface cache coverage. It can increase the chance of mapping data onto the wrong triangles but primitive types using this function to build cards usually already have high mismatch between their RT representations and cards. Leaning towards better coverage seems to work better.

#tests tested in editor and on a console

diff --git a/Engine/Source/Runtime/Engine/Private/MeshCardRepresentation.cpp b/Engine/Source/Runtime/Engine/Private/MeshCardRepresentation.cpp
--- a/Engine/Source/Runtime/Engine/Private/MeshCardRepresentation.cpp
+++ b/Engine/Source/Runtime/Engine/Private/MeshCardRepresentation.cpp
@@ -114,7 +114,7 @@ FVector3f MeshCardRepresentation::GetAxisAlignedDirection(uint32 AxisAlignedDire
 	return Direction;
 }
 
-void MeshCardRepresentation::SetCardsFromBounds(FMeshCardsBuildData& CardData)
+void MeshCardRepresentation::SetCardsFromBounds(FMeshCardsBuildData& CardData, bool bCardCoversHalfBounds)
 {
 	CardData.CardBuildData.SetNum(6);
 
@@ -136,9 +136,17 @@ void MeshCardRepresentation::SetCardsFromBounds(FMeshCardsBuildData& CardData)
 		OBB.Extent.Y = FMath::Abs(OBB.AxisY.Dot(MeshExtent));
 		OBB.Extent.Z = FMath::Abs(OBB.AxisZ.Dot(MeshExtent));
 
-		// Shrink and shift the card box to cover half the mesh bounds along Z and then bump it forward a bit
-		OBB.Extent.Z = OBB.Extent.Z * 0.5f + CardZOffset;
-		OBB.Origin = MeshCenter + OBB.Extent.Z * OBB.AxisZ;
+		if (bCardCoversHalfBounds)
+		{
+			// Shrink and shift the card box to cover half the mesh bounds along Z and then bump it forward a bit
+			OBB.Extent.Z = OBB.Extent.Z * 0.5f + CardZOffset;
+			OBB.Origin = MeshCenter + OBB.Extent.Z * OBB.AxisZ;
+		}
+		else
+		{
+			OBB.Extent.Z += CardZOffset;
+			OBB.Origin = MeshCenter + CardZOffset * OBB.AxisZ;
+		}
 
 		CardData.CardBuildData[DirectionIndex].OBB = OBB;
 		CardData.CardBuildData[DirectionIndex].AxisAlignedDirectionIndex = DirectionIndex;
diff --git a/Engine/Source/Runtime/Engine/Public/MeshCardRepresentation.h b/Engine/Source/Runtime/Engine/Public/MeshCardRepresentation.h
--- a/Engine/Source/Runtime/Engine/Public/MeshCardRepresentation.h
+++ b/Engine/Source/Runtime/Engine/Public/MeshCardRepresentation.h
@@ -25,7 +25,7 @@ namespace MeshCardRepresentation
 
 	// Util
 	ENGINE_API FVector3f GetAxisAlignedDirection(uint32 AxisAlignedDirectionIndex);
-	ENGINE_API void SetCardsFromBounds(FMeshCardsBuildData& CardData);
+	ENGINE_API void SetCardsFromBounds(FMeshCardsBuildData& CardData, bool bCardCoversHalfBounds = false);
 };
 
 template<typename T>
