Fix for shared linear + cpu accessible textures

diff --git a/Engine/Source/Runtime/Engine/Private/TextureDerivedData.cpp b/Engine/Source/Runtime/Engine/Private/TextureDerivedData.cpp
--- a/Engine/Source/Runtime/Engine/Private/TextureDerivedData.cpp
+++ b/Engine/Source/Runtime/Engine/Private/TextureDerivedData.cpp
@@ -976,7 +976,8 @@ static void FinalizeBuildSettingsForLayer(
 			OutSettings.TilerEvenIfNotSharedLinear = ChildTextureFormat->GetTiler();
 		}
 
-		if (FResolvedTextureEncodingSettings::Get().Project.bSharedLinearTextureEncoding)
+		if (FResolvedTextureEncodingSettings::Get().Project.bSharedLinearTextureEncoding &&
+			!OutSettings.bCPUAccessible) // CPU textures are not tiled.
 		{
 			//
 			// We want to separate out textures involved in shared linear encoding in order to facilitate
