Tool to migrate legacy subtitles to the newer Subtitles Subsystem plugin.

#virtualized

diff --git a/Engine/Source/Runtime/Engine/Public/Subtitles/SubtitlesAndClosedCaptionsDelegates.h b/Engine/Source/Runtime/Engine/Public/Subtitles/SubtitlesAndClosedCaptionsDelegates.h
--- a/Engine/Source/Runtime/Engine/Public/Subtitles/SubtitlesAndClosedCaptionsDelegates.h
+++ b/Engine/Source/Runtime/Engine/Public/Subtitles/SubtitlesAndClosedCaptionsDelegates.h
@@ -8,6 +8,8 @@
 class FCanvas;
 class UAssetUserData;
 
+// Externally-Timed subtitles must be manually added and removed with USubtitlesSubsystem::QueueSubtitle and ::StopSubtitle.
+// For the initial delay before becoming visible, use USubtitleAssetUserData::StartOffset instead of this enum.
 UENUM()
 enum class ESubtitleTiming : uint8
 {
