USoundNodeRandom crash fix. Added PostEditImport() to fixup state after alt-drag.

diff --git a/Engine/Source/Runtime/Engine/Classes/Sound/SoundNodeRandom.h b/Engine/Source/Runtime/Engine/Classes/Sound/SoundNodeRandom.h
--- a/Engine/Source/Runtime/Engine/Classes/Sound/SoundNodeRandom.h
+++ b/Engine/Source/Runtime/Engine/Classes/Sound/SoundNodeRandom.h
@@ -75,6 +75,7 @@ class USoundNodeRandom : public USoundNode
 public:
 	//~ Begin UObject Interface
 	virtual void PostLoad() override;
+	virtual void PostEditImport() override;
 	//~ End UObject Interface
 
 	//~ Begin USoundNode Interface.
diff --git a/Engine/Source/Runtime/Engine/Private/SoundNodeRandom.cpp b/Engine/Source/Runtime/Engine/Private/SoundNodeRandom.cpp
--- a/Engine/Source/Runtime/Engine/Private/SoundNodeRandom.cpp
+++ b/Engine/Source/Runtime/Engine/Private/SoundNodeRandom.cpp
@@ -95,6 +95,14 @@ void USoundNodeRandom::PostLoad()
 	FixHasBeenUsedArray();
 }
 
+void USoundNodeRandom::PostEditImport()
+{
+	Super::PostEditImport();
+	
+	FixWeightsArray();
+	FixHasBeenUsedArray();
+}
+
 int32 USoundNodeRandom::ChooseNodeIndex(FActiveSound& ActiveSound)
 {
 	int32 NodeIndex = 0;
