Transactional support for FTickableGameObject

diff --git a/Engine/Source/Runtime/Engine/Private/Tickable.cpp b/Engine/Source/Runtime/Engine/Private/Tickable.cpp
--- a/Engine/Source/Runtime/Engine/Private/Tickable.cpp
+++ b/Engine/Source/Runtime/Engine/Private/Tickable.cpp
@@ -193,7 +193,16 @@ void FTickableGameObject::TickObjects(UWorld* World, const ELevelTick LevelTickT
 						(bIsGameWorld && ((!bIsPaused && LevelTickType != LEVELTICK_TimeOnly) || (bIsPaused && TickableObject->IsTickableWhenPaused()))))
 					{
 						SCOPE_CYCLE_COUNTER_STATID(TickableObject->GetStatId());
-						TickableObject->Tick(DeltaSeconds);
+
+#if UE_WITH_REMOTE_OBJECT_HANDLE
+						static FName TransactionalWorkName(TEXT("FTickableGameObject::TickObjects"));
+						UE::RemoteExecutor::ExecuteTransactional(TransactionalWorkName, [TickableObject, DeltaSeconds]
+						{
+#endif
+							TickableObject->Tick(DeltaSeconds);
+#if UE_WITH_REMOTE_OBJECT_HANDLE
+						});
+#endif
 					}
 				}
 			}
