[Backout] - CL41802852
#rnx
Original CL Desc
-----------------------------------------------------------------
Transactional support for FTickableGameObject

diff --git a/Engine/Source/Runtime/Engine/Private/Tickable.cpp b/Engine/Source/Runtime/Engine/Private/Tickable.cpp
--- a/Engine/Source/Runtime/Engine/Private/Tickable.cpp
+++ b/Engine/Source/Runtime/Engine/Private/Tickable.cpp
@@ -193,16 +193,7 @@ void FTickableGameObject::TickObjects(UWorld* World, const ELevelTick LevelTickT
 						(bIsGameWorld && ((!bIsPaused && LevelTickType != LEVELTICK_TimeOnly) || (bIsPaused && TickableObject->IsTickableWhenPaused()))))
 					{
 						SCOPE_CYCLE_COUNTER_STATID(TickableObject->GetStatId());
-
-#if UE_WITH_REMOTE_OBJECT_HANDLE
-						static FName TransactionalWorkName(TEXT("FTickableGameObject::TickObjects"));
-						UE::RemoteExecutor::ExecuteTransactional(TransactionalWorkName, [TickableObject, DeltaSeconds]
-						{
-#endif
-							TickableObject->Tick(DeltaSeconds);
-#if UE_WITH_REMOTE_OBJECT_HANDLE
-						});
-#endif
+						TickableObject->Tick(DeltaSeconds);
 					}
 				}
 			}
