Add a deferred component move handler on the UWorld to allow scene components to request movement to be propegated later on the frame as a larger group of updates to help improve performance.
This can be enabled with s.GroupedComponentMovement.Enable 1 and s.GroupedComponentMovement.BufferSize modifies the max number of moves to apply at once.
Add support for this to MoverComponent using bUseDeferredGroupMovement flag and a set function.

diff --git a/Engine/Plugins/Experimental/Mover/Source/Mover/Private/MoverComponent.cpp b/Engine/Plugins/Experimental/Mover/Source/Mover/Private/MoverComponent.cpp
--- a/Engine/Plugins/Experimental/Mover/Source/Mover/Private/MoverComponent.cpp
+++ b/Engine/Plugins/Experimental/Mover/Source/Mover/Private/MoverComponent.cpp
@@ -787,9 +787,27 @@ void UMoverComponent::SetFrameStateFromContext(const FMoverSyncState* SyncState,
 		const FVector WorldVelocity = MoverState->GetVelocity_WorldSpace();
 
 		// Apply the desired transform to the scene component
-		FTransform Transform(WorldOrientation, WorldLocation, UpdatedComponent->GetComponentTransform().GetScale3D());
-		UpdatedComponent->SetWorldTransform(Transform, /*bSweep*/false, nullptr, ETeleportType::TeleportPhysics);
-		UpdatedComponent->ComponentVelocity = WorldVelocity;
+
+		// If we can, then we can utilize grouped movement updates to reduce the number of calls to SendPhysicsTransform
+		if (IsUsingDeferredGroupMovement())
+		{
+			// Signal to the USceneComponent that we are moving that this should be in a grouped update
+			// and not apply changes on the physics thread immediately
+			FScopedMovementUpdate MovementUpdate(
+				UpdatedComponent,
+				EScopedUpdate::DeferredGroupUpdates,
+				/*bRequireOverlapsEventFlagToQueueOverlaps*/ true);
+
+			FTransform Transform(WorldOrientation, WorldLocation, UpdatedComponent->GetComponentTransform().GetScale3D());
+			UpdatedComponent->SetWorldTransform(Transform, /*bSweep*/false, nullptr, ETeleportType::TeleportPhysics);
+			UpdatedComponent->ComponentVelocity = WorldVelocity;
+		}
+		else
+		{
+			FTransform Transform(WorldOrientation, WorldLocation, UpdatedComponent->GetComponentTransform().GetScale3D());
+			UpdatedComponent->SetWorldTransform(Transform, /*bSweep*/false, nullptr, ETeleportType::TeleportPhysics);
+			UpdatedComponent->ComponentVelocity = WorldVelocity;
+		}
 	}
 }
 
@@ -1712,7 +1730,7 @@ void UMoverComponent::SetUpdatedComponent(USceneComponent* NewUpdatedComponent)
 	if (UpdatedComponent)
 	{
 		UpdatedComponent->SetShouldUpdatePhysicsVolume(false);
-		UpdatedComponent->SetPhysicsVolume(NULL, true);
+		UpdatedComponent->SetPhysicsVolume(nullptr, true);
 		UpdatedComponent->PhysicsVolumeChangedDelegate.RemoveDynamic(this, &UMoverComponent::PhysicsVolumeChanged);
 
 		// remove from tick prerequisite
@@ -1878,6 +1896,18 @@ FTransform UMoverComponent::GetBaseVisualComponentTransform() const
 	return BaseVisualComponentTransform;
 }
 
+void UMoverComponent::SetUseDeferredGroupMovement(bool bEnable)
+{
+	bUseDeferredGroupMovement = bEnable;
+
+	// TODO update any necessary dependencies as needed
+}
+
+bool UMoverComponent::IsUsingDeferredGroupMovement() const
+{
+	return bUseDeferredGroupMovement && USceneComponent::IsGroupedComponentMovementEnabled();
+}
+
 TArray<FTrajectorySampleInfo> UMoverComponent::GetFutureTrajectory(float FutureSeconds, float SamplesPerSecond)
 {
 	FMoverPredictTrajectoryParams PredictionParams;
diff --git a/Engine/Plugins/Experimental/Mover/Source/Mover/Public/MoverComponent.h b/Engine/Plugins/Experimental/Mover/Source/Mover/Public/MoverComponent.h
--- a/Engine/Plugins/Experimental/Mover/Source/Mover/Public/MoverComponent.h
+++ b/Engine/Plugins/Experimental/Mover/Source/Mover/Public/MoverComponent.h
@@ -260,6 +260,14 @@ public:
 	UFUNCTION(BlueprintCallable, Category = Mover)
 	FTransform GetBaseVisualComponentTransform() const;
 
+	/** Sets whether this mover component can use grouped movement updates, which improve performance but can cause attachments to update later than expected */
+	UFUNCTION(BlueprintCallable, Category = Mover)
+	void SetUseDeferredGroupMovement(bool bEnable);
+
+	/** Returns true if this component is actually using grouped movement updates, which checks the flag and any global settings */
+	UFUNCTION(BlueprintCallable, Category = Mover)
+	bool IsUsingDeferredGroupMovement() const;
+	
 public:
 
 	/**
@@ -635,6 +643,19 @@ private:
 	UPROPERTY(EditDefaultsOnly, Category="Mover|Gravity")
 	bool bHasGravityOverride = false;
 
+	/**
+     * If true, then the transform updates applied in UMoverComponent::SetFrameStateFromContext will use a "deferred group move"
+     * to improve performance.
+     *
+     * It is not recommended that you enable this when you need exact, high fidelity characters such as your player character.
+     * This is mainly a benefit for scenarios with large amounts of NPCs or lower fidelity characters where it is acceptable
+     * to not have immediately applied transforms.
+     *
+     * This only does something if the "s.GroupedComponentMovement.Enable" CVar is set to true.
+     */
+    UPROPERTY(EditDefaultsOnly, Category="Mover", meta = (EditCondition = "Engine.SceneComponent.IsGroupedComponentMovementEnabled"))
+    bool bUseDeferredGroupMovement = false;
+
 	/** Transient flag indicating whether we are executing OnRegister(). */
 	bool bInOnRegister = false;
 
diff --git a/Engine/Source/Runtime/Engine/Classes/Components/PrimitiveComponent.h b/Engine/Source/Runtime/Engine/Classes/Components/PrimitiveComponent.h
--- a/Engine/Source/Runtime/Engine/Classes/Components/PrimitiveComponent.h
+++ b/Engine/Source/Runtime/Engine/Classes/Components/PrimitiveComponent.h
@@ -13,6 +13,7 @@
 #include "Templates/SubclassOf.h"
 #include "Engine/EngineTypes.h"
 #include "Engine/ScopedMovementUpdate.h"
+#include "Engine/OverlapInfo.h"
 #include "Components/SceneComponent.h"
 #include "Components/ActorPrimitiveComponentInterface.h"
 #include "RenderCommandFence.h"
@@ -1243,7 +1244,7 @@ protected:
 	/** Set of components that this component is currently overlapping. */
 	TArray<FOverlapInfo> OverlappingComponents;
 
-private:
+public:
 	/** Convert a set of overlaps from a sweep to a subset that includes only those at the end location (filling in OverlapsAtEndLocation). */
 	template<typename AllocatorType>
 	bool ConvertSweptOverlapsToCurrentOverlaps(TArray<FOverlapInfo, AllocatorType>& OutOverlapsAtEndLocation, const TOverlapArrayView& SweptOverlaps, int32 SweptOverlapsIndex, const FVector& EndLocation, const FQuat& EndRotationQuat);
@@ -1252,6 +1253,7 @@ private:
 	template<typename AllocatorType>
 	bool ConvertRotationOverlapsToCurrentOverlaps(TArray<FOverlapInfo, AllocatorType>& OutOverlapsAtEndLocation, const TOverlapArrayView& CurrentOverlaps);
 
+private:
 	template<typename AllocatorType>
 	bool GetOverlapsWithActor_Template(const AActor* Actor, TArray<FOverlapInfo, AllocatorType>& OutOverlaps) const;
 
diff --git a/Engine/Source/Runtime/Engine/Classes/Components/SceneComponent.h b/Engine/Source/Runtime/Engine/Classes/Components/SceneComponent.h
--- a/Engine/Source/Runtime/Engine/Classes/Components/SceneComponent.h
+++ b/Engine/Source/Runtime/Engine/Classes/Components/SceneComponent.h
@@ -253,7 +253,7 @@ public:
 	UFUNCTION(BlueprintSetter)
 	ENGINE_API void SetShouldUpdatePhysicsVolume(bool bInShouldUpdatePhysicsVolume);
 
-	/** If true, this component stops the the walk up the attachment chain in GetActorPositionForRenderer(). Instead this component's child will be used as the attachment root. */
+	/** If true, this component stops the walk up the attachment chain in GetActorPositionForRenderer(). Instead this component's child will be used as the attachment root. */
 	UPROPERTY()
 	uint8 bIsNotRenderAttachmentRoot : 1;
 
@@ -263,6 +263,12 @@ protected:
 
 	/** If true, OnUpdateTransform virtual will be called each time this component is moved. */
 	uint8 bWantsOnUpdateTransform:1;
+	
+	/**
+    * If true, then this scene component has requested to have some movement
+    * applied as part of a deferred group. 
+    */
+    uint8 bHasGroupedMoveThisFrame : 1 = false;
 
 private:
 	uint8 bNetUpdateTransform : 1;
@@ -309,6 +315,19 @@ public:
 	/** Returns the current scoped movement update, or NULL if there is none. @see FScopedMovementUpdate */
 	ENGINE_API FScopedMovementUpdate* GetCurrentScopedMovement() const;
 
+	/**
+	* Process a single instance of a scoped movement. This is called from the world when we are flushing the 
+	* deferred group movement queue.
+	* @see UWorld::ProcessPendingGroupMoves
+	*/
+	void ProcessDeferredMovementGroup(const struct FDeferredComponentMoveData& MoveData);
+
+	/**
+	 * @return true if grouped component movement can be used, set by s.GroupedComponentMovement.Enable
+	 */
+	UFUNCTION()
+	static ENGINE_API bool IsGroupedComponentMovementEnabled();
+
 private:
 	/** Stack of current movement scopes. */
 	TArray<FScopedMovementUpdate*> ScopedMovementStack;
@@ -993,6 +1012,8 @@ protected:
 
 	ENGINE_API bool IsDeferringMovementUpdates(const FScopedMovementUpdate& ScopedUpdate) const;
 
+	ENGINE_API bool HasGroupedMovementUpdate(const FScopedMovementUpdate& ScopedUpdate) const;
+
 public:
 	/** Call UpdateComponentToWorld if bComponentToWorldUpdated is false. */
 	ENGINE_API void ConditionalUpdateComponentToWorld();
@@ -1000,6 +1021,9 @@ public:
 	/** Returns true if movement is currently within the scope of an FScopedMovementUpdate. */
 	ENGINE_API bool IsDeferringMovementUpdates() const;
 
+	/** Returns true if this component is part of a grouped movement update that will apply later */
+	ENGINE_API bool HasGroupedMovementUpdates() const;
+
 	/** Called when AttachParent changes, to allow the scene to update its attachment state. */
 	virtual void OnAttachmentChanged() {}
 
@@ -1735,10 +1759,10 @@ private:
 
 	// This class can only be created on the stack, otherwise the ordering constraints
 	// of the constructor and destructor between encapsulated scopes could be violated.
-	void*	operator new		(size_t);
-	void*	operator new[]		(size_t);
-	void	operator delete		(void *);
-	void	operator delete[]	(void*);
+	void* operator new (size_t) = delete;
+	void* operator new[] (size_t) = delete;
+	void operator delete (void *) = delete;
+	void operator delete[] (void*) = delete;
 };
 
 FORCEINLINE_DEBUGGABLE FScopedMovementUpdate* USceneComponent::GetCurrentScopedMovement() const
@@ -1760,11 +1784,14 @@ FORCEINLINE_DEBUGGABLE bool USceneComponent::IsDeferringMovementUpdates() const
 	return false;
 }
 
-FORCEINLINE_DEBUGGABLE void USceneComponent::BeginScopedMovementUpdate(FScopedMovementUpdate& ScopedUpdate)
+FORCEINLINE_DEBUGGABLE bool USceneComponent::HasGroupedMovementUpdates() const
 {
-	checkSlow(IsInGameThread());
-	checkSlow(IsDeferringMovementUpdates(ScopedUpdate));
-	ScopedMovementStack.Push(&ScopedUpdate);
+	if (ScopedMovementStack.Num() > 0)
+	{
+		checkSlow(HasGroupedMovementUpdate(*ScopedMovementStack.Last()));
+		return true;
+	}
+	return false;
 }
 
 FORCEINLINE_DEBUGGABLE bool USceneComponent::GetShouldUpdatePhysicsVolume() const
diff --git a/Engine/Source/Runtime/Engine/Classes/Engine/World.h b/Engine/Source/Runtime/Engine/Classes/Engine/World.h
--- a/Engine/Source/Runtime/Engine/Classes/Engine/World.h
+++ b/Engine/Source/Runtime/Engine/Classes/Engine/World.h
@@ -32,6 +32,7 @@
 #include "RHIFeatureLevel.h"
 #include "Engine/WorldInitializationValues.h"
 #include "Misc/TransactionallySafeRWLock.h"
+#include "Engine/ScopedMovementUpdate.h"
 
 #include "World.generated.h"
 
@@ -859,6 +860,48 @@ private:
 	static void BroadcastWorldPartitionUninitialized(UWorld* InWorld, UWorldPartition* InWorldPartition);
 };
 
+/**
+* Represents a single "move" of a component generated by a scoped movement update.
+* Keeps track of movement data for scene components who have requested to processed
+* as part of a Deferred Movement Group. This is populated by using ScopedMovementUpdates of type
+* EScopedUpdate::DeferredGroupUpdates.
+* 
+* This is a snapshot of the state of a FScopedMovementUpdate and keep track of to use for later.
+* Populated in UWorld::PushComponentGroupMove after a DeferredGroupMove finishes on a scene component.
+*
+* The world will process these deferred moves all at once to improve game thread performance,
+* including propagating changes to the physics simulation state.
+* 
+* @see FScopedMovementUpdate
+*/
+struct FDeferredComponentMoveData
+{
+	FDeferredComponentMoveData() = default;
+	
+	/**
+	 * The component which this move data should be applied to. 
+	 */
+	TWeakObjectPtr<USceneComponent> ComponentToMove = nullptr;
+	
+	FTransform InitialTransform;
+
+	FScopedMovementUpdate::TScopedOverlapInfoArray PendingOverlaps;
+	
+	FScopedMovementUpdate::TScopedBlockingHitArray BlockingHits;
+	
+	bool bHasTransformMovedForGroupUpdate = false;
+
+	bool bHasMoved = false;
+	
+	FScopedMovementUpdate::EOverlapState OverlapState = FScopedMovementUpdate::EOverlapState::eUnknown;
+		
+	ETeleportType TeleportType = ETeleportType::None;
+
+	// If not INDEX_NONE, overlaps at this index and beyond in PendingOverlaps are at the final destination
+	int32 FinalOverlapCandidatesIndex = INDEX_NONE;
+};
+
+
 /** 
  * The World is the top level object representing a map or a sandbox in which Actors and Components will exist and be rendered.  
  *
@@ -1078,7 +1121,33 @@ public:
 	bool IsActorTickAndUserCallbacksEnabled() const { return bEnableActorTickAndUserCallbacks; };
 #endif
 
+	/**
+	 * Experimental.
+	 * 
+	 * Adds the given scene component and movement data to be applied later in the frame (ProcessPendingGroupMoves)
+	 * 
+	 * @param SceneComp The Scene component to move later 
+	 * @param Scope The scoped movement update which represents the movement state to be applied later.
+	 */
+	void PushComponentGroupMove(const FScopedMovementUpdate* Scope);
+
+	/**
+	 * Experimental.
+	 *
+	 * Processes any currently pending grouped move updates.
+	 * 
+	 * @param bProcessAllMoves if true, always process all pending moves. False for incremental partial updates
+	 */
+	void ProcessPendingGroupMoves(bool bProcessAllMoves = true);
+
 private:
+	/**
+	 * Experimental.
+	 *
+	 * Tracks data for movements which should be applied to components later on in the frame
+	 */
+	TArray<FDeferredComponentMoveData> DeferredComponentMoves;
+	
 	/** Line Batchers. Persistent types don't get flushed every frame. */
 	UPROPERTY(Transient)
 	TObjectPtr<class ULineBatchComponent> LineBatchers[(int32)ELineBatcherType::NUM];
diff --git a/Engine/Source/Runtime/Engine/Private/Components/SceneComponent.cpp b/Engine/Source/Runtime/Engine/Private/Components/SceneComponent.cpp
--- a/Engine/Source/Runtime/Engine/Private/Components/SceneComponent.cpp
+++ b/Engine/Source/Runtime/Engine/Private/Components/SceneComponent.cpp
@@ -8,6 +8,7 @@
 #include "Components/SceneComponent.h"
 #include "Engine/OverlapResult.h"
 #include "Engine/Level.h"
+#include "Engine/World.h"
 #include "EngineStats.h"
 #include "Components/StaticMeshComponent.h"
 #include "AI/NavigationSystemBase.h"
@@ -55,6 +56,15 @@ namespace SceneComponentStatics
 
 namespace SceneComponentCVars
 {
+	bool bEnableGroupedComponentMovement = false;
+	static FAutoConsoleVariableRef CVarEnableGroupedComponentMovement(
+		TEXT("s.GroupedComponentMovement.Enable"),
+		bEnableGroupedComponentMovement,
+		TEXT("Scene Component bEnableGroupedComponentMovement:\n")
+		TEXT("\tTrue: Allows the deferring of propagating component transforms when using scoped movements using EScopedUpdate::DeferredGroupUpdates.\n")
+		TEXT("\tFalse: Transforms will always be propagated immediately upon the ending of a movement scope."),
+		ECVF_Default);
+
 	bool bCheckRootComponentReplicationOnAttachedChildren = true;
 	static FAutoConsoleVariableRef CVarCheckRootComponentReplicationOnAttachedChildren(
 		TEXT("s.CheckRootComponentReplicationOnAttachedChildren"),
@@ -1035,6 +1045,11 @@ bool USceneComponent::IsDeferringMovementUpdates(const FScopedMovementUpdate& Sc
 	return ScopedUpdate.IsDeferringUpdates();
 }
 
+bool USceneComponent::HasGroupedMovementUpdate(const FScopedMovementUpdate& ScopedUpdate) const
+{
+	return ScopedUpdate.IsGroupUpdate();
+}
+
 bool USceneComponent::UpdateOverlaps(const TOverlapArrayView* PendingOverlaps /* = nullptr */, bool bDoNotifies /* = true */, const TOverlapArrayView* OverlapsAtEndLocation /* = nullptr */)
 {
 	TRACE_CPUPROFILER_EVENT_SCOPE(USceneComponent::UpdateOverlaps)
@@ -1050,6 +1065,156 @@ bool USceneComponent::UpdateOverlaps(const TOverlapArrayView* PendingOverlaps /*
 	return bSkipUpdateOverlaps;
 }
 
+static TOptional<TOverlapArrayView> GetOverlapsAtEnd(
+	UPrimitiveComponent& PrimComponent, 
+	TInlineOverlapInfoArray& OutEndOverlaps,
+	const FDeferredComponentMoveData& GroupedMoveData)
+{
+	const bool bTransformChanged = GroupedMoveData.bHasTransformMovedForGroupUpdate;
+	
+	TOptional<TOverlapArrayView> Result;
+	switch (GroupedMoveData.OverlapState)
+	{
+	case FScopedMovementUpdate::EOverlapState::eUseParent:
+		{
+			// Only rotation could have possibly changed
+			if (bTransformChanged && PrimComponent.AreSymmetricRotations(GroupedMoveData.InitialTransform.GetRotation(), PrimComponent.GetComponentQuat(), PrimComponent.GetComponentScale()))
+			{
+				if (PrimComponent.ConvertRotationOverlapsToCurrentOverlaps(OutEndOverlaps, PrimComponent.GetOverlapInfos()))
+				{
+					Result = OutEndOverlaps;
+				}
+			}
+			else
+			{
+				// Use current overlaps (unchanged)
+				Result = PrimComponent.GetOverlapInfos();
+			}
+			break;
+		}
+	case FScopedMovementUpdate::EOverlapState::eUnknown:
+	case FScopedMovementUpdate::EOverlapState::eForceUpdate:
+		{
+			break;
+		}
+	case FScopedMovementUpdate::EOverlapState::eIncludesOverlaps:
+		{
+			if (GroupedMoveData.FinalOverlapCandidatesIndex == INDEX_NONE)
+			{
+				// Overlapping nothing
+				Result = OutEndOverlaps;
+			}
+			else
+			{
+				// Fill in EndOverlaps with overlaps valid at the end location.
+				const bool bMatchingScale = FTransform::AreScale3DsEqual(GroupedMoveData.InitialTransform, PrimComponent.GetComponentTransform());
+				if (bMatchingScale)
+				{
+					const bool bHasEndOverlaps = PrimComponent.ConvertSweptOverlapsToCurrentOverlaps(
+						OutEndOverlaps, GroupedMoveData.PendingOverlaps, GroupedMoveData.FinalOverlapCandidatesIndex,
+						PrimComponent.GetComponentLocation(), PrimComponent.GetComponentQuat());
+
+					if (bHasEndOverlaps)
+					{
+						Result = OutEndOverlaps;
+					}
+				}
+			}
+			break;
+		}
+	default:
+		{
+			checkf(false, TEXT("Unknown FScopedMovementUpdate::EOverlapState value"));
+			break;
+		}
+	}
+
+	return Result;
+}
+
+void USceneComponent::ProcessDeferredMovementGroup(const FDeferredComponentMoveData& GroupedMoveData)
+{
+	ensure(IsGroupedComponentMovementEnabled());
+	
+	// If we need to apply the deferred movement group, then do so here so it all happens at once
+	if (!bHasGroupedMoveThisFrame)
+	{
+		ensureMsgf(false, TEXT("Attempting to move a deferred group when none have been applied"));
+		return;
+	}
+
+	// The expense part of moving scene components comes here, as this function will block the async physics thread
+	if (GroupedMoveData.bHasTransformMovedForGroupUpdate)
+	{
+		// Actually propagate the transform update for the grouped move now all at once
+		PropagateTransformUpdate(true, EUpdateTransformFlags::None, GroupedMoveData.TeleportType);
+	}
+
+	if (GroupedMoveData.bHasTransformMovedForGroupUpdate || GroupedMoveData.bHasMoved)
+	{
+		if (UPrimitiveComponent* PrimitiveThis = Cast<UPrimitiveComponent>(this))
+		{
+			// NOTE: UpdateOverlaps filters events to only consider overlaps where bGenerateOverlapEvents is true for both components, so it's ok if we queued up other overlaps.
+			TInlineOverlapInfoArray EndOverlaps;
+			const TOverlapArrayView PendingOverlaps(GroupedMoveData.PendingOverlaps);
+			const TOptional<TOverlapArrayView> EndOverlapsOptional = GetOverlapsAtEnd(*PrimitiveThis, EndOverlaps, GroupedMoveData);
+			UpdateOverlaps(&PendingOverlaps, true, EndOverlapsOptional.IsSet() ? &(EndOverlapsOptional.GetValue()) : nullptr);
+		}
+		else
+		{
+			UpdateOverlaps(nullptr, true, nullptr);
+		}
+	}
+
+	// Dispatch all deferred blocking hits
+	if (GroupedMoveData.BlockingHits.Num() > 0)
+	{
+		if (AActor* const Owner = GetOwner())
+		{
+			// If we have blocking hits, we must be a primitive component.
+			UPrimitiveComponent* PrimitiveThis = CastChecked<UPrimitiveComponent>(this);
+			for (const FHitResult& Hit : GroupedMoveData.BlockingHits)
+			{
+				// Overlaps may have caused us to be destroyed, as could other queued blocking hits.
+				if (!IsValid(PrimitiveThis))
+				{
+					break;
+				}
+	
+				// Collision response may change (due to overlaps or multiple blocking hits), make sure it's still considered blocking.
+				if (PrimitiveThis->GetCollisionResponseToComponent(Hit.GetComponent()) == ECR_Block)
+				{
+					PrimitiveThis->DispatchBlockingHit(*Owner, Hit);
+				}
+			}
+		}
+	}
+
+	// Reset our flag for if we have a deferred move this frame
+	bHasGroupedMoveThisFrame = false;
+}
+
+bool USceneComponent::IsGroupedComponentMovementEnabled()
+{
+	return SceneComponentCVars::bEnableGroupedComponentMovement;	
+}
+
+void USceneComponent::BeginScopedMovementUpdate(FScopedMovementUpdate& ScopedUpdate)
+{
+	checkSlow(IsInGameThread());
+	checkSlow(IsDeferringMovementUpdates(ScopedUpdate) || HasGroupedMovementUpdate(ScopedUpdate));
+
+	// What happens if your component is in a grouped update here, but gets moved elsewhere too? 	
+	if (ScopedUpdate.bIsGroupDeferedUpdate && IsGroupedComponentMovementEnabled())
+	{		
+		// TODO_BH: How do we want to handle multiple group moves per frame? 
+		// We would need to keep track of an array of deferred moves to apply at once instead of just one GroupedMoveData instance
+		bHasGroupedMoveThisFrame = true;
+	}
+
+	ScopedMovementStack.Push(&ScopedUpdate);
+}
+
 void USceneComponent::EndScopedMovementUpdate(class FScopedMovementUpdate& CompletedScope)
 {
 	SCOPE_CYCLE_COUNTER(STAT_EndScopedMovementUpdate);
@@ -1070,52 +1235,71 @@ void USceneComponent::EndScopedMovementUpdate(class FScopedMovementUpdate& Compl
 		{
 			// This was the last item on the stack, time to apply the updates if necessary
 			const bool bTransformChanged = CurrentScopedUpdate->IsTransformDirty();
-			if (bTransformChanged)
+
+			// For group updates, we want to note the current state but do all the work later on in tick
+			// If group movements are disabled, then just proceed as normal.
+			if (CurrentScopedUpdate->bIsGroupDeferedUpdate && IsGroupedComponentMovementEnabled())
 			{
-				// Pass teleport flag if set
-				PropagateTransformUpdate(true, EUpdateTransformFlags::None, CurrentScopedUpdate->TeleportType);
+				ensure(bHasGroupedMoveThisFrame == true);
+				
+				// Notify the outer world of this component's desire to be updated in a 
+				// deferred movement group, later in the frame 
+				UWorld* World = GetWorld();
+				check(World);
+
+				World->PushComponentGroupMove(CurrentScopedUpdate);					
 			}
-
-			// We may have moved somewhere and then moved back to the start, we still need to update overlaps if we touched things along the way.
-			// If no movement and no change in transform, nothing changed.
-			if (bTransformChanged || CurrentScopedUpdate->bHasMoved)
+			else
 			{
-				UPrimitiveComponent* PrimitiveThis = Cast<UPrimitiveComponent>(this);
-				if (PrimitiveThis)
+				if (bTransformChanged)
 				{
-					// NOTE: UpdateOverlaps filters events to only consider overlaps where bGenerateOverlapEvents is true for both components, so it's ok if we queued up other overlaps.
-					TInlineOverlapInfoArray EndOverlaps;
-					const TOverlapArrayView PendingOverlaps(CurrentScopedUpdate->GetPendingOverlaps());
-					const TOptional<TOverlapArrayView> EndOverlapsOptional = CurrentScopedUpdate->GetOverlapsAtEnd(*PrimitiveThis, EndOverlaps, bTransformChanged);
-					UpdateOverlaps(&PendingOverlaps, true, EndOverlapsOptional.IsSet() ? &(EndOverlapsOptional.GetValue()) : nullptr);
+					// Pass teleport flag if set
+					// Note: This call to "PropagateTransformUpdate" is what will eventually block our async physics thread, 
+					// and is what we want to avoid when doing "grouped updates"
+					PropagateTransformUpdate(true, EUpdateTransformFlags::None, CurrentScopedUpdate->TeleportType);
 				}
-				else
+
+				// We may have moved somewhere and then moved back to the start, we still need to update overlaps if we touched things along the way.
+				// If no movement and no change in transform, nothing changed.
+				if (bTransformChanged || CurrentScopedUpdate->bHasMoved)
 				{
-					UpdateOverlaps(nullptr, true, nullptr);
+					UPrimitiveComponent* PrimitiveThis = Cast<UPrimitiveComponent>(this);
+					if (PrimitiveThis)
+					{
+						// NOTE: UpdateOverlaps filters events to only consider overlaps where bGenerateOverlapEvents is true for both components, so it's ok if we queued up other overlaps.
+						TInlineOverlapInfoArray EndOverlaps;
+						const TOverlapArrayView PendingOverlaps(CurrentScopedUpdate->GetPendingOverlaps());
+						const TOptional<TOverlapArrayView> EndOverlapsOptional = CurrentScopedUpdate->GetOverlapsAtEnd(*PrimitiveThis, EndOverlaps, bTransformChanged);
+						UpdateOverlaps(&PendingOverlaps, true, EndOverlapsOptional.IsSet() ? &(EndOverlapsOptional.GetValue()) : nullptr);
+					}
+					else
+					{
+						UpdateOverlaps(nullptr, true, nullptr);
+					}
 				}
-			}
 
-			// Dispatch all deferred blocking hits
-			if (CurrentScopedUpdate->BlockingHits.Num() > 0)
-			{
-				AActor* const Owner = GetOwner();
-				if (Owner)
+				// Dispatch all deferred blocking hits
+				if (CurrentScopedUpdate->BlockingHits.Num() > 0)
 				{
-					// If we have blocking hits, we must be a primitive component.
-					UPrimitiveComponent* PrimitiveThis = CastChecked<UPrimitiveComponent>(this);
-					for (const FHitResult& Hit : CurrentScopedUpdate->BlockingHits)
+					AActor* const Owner = GetOwner();
+					if (Owner)
 					{
-						// Overlaps may have caused us to be destroyed, as could other queued blocking hits.
-						if (!IsValid(PrimitiveThis))
+						// If we have blocking hits, we must be a primitive component.
+						UPrimitiveComponent* PrimitiveThis = CastChecked<UPrimitiveComponent>(this);
+						for (const FHitResult& Hit : CurrentScopedUpdate->BlockingHits)
 						{
-							break;
-						}
+							// Overlaps may have caused us to be destroyed, as could other queued blocking hits.
+							if (!IsValid(PrimitiveThis))
+							{
+								break;
+							}
 
-						// Collision response may change (due to overlaps or multiple blocking hits), make sure it's still considered blocking.
-						if (PrimitiveThis->GetCollisionResponseToComponent(Hit.GetComponent()) == ECR_Block)
-						{
-							PrimitiveThis->DispatchBlockingHit(*Owner, Hit);
-						}						
+							// Collision response may change (due to overlaps or multiple blocking hits), make sure it's still considered blocking.
+							if (PrimitiveThis->GetCollisionResponseToComponent(Hit.GetComponent()) == ECR_Block)
+							{
+								PrimitiveThis->DispatchBlockingHit(*Owner, Hit);
+							}
+						}
 					}
 				}
 			}
diff --git a/Engine/Source/Runtime/Engine/Private/Engine/ScopedMovementUpdate.cpp b/Engine/Source/Runtime/Engine/Private/Engine/ScopedMovementUpdate.cpp
--- a/Engine/Source/Runtime/Engine/Private/Engine/ScopedMovementUpdate.cpp
+++ b/Engine/Source/Runtime/Engine/Private/Engine/ScopedMovementUpdate.cpp
@@ -16,6 +16,7 @@ FScopedMovementUpdate::FScopedMovementUpdate( class USceneComponent* Component,
 , bDeferUpdates(ScopeBehavior == EScopedUpdate::DeferredUpdates)
 , bHasMoved(false)
 , bRequireOverlapsEventFlag(bRequireOverlapsEventFlagToQueueOverlaps)
+, bIsGroupDeferedUpdate(ScopeBehavior == EScopedUpdate::DeferredGroupUpdates)
 {
 	if (IsValid(Component))
 	{
@@ -38,9 +39,23 @@ FScopedMovementUpdate::FScopedMovementUpdate( class USceneComponent* Component,
 
 				bDeferUpdates = true;
 			}
-		}			
+		}
+
+		// Do not allow for GROUP deferred updates if there is already an outer scope 
+		// which is not part of the group.
+		
+		// Safe gaurd against other possible input scopes so that you can only use grouped updates on their own, not with others
+		if (bIsGroupDeferedUpdate && OuterDeferredScope && !OuterDeferredScope->bIsGroupDeferedUpdate)
+		{
+			UE_LOG(LogScopedMovementUpdate, Warning, TEXT("FScopedMovementUpdate attempting to use group updates within an ungrouped scope, will use deferred updates instead."));
+			bIsGroupDeferedUpdate = false;
+
+			// Fallback to just using normal deferred updates 
+			bDeferUpdates = true;
+		}
 
-		if (bDeferUpdates)
+		// Probably better to have a "Begin Grouped movement update" instead to keep them separate. 
+		if (bDeferUpdates || bIsGroupDeferedUpdate)
 		{
 			Component->BeginScopedMovementUpdate(*this);
 		}
@@ -51,12 +66,13 @@ FScopedMovementUpdate::FScopedMovementUpdate( class USceneComponent* Component,
 		// Having no owner should be a no-op on all moves, but if someone requests deferred updates on a null component
 		// we should not return true for IsDeferringUpdates() so force it off here.
 		bDeferUpdates = false;
+		bIsGroupDeferedUpdate = false;
 	}
 }
 
 FScopedMovementUpdate::~FScopedMovementUpdate()
 {
-	if (bDeferUpdates && IsValid(Owner))
+	if ((bDeferUpdates || bIsGroupDeferedUpdate) && IsValid(Owner))
 	{
 		Owner->EndScopedMovementUpdate(*this);
 	}
diff --git a/Engine/Source/Runtime/Engine/Private/LevelTick.cpp b/Engine/Source/Runtime/Engine/Private/LevelTick.cpp
--- a/Engine/Source/Runtime/Engine/Private/LevelTick.cpp
+++ b/Engine/Source/Runtime/Engine/Private/LevelTick.cpp
@@ -781,6 +781,10 @@ void UWorld::RunTickGroup(ETickingGroup Group, bool bBlockTillComplete = true)
 {
 	check(TickGroup == Group); // this should already be at the correct value, but we want to make sure things are happening in the right order
 	FTickTaskManagerInterface::Get().RunTickGroup(Group, bBlockTillComplete);
+
+	// Flush any grouped updates (TODO replace this with a more generic end of tick group callback)
+	ProcessPendingGroupMoves(true);
+
 	TickGroup = ETickingGroup(TickGroup + 1); // new actors go into the next tick group because this one is already gone
 }
 
diff --git a/Engine/Source/Runtime/Engine/Private/World.cpp b/Engine/Source/Runtime/Engine/Private/World.cpp
--- a/Engine/Source/Runtime/Engine/Private/World.cpp
+++ b/Engine/Source/Runtime/Engine/Private/World.cpp
@@ -226,6 +226,14 @@ static TAutoConsoleVariable<int32> CVarPurgeEditorSceneDuringPIE(
 	TEXT("0 to keep editor scene fully initialized during PIE (default)\n")
 	TEXT("1 to purge editor scene from memory during PIE and restore when the session finishes."));
 
+static int32 GGroupedComponentMovementBufferSize = 20;
+static FAutoConsoleVariableRef CVarGroupedComponentMovementBufferSize(
+	TEXT("s.GroupedComponentMovement.BufferSize"),
+	GGroupedComponentMovementBufferSize,
+	TEXT("If s.GroupedComponentMovement.Enable is enabled, sets the max number of grouped component moves that can accumulate before it immediately flushes them.\n")
+	TEXT("It will always flush any remaining updates at the end of the current tick group"),
+	ECVF_Default);
+
 namespace UE::Private::World
 {
 	struct FStreamingLevelsToConsiderIterationScope
@@ -3348,6 +3356,67 @@ void UWorld::EnableActorTickAndUserCallbacks(bool bEnable)
 }
 #endif
 
+void UWorld::PushComponentGroupMove(const FScopedMovementUpdate* Scope)
+{
+	if (!ensure(USceneComponent::IsGroupedComponentMovementEnabled()))
+	{
+		return;
+	}
+
+	if (!Scope)
+	{
+		return;
+	}
+
+	check(IsInGameThread());
+	
+	FDeferredComponentMoveData& GroupedMoveData = DeferredComponentMoves.AddDefaulted_GetRef();
+
+	GroupedMoveData.ComponentToMove = MakeWeakObjectPtr<USceneComponent>(Scope->GetOwner());
+	GroupedMoveData.InitialTransform = Scope->GetInitialTransform();
+	GroupedMoveData.PendingOverlaps = Scope->GetPendingOverlaps();
+	GroupedMoveData.BlockingHits = Scope->GetPendingBlockingHits();
+	GroupedMoveData.bHasTransformMovedForGroupUpdate |= Scope->IsTransformDirty();
+	GroupedMoveData.bHasMoved |= Scope->GetHasMoved();
+	GroupedMoveData.OverlapState = Scope->GetOverlapState();
+	GroupedMoveData.TeleportType = Scope->GetTeleportType();
+	GroupedMoveData.FinalOverlapCandidatesIndex = Scope->GetFinalOverlapCandidatesIndex();
+
+	// Do an incremental update if needed
+	ProcessPendingGroupMoves(false);
+}
+
+void UWorld::ProcessPendingGroupMoves(bool bProcessAllMoves)
+{
+	if (bProcessAllMoves)
+	{
+		if (DeferredComponentMoves.Num() == 0)
+		{
+			return;
+		}
+	}
+	else if (DeferredComponentMoves.Num() < GGroupedComponentMovementBufferSize)
+	{
+		// Don't incrementally process until it hits the buffer size
+		return;
+	}
+
+	check(IsInGameThread());
+	ensure(USceneComponent::IsGroupedComponentMovementEnabled());
+	TRACE_CPUPROFILER_EVENT_SCOPE(ProcessPendingGroupMoves);
+
+	// Apply the movement updates of any component who has requested a move be deferred for later
+	for (const FDeferredComponentMoveData& MoveData : DeferredComponentMoves)
+	{
+		if (USceneComponent* SceneComp = MoveData.ComponentToMove.Get())
+		{
+			SceneComp->ProcessDeferredMovementGroup(MoveData);
+		}
+	}
+
+	DeferredComponentMoves.Reset();
+}
+
 void UWorld::SetLevelPendingVisibilityState(ULevel* InLevel, ELevelPendingVisibilityState InState)
 {
 	UE::TWriteScopeLock WriteLock(PendingVisibilityLock);
diff --git a/Engine/Source/Runtime/Engine/Public/Engine/ScopedMovementUpdate.h b/Engine/Source/Runtime/Engine/Public/Engine/ScopedMovementUpdate.h
--- a/Engine/Source/Runtime/Engine/Public/Engine/ScopedMovementUpdate.h
+++ b/Engine/Source/Runtime/Engine/Public/Engine/ScopedMovementUpdate.h
@@ -20,8 +20,28 @@ namespace EScopedUpdate
 {
 	enum Type
 	{
+		/** Apply changes immediately, this cannot be nested inside deferred updates */
 		ImmediateUpdates,
-		DeferredUpdates	
+
+		/** Apply changes when the last FScopedMovementUpdate on the stack ends. These can be safely nested */
+		DeferredUpdates,
+
+		/**
+		 * Experimental.
+		 * 
+		 * Defer the expensive aspects of updating a component's transform to later on the
+		 * frame instead of happening immediately upon the end of this movement scope.
+		 *
+		 * This includes:
+		 * - The call to PropagateTransformUpdate, which is where actual physics BodyInstance will be updated with
+		 *   the component's transform information.
+		 * - The updating overlaps for this component that may have occurred as a result of this component moving.
+		 * - Dispatching of any potential FHitResults (collision responses) that may have occurred as a result of
+		 *   this component moving.
+		 *
+		 * @see USceneComponent::ProcessDeferredMovementGroup
+		 */
+		DeferredGroupUpdates,
 	};
 }
 
@@ -60,6 +80,9 @@ public:
 
 	/** Return true if deferring updates, false if updates are applied immediately. */
 	ENGINE_API bool IsDeferringUpdates() const;
+
+	/** Return true if this movement update should be deferred and applied later in the frame as part of a larger group of components. */
+	ENGINE_API bool IsGroupUpdate() const;
 	
 	/** Revert movement to the initial location of the Component at the start of the scoped update. Also clears pending overlaps and sets bHasMoved to false. */
 	ENGINE_API void RevertMove();
@@ -106,6 +129,26 @@ public:
 	/** Registers that this move is a teleport */
 	ENGINE_API void SetHasTeleported(ETeleportType InTeleportType);
 
+	/** The type of teleport which this scoped move should use to determine overlaps. */
+	ENGINE_API ETeleportType GetTeleportType() const;
+	
+	/** True if this scoped movement resulted in the component being moved */
+	ENGINE_API bool GetHasMoved() const;
+	
+	/** 
+	* If not INDEX_NONE, overlaps at this index and beyond in PendingOverlaps are at the final destination
+	*/
+	ENGINE_API int32 GetFinalOverlapCandidatesIndex() const;
+	
+	/** The current overlap state of this scoped move. */
+	ENGINE_API EOverlapState GetOverlapState() const;
+	
+	/** The initial transform when this scoped movement began. */
+	ENGINE_API const FTransform& GetInitialTransform() const;
+	
+	/** The owning scene component this scoped movement is modifying. */
+	ENGINE_API USceneComponent* GetOwner() const;
+
 protected:
 	/** Fills in the list of overlaps at the end location (in EndOverlaps). Returns pointer to the list, or null if it can't be computed. */
 	ENGINE_API TOptional<TOverlapArrayView> GetOverlapsAtEnd(class UPrimitiveComponent& PrimComponent, TInlineOverlapInfoArray& OutEndOverlaps, bool bTransformChanged) const;
@@ -118,10 +161,10 @@ private:
 	
 	// This class can only be created on the stack, otherwise the ordering constraints
 	// of the constructor and destructor between encapsulated scopes could be violated.
-	ENGINE_API void*	operator new		(size_t);
-	void*	operator new[]		(size_t);
-	ENGINE_API void	operator delete		(void *);
-	void	operator delete[]	(void*);
+	ENGINE_API void* operator new (size_t) = delete;
+	void* operator new[] (size_t) = delete;
+	ENGINE_API void	operator delete (void *) = delete;
+	void operator delete[] (void*) = delete;
 
 protected:
 	USceneComponent* Owner;
@@ -142,6 +185,11 @@ protected:
 	uint8 bDeferUpdates:1;
 	uint8 bHasMoved:1;
 	uint8 bRequireOverlapsEventFlag:1;
+	
+	/**
+	* True when DeferredGroupUpdates is the type applied. See notes above.
+	*/
+	uint8 bIsGroupDeferedUpdate : 1;
 
 	friend class USceneComponent;
 };
@@ -159,6 +207,11 @@ FORCEINLINE bool FScopedMovementUpdate::IsDeferringUpdates() const
 	return bDeferUpdates;
 }
 
+FORCEINLINE bool FScopedMovementUpdate::IsGroupUpdate() const
+{
+	return bIsGroupDeferedUpdate;
+}
+
 FORCEINLINE bool FScopedMovementUpdate::HasMoved(EHasMovedTransformOption CheckTransform) const
 {
 	return bHasMoved || (CheckTransform == EHasMovedTransformOption::eTestTransform && IsTransformDirty());
@@ -214,3 +267,33 @@ FORCEINLINE_DEBUGGABLE void FScopedMovementUpdate::SetHasTeleported(ETeleportTyp
 	// Request an initialization. Teleport type can only go higher - i.e. if we have requested a reset, then a teleport will still reset fully
 	TeleportType = ((InTeleportType > TeleportType) ? InTeleportType : TeleportType); 
 }
+
+FORCEINLINE ETeleportType FScopedMovementUpdate::GetTeleportType() const
+{
+	return TeleportType;
+}
+
+FORCEINLINE bool FScopedMovementUpdate::GetHasMoved() const
+{
+	return bHasMoved;
+}
+
+FORCEINLINE int32 FScopedMovementUpdate::GetFinalOverlapCandidatesIndex() const
+{
+	return FinalOverlapCandidatesIndex;
+}
+
+FORCEINLINE FScopedMovementUpdate::EOverlapState FScopedMovementUpdate::GetOverlapState() const
+{
+	return CurrentOverlapState;
+}
+
+FORCEINLINE const FTransform& FScopedMovementUpdate::GetInitialTransform() const
+{
+	return InitialTransform;
+}
+
+FORCEINLINE USceneComponent* FScopedMovementUpdate::GetOwner() const
+{
+	return Owner;
+}
