Build fix to remove invalid DLL export
Intended for UE5

diff --git a/Engine/Source/Runtime/Engine/Public/Engine/ScopedMovementUpdate.h b/Engine/Source/Runtime/Engine/Public/Engine/ScopedMovementUpdate.h
--- a/Engine/Source/Runtime/Engine/Public/Engine/ScopedMovementUpdate.h
+++ b/Engine/Source/Runtime/Engine/Public/Engine/ScopedMovementUpdate.h
@@ -161,9 +161,9 @@ private:
 	
 	// This class can only be created on the stack, otherwise the ordering constraints
 	// of the constructor and destructor between encapsulated scopes could be violated.
-	ENGINE_API void* operator new (size_t) = delete;
+	void* operator new (size_t) = delete;
 	void* operator new[] (size_t) = delete;
-	ENGINE_API void	operator delete (void *) = delete;
+	void operator delete (void *) = delete;
 	void operator delete[] (void*) = delete;
 
 protected:
