RemoteObjects:  PostMigrate can often migrate Destroyed/Garbage objects.  We need to know the state change of the ActorComponent to properly apply the ReturnLoaned State.  We use the same strategy as AActor and use UObject Annotations to record the loaned-out state and compare it to the returned state.

#rnx

diff --git a/Engine/Source/Runtime/Engine/Private/Components/ActorComponent.cpp b/Engine/Source/Runtime/Engine/Private/Components/ActorComponent.cpp
--- a/Engine/Source/Runtime/Engine/Private/Components/ActorComponent.cpp
+++ b/Engine/Source/Runtime/Engine/Private/Components/ActorComponent.cpp
@@ -687,11 +687,40 @@ void UActorComponent::PostLoad()
 #if UE_WITH_REMOTE_OBJECT_HANDLE
 void UActorComponent::PostMigrate(const struct FUObjectMigrationContext& MigrationContext)
 {
+	// We want to be able to store "annotations" (metadata) about this Component while we loan it out, so that
+	// when we have the loan returned to us, we can check what has changed...
+	struct FLoanedState
+	{
+		bool bWasRegistered = false;
+		bool bWasDestroyed = false;
+		bool bWasPhysicsStateCreated = false;
+		bool bWasTickFunctionRegistered = false;
+		bool bWasComponentTickEnabled = false;
+
+		bool operator==(const FLoanedState&) const = default;
+		bool IsDefault() const
+		{
+			static const FLoanedState DefaultLoanedState;
+			return *this == DefaultLoanedState;
+		}
+	};
+	static FUObjectAnnotationSparse<FLoanedState,/*bAutoRemove=*/true> RemoteObjects_LoanedStates;
+	FLoanedState NewState {
+		.bWasRegistered = bRegistered, .bWasDestroyed = IsBeingDestroyed(), .bWasPhysicsStateCreated = bPhysicsStateCreated,
+		.bWasTickFunctionRegistered = bTickFunctionsRegistered, .bWasComponentTickEnabled = IsComponentTickEnabled()
+	};
+
 	if (MigrationContext.MigrationSide == EObjectMigrationSide::Send)
 	{
-		EObjectMigrationSendType SendType = MigrationContext.GetObjectMigrationSendType(this);
+		const EObjectMigrationSendType SendType = MigrationContext.GetObjectMigrationSendType(this);
+
 		const bool bKeepOwnership = (SendType == EObjectMigrationSendType::Loan);
-		if (!bKeepOwnership)
+		if (bKeepOwnership)
+		{
+			checkfSlow(!RemoteObjects_LoanedStates.GetAnnotationMap().Find(this), TEXT("%hs: Expected no existing Loaned State for %s (it should have been consumed when received, or not loaned out twice in a row)"), __func__, *GetName());
+			RemoteObjects_LoanedStates.AddAnnotation(this, NewState);
+		}
+		else // if (!bKeepOwnership)
 		{
 			// If we're transferring ownership, make sure we don't keep anything registered locally
 			if (IsRegistered())
@@ -705,7 +734,7 @@ void UActorComponent::PostMigrate(const struct FUObjectMigrationContext& Migrati
 	}
 	else if (ensureMsgf(MigrationContext.MigrationSide == EObjectMigrationSide::Receive, TEXT("Invalid MigrationContext: Side was not Send or Receive")))
 	{
-		EObjectMigrationRecvType RecvType = MigrationContext.GetObjectMigrationRecvType(this);
+		const EObjectMigrationRecvType RecvType = MigrationContext.GetObjectMigrationRecvType(this);
 
 		// We are expecting this flow from AActor, but now we must do this per-component in such a way
 		// that we are not mutating our previously deserialized data.
@@ -714,7 +743,71 @@ void UActorComponent::PostMigrate(const struct FUObjectMigrationContext& Migrati
 		// InitializeComponents();
 		// PostInitializeComponents();
 
-		if (RecvType != EObjectMigrationRecvType::ReturnedLoan)
+		if (RecvType == EObjectMigrationRecvType::ReturnedLoan)
+		{
+			FLoanedState LoanedState = RemoteObjects_LoanedStates.GetAndRemoveAnnotation(this);
+			if (NewState != LoanedState)
+			{
+				const bool bNewlyDestroyed = IsBeingDestroyed() && !LoanedState.bWasDestroyed;
+				UE_CLOG(bNewlyDestroyed, LogActorComponent, Warning, TEXT("%hs: Destroying %s back on owning server (it was destroyed while loaned-out)"), __func__, *GetName());
+
+				// Reset these variables to how they were when we loaned them out, and then go through the code path locally
+				bIsBeingDestroyed = LoanedState.bWasDestroyed;
+				bRegistered = LoanedState.bWasRegistered;
+				bPhysicsStateCreated = LoanedState.bWasPhysicsStateCreated;
+				bTickFunctionsRegistered = LoanedState.bWasTickFunctionRegistered;
+
+				// These are a chain of possibilities, with destroying covering all other cases
+				if (NewState.bWasDestroyed != LoanedState.bWasDestroyed)
+				{
+					DestroyComponent();
+				}
+				else if (NewState.bWasRegistered != LoanedState.bWasRegistered)
+				{
+					if (NewState.bWasRegistered)
+					{
+						RegisterComponent();
+					}
+					else
+					{
+						UnregisterComponent();
+					}
+				}
+				else // These are all individually set states (not chained, like the above are)
+				{
+					if (NewState.bWasPhysicsStateCreated != LoanedState.bWasPhysicsStateCreated)
+					{
+						if (NewState.bWasPhysicsStateCreated)
+						{
+							CreatePhysicsState();
+						}
+						else
+						{
+							DestroyPhysicsState();
+						}
+					}
+
+					if (NewState.bWasTickFunctionRegistered != LoanedState.bWasTickFunctionRegistered)
+					{
+						RegisterAllComponentTickFunctions(NewState.bWasTickFunctionRegistered);
+					}
+				}
+
+				// Since the tick can already be scheduled, we can enable/disable this regardless of the above states
+				if (NewState.bWasComponentTickEnabled != LoanedState.bWasComponentTickEnabled)
+				{
+					SetComponentTickEnabled(NewState.bWasComponentTickEnabled);
+				}
+
+				FLoanedState AppliedState {
+					.bWasRegistered = bRegistered, .bWasDestroyed = IsBeingDestroyed(), .bWasPhysicsStateCreated = bPhysicsStateCreated,
+					.bWasTickFunctionRegistered = bTickFunctionsRegistered, .bWasComponentTickEnabled = IsComponentTickEnabled()
+				};
+
+				ensureMsgf(AppliedState == NewState, TEXT("%hs: ReturnedLoan State Application code path was not correct for %s"), __func__, *GetName());
+			}
+		}
+		else // if (RecvType != EObjectMigrationRecvType::ReturnedLoan)
 		{
 			// If we did not previously have this Component, we need to register it (RegisterAllComponents)
 			// However, since that causes a bunch of side effects which have already happened at some point
diff --git a/Engine/Source/Runtime/Engine/Private/Components/PrimitiveComponent.cpp b/Engine/Source/Runtime/Engine/Private/Components/PrimitiveComponent.cpp
--- a/Engine/Source/Runtime/Engine/Private/Components/PrimitiveComponent.cpp
+++ b/Engine/Source/Runtime/Engine/Private/Components/PrimitiveComponent.cpp
@@ -1443,11 +1443,6 @@ void UPrimitiveComponent::PostMigrate(const struct FUObjectMigrationContext& Mig
 				{
 					UE_LOG(LogPrimitiveComponent, Error, TEXT("[%hs] Received a loaned object that destroyed our physics state, but we still have pending data to process. This should not happen..."), __func__);
 				}
-
-				// If the borrowing server destroyed the object, we will not receive a physics state to apply, but the component flag bPhysicsStateCreated will be migrated as false
-				// If so, we need to destroy our local physics state.
-				bPhysicsStateCreated = true;
-				DestroyPhysicsState();
 			}
 		}
 
@@ -1455,7 +1450,7 @@ void UPrimitiveComponent::PostMigrate(const struct FUObjectMigrationContext& Mig
 		{
 			ApplyMigratedPhysicsState();
 		}
-		else
+		else if (IsValid(this))
 		{
 			OnComponentPhysicsStateChanged.AddUniqueDynamic(this, &UPrimitiveComponent::HandlePostMigrationPhysicsStateChanged);
 		}
