Do not force static permutation update on PostEditChangeProperty for anything other that MaterialInstanceConstant

* this avoids an issue where a MIC is reloaded and the MIDs that reference it were asserting while they updated

diff --git a/Engine/Source/Runtime/Engine/Private/Materials/MaterialInstance.cpp b/Engine/Source/Runtime/Engine/Private/Materials/MaterialInstance.cpp
--- a/Engine/Source/Runtime/Engine/Private/Materials/MaterialInstance.cpp
+++ b/Engine/Source/Runtime/Engine/Private/Materials/MaterialInstance.cpp
@@ -4523,7 +4523,7 @@ void UMaterialInstance::PostEditChangeProperty(FPropertyChangedEvent& PropertyCh
 	InitResources();
 
 	// Force UpdateStaticPermutation when change type is Redirected as this probably means a Material or MaterialInstance parent asset was deleted.
-	const bool bForceStaticPermutationUpdate = PropertyChangedEvent.ChangeType == EPropertyChangeType::Redirected;
+	const bool bForceStaticPermutationUpdate = IsA(UMaterialInstanceConstant::StaticClass()) && PropertyChangedEvent.ChangeType == EPropertyChangeType::Redirected;
 	if (bForceStaticPermutationUpdate)
 	{
 		// This can run before UMaterial::PostEditChangeProperty has a chance to run, so explicitly call UpdateCachedExpressionData here
