Turns atomic UMaterialInstance::UsedByRT into two non-atomic booleans. The change is still safe and required by AutoRTFM , since it does not support handling atomic values.

#rnx

diff --git a/Engine/Source/Runtime/Engine/Private/Materials/MaterialInstance.cpp b/Engine/Source/Runtime/Engine/Private/Materials/MaterialInstance.cpp
--- a/Engine/Source/Runtime/Engine/Private/Materials/MaterialInstance.cpp
+++ b/Engine/Source/Runtime/Engine/Private/Materials/MaterialInstance.cpp
@@ -87,7 +87,7 @@ const FMaterialInstanceCachedData FMaterialInstanceCachedData::EmptyData{};
 
 void UMaterialInstance::StartCacheUniformExpressions() const
 {
-	UsedByRT |= (uint32)EMaterialInstanceUsedByRTFlag::CacheUniformExpressions;
+	bCachingUniformExpressions = true;
 #if WITH_ODSC
 	FODSCManager::RegisterMaterialInstance(this);
 #endif
@@ -95,7 +95,7 @@ void UMaterialInstance::StartCacheUniformExpressions() const
 
 void UMaterialInstance::FinishCacheUniformExpressions() const
 {
-	UsedByRT &= ~(uint32)EMaterialInstanceUsedByRTFlag::CacheUniformExpressions;
+	bCachingUniformExpressions = false;
 }
 
 void FMaterialInstanceResource::StartCacheUniformExpressions() const
@@ -811,7 +811,6 @@ bool UMaterialInstance::UpdateParameters()
 
 UMaterialInstance::UMaterialInstance(const FObjectInitializer& ObjectInitializer)
 	: Super(ObjectInitializer)
-	, UsedByRT((uint32)EMaterialInstanceUsedByRTFlag::None)
 {
 	bHasStaticPermutationResource = false;
 	bLoadedCachedData = false;
@@ -836,7 +835,7 @@ void UMaterialInstance::PostInitProperties()
 	if(!HasAnyFlags(RF_ClassDefaultObject) && FApp::CanEverRenderOrProduceRenderData())
 	{
 		Resource = new FMaterialInstanceResource(this);
-		UsedByRT |= (uint32)EMaterialInstanceUsedByRTFlag::ResourceCreate;
+		bResourceCreated = true;
 	}
 }
 
@@ -3671,15 +3670,13 @@ void UMaterialInstance::BeginDestroy()
 
 	if (Resource || ResourcesToDestroy.Num() > 0)
 	{
-		FMaterialRenderProxy* LocalResource = Resource;
-		std::atomic<uint32>* Used = &UsedByRT;
 		ENQUEUE_RENDER_COMMAND(BeginDestroyCommand)(
-		[ResourcesToDestroy = MoveTemp(ResourcesToDestroy), LocalResource, Used](FRHICommandListImmediate& RHICmdList) mutable
+		[ResourcesToDestroy = MoveTemp(ResourcesToDestroy), this](FRHICommandListImmediate& RHICmdList) mutable
 		{
-			if (LocalResource)
+			if (Resource)
 			{
-				LocalResource->MarkForGarbageCollection();
-				LocalResource->ReleaseResource();
+				Resource->MarkForGarbageCollection();
+				Resource->ReleaseResource();
 			}
 
 			for (FMaterialResource* CurrentResource : ResourcesToDestroy)
@@ -3691,15 +3688,15 @@ void UMaterialInstance::BeginDestroy()
 			ResourcesToDestroy.Empty();
 
 			// Clear flag set when Resource was created
-			*Used &= ~(uint32)EMaterialInstanceUsedByRTFlag::ResourceCreate;
+			bResourceCreated = false;
 
 			// And remove from deferred uniform expression cache queue if it's in that
-			if (*Used & (uint32)EMaterialInstanceUsedByRTFlag::CacheUniformExpressions)
+			if (bCachingUniformExpressions)
 			{
-				LocalResource->CancelCacheUniformExpressions();
-				*Used &= ~(uint32)EMaterialInstanceUsedByRTFlag::CacheUniformExpressions;
+				Resource->CancelCacheUniformExpressions();
+				bCachingUniformExpressions = false;
 			}
-		});		
+		});
 	}
 }
 
@@ -3707,7 +3704,7 @@ bool UMaterialInstance::IsReadyForFinishDestroy()
 {
 	bool bIsReady = Super::IsReadyForFinishDestroy();
 
-	return bIsReady && UsedByRT == 0;
+	return bIsReady && !bResourceCreated && !bCachingUniformExpressions;
 }
 
 void UMaterialInstance::FinishDestroy()
diff --git a/Engine/Source/Runtime/Engine/Public/Materials/MaterialInstance.h b/Engine/Source/Runtime/Engine/Public/Materials/MaterialInstance.h
--- a/Engine/Source/Runtime/Engine/Public/Materials/MaterialInstance.h
+++ b/Engine/Source/Runtime/Engine/Public/Materials/MaterialInstance.h
@@ -543,15 +543,6 @@ enum class EMaterialInstanceClearParameterFlag
 };
 ENUM_CLASS_FLAGS(EMaterialInstanceClearParameterFlag);
 
-enum class EMaterialInstanceUsedByRTFlag : uint32
-{
-	None = 0u,
-	ResourceCreate = (1u << 0),
-	CacheUniformExpressions = (1u << 1),
-
-	All = ResourceCreate | CacheUniformExpressions,
-};
-
 #if WITH_EDITORONLY_DATA
 class FMaterialInstanceParameterUpdateContext
 {
@@ -813,8 +804,12 @@ private:
 	/** Material resources being cached for cooking. */
 	TMap<const ITargetPlatform*, TArray<FMaterialResourceForCooking>> CachedMaterialResourcesForCooking;
 #endif
-	/** Thread-safe flags used to track whether this instance is still in use by the render thread (EMaterialInstanceUsedByRTFlag) */
-	mutable std::atomic<uint32> UsedByRT;
+
+	/** This instance resource has been created and is still alive (the render thread hasn't destroyed it yet). */
+	mutable bool bResourceCreated = false;
+
+	/** This instance is caching uniform expressions (the render thread hasn't canceled its creation). */
+	mutable bool bCachingUniformExpressions = false;
 	
 #if WITH_EDITOR
 	struct FTextureParameterOverride
