[Shaders] Move CVar out of shader backend as INI system is not loaded in SCW.

This needs to be a new CFLAG entry because GenerateSymbols also sets certain backend compiler options that are undesired to ensure matching shader hashes for debugging.

#rnx

diff --git a/Engine/Source/Runtime/Engine/Private/ShaderCompiler/ShaderCompiler.cpp b/Engine/Source/Runtime/Engine/Private/ShaderCompiler/ShaderCompiler.cpp
--- a/Engine/Source/Runtime/Engine/Private/ShaderCompiler/ShaderCompiler.cpp
+++ b/Engine/Source/Runtime/Engine/Private/ShaderCompiler/ShaderCompiler.cpp
@@ -513,6 +513,14 @@ static FAutoConsoleVariableRef CVarCreateShadersOnLoad(
 	GCreateShadersOnLoad,
 	TEXT("Whether to create shaders on load, which can reduce hitching, but use more memory.  Otherwise they will be created as needed."));
 
+static TAutoConsoleVariable<bool> CVarForceSpirvDebugInfo(
+	TEXT("r.ShaderCompiler.ForceSpirvDebugInfo"),
+	false,
+	TEXT("Enable SPIR-V specific debug information independently of debug and optimization compilation options.\n")
+	TEXT(" false: Disable (default)\n")
+	TEXT(" true: Enable"),
+	ECVF_ReadOnly);
+
 bool CreateShadersOnLoad()
 {
 	return GCreateShadersOnLoad != 0;
@@ -3342,6 +3350,12 @@ void GlobalBeginCompileShader(
 		SET_SHADER_DEFINE(Input.Environment, RECT_LIGHT_AS_SPOTLIGHT, 1);
 	}
 
+	static FShaderPlatformCachedIniValue<bool> ForceSpirvDebugInfoCVar(TEXT("r.ShaderCompiler.ForceSpirvDebugInfo"));
+	if (ForceSpirvDebugInfoCVar.Get((EShaderPlatform)Target.GetPlatform()))
+	{
+		Input.Environment.CompilerFlags.Add(CFLAG_ForceSpirvDebugInfo);
+	}
+
 	if (ShaderPlatform == SP_VULKAN_ES3_1_ANDROID || ShaderPlatform == SP_VULKAN_SM5_ANDROID)
 	{
 		bool bIsStripReflect = true;
diff --git a/Engine/Source/Runtime/RenderCore/Public/ShaderCompilerFlags.inl b/Engine/Source/Runtime/RenderCore/Public/ShaderCompilerFlags.inl
--- a/Engine/Source/Runtime/RenderCore/Public/ShaderCompilerFlags.inl
+++ b/Engine/Source/Runtime/RenderCore/Public/ShaderCompilerFlags.inl
@@ -105,6 +105,8 @@ SHADER_COMPILER_FLAGS_ENTRY(DisableSourceStripping)
 SHADER_COMPILER_FLAGS_ENTRY(ShaderBindingLayout)
 // Request full shader analysis artifacts in output statistics. This may contain multiple compilation steps in full text form.
 SHADER_COMPILER_FLAGS_ENTRY(OutputAnalysisArtifacts)
+// Force to generate debug info, i.e. FShaderConductorOptions::bEnableDebugInfo.
+SHADER_COMPILER_FLAGS_ENTRY(ForceSpirvDebugInfo)
 
 #undef SHADER_COMPILER_FLAGS_ENTRY
 #undef SHADER_COMPILER_FLAGS_ENTRY_DEPRECATED
