Use helper function for virtual texture lightmap support instead of doing inline cvar lookup.
#rnx

diff --git a/Engine/Plugins/Experimental/GPULightmass/Source/GPULightmass/Private/Scene/Scene.cpp b/Engine/Plugins/Experimental/GPULightmass/Source/GPULightmass/Private/Scene/Scene.cpp
--- a/Engine/Plugins/Experimental/GPULightmass/Source/GPULightmass/Private/Scene/Scene.cpp
+++ b/Engine/Plugins/Experimental/GPULightmass/Source/GPULightmass/Private/Scene/Scene.cpp
@@ -1878,8 +1878,7 @@ void FScene::ApplyFinishedLightmapsToWorld()
 {
 	UWorld* World = GPULightmass->World;
 
-	static const auto CVar = IConsoleManager::Get().FindTConsoleVariableDataInt(TEXT("r.VirtualTexturedLightmaps"));
-	const bool bUseVirtualTextures = (CVar->GetValueOnAnyThread() != 0) && UseVirtualTexturing(GetFeatureLevelShaderPlatform(World->GetFeatureLevel()));
+	const bool bUseVirtualTextures = UseVirtualTextureLightmap(GetFeatureLevelShaderPlatform(World->GetFeatureLevel()));
 
 	bool bHasSkyShadowing = LightScene.SkyLight.IsSet() && LightScene.SkyLight->CastsStationaryShadow();
 	
diff --git a/Engine/Source/Runtime/Engine/Private/InstancedStaticMesh.cpp b/Engine/Source/Runtime/Engine/Private/InstancedStaticMesh.cpp
--- a/Engine/Source/Runtime/Engine/Private/InstancedStaticMesh.cpp
+++ b/Engine/Source/Runtime/Engine/Private/InstancedStaticMesh.cpp
@@ -3112,8 +3112,7 @@ void FStaticLightingTextureMapping_InstancedStaticMesh::Serialize(FArchive& Ar)
 
 void UInstancedStaticMeshComponent::ApplyLightMapping(FStaticLightingTextureMapping_InstancedStaticMesh* InMapping, const FStaticLightingBuildContext* LightingContext)
 {
-	static const auto CVar = IConsoleManager::Get().FindTConsoleVariableDataInt(TEXT("r.VirtualTexturedLightmaps"));
-	const bool bUseVirtualTextures = (CVar->GetValueOnAnyThread() != 0) && UseVirtualTexturing(GMaxRHIShaderPlatform);
+	const bool bUseVirtualTextures = UseVirtualTextureLightmap(GMaxRHIShaderPlatform);
 
 	// restore the CachedMappings ptr (in case we're using DeferredMappings)
 	uint32 Index = (InMapping->GetLODIndex() * PerInstanceSMData.Num()) + InMapping->InstanceIndex;
diff --git a/Engine/Source/Runtime/Engine/Private/LightMap.cpp b/Engine/Source/Runtime/Engine/Private/LightMap.cpp
--- a/Engine/Source/Runtime/Engine/Private/LightMap.cpp
+++ b/Engine/Source/Runtime/Engine/Private/LightMap.cpp
@@ -1262,7 +1262,7 @@ void FLightMapPendingTexture::CreateUObjects()
 	++GLightmapCounter;
 	
 	// Only build VT lightmaps if they are enabled
-	const bool bUseVirtualTextures = (CVarVirtualTexturedLightMaps.GetValueOnAnyThread() != 0) && UseVirtualTexturing(GMaxRHIShaderPlatform);
+	const bool bUseVirtualTextures = UseVirtualTextureLightmap(GMaxRHIShaderPlatform);
 	const bool bIncludeNonVirtualTextures = !bUseVirtualTextures || (CVarIncludeNonVirtualTexturedLightMaps.GetValueOnAnyThread() != 0);
 	
 	if (bIncludeNonVirtualTextures)
@@ -2463,7 +2463,7 @@ void FLightMap2D::EncodeTextures(const FStaticLightingBuildContext* LightingCont
 #if WITH_EDITOR
 	if (bLightingSuccessful)
 	{
-		const bool bUseVirtualTextures = (CVarVirtualTexturedLightMaps.GetValueOnAnyThread() != 0) && UseVirtualTexturing(GMaxRHIShaderPlatform);
+		const bool bUseVirtualTextures = UseVirtualTextureLightmap(GMaxRHIShaderPlatform);
 		const bool bIncludeNonVirtualTextures = !bUseVirtualTextures || (CVarIncludeNonVirtualTexturedLightMaps.GetValueOnAnyThread() != 0);
 
 		GWarn->BeginSlowTask( NSLOCTEXT("LightMap2D", "BeginEncodingLightMapsTask", "Encoding light-maps"), false );
@@ -3040,7 +3040,7 @@ void FLightMap2D::Serialize(FArchive& Ar)
 
 	FLightMap::Serialize(Ar);
 
-	const bool bUsingVTLightmaps = (CVarVirtualTexturedLightMaps.GetValueOnAnyThread() != 0) && UseVirtualTexturing(GMaxRHIShaderPlatform, Ar.CookingTarget());
+	const bool bUsingVTLightmaps = UseVirtualTextureLightmap(GMaxRHIShaderPlatform, Ar.CookingTarget());
 
 	if( Ar.IsLoading() && Ar.UEVer() < VER_UE4_LOW_QUALITY_DIRECTIONAL_LIGHTMAPS )
 	{
@@ -3252,7 +3252,7 @@ FLightMapInteraction FLightMap2D::GetInteraction(ERHIFeatureLevel::Type InFeatur
 
 	int32 LightmapIndex = bHighQuality ? 0 : 1;
 
-	const bool bUseVirtualTextures = (CVarVirtualTexturedLightMaps.GetValueOnAnyThread() != 0) && UseVirtualTexturing(GetFeatureLevelShaderPlatform(InFeatureLevel));
+	const bool bUseVirtualTextures = UseVirtualTextureLightmap(GetFeatureLevelShaderPlatform(InFeatureLevel));
 	if (!bUseVirtualTextures)
 	{
 		bool bValidTextures = Textures[LightmapIndex] && Textures[LightmapIndex]->GetResource();
@@ -3282,7 +3282,7 @@ FShadowMapInteraction FLightMap2D::GetShadowInteraction(ERHIFeatureLevel::Type I
 
 	int32 LightmapIndex = bHighQuality ? 0 : 1;
 
-	const bool bUseVirtualTextures = (CVarVirtualTexturedLightMaps.GetValueOnAnyThread() != 0) && UseVirtualTexturing(GetFeatureLevelShaderPlatform(InFeatureLevel));
+	const bool bUseVirtualTextures = UseVirtualTextureLightmap(GetFeatureLevelShaderPlatform(InFeatureLevel));
 	if (bUseVirtualTextures)
 	{
 		// Preview lightmaps don't stream from disk, thus no FVirtualTexture2DResource
@@ -3453,7 +3453,7 @@ FLightmapResourceCluster::~FLightmapResourceCluster()
 
 bool FLightmapResourceCluster::GetUseVirtualTexturing() const
 {
-	return (CVarVirtualTexturedLightMaps.GetValueOnRenderThread() != 0) && UseVirtualTexturing(GetFeatureLevelShaderPlatform(GetFeatureLevel()));
+	return UseVirtualTextureLightmap(GetFeatureLevelShaderPlatform(GetFeatureLevel()));
 }
 
 // Two stage initialization of FLightmapResourceCluster
diff --git a/Engine/Source/Runtime/Engine/Private/LightmapUniformShaderParameters.cpp b/Engine/Source/Runtime/Engine/Private/LightmapUniformShaderParameters.cpp
--- a/Engine/Source/Runtime/Engine/Private/LightmapUniformShaderParameters.cpp
+++ b/Engine/Source/Runtime/Engine/Private/LightmapUniformShaderParameters.cpp
@@ -64,8 +64,7 @@ void GetPrecomputedLightingParameters(
 	const FLightCacheInterface* LCI
 )
 {
-	static const auto CVar = IConsoleManager::Get().FindTConsoleVariableDataInt(TEXT("r.VirtualTexturedLightmaps"));
-	const bool bUseVirtualTextures = (CVar->GetValueOnRenderThread() != 0) && UseVirtualTexturing(GetFeatureLevelShaderPlatform(FeatureLevel));
+	const bool bUseVirtualTextures = UseVirtualTextureLightmap(GetFeatureLevelShaderPlatform(FeatureLevel));
 
 	// TDistanceFieldShadowsAndLightMapPolicy
 	const FShadowMapInteraction ShadowMapInteraction = LCI ? LCI->GetShadowMapInteraction(FeatureLevel) : FShadowMapInteraction();
diff --git a/Engine/Source/Runtime/Engine/Private/ModelLight.cpp b/Engine/Source/Runtime/Engine/Private/ModelLight.cpp
--- a/Engine/Source/Runtime/Engine/Private/ModelLight.cpp
+++ b/Engine/Source/Runtime/Engine/Private/ModelLight.cpp
@@ -999,8 +999,7 @@ void UModel::GroupAllNodes(ULevel* Level, const TArray<ULightComponentBase*>& Li
 void UModel::ApplyStaticLighting(const FStaticLightingBuildContext* LightingContext)
 {
 #if WITH_EDITOR
-	static const auto CVar = IConsoleManager::Get().FindTConsoleVariableDataInt(TEXT("r.VirtualTexturedLightmaps"));
-	const bool bUseVirtualTextures = (CVar->GetValueOnAnyThread() != 0) && UseVirtualTexturing(GMaxRHIShaderPlatform);
+	const bool bUseVirtualTextures = UseVirtualTextureLightmap(GMaxRHIShaderPlatform);
 
 	check(CachedMappings[0]->QuantizedData);
 
diff --git a/Engine/Source/Runtime/Engine/Private/SceneManagement.cpp b/Engine/Source/Runtime/Engine/Private/SceneManagement.cpp
--- a/Engine/Source/Runtime/Engine/Private/SceneManagement.cpp
+++ b/Engine/Source/Runtime/Engine/Private/SceneManagement.cpp
@@ -1449,8 +1449,7 @@ void GetLightmapClusterResourceParameters(
 {
 	const bool bAllowHighQualityLightMaps = AllowHighQualityLightmaps(FeatureLevel);
 
-	static const auto CVar = IConsoleManager::Get().FindTConsoleVariableDataInt(TEXT("r.VirtualTexturedLightmaps"));
-	const bool bUseVirtualTextures = (CVar->GetValueOnRenderThread() != 0) && UseVirtualTexturing(GetFeatureLevelShaderPlatform(FeatureLevel));
+	const bool bUseVirtualTextures = UseVirtualTextureLightmap(GetFeatureLevelShaderPlatform(FeatureLevel));
 
 	Parameters.LightMapTexture = GBlackTexture->TextureRHI;
 	Parameters.SkyOcclusionTexture = GWhiteTexture->TextureRHI;
diff --git a/Engine/Source/Runtime/Engine/Private/StaticMeshLight.cpp b/Engine/Source/Runtime/Engine/Private/StaticMeshLight.cpp
--- a/Engine/Source/Runtime/Engine/Private/StaticMeshLight.cpp
+++ b/Engine/Source/Runtime/Engine/Private/StaticMeshLight.cpp
@@ -239,8 +239,7 @@ void FStaticMeshStaticLightingTextureMapping::Serialize(FArchive& Ar)
 // FStaticLightingTextureMapping interface
 void FStaticMeshStaticLightingTextureMapping::Apply(FQuantizedLightmapData* QuantizedData, const TMap<ULightComponent*,FShadowMapData2D*>& ShadowMapData, const FStaticLightingBuildContext* LightingContext)
 {
-	static const auto CVar = IConsoleManager::Get().FindTConsoleVariableDataInt(TEXT("r.VirtualTexturedLightmaps"));
-	const bool bUseVirtualTextures = (CVar->GetValueOnAnyThread() != 0) && UseVirtualTexturing(GMaxRHIShaderPlatform);
+	const bool bUseVirtualTextures = UseVirtualTextureLightmap(GMaxRHIShaderPlatform);
 
 	UStaticMeshComponent* StaticMeshComponent = Primitive.Get();
 
diff --git a/Engine/Source/Runtime/Landscape/Private/LandscapeLight.cpp b/Engine/Source/Runtime/Landscape/Private/LandscapeLight.cpp
--- a/Engine/Source/Runtime/Landscape/Private/LandscapeLight.cpp
+++ b/Engine/Source/Runtime/Landscape/Private/LandscapeLight.cpp
@@ -71,8 +71,7 @@ void FLandscapeStaticLightingTextureMapping::Serialize(FArchive& Ar)
 
 void FLandscapeStaticLightingTextureMapping::Apply(FQuantizedLightmapData* QuantizedData, const TMap<ULightComponent*,FShadowMapData2D*>& ShadowMapData, const FStaticLightingBuildContext* LightingContext)
 {
-	static const auto CVar = IConsoleManager::Get().FindTConsoleVariableDataInt(TEXT("r.VirtualTexturedLightmaps"));
-	const bool bUseVirtualTextures = (CVar->GetValueOnAnyThread() != 0) && UseVirtualTexturing(GMaxRHIShaderPlatform);
+	const bool bUseVirtualTextures = UseVirtualTextureLightmap(GMaxRHIShaderPlatform);
 
 	//ELightMapPaddingType PaddingType = GAllowLightmapPadding ? LMPT_NormalPadding : LMPT_NoPadding;
 	ELightMapPaddingType PaddingType = LMPT_NoPadding;
diff --git a/Engine/Source/Runtime/Renderer/Private/LightMapRendering.cpp b/Engine/Source/Runtime/Renderer/Private/LightMapRendering.cpp
--- a/Engine/Source/Runtime/Renderer/Private/LightMapRendering.cpp
+++ b/Engine/Source/Runtime/Renderer/Private/LightMapRendering.cpp
@@ -41,8 +41,7 @@ void LightMapPolicyImpl::ModifyCompilationEnvironment(ELightmapQuality LightmapQ
 	OutEnvironment.SetDefine(GLightmapDefineName[LightmapQuality], TEXT("1"));
 	OutEnvironment.SetDefine(TEXT("NUM_LIGHTMAP_COEFFICIENTS"), GNumLightmapCoefficients[LightmapQuality]);
 
-	static const auto CVar = IConsoleManager::Get().FindTConsoleVariableDataInt(TEXT("r.VirtualTexturedLightmaps"));
-	const bool VirtualTextureLightmaps = (CVar->GetValueOnAnyThread() != 0) && UseVirtualTexturing(Parameters.Platform);
+	const bool VirtualTextureLightmaps = UseVirtualTextureLightmap(Parameters.Platform);
 	OutEnvironment.SetDefine(TEXT("LIGHTMAP_VT_ENABLED"), VirtualTextureLightmaps);
 }
 
diff --git a/Engine/Source/Runtime/Renderer/Private/RayTracing/RayTracingMaterialHitShaders.cpp b/Engine/Source/Runtime/Renderer/Private/RayTracing/RayTracingMaterialHitShaders.cpp
--- a/Engine/Source/Runtime/Renderer/Private/RayTracing/RayTracingMaterialHitShaders.cpp
+++ b/Engine/Source/Runtime/Renderer/Private/RayTracing/RayTracingMaterialHitShaders.cpp
@@ -189,8 +189,7 @@ public:
 		OutEnvironment.SetDefine(TEXT("SCENE_TEXTURES_DISABLED"), 1);
 		LightMapPolicyType::ModifyCompilationEnvironment(Parameters, OutEnvironment);
 		FMeshMaterialShader::ModifyCompilationEnvironment(Parameters, OutEnvironment);
-		static const auto CVar = IConsoleManager::Get().FindTConsoleVariableDataInt(TEXT("r.VirtualTexturedLightmaps"));
-		const bool VirtualTextureLightmaps = (CVar->GetValueOnAnyThread() != 0) && UseVirtualTexturing(Parameters.Platform);
+		const bool VirtualTextureLightmaps = UseVirtualTextureLightmap(Parameters.Platform);
 		OutEnvironment.SetDefine(TEXT("LIGHTMAP_VT_ENABLED"), VirtualTextureLightmaps);
 	}
 
