Fix for crashing when opening the texture editor with a bad sdk installation.

#rnx

diff --git a/Engine/Source/Developer/TextureBuildUtilities/Private/TextureBuildUtilities.cpp b/Engine/Source/Developer/TextureBuildUtilities/Private/TextureBuildUtilities.cpp
--- a/Engine/Source/Developer/TextureBuildUtilities/Private/TextureBuildUtilities.cpp
+++ b/Engine/Source/Developer/TextureBuildUtilities/Private/TextureBuildUtilities.cpp
@@ -433,6 +433,11 @@ static int64 CalculateTextureSourceBytesFromImageInfo(const FImageInfo & InImage
 
 TEXTUREBUILDUTILITIES_API EPixelFormat GetOutputPixelFormatWithFallback(const FTextureBuildSettings& InBuildSettings, bool bInKnownAlphaFallback)
 {
+	if (InBuildSettings.BaseTextureFormat == nullptr)
+	{
+		return PF_Unknown;
+	}
+
 	bool bHasAlpha = false;
 	InBuildSettings.GetOutputAlphaFromKnownAlphaOrFallback(&bHasAlpha, bInKnownAlphaFallback);
 
diff --git a/Engine/Source/Runtime/Engine/Private/TextureDerivedData.cpp b/Engine/Source/Runtime/Engine/Private/TextureDerivedData.cpp
--- a/Engine/Source/Runtime/Engine/Private/TextureDerivedData.cpp
+++ b/Engine/Source/Runtime/Engine/Private/TextureDerivedData.cpp
@@ -1048,8 +1048,9 @@ static void FinalizeBuildSettingsForLayer(
 		}
 	}
 	
-	// this is called once first with NAME_None and then called again after Name is set up
-	if ( ! OutSettings.bVirtualStreamable && ! OutSettings.TextureFormatName.IsNone() )
+	// this is called once first with NAME_None and then called again after Name is set up. TextureFormat might also
+	// be null due to incorrect SDK configuration.
+	if ( ! OutSettings.bVirtualStreamable && TextureFormat )
 	{
 		check( LayerIndex == 0 );
 		ModifyMaxTextureResolutionBuildSettingsForPlatformLimit(Texture,TargetPlatform,TextureFormat,OutSettings);
