Override and blend blendable objects

diff --git a/Engine/Source/Runtime/Engine/Private/PostProcessUtils.cpp b/Engine/Source/Runtime/Engine/Private/PostProcessUtils.cpp
--- a/Engine/Source/Runtime/Engine/Private/PostProcessUtils.cpp
+++ b/Engine/Source/Runtime/Engine/Private/PostProcessUtils.cpp
@@ -149,14 +149,6 @@ bool FPostProcessUtils::OverridePostProcessSettings(FPostProcessSettings& ThisFr
 		UE_SET_PP(DepthOfFieldDepthBlurRadius);
 		UE_SET_PP(DepthOfFieldUseHairDepth)
 		UE_SET_PP(DepthOfFieldDepthBlurAmount);
-		if (OtherTo.bOverride_DepthOfFieldMatteBoxFlags)
-		{
-			for (uint32 Index = 0; Index < UE_ARRAY_COUNT(OtherTo.DepthOfFieldMatteBoxFlags); ++Index)
-			{
-				ThisFrom.DepthOfFieldMatteBoxFlags[Index] = OtherTo.DepthOfFieldMatteBoxFlags[Index];
-			}
-			bAnyOverwritten = true;
-		}
 		UE_SET_PP(DepthOfFieldFocalRegion);
 		UE_SET_PP(DepthOfFieldNearTransitionRegion);
 		UE_SET_PP(DepthOfFieldFarTransitionRegion);
@@ -289,6 +281,10 @@ bool FPostProcessUtils::OverridePostProcessSettings(FPostProcessSettings& ThisFr
 
 #undef UE_SET_PP
 
+	// Override blendable objects.
+	bAnyOverwritten |= (ThisFrom.WeightedBlendables.Array.Num() > 0 || OtherTo.WeightedBlendables.Array.Num() > 0);
+	ThisFrom.WeightedBlendables = OtherTo.WeightedBlendables;
+
 	return bAnyOverwritten;
 }
 
@@ -594,11 +590,38 @@ bool FPostProcessUtils::BlendPostProcessSettings(FPostProcessSettings& ThisFrom,
 		UE_SET_PP(UserFlags);
 	}
 
-	// No support for blendable objects for now.
-
 #undef UE_SET_PP
 #undef UE_LERP_PP
 
+	// Blend out old objects, blend in new objects.
+	// This uses an O(n2) lookup for now since 99% of the time we have zero blendables, and the rest
+	// of the time we have only one or two.
+	bAnyOverwritten |= (ThisFrom.WeightedBlendables.Array.Num() > 0 || OtherTo.WeightedBlendables.Array.Num() > 0);
+	for (FWeightedBlendable& ThisBlendable : ThisFrom.WeightedBlendables.Array)
+	{
+		const bool bFoundInOther = OtherTo.WeightedBlendables.Array.ContainsByPredicate(
+				[&ThisBlendable](const FWeightedBlendable& OtherBlendable)
+				{
+					return ThisBlendable.Object == OtherBlendable.Object;
+				});
+		if (!bFoundInOther)
+		{
+			ThisBlendable.Weight *= (1.0f - BlendFactor);
+		}
+	}
+	for (const FWeightedBlendable& OtherBlendable : OtherTo.WeightedBlendables.Array)
+	{
+		const bool bFoundInThis = ThisFrom.WeightedBlendables.Array.ContainsByPredicate(
+				[&OtherBlendable](FWeightedBlendable& ThisBlendable)
+				{
+					return ThisBlendable.Object == OtherBlendable.Object;
+				});
+		if (!bFoundInThis)
+		{
+			ThisFrom.WeightedBlendables.Array.Emplace(OtherBlendable.Weight * BlendFactor, OtherBlendable.Object);
+		}
+	}
+
 	return bAnyOverwritten;
 }
 
