Remove old deprecated GameplayFramework code. This only removes functions that were properly deprecated at the time

diff --git a/Engine/Plugins/Runtime/RigVM/Source/RigVMEditor/Private/Editor/Kismet/RigVMBlueprintAssetHandler.h b/Engine/Plugins/Runtime/RigVM/Source/RigVMEditor/Private/Editor/Kismet/RigVMBlueprintAssetHandler.h
--- a/Engine/Plugins/Runtime/RigVM/Source/RigVMEditor/Private/Editor/Kismet/RigVMBlueprintAssetHandler.h
+++ b/Engine/Plugins/Runtime/RigVM/Source/RigVMEditor/Private/Editor/Kismet/RigVMBlueprintAssetHandler.h
@@ -41,18 +41,6 @@ public:
 	 *@return True if the asset contains a blueprint, false otherwise
 	 */
 	virtual bool AssetContainsBlueprint(const FAssetData& InAssetData) const = 0;
-
-
-	/**
-	 * Check whether the specified asset supports nativization
-	 *
-	 *@param InAsset         The asset that is being queired for nativization support
-	 *@param InBlueprint     The blueprint that is contained within InAsset
-	 *@param OutReason       (Optional) An optional failure text to set
-	 *@return true if the specified asset supports nativization, false otherwise
-	 */
-	UE_DEPRECATED(5.0, "Blueprint Nativization has been removed as a supported feature. This API will eventually be removed.")
-	virtual bool SupportsNativization(const UObject* InAsset, const UBlueprint* InBlueprint, FText* OutReason) const { return false; }
 };
 
 
diff --git a/Engine/Source/Editor/Kismet/Public/BlueprintAssetHandler.h b/Engine/Source/Editor/Kismet/Public/BlueprintAssetHandler.h
--- a/Engine/Source/Editor/Kismet/Public/BlueprintAssetHandler.h
+++ b/Engine/Source/Editor/Kismet/Public/BlueprintAssetHandler.h
@@ -39,18 +39,6 @@ public:
 	 *@return True if the asset contains a blueprint, false otherwise
 	 */
 	virtual bool AssetContainsBlueprint(const FAssetData& InAssetData) const = 0;
-
-
-	/**
-	 * Check whether the specified asset supports nativization
-	 *
-	 *@param InAsset         The asset that is being queired for nativization support
-	 *@param InBlueprint     The blueprint that is contained within InAsset
-	 *@param OutReason       (Optional) An optional failure text to set
-	 *@return true if the specified asset supports nativization, false otherwise
-	 */
-	UE_DEPRECATED(5.0, "Blueprint Nativization has been removed as a supported feature. This API will eventually be removed.")
-	virtual bool SupportsNativization(const UObject* InAsset, const UBlueprint* InBlueprint, FText* OutReason) const { return false; }
 };
 
 
diff --git a/Engine/Source/Editor/UnrealEd/Private/CookOnTheFlyServer.cpp b/Engine/Source/Editor/UnrealEd/Private/CookOnTheFlyServer.cpp
--- a/Engine/Source/Editor/UnrealEd/Private/CookOnTheFlyServer.cpp
+++ b/Engine/Source/Editor/UnrealEd/Private/CookOnTheFlyServer.cpp
@@ -8398,15 +8398,6 @@ void UCookOnTheFlyServer::CollectFilesToCook(TArray<FName>& FilesInPath, TMap<FN
 	{
 		UE_SCOPED_HIERARCHICAL_COOKTIMER_AND_DURATION(CookModificationDelegate, DetailedCookStats::GameCookModificationDelegateTimeSec);
 
-		TArray<FString> FilesInPathStrings;
-		PRAGMA_DISABLE_DEPRECATION_WARNINGS;
-		FGameDelegates::Get().GetCookModificationDelegate().ExecuteIfBound(FilesInPathStrings);
-		PRAGMA_ENABLE_DEPRECATION_WARNINGS;
-		for (const FString& FileString : FilesInPathStrings)
-		{
-			AddFileToCook(FilesInPath, Instigators, FileString, EInstigator::CookModificationDelegate);
-		}
-
 		FModifyCookDelegate& ModifyCookDelegate = FGameDelegates::Get().GetModifyCookDelegate();
 		TArray<FName> PackagesToNeverCook;
 
diff --git a/Engine/Source/Runtime/CoreUObject/Public/UObject/ICookInfo.h b/Engine/Source/Runtime/CoreUObject/Public/UObject/ICookInfo.h
--- a/Engine/Source/Runtime/CoreUObject/Public/UObject/ICookInfo.h
+++ b/Engine/Source/Runtime/CoreUObject/Public/UObject/ICookInfo.h
@@ -86,7 +86,6 @@ namespace UE::Cook
 	CallbackMacro(InputSettingsIni, true) \
 	CallbackMacro(StartupSoftObjectPath, true) \
 	CallbackMacro(PackagingSettingsMapToCook, true) \
-	CallbackMacro(CookModificationDelegate, true) \
 	CallbackMacro(ModifyCookDelegate, true) \
 	CallbackMacro(AssetManagerModifyCook, true) \
 	CallbackMacro(AssetManagerModifyDLCCook, true) \
diff --git a/Engine/Source/Runtime/Engine/Classes/Camera/PlayerCameraManager.h b/Engine/Source/Runtime/Engine/Classes/Camera/PlayerCameraManager.h
--- a/Engine/Source/Runtime/Engine/Classes/Camera/PlayerCameraManager.h
+++ b/Engine/Source/Runtime/Engine/Classes/Camera/PlayerCameraManager.h
@@ -798,19 +798,6 @@ public:
 	UFUNCTION(BlueprintCallable, Category = "Camera")
 	ENGINE_API virtual void ClearCameraLensEffects();
 
-	//
-	// Legacy Camera Lens Effect Functions
-	//
-
-	UE_DEPRECATED(5.0, "APlayerCameraManager::FindGenericCameraLensEffect is favored now, and this function forwards to that one.")
-	ENGINE_API virtual AEmitterCameraLensEffectBase* FindCameraLensEffect(TSubclassOf<AEmitterCameraLensEffectBase> LensEffectEmitterClass);
-
-	UFUNCTION(meta = (DeprecatedFunction, DeprecationMessage = "APlayerCameraManager::AddGenericCameraLensEffect is favored now, and this function forwards to that one."))
-	ENGINE_API virtual AEmitterCameraLensEffectBase* AddCameraLensEffect(TSubclassOf<AEmitterCameraLensEffectBase> LensEffectEmitterClass);
-
-	UFUNCTION(meta = (DeprecatedFunction, DeprecationMessage = "APlayerCameraManager::RemoveGenericCameraLensEffect is favored now, and this function forwards to that one."))
-	ENGINE_API virtual void RemoveCameraLensEffect(AEmitterCameraLensEffectBase* Emitter);
-
 	//
 	// Camera Shakes.
 	//
diff --git a/Engine/Source/Runtime/Engine/Classes/Components/PrimitiveComponent.h b/Engine/Source/Runtime/Engine/Classes/Components/PrimitiveComponent.h
--- a/Engine/Source/Runtime/Engine/Classes/Components/PrimitiveComponent.h
+++ b/Engine/Source/Runtime/Engine/Classes/Components/PrimitiveComponent.h
@@ -1336,11 +1336,6 @@ public:
 	 */
 	ENGINE_API virtual bool UpdateOverlapsImpl(const TOverlapArrayView* NewPendingOverlaps=nullptr, bool bDoNotifies=true, const TOverlapArrayView* OverlapsAtEndLocation=nullptr) override;
 
-#if WITH_EDITOR
-	UE_DEPRECATED(5.2, "Use GetIgnoreBoundsForEditorFocus instead")
-	virtual bool IgnoreBoundsForEditorFocus() const { return bIgnoreBoundsForEditorFocus; }
-#endif	
-
 	/**
 	 * Whether or not the bounds of this component should be considered when focusing the editor camera to an actor with this component in it.
 	 * Useful for debug components which need a bounds for rendering but don't contribute to the visible part of the mesh in a meaningful way
@@ -2407,20 +2402,6 @@ public:
 	virtual bool SupportsWorldPositionOffsetVelocity() const { return true; }
 
 #if WITH_EDITOR
-	/**
-	 * Determines whether the supplied bounding box intersects with the component.
-	 * Used by the editor in orthographic viewports.
-	 *
-	 * @param	InSelBBox						Bounding box to test against
-	 * @param	ShowFlags						Engine ShowFlags for the viewport
-	 * @param	bConsiderOnlyBSP				If only BSP geometry should be tested
-	 * @param	bMustEncompassEntireComponent	Whether the component bounding box must lay wholly within the supplied bounding box
-	 *
-	 * @return	true if the supplied bounding box is determined to intersect the component (partially or wholly)
-	 */
-	UE_DEPRECATED(5.1, "This function is deprecated. Use the function IsShown and the overload that doesn't take an EngineShowFlags instead.")
-	ENGINE_API virtual bool ComponentIsTouchingSelectionBox(const FBox& InSelBBox, const FEngineShowFlags& ShowFlags, const bool bConsiderOnlyBSP, const bool bMustEncompassEntireComponent) const;
-
 	/**
 	 * Determines whether the supplied bounding box intersects with the component.
 	 * Used by the editor in orthographic viewports.
@@ -2433,20 +2414,6 @@ public:
 	 */
 	ENGINE_API virtual bool ComponentIsTouchingSelectionBox(const FBox& InSelBBox, const bool bConsiderOnlyBSP, const bool bMustEncompassEntireComponent) const;
 
-	/**
-	 * Determines whether the supplied frustum intersects with the component.
-	 * Used by the editor in perspective viewports.
-	 *
-	 * @param	InFrustum						Frustum to test against
-	 * @param	ShowFlags						Engine ShowFlags for the viewport
-	 * @param	bConsiderOnlyBSP				If only BSP geometry should be tested
-	 * @param	bMustEncompassEntireComponent	Whether the component bounding box must lay wholly within the supplied bounding box
-	 *
-	 * @return	true if the supplied bounding box is determined to intersect the component (partially or wholly)
-	 */
-	UE_DEPRECATED(5.1, "This function is deprecated. Use the function IsShown and the overload that doesn't take an EngineShowFlags instead.")
-	ENGINE_API virtual bool ComponentIsTouchingSelectionFrustum(const FConvexVolume& InFrustum, const FEngineShowFlags& ShowFlags, const bool bConsiderOnlyBSP, const bool bMustEncompassEntireComponent) const;
-
 	/**
 	 * Determines whether the supplied frustum intersects with the component.
 	 * Used by the editor in perspective viewports.
diff --git a/Engine/Source/Runtime/Engine/Classes/EdGraph/EdGraphPin.h b/Engine/Source/Runtime/Engine/Classes/EdGraph/EdGraphPin.h
--- a/Engine/Source/Runtime/Engine/Classes/EdGraph/EdGraphPin.h
+++ b/Engine/Source/Runtime/Engine/Classes/EdGraph/EdGraphPin.h
@@ -317,9 +317,6 @@ public:
 
 	/** If true, the default value on this pin is ignored and should not be set. */
 	uint8 bDefaultValueIsIgnored:1;
-
-	UE_DEPRECATED(5.0, "bIsDiffing is deprecated. Please use SGraphPanel::DiffResults and SGraphPanel::FocusedDiffResult instead")
-	uint32 bIsDiffing:1;
 	
 	/** If true, the pin may be hidden by user. */
 	uint8 bAdvancedView:1;
diff --git a/Engine/Source/Runtime/Engine/Classes/Engine/Blueprint.h b/Engine/Source/Runtime/Engine/Classes/Engine/Blueprint.h
--- a/Engine/Source/Runtime/Engine/Classes/Engine/Blueprint.h
+++ b/Engine/Source/Runtime/Engine/Classes/Engine/Blueprint.h
@@ -368,14 +368,6 @@ struct FBPEditorBookmarkNode
 	}
 };
 
-UENUM()
-enum class UE_DEPRECATED(5.0, "Blueprint Nativization has been removed as a supported feature. This type will eventually be removed.") EBlueprintNativizationFlag : uint8
-{
-	Disabled,
-	Dependency, // conditionally enabled (set from sub-class as a dependency)
-	ExplicitlyEnabled
-};
-
 UENUM()
 enum class EShouldCookBlueprintPropertyGuids
 {
@@ -503,13 +495,6 @@ class UBlueprint : public UBlueprintCore, public IBlueprintPropertyGuidProvider
 	ENGINE_API bool ShouldCookPropertyGuids() const;
 
 public:
-	/** When exclusive nativization is enabled, then this asset will be nativized. All super classes must be also nativized. */
-	UE_DEPRECATED(5.0, "Blueprint Nativization has been removed as a supported feature. This property will eventually be removed.")
-PRAGMA_DISABLE_DEPRECATION_WARNINGS
-	UPROPERTY(transient)
-	EBlueprintNativizationFlag NativizationFlag;
-PRAGMA_ENABLE_DEPRECATION_WARNINGS
-
 	/** The mode that will be used when compiling this class. */
 	UPROPERTY(EditAnywhere, Category=ClassOptions, AdvancedDisplay)
 	EBlueprintCompileMode CompileMode;
@@ -879,18 +864,6 @@ public:
 	 */
 	ENGINE_API virtual bool CanAlwaysRecompileWhilePlayingInEditor() const;
 
-	UE_DEPRECATED(5.0, "CanRecompileWhilePlayingInEditor was renamed to CanAlwaysRecompileWhilePlayingInEditor to better explain usage.")
-	bool CanRecompileWhilePlayingInEditor() const
-	{
-		return CanAlwaysRecompileWhilePlayingInEditor();
-	}
-
-	/**
-	 * Check whether this blueprint can be nativized or not
-	 */
-	UE_DEPRECATED(5.0, "Blueprint Nativization has been removed as a supported feature. This API will eventually be removed.")
-	virtual bool SupportsNativization(FText* OutReason = nullptr) const { return false; }
-
 private:
 
 	/** Sets the current object being debugged */
diff --git a/Engine/Source/Runtime/Engine/Classes/Engine/DataTable.h b/Engine/Source/Runtime/Engine/Classes/Engine/DataTable.h
--- a/Engine/Source/Runtime/Engine/Classes/Engine/DataTable.h
+++ b/Engine/Source/Runtime/Engine/Classes/Engine/DataTable.h
@@ -152,8 +152,6 @@ public:
 	virtual bool NeedsLoadForClient() const override { return bStripFromClientBuilds ? false : Super::NeedsLoadForClient(); }
 	virtual bool NeedsLoadForEditorGame() const override { return bStripFromClientBuilds ? false : Super::NeedsLoadForEditorGame(); }
 #if WITH_EDITORONLY_DATA
-	UE_DEPRECATED(5.1, "Class names are now represented by path names. Please use GetRowStructPathName.")
-	ENGINE_API FName GetRowStructName() const;
 	ENGINE_API FTopLevelAssetPath GetRowStructPathName() const;
 	ENGINE_API virtual void GetAssetRegistryTags(FAssetRegistryTagsContext Context) const override;
 	UE_DEPRECATED(5.4, "Implement the version that takes FAssetRegistryTagsContext instead.")
diff --git a/Engine/Source/Runtime/Engine/Classes/Engine/GameInstance.h b/Engine/Source/Runtime/Engine/Classes/Engine/GameInstance.h
--- a/Engine/Source/Runtime/Engine/Classes/Engine/GameInstance.h
+++ b/Engine/Source/Runtime/Engine/Classes/Engine/GameInstance.h
@@ -343,15 +343,6 @@ public:
 	 */
 	ENGINE_API ULocalPlayer* CreateLocalPlayer(FPlatformUserId UserId, FString& OutError, bool bSpawnPlayerController);
 
-	/**
-	 * Adds a LocalPlayer to the local and global list of Players.
-	 *
-	 * @param	NewPlayer	the player to add
-	 * @param	ControllerId id of the controller associated with the player
-	 */
-	UE_DEPRECATED(5.1, "This version of AddLocalPlayer has been deprecated, pleasse use the version that takes a FPlatformUserId instead.")
-	ENGINE_API virtual int32			AddLocalPlayer(ULocalPlayer* NewPlayer, int32 ControllerId);
-
 	/**
 	 * Adds a LocalPlayer to the local and global list of Players.
 	 *
diff --git a/Engine/Source/Runtime/Engine/Classes/Engine/Level.h b/Engine/Source/Runtime/Engine/Classes/Engine/Level.h
--- a/Engine/Source/Runtime/Engine/Classes/Engine/Level.h
+++ b/Engine/Source/Runtime/Engine/Classes/Engine/Level.h
@@ -1404,12 +1404,6 @@ public:
 	 */
 	ENGINE_API void CleanupLevelScriptBlueprint();
 
-	/**
-	 *  Returns a list of all blueprints contained within the level
-	 */
-	UE_DEPRECATED(5.0, "This function is deprecated, we only support having a single level script blueprint. Use GetLevelScriptBlueprint instead.")
-	ENGINE_API TArray<class UBlueprint*> GetLevelBlueprints() const;
-
 	/**
 	 *  Called when the level script blueprint has been successfully changed and compiled.  Handles creating an instance of the blueprint class in LevelScriptActor
 	 */
diff --git a/Engine/Source/Runtime/Engine/Classes/Engine/LocalPlayer.h b/Engine/Source/Runtime/Engine/Classes/Engine/LocalPlayer.h
--- a/Engine/Source/Runtime/Engine/Classes/Engine/LocalPlayer.h
+++ b/Engine/Source/Runtime/Engine/Classes/Engine/LocalPlayer.h
@@ -554,20 +554,6 @@ public:
 	 */
 	ENGINE_API virtual FUniqueNetIdRepl GetPreferredUniqueNetId() const;
 
-	UE_DEPRECATED(5.0, "Platform User Id now has priority over ControllerId, these are not expected to be the same")
-	ENGINE_API bool IsCachedUniqueNetIdPairedWithControllerId() const;
-
-	struct FOptionalAllottedSize
-	{
-		FVector2f Value;
-
-		ENGINE_API FOptionalAllottedSize(std::nullptr_t Empty);
-		ENGINE_API FOptionalAllottedSize(const FVector2d* InVector2D);
-		ENGINE_API FOptionalAllottedSize(const FVector2f* InVector2D);
-
-		ENGINE_API explicit operator bool() const;
-	};
-
 	/**
 	 * This function will give you two points in Pixel Space that surround the World Space box.
 	 *
@@ -579,11 +565,6 @@ public:
 	ENGINE_API bool GetPixelBoundingBox(const FBox& ActorBox, FVector2D& OutLowerLeft, FVector2D& OutUpperRight, const FVector2f* OptionalAllotedSize = nullptr);
 	static ENGINE_API bool GetPixelBoundingBox(const FSceneViewProjectionData& ProjectionData, const FBox& ActorBox, FVector2D& OutLowerLeft, FVector2D& OutUpperRight, const FVector2f* OptionalAllotedSize = nullptr);
 
-	UE_DEPRECATED(5.2, "Please use const FVector2f* directly")
-	ENGINE_API bool GetPixelBoundingBox(const FBox& ActorBox, FVector2D& OutLowerLeft, FVector2D& OutUpperRight, FOptionalAllottedSize OptionalAllotedSize);
-	UE_DEPRECATED(5.2, "Please use const FVector2f* directly")
-	static ENGINE_API bool GetPixelBoundingBox(const FSceneViewProjectionData& ProjectionData, const FBox& ActorBox, FVector2D& OutLowerLeft, FVector2D& OutUpperRight, FOptionalAllottedSize OptionalAllotedSize);
-
 	/**
 	 * This function will give you a point in Pixel Space from a World Space position
 	 *
@@ -594,11 +575,6 @@ public:
 	ENGINE_API bool GetPixelPoint(const FVector& InPoint, FVector2D& OutPoint, const FVector2f* OptionalAllotedSize = nullptr);
 	static ENGINE_API bool GetPixelPoint(const FSceneViewProjectionData& ProjectionData, const FVector& InPoint, FVector2D& OutPoint, const FVector2f* OptionalAllotedSize = nullptr);
 
-	UE_DEPRECATED(5.2, "Please use const FVector2f* directly")
-	ENGINE_API bool GetPixelPoint(const FVector& InPoint, FVector2D& OutPoint, FOptionalAllottedSize OptionalAllotedSize);
-	UE_DEPRECATED(5.2, "Please use const FVector2f* directly")
-	static ENGINE_API bool GetPixelPoint(const FSceneViewProjectionData& ProjectionData, const FVector& InPoint, FVector2D& OutPoint, FOptionalAllottedSize OptionalAllotedSize);
-
 	/**
 	 * Helper function for deriving various bits of data needed for projection
 	 *
diff --git a/Engine/Source/Runtime/Engine/Private/Components/PrimitiveComponent.cpp b/Engine/Source/Runtime/Engine/Private/Components/PrimitiveComponent.cpp
--- a/Engine/Source/Runtime/Engine/Private/Components/PrimitiveComponent.cpp
+++ b/Engine/Source/Runtime/Engine/Private/Components/PrimitiveComponent.cpp
@@ -3849,10 +3849,6 @@ bool UPrimitiveComponent::IsShown(const FEngineShowFlags& ShowFlags) const
 }
 
 #if WITH_EDITOR
-bool UPrimitiveComponent::ComponentIsTouchingSelectionBox(const FBox& InSelBBox, const FEngineShowFlags& ShowFlags, const bool bConsiderOnlyBSP, const bool bMustEncompassEntireComponent) const
-{
-	return IsShown(ShowFlags) && ComponentIsTouchingSelectionBox(InSelBBox, bConsiderOnlyBSP, bMustEncompassEntireComponent);
-}
 
 bool UPrimitiveComponent::ComponentIsTouchingSelectionBox(const FBox& InSelBBox, const bool bConsiderOnlyBSP, const bool bMustEncompassEntireComponent) const
 {
@@ -3873,11 +3869,6 @@ bool UPrimitiveComponent::ComponentIsTouchingSelectionBox(const FBox& InSelBBox,
 	return false;
 }
 
-bool UPrimitiveComponent::ComponentIsTouchingSelectionFrustum(const FConvexVolume& InFrustum, const FEngineShowFlags& ShowFlags, const bool bConsiderOnlyBSP, const bool bMustEncompassEntireComponent) const
-{
-	return IsShown(ShowFlags) && ComponentIsTouchingSelectionFrustum(InFrustum, bConsiderOnlyBSP, bMustEncompassEntireComponent);
-}
-
 bool UPrimitiveComponent::ComponentIsTouchingSelectionFrustum(const FConvexVolume& InFrustum, const bool bConsiderOnlyBSP, const bool bMustEncompassEntireComponent) const
 {
 	if (!bConsiderOnlyBSP)
diff --git a/Engine/Source/Runtime/Engine/Private/DataTable.cpp b/Engine/Source/Runtime/Engine/Private/DataTable.cpp
--- a/Engine/Source/Runtime/Engine/Private/DataTable.cpp
+++ b/Engine/Source/Runtime/Engine/Private/DataTable.cpp
@@ -336,10 +336,6 @@ void UDataTable::FinishDestroy()
 }
 
 #if WITH_EDITORONLY_DATA
-FName UDataTable::GetRowStructName() const
-{
-	return (RowStruct) ? RowStruct->GetFName() : RowStructName_DEPRECATED;
-}
 
 FTopLevelAssetPath UDataTable::GetRowStructPathName() const
 {
diff --git a/Engine/Source/Runtime/Engine/Private/EdGraph/EdGraphPin.cpp b/Engine/Source/Runtime/Engine/Private/EdGraph/EdGraphPin.cpp
--- a/Engine/Source/Runtime/Engine/Private/EdGraph/EdGraphPin.cpp
+++ b/Engine/Source/Runtime/Engine/Private/EdGraph/EdGraphPin.cpp
@@ -490,7 +490,6 @@ namespace PinHelpers
 	static FString bNotConnectableName = GET_MEMBER_NAME_STRING_CHECKED(UEdGraphPin, bNotConnectable);
 	static FString bDefaultValueIsReadOnlyName = GET_MEMBER_NAME_STRING_CHECKED(UEdGraphPin, bDefaultValueIsReadOnly);
 	static FString bDefaultValueIsIgnoredName = GET_MEMBER_NAME_STRING_CHECKED(UEdGraphPin, bDefaultValueIsIgnored);
-	// Don't need bIsDiffing's name because it is transient
 	static FString bAdvancedViewName = GET_MEMBER_NAME_STRING_CHECKED(UEdGraphPin, bAdvancedView);
 	// Don't need bDisplayAsMutableRef's name because it is transient
 	static FString bOrphanedPinName = GET_MEMBER_NAME_STRING_CHECKED(UEdGraphPin, bOrphanedPin);
@@ -1085,8 +1084,6 @@ bool UEdGraphPin::ExportTextItem(FString& ValueStr, int32 PortFlags) const
 		BoolPropCDO->ExportTextItem_Direct(ValueStr, &LocalDefaultValueIsIgnored, nullptr, nullptr, PortFlags, nullptr);
 		ValueStr += PinHelpers::ExportTextPropDelimiter;
 
-		// Intentionally not exporting bIsDiffing as it is transient
-
 		ValueStr += PinHelpers::bAdvancedViewName + TEXT("=");
 		bool LocalAdvancedView = bAdvancedView;
 		BoolPropCDO->ExportTextItem_Direct(ValueStr, &LocalAdvancedView, nullptr, nullptr, PortFlags, nullptr);
@@ -1337,7 +1334,6 @@ bool UEdGraphPin::ImportTextItem(const TCHAR*& Buffer, int32 PortFlags, class UO
 			bParseSuccess = (Buffer != nullptr);
 			bDefaultValueIsIgnored = LocalDefaultValueIsIgnored;
 		}
-		// Intentionally not importing bIsDiffing as it is transient
 		else if (PropertyToken == PinHelpers::bAdvancedViewName)
 		{
 			bool LocalAdvancedView = bAdvancedView;
diff --git a/Engine/Source/Runtime/Engine/Private/GameInstance.cpp b/Engine/Source/Runtime/Engine/Private/GameInstance.cpp
--- a/Engine/Source/Runtime/Engine/Private/GameInstance.cpp
+++ b/Engine/Source/Runtime/Engine/Private/GameInstance.cpp
@@ -929,14 +929,6 @@ ULocalPlayer* UGameInstance::CreateLocalPlayer(FPlatformUserId UserId, FString&
 	return NewPlayer;
 }
 
-int32 UGameInstance::AddLocalPlayer(ULocalPlayer* NewLocalPlayer, int32 ControllerId)
-{
-	FPlatformUserId UserId = FGenericPlatformMisc::GetPlatformUserForUserIndex(ControllerId);
-	FInputDeviceId DummyInputDevice = INPUTDEVICEID_NONE;
-	IPlatformInputDeviceMapper::Get().RemapControllerIdToPlatformUserAndDevice(ControllerId, UserId, DummyInputDevice);
-	return AddLocalPlayer(NewLocalPlayer, UserId);
-}
-
 int32 UGameInstance::AddLocalPlayer(ULocalPlayer* NewLocalPlayer, FPlatformUserId UserId)
 {
 	if (NewLocalPlayer == nullptr)
diff --git a/Engine/Source/Runtime/Engine/Private/Level.cpp b/Engine/Source/Runtime/Engine/Private/Level.cpp
--- a/Engine/Source/Runtime/Engine/Private/Level.cpp
+++ b/Engine/Source/Runtime/Engine/Private/Level.cpp
@@ -4687,22 +4687,6 @@ void ULevel::OnApplyNewLightingData(bool bLightingSuccessful)
 	}
 }
 
-TArray<UBlueprint*> ULevel::GetLevelBlueprints() const
-{
-	TArray<UBlueprint*> LevelBlueprints;
-
-	ForEachObjectWithOuter(this, [&LevelBlueprints](UObject* LevelChild)
-	{
-		UBlueprint* LevelChildBP = Cast<UBlueprint>(LevelChild);
-		if (LevelChildBP)
-		{
-			LevelBlueprints.Add(LevelChildBP);
-		}
-	}, false, RF_NoFlags, EInternalObjectFlags::Garbage);
-
-	return LevelBlueprints;
-}
-
 ULevelScriptBlueprint* ULevel::GetLevelScriptBlueprint(bool bDontCreate)
 {
 	const FString LevelScriptName = ULevelScriptBlueprint::CreateLevelScriptNameFromLevel(this);
diff --git a/Engine/Source/Runtime/Engine/Private/LocalPlayer.cpp b/Engine/Source/Runtime/Engine/Private/LocalPlayer.cpp
--- a/Engine/Source/Runtime/Engine/Private/LocalPlayer.cpp
+++ b/Engine/Source/Runtime/Engine/Private/LocalPlayer.cpp
@@ -957,42 +957,6 @@ FSceneView* ULocalPlayer::CalcSceneView( class FSceneViewFamily* ViewFamily,
 	return View;
 }
 
-ULocalPlayer::FOptionalAllottedSize::FOptionalAllottedSize(std::nullptr_t Empty)
-	: Value(-std::numeric_limits<float>::infinity(), -std::numeric_limits<float>::infinity())
-{}
-
-ULocalPlayer::FOptionalAllottedSize::FOptionalAllottedSize(const FVector2d* InVector2D)
-{
-	if (InVector2D)
-	{
-		Value.X = static_cast<float>(InVector2D->X);
-		Value.Y = static_cast<float>(InVector2D->Y);
-	}
-	else
-	{
-		Value.X = -std::numeric_limits<float>::infinity();
-		Value.Y = -std::numeric_limits<float>::infinity();
-	}
-}
-
-ULocalPlayer::FOptionalAllottedSize::FOptionalAllottedSize(const FVector2f* InVector)
-{
-	if (InVector)
-	{
-		Value = *InVector;
-	}
-	else
-	{
-		Value.X = -std::numeric_limits<float>::infinity();
-		Value.Y = -std::numeric_limits<float>::infinity();
-	}
-}
-
-ULocalPlayer::FOptionalAllottedSize::operator bool() const
-{
-	return Value.X != -std::numeric_limits<float>::infinity();
-}
-
 bool ULocalPlayer::GetPixelBoundingBox(const FBox& ActorBox, FVector2D& OutLowerLeft, FVector2D& OutUpperRight, const FVector2f* OptionalAllotedSize)
 {
 	//@TODO: CAMERA: This has issues with aspect-ratio constrained cameras
@@ -1071,30 +1035,6 @@ bool ULocalPlayer::GetPixelBoundingBox(const FSceneViewProjectionData& Projectio
 	return SuccessCount >= 2;
 }
 
-bool ULocalPlayer::GetPixelBoundingBox(const FBox& ActorBox, FVector2D& OutLowerLeft, FVector2D& OutUpperRight, FOptionalAllottedSize OptionalAllotedSize)
-{
-	if (OptionalAllotedSize)
-	{
-		return GetPixelBoundingBox(ActorBox, OutLowerLeft, OutUpperRight, &OptionalAllotedSize.Value);
-	}
-	else
-	{
-		return GetPixelBoundingBox(ActorBox, OutLowerLeft, OutUpperRight);
-	}
-}
-
-bool ULocalPlayer::GetPixelBoundingBox(const FSceneViewProjectionData& ProjectionData, const FBox& ActorBox, FVector2D& OutLowerLeft, FVector2D& OutUpperRight, FOptionalAllottedSize OptionalAllotedSize)
-{
-	if (OptionalAllotedSize)
-	{
-		return GetPixelBoundingBox(ProjectionData, ActorBox, OutLowerLeft, OutUpperRight, &OptionalAllotedSize.Value);
-	}
-	else
-	{
-		return GetPixelBoundingBox(ProjectionData, ActorBox, OutLowerLeft, OutUpperRight);
-	}
-}
-
 bool ULocalPlayer::GetPixelPoint(const FVector& InPoint, FVector2D& OutPoint, const FVector2f* OptionalAllotedSize)
 {
 	//@TODO: CAMERA: This has issues with aspect-ratio constrained cameras
@@ -1147,30 +1087,6 @@ bool ULocalPlayer::GetPixelPoint(const FSceneViewProjectionData& ProjectionData,
 	return bInFrontOfCamera;
 }
 
-bool ULocalPlayer::GetPixelPoint(const FVector& InPoint, FVector2D& OutPoint, FOptionalAllottedSize OptionalAllotedSize)
-{
-	if (OptionalAllotedSize)
-	{
-		return GetPixelPoint(InPoint, OutPoint, &OptionalAllotedSize.Value);
-	}
-	else
-	{
-		return GetPixelPoint(InPoint, OutPoint);
-	}
-}
-
-bool ULocalPlayer::GetPixelPoint(const FSceneViewProjectionData& ProjectionData, const FVector& InPoint, FVector2D& OutPoint, FOptionalAllottedSize OptionalAllotedSize)
-{
-	if (OptionalAllotedSize)
-	{
-		return GetPixelPoint(ProjectionData, InPoint, OutPoint, &OptionalAllotedSize.Value);
-	}
-	else
-	{
-		return GetPixelPoint(ProjectionData, InPoint, OutPoint);
-	}
-}
-
 bool ULocalPlayer::GetProjectionData(FViewport* Viewport, FSceneViewProjectionData& ProjectionData, int32 StereoViewIndex) const
 {
 	// If the actor
@@ -1816,13 +1732,6 @@ FUniqueNetIdRepl ULocalPlayer::GetPreferredUniqueNetId() const
 	return GetUniqueNetIdForPlatformUser();
 }
 
-bool ULocalPlayer::IsCachedUniqueNetIdPairedWithControllerId() const
-{
-	// Get the UniqueNetId that is paired with the controller
-	FUniqueNetIdRepl UniqueIdFromController = GetUniqueNetIdFromCachedControllerId();
-	return (CachedUniqueNetId == UniqueIdFromController);
-}
-
 TSharedPtr<FSlateUser> ULocalPlayer::GetSlateUser()
 {
 	return FSlateApplication::Get().GetUserFromControllerId(ControllerId);
diff --git a/Engine/Source/Runtime/Engine/Private/PlayerCameraManager.cpp b/Engine/Source/Runtime/Engine/Private/PlayerCameraManager.cpp
--- a/Engine/Source/Runtime/Engine/Private/PlayerCameraManager.cpp
+++ b/Engine/Source/Runtime/Engine/Private/PlayerCameraManager.cpp
@@ -1222,30 +1222,11 @@ void APlayerCameraManager::ClearCameraLensEffects()
 		}
 	}
 
-	// empty the array.  unnecessary, since destruction will call RemoveCameraLensEffect,
+	// empty the array.  unnecessary, since destruction will call RemoveGenericCameraLensEffect,
 	// but this gets it done in one fell swoop.
 	CameraLensEffects.Empty();
 }
 
-AEmitterCameraLensEffectBase* APlayerCameraManager::FindCameraLensEffect(TSubclassOf<AEmitterCameraLensEffectBase> LensEffectEmitterClass)
-{
-	static_assert(TIsDerivedFrom<AEmitterCameraLensEffectBase, ICameraLensEffectInterface>::IsDerived, "Unexpected: AEmitterCameraLensEffectBase does not implement ICameraLensEffectInterface! Partial engine merge?");
-	return CastChecked<AEmitterCameraLensEffectBase>(FindGenericCameraLensEffect(LensEffectEmitterClass).GetObject(), ECastCheckedType::NullAllowed);
-}
-
-AEmitterCameraLensEffectBase* APlayerCameraManager::AddCameraLensEffect(TSubclassOf<AEmitterCameraLensEffectBase> LensEffectEmitterClass)
-{
-	static_assert(TIsDerivedFrom<AEmitterCameraLensEffectBase, ICameraLensEffectInterface>::IsDerived, "Unexpected: AEmitterCameraLensEffectBase does not implement ICameraLensEffectInterface! Partial engine merge?");
-	return CastChecked<AEmitterCameraLensEffectBase>(AddGenericCameraLensEffect(LensEffectEmitterClass).GetObject(), ECastCheckedType::NullAllowed);
-}
-
-void APlayerCameraManager::RemoveCameraLensEffect(AEmitterCameraLensEffectBase* Emitter)
-{
-	static_assert(TIsDerivedFrom<AEmitterCameraLensEffectBase, ICameraLensEffectInterface>::IsDerived, "Unexpected: AEmitterCameraLensEffectBase does not implement ICameraLensEffectInterface! Partial engine merge?");
-	TScriptInterface<ICameraLensEffectInterface> LensEffect{Emitter};
-	RemoveGenericCameraLensEffect(LensEffect);
-}
-
 /** ------------------------------------------------------------
  *  Camera Shakes
  *  ------------------------------------------------------------ */
diff --git a/Engine/Source/Runtime/Engine/Public/GameDelegates.h b/Engine/Source/Runtime/Engine/Public/GameDelegates.h
--- a/Engine/Source/Runtime/Engine/Public/GameDelegates.h
+++ b/Engine/Source/Runtime/Engine/Public/GameDelegates.h
@@ -94,12 +94,6 @@ public:
 	DEFINE_GAME_DELEGATE(ExtendedSaveGameInfoDelegate);
 	DEFINE_GAME_DELEGATE(WebServerActionDelegate);	
 
-public:
-	UE_DEPRECATED(5.0, "Use AssetManager or ModifyCookDelegate instead.")
-	FCookModificationDelegate& GetCookModificationDelegate() { return CookModificationDelegate; }
-private:
-	FCookModificationDelegate CookModificationDelegate;
-
 	// Called at start of CookByTheBook
 	DEFINE_GAME_DELEGATE(ModifyCookDelegate);
 
