Fix non unity build error.
#rnx

diff --git a/Engine/Source/Runtime/Engine/Private/VT/VirtualTexture.cpp b/Engine/Source/Runtime/Engine/Private/VT/VirtualTexture.cpp
--- a/Engine/Source/Runtime/Engine/Private/VT/VirtualTexture.cpp
+++ b/Engine/Source/Runtime/Engine/Private/VT/VirtualTexture.cpp
@@ -4,11 +4,11 @@
 
 #include "EngineModule.h"
 #include "FileCache/FileCache.h"
+#include "Interfaces/ITargetPlatform.h"
 #include "RenderUtils.h"
 #include "RendererInterface.h"
 #include "Serialization/CustomVersion.h"
 #include "VT/VirtualTextureBuiltData.h"
-#include "EngineModule.h"
 
 #include UE_INLINE_GENERATED_CPP_BY_NAME(VirtualTexture)
 
