Use of virtual texture for opacity mask is no longer a material compilation error.
The renderer setting for this now only applies to importer behavior.
Naming of the property text updated to clarify this.
Also tidied up the order/wording for some other virtual texture renderer settings.

diff --git a/Engine/Source/Runtime/Engine/Classes/Engine/RendererSettings.h b/Engine/Source/Runtime/Engine/Classes/Engine/RendererSettings.h
--- a/Engine/Source/Runtime/Engine/Classes/Engine/RendererSettings.h
+++ b/Engine/Source/Runtime/Engine/Classes/Engine/RendererSettings.h
@@ -393,17 +393,10 @@
 	 */
 	UPROPERTY(config, EditAnywhere, Category = VirtualTextures, meta = (
 		ConsoleVariable = "r.VirtualTextures", DisplayName = "Enable virtual texture support",
-		ToolTip = "When enabled, Textures can be streamed using the virtual texture system. Changing this setting requires restarting the editor.",
+		ToolTip = "When enabled, textures can be streamed using the virtual texture system. Changing this setting requires restarting the editor.",
 		ConfigRestartRequired = true))
 	uint32 bVirtualTextures : 1;
 
 	UPROPERTY(config, EditAnywhere, Category = VirtualTextures, meta = (
 		EditCondition = "bVirtualTextures",
-		ConsoleVariable = "r.VT.EnableAutoImport", DisplayName = "Enable virtual texture on texture import",
-		ToolTip = "Set the 'Virtual Texture Streaming' setting for imported textures based on 'Auto Virtual Texturing Size' in the texture import settings.",
-		ConfigRestartRequired = false))
-	uint32 bVirtualTextureEnableAutoImport : 1;
-
-	UPROPERTY(config, EditAnywhere, Category = VirtualTextures, meta = (
-		EditCondition = "bVirtualTextures",
 		ConsoleVariable = "r.VirtualTexturedLightmaps", DisplayName = "Enable virtual texture lightmaps",
@@ -413,17 +406,4 @@
 
 	UPROPERTY(config, EditAnywhere, Category = VirtualTextures, meta = (
 		EditCondition = "bVirtualTextures",
-		ConsoleVariable = "r.VT.AnisotropicFiltering", DisplayName = "Enable virtual texture anisotropic filtering",
-		ToolTip = "When enabled, virtual textures will use anisotropic filtering. This adds a cost to all shaders using virtual textures. Changing this setting requires restarting the editor.",
-		ConfigRestartRequired = true))
-	uint32 bVirtualTextureAnisotropicFiltering : 1;
-
-	UPROPERTY(config, EditAnywhere, Category = VirtualTextures, meta = (
-		EditCondition = "bVirtualTextures",
-		DisplayName = "Enable virtual textures for Opacity Mask",
-		ToolTip = "Relax restriction on virtual textures contributing to Opacity Mask. In some edge cases this can lead to low resolution shadow edges."))
-		uint32 bEnableVirtualTextureOpacityMask : 1;
-
-	UPROPERTY(config, EditAnywhere, Category = VirtualTextures, meta = (
-		EditCondition = "bVirtualTextures",
 		ConsoleVariable = "r.VT.TileSize", DisplayName = "Tile size",
@@ -440,45 +420,66 @@
 
 	UPROPERTY(config, EditAnywhere, Category = VirtualTextures, meta = (
 		EditCondition = "bVirtualTextures",
-		ConsoleVariable = "r.vt.FeedbackFactor", DisplayName = "Feedback resolution factor",
-		ToolTip = "Lower factor will increase virtual texture feedback resolution which increases CPU/GPU overhead, but may decrease streaming latency, especially if materials use many virtual textures."))
-	uint32 VirtualTextureFeedbackFactor;
+		ConsoleVariable = "r.VT.AnisotropicFiltering", DisplayName = "Enable virtual texture anisotropic filtering",
+		ToolTip = "When enabled, virtual textures will use anisotropic filtering. This adds a cost to all shaders using virtual textures. Changing this setting requires restarting the editor.",
+		ConfigRestartRequired = true))
+	uint32 bVirtualTextureAnisotropicFiltering : 1;
 
 	UPROPERTY(config, EditAnywhere, Category = VirtualTextures, meta = (
 		EditCondition = "bVirtualTextures",
+		ConsoleVariable = "r.VT.EnableAutoImport", DisplayName = "Enable virtual texture on texture import",
+		ToolTip = "Enable virtual texture streaming for imported textures based on the 'Auto Virtual Texturing Size' in the texture import settings.",
+		ConfigRestartRequired = false))
+	uint32 bVirtualTextureEnableAutoImport : 1;
+
+	UPROPERTY(config, EditAnywhere, Category = VirtualTextures, meta = (
+		EditCondition = "bVirtualTextures",
+		DisplayName = "Enable virtual texture for opacity mask on material import.",
+		ToolTip = "When virtual textures contribute to opacity mask it can sometimes lead to low resolution shadow edges."))
+	uint32 bEnableVirtualTextureOpacityMask : 1;
+
+	UPROPERTY(config, EditAnywhere, Category = MeshPaintVirtualTextures, meta = (
+		EditCondition = "bVirtualTextures",
+		ConsoleVariable = "r.MeshPaintVirtualTexture.Support", DisplayName = "Enable virtual texture mesh painting",
+		ToolTip = "When enabled, mesh paint virtual textures are supported. Changing this setting requires restarting the editor.",
+		ConfigRestartRequired = true))
+	uint32 bMeshPaintVirtualTextureSupport : 1;
+		
+	UPROPERTY(config, EditAnywhere, Category = MeshPaintVirtualTextures, meta = (
+		EditCondition = "bVirtualTextures && bMeshPaintVirtualTextureSupport",
 		ConsoleVariable = "r.MeshPaintVirtualTexture.TileSize", DisplayName = "Mesh paint tile size",
 		ToolTip = "Size in pixels for mesh paint virtual texture tiles. Will be rounded to next power-of-2. Changing this setting requires restarting the editor.",
 		ConfigRestartRequired = true))
 	uint32 MeshPaintVirtualTextureTileSize;
 
-	UPROPERTY(config, EditAnywhere, Category = VirtualTextures, meta = (
-		EditCondition = "bVirtualTextures",
+	UPROPERTY(config, EditAnywhere, Category = MeshPaintVirtualTextures, meta = (
+		EditCondition = "bVirtualTextures && bMeshPaintVirtualTextureSupport",
 		ConsoleVariable = "r.MeshPaintVirtualTexture.TileBorderSize", DisplayName = "Mesh paint tile border size",
 		ToolTip = "Size in pixels for virtual texture tile borders. Will be rounded next to multiple-of-2. Changing this setting requires restarting the editor.",
 		ConfigRestartRequired = true))
 	uint32 MeshPaintVirtualTextureTileBorderSize;
 
-	UPROPERTY(config, EditAnywhere, Category = VirtualTextures, meta = (
-		EditCondition = "bVirtualTextures",
+	UPROPERTY(config, EditAnywhere, Category = MeshPaintVirtualTextures, meta = (
+		EditCondition = "bVirtualTextures && bMeshPaintVirtualTextureSupport",
 		ConsoleVariable = "r.MeshPaintVirtualTexture.UseCompression", DisplayName = "Mesh paint use compressed textures",
 		ToolTip = "Whether to use a compressed texture format for storing mesh paint textures.",
 		ConfigRestartRequired = true))
 	bool MeshPaintVirtualTextureUseCompression;
 
-	UPROPERTY(config, EditAnywhere, Category = VirtualTextures, meta = (
-		EditCondition = "bVirtualTextures",
+	UPROPERTY(config, EditAnywhere, Category = MeshPaintVirtualTextures, meta = (
+		EditCondition = "bVirtualTextures && bMeshPaintVirtualTextureSupport",
 		ConsoleVariable = "r.StaticMesh.DefaultMeshPaintTextureSupport", DisplayName = "Mesh paint default static mesh support",
 		ToolTip = "Default setting for whether static mesh assets support mesh paint textures."))
 	bool MeshPaintDefaultOnStaticMesh;
 
-	UPROPERTY(config, EditAnywhere, Category = VirtualTextures, meta = (
-		EditCondition = "bVirtualTextures",
+	UPROPERTY(config, EditAnywhere, Category = MeshPaintVirtualTextures, meta = (
+		EditCondition = "bVirtualTextures && bMeshPaintVirtualTextureSupport",
 		ConsoleVariable = "r.MeshPaintVirtualTexture.DefaultTexelsPerVertex", DisplayName = "Mesh paint texels per vertex",
 		ToolTip = "Default ratio of texels to vertices when creating a mesh paint texture for a mesh."))
 	int32 MeshPaintDefaultTexelsPerVertex;
 
-	UPROPERTY(config, EditAnywhere, Category = VirtualTextures, meta = (
-		EditCondition = "bVirtualTextures",
+	UPROPERTY(config, EditAnywhere, Category = MeshPaintVirtualTextures, meta = (
+		EditCondition = "bVirtualTextures && bMeshPaintVirtualTextureSupport",
 		ConsoleVariable = "r.MeshPaintVirtualTexture.MaxTextureSize", DisplayName = "Mesh paint maximum texture size",
 		ToolTip = "Maximum size in pixels for mesh paint virtual textures. Will be rounded to next power-of-2."))
 	int32 MeshPaintVirtualTextureMaxTextureSize;
diff --git a/Engine/Source/Runtime/Engine/Private/Materials/HLSLMaterialTranslator.cpp b/Engine/Source/Runtime/Engine/Private/Materials/HLSLMaterialTranslator.cpp
--- a/Engine/Source/Runtime/Engine/Private/Materials/HLSLMaterialTranslator.cpp
+++ b/Engine/Source/Runtime/Engine/Private/Materials/HLSLMaterialTranslator.cpp
@@ -1222,60 +1222,6 @@ EMaterialExpressionVisitResult FHLSLMaterialTranslator::VisitExpressionsForPrope
 	return MVR_STOP;
 }
 
-void FHLSLMaterialTranslator::ValidateVtPropertyLimits()
-{
-	class FFindVirtualTextureVisitor : public IMaterialExpressionVisitor
-	{
-	public:
-		virtual EMaterialExpressionVisitResult Visit(UMaterialExpression* InExpression) override
-		{
-			if (UMaterialExpressionTextureBase *TextureExpr = Cast<UMaterialExpressionTextureBase>(InExpression))
-			{
-				if (IsVirtualSamplerType(TextureExpr->SamplerType))
-				{
-					FoundVirtualTexture = true;
-					return MVR_STOP;
-				}
-			}
-			return MVR_CONTINUE;
-		}
-
-		bool FoundVirtualTexture = false;
-	};
-	
-	for (uint32 PropertyIndex = 0u; PropertyIndex < MP_MAX; ++PropertyIndex)
-	{
-		const EMaterialProperty PropertyToValidate = (EMaterialProperty)PropertyIndex;
-		if (PropertyToValidate == MP_MaterialAttributes || PropertyToValidate == MP_CustomOutput)
-		{
-			// These properties are "special", attempting to pass them to FMaterialAttributeDefinitionMap::GetShaderFrequency() will generate log spam
-			continue;
-		}
-
-		const EShaderFrequency ShaderFrequencyToValidate = FMaterialAttributeDefinitionMap::GetShaderFrequency(PropertyToValidate);
-
-		// When converting legacy material the opacty mask is automatically connected when the material uses MaterialAttributes. 
-		// To avoid false positive, we bypass VT validation for OpactyMask input
-		if (Substrate::IsSubstrateEnabled() && PropertyToValidate == MP_OpacityMask && bSubstrateUsesConversionFromLegacy)
-		{
-			continue;
-		}
-
-		// check to see if this is a property that doesn't support virtual texture connections
-		if (PropertyToValidate == MP_OpacityMask || ShaderFrequencyToValidate != SF_Pixel)
-		{
-			FFindVirtualTextureVisitor Visitor;
-			VisitExpressionsForProperty(PropertyToValidate, Visitor);
-			if (Visitor.FoundVirtualTexture)
-			{
-				if (PropertyToValidate == MP_OpacityMask && GetDefault<URendererSettings>()->bEnableVirtualTextureOpacityMask == false)
-				{
-					Errorf(TEXT("Sampling a virtual texture for the Opacity Mask material attribute is disabled in the Renderer Settings."));
-				}
-			}
-		}
-	}
-}
 
 UE_TRACE_EVENT_BEGIN(Cpu, FHLSLMaterialTranslatorTranslate, NoSync)
 	UE_TRACE_EVENT_FIELD(UE::Trace::WideString, MaterialName)
@@ -1795,15 +1741,6 @@ void FHLSLMaterialTranslator::TranslateMaterial()
 	// Some custom outputs must be pre-compiled so they can be re-used as shared inputs
 	CompileCustomOutputs(CustomOutputExpressions, SeenCustomOutputExpressionsClasses, true);
 
-	// Validate some things on the VT system. Since generated code for expressions shared between multiple properties
-	// (e.g. a texture sample connected to both diffuse and opacity mask) is reused we can't check based on the MaterialProperty
-	// variable inside the actual code generation pass. So we do a pre-pass over it here.
-	if (UseVirtualTexturing(Platform))
-	{
-		ValidateVtPropertyLimits();
-	}
-
-
 	// Normal must always be compiled first; this will ensure its chunk calculations are the first to be added
 	{
 		// Verify that start chunk is 0
