Optimize QueueTickFunction to avoid allocating memory and check status before starting expensive nested function calls.
Fix some incorrect comments about how disabled dependencies work.

diff --git a/Engine/Source/Runtime/Engine/Private/TickTaskManager.cpp b/Engine/Source/Runtime/Engine/Private/TickTaskManager.cpp
--- a/Engine/Source/Runtime/Engine/Private/TickTaskManager.cpp
+++ b/Engine/Source/Runtime/Engine/Private/TickTaskManager.cpp
@@ -236,7 +236,7 @@ struct FTickGroupCondition
 			TEnumAsByte<ETickingGroup> StartGroup;
 			TEnumAsByte<ETickingGroup> EndGroup;
 			bool bHighPriority;
-			bool bIsBatch;
+			bool bRunOnAnyThread;
 		};
 	};
 	
@@ -249,7 +249,7 @@ struct FTickGroupCondition
 		: StartGroup(TickFunction->GetActualTickGroup())
 		, EndGroup(TickFunction->GetActualEndTickGroup())
 		, bHighPriority(TickFunction->bHighPriority)
-		, bIsBatch(TickFunction->bAllowTickBatching && !TickFunction->bRunOnAnyThread)
+		, bRunOnAnyThread(TickFunction->bRunOnAnyThread)
 	{
 	}
 
@@ -583,6 +583,12 @@ public:
 		return (FTickGraphTask*)TickFunction->GetTaskPointer(FTickFunction::ETickTaskState::HasTask);
 	}
 
+	/** Non-threadsafe accessor for checking if this has been visited for queuing this frame, only valid on registered functiosn */
+	FORCEINLINE static bool HasBeenVisited(const FTickFunction* TickFunction, uint32 CurrentFrameCounter)
+	{
+		return (TickFunction->InternalData->TickVisitedGFrameCounter.load(std::memory_order_relaxed) == CurrentFrameCounter);
+	}
+
 	/** Sets up TickContext for a tick function that will possibly tick */
 	FORCEINLINE FTickContext SetupTickContext(FTickFunction* TickFunction, const FTickContext& TickContext)
 	{
@@ -669,7 +675,7 @@ public:
 	FORCEINLINE bool CanBatchCondition(FTickGroupCondition Condition)
 	{
 		// Don't batch high priority ticks or ones that last more than a single tick group
-		return Condition.bIsBatch && !Condition.bHighPriority && Condition.StartGroup == Condition.EndGroup;
+		return !Condition.bRunOnAnyThread && !Condition.bHighPriority && Condition.StartGroup == Condition.EndGroup;
 	}
 
 	/** Call before doing any batched ticks */
@@ -705,9 +711,9 @@ public:
 	}
 
 	/** This will add to an existing batch, create a new batch, or just spawn a single task and return null */
-	FTickBatchInfo* QueueOrBatchTickTask(TArray<FTickFunction*>& Prerequisites, FTickFunction* TickFunction, const FTickContext& TickContext)
+	FTickBatchInfo* QueueOrBatchTickTask(TArray<FTickFunction*, TInlineAllocator<2>>& Prerequisites, FTickFunction* TickFunction, const FTickContext& TickContext)
 	{
-		if (bAllowBatchedTicksForFrame)
+		if (bAllowBatchedTicksForFrame && TickFunction->bAllowTickBatching)
 		{
 			FTickGroupCondition Condition = FTickGroupCondition(TickFunction);
 
@@ -1413,10 +1419,17 @@ public:
 	void QueueAllTicks()
 	{
 		FTickTaskSequencer& TTS = FTickTaskSequencer::Get();
+
+		// Only use the lower 32 bits of the frame counter
+		uint32 CurrentFrameCounter = (uint32)GFrameCounter;
+		
 		for (TSet<FTickFunction*>::TIterator It(AllEnabledTickFunctions); It; ++It)
 		{
 			FTickFunction* TickFunction = *It;
-			TickFunction->QueueTickFunction(TTS, Context);
+			if (!TTS.HasBeenVisited(TickFunction, CurrentFrameCounter))
+			{
+				TickFunction->QueueTickFunction(TTS, Context);
+			}
 
 			if (TickFunction->TickInterval > 0.f)
 			{
@@ -1431,7 +1444,10 @@ public:
 			if (TickFunction->TickState == FTickFunction::ETickState::Enabled)
 			{
 				CumulativeCooldown += TickFunction->InternalData->RelativeTickCooldown;
-				TickFunction->QueueTickFunction(TTS, Context);
+				if (!TTS.HasBeenVisited(TickFunction, CurrentFrameCounter))
+				{
+					TickFunction->QueueTickFunction(TTS, Context);
+				}
 				RescheduleForInterval(TickFunction, TickFunction->TickInterval - (Context.DeltaSeconds - CumulativeCooldown)); // Give credit for any overrun
 				AllCoolingDownTickFunctions.Head = TickFunction->InternalData->Next;
 			}
@@ -1453,11 +1469,15 @@ public:
 		// Constructing set iterators is not trivial, so avoid the following block if it will have no effect
 		if (NewlySpawnedTickFunctions.Num() != 0)
 		{
+			uint32 CurrentFrameCounter = (uint32)GFrameCounter;
 			FTickTaskSequencer& TTS = FTickTaskSequencer::Get();
 			for (TSet<FTickFunction*>::TIterator It(NewlySpawnedTickFunctions); It; ++It)
 			{
 				FTickFunction* TickFunction = *It;
-				TickFunction->QueueTickFunction(TTS, Context);
+				if (!TTS.HasBeenVisited(TickFunction, CurrentFrameCounter))
+				{
+					TickFunction->QueueTickFunction(TTS, Context);
+				}
 				Num++;
 
 				if (TickFunction->TickInterval > 0.f)
@@ -2546,87 +2566,93 @@ void FTickFunction::SetTaskPointer(ETickTaskState NewState, void* InTaskPointer)
 
 void FTickFunction::QueueTickFunction(FTickTaskSequencer& TTS, const struct FTickContext& TickContext)
 {
-	checkSlow(TickContext.Thread == ENamedThreads::GameThread); // we assume same thread here
-	check(IsTickFunctionRegistered());
-
 	// Only compare the 32bit part of the frame counter
 	uint32 CurrentFrameCounter = (uint32)GFrameCounter;
-	if (InternalData->TickVisitedGFrameCounter.load(std::memory_order_relaxed) != CurrentFrameCounter)
+
+	checkSlow(TickContext.Thread == ENamedThreads::GameThread); // we assume same thread here
+	check(IsTickFunctionRegistered() && !TTS.HasBeenVisited(this, CurrentFrameCounter));
+
+	// Mark visited at start of function
+	InternalData->TickVisitedGFrameCounter.store(CurrentFrameCounter, std::memory_order_relaxed);
+	if (TickState != FTickFunction::ETickState::Disabled)
 	{
-		InternalData->TickVisitedGFrameCounter.store(CurrentFrameCounter, std::memory_order_relaxed);
-		if (TickState != FTickFunction::ETickState::Disabled)
-		{
-			ETickingGroup MaxStartTickGroup = ETickingGroup(0);
-			ETickingGroup MaxEndTickGroup = ETickingGroup(0);
+		ETickingGroup MaxStartTickGroup = ETickingGroup(0);
+		ETickingGroup MaxEndTickGroup = ETickingGroup(0);
 
-			TArray<FTickFunction*> RawPrerequisites;
-			for (int32 PrereqIndex = 0; PrereqIndex < Prerequisites.Num(); PrereqIndex++)
+		TArray<FTickFunction*, TInlineAllocator<2>> RawPrerequisites;
+		for (int32 PrereqIndex = 0; PrereqIndex < Prerequisites.Num(); PrereqIndex++)
+		{
+			FTickFunction* Prereq = Prerequisites[PrereqIndex].PrerequisiteTickFunction;
+			if (!Prerequisites[PrereqIndex].PrerequisiteObject.IsValid(true))
+			{
+				// stale prereq, delete it
+				Prerequisites.RemoveAtSwap(PrereqIndex--);
+			}
+			else if (Prereq->IsTickFunctionRegistered())
 			{
-				FTickFunction* Prereq = Prerequisites[PrereqIndex].Get();
-				if (!Prereq)
+				// If the prerequisite object is valid Prereq can't be null, and if it is registered InternalData can't be null either
+				if (!TTS.HasBeenVisited(Prereq, CurrentFrameCounter))
 				{
-					// stale prereq, delete it
-					Prerequisites.RemoveAtSwap(PrereqIndex--);
+					// If the prerequisite hasn't been visited, queue it now
+					Prereq->QueueTickFunction(TTS, TickContext);
 				}
-				else if (Prereq->IsTickFunctionRegistered())
+
+				if (Prereq->InternalData->TickQueuedGFrameCounter.load(std::memory_order_relaxed) != CurrentFrameCounter)
 				{
-					// recursive call to make sure my prerequisite is set up so I can use its completion handle
-					Prereq->QueueTickFunction(TTS, TickContext);
-					if (Prereq->InternalData->TickQueuedGFrameCounter.load(std::memory_order_relaxed) != CurrentFrameCounter)
-					{
-						// this must be up the call stack, therefore this is a cycle
-						UE_LOG(LogTick, Warning, TEXT("While processing prerequisites for %s, could use %s because it would form a cycle."),*DiagnosticMessage(), *Prereq->DiagnosticMessage());
-					}
-					else if (Prereq->InternalData->TaskState == ETickTaskState::NotQueued)
-					{
-						//ok UE_LOG(LogTick, Warning, TEXT("While processing prerequisites for %s, could use %s because it is disabled."),*DiagnosticMessage(), *Prereq->DiagnosticMessage());
-					}
-					else if (TTS.ShouldConsiderPrerequisite(this, Prereq))
-					{
-						MaxStartTickGroup = FMath::Max<ETickingGroup>(MaxStartTickGroup, Prereq->InternalData->ActualStartTickGroup.GetValue());
-						MaxEndTickGroup = FMath::Max<ETickingGroup>(MaxEndTickGroup, Prereq->InternalData->ActualEndTickGroup.GetValue());
-						RawPrerequisites.Add(Prereq);
-					}
+					// This can only happen if the prerequisite is is partially queued in the current stack
+					UE_LOG(LogTick, Warning, TEXT("While processing prerequisites for %s, could not use %s because it would form a cycle."), *DiagnosticMessage(), *Prereq->DiagnosticMessage());
+				}
+				else if (Prereq->InternalData->TaskState == ETickTaskState::NotQueued)
+				{
+					// Ignore disabled dependencies, this means that intermediate scene components will break the automatic dependency setting
+				}
+				else if (TTS.ShouldConsiderPrerequisite(this, Prereq))
+				{
+					MaxStartTickGroup = FMath::Max<ETickingGroup>(MaxStartTickGroup, Prereq->InternalData->ActualStartTickGroup.GetValue());
+					MaxEndTickGroup = FMath::Max<ETickingGroup>(MaxEndTickGroup, Prereq->InternalData->ActualEndTickGroup.GetValue());
+					RawPrerequisites.Add(Prereq);
 				}
 			}
+		}
 
-			// tick group is the max of the prerequisites, the current tick group, and the desired tick group
-			ETickingGroup MyActualTickGroup = FMath::Max<ETickingGroup>(MaxStartTickGroup, FMath::Max<ETickingGroup>(TickGroup.GetValue(), TickContext.TickGroup));
-			if (MyActualTickGroup != TickGroup)
+		// tick group is the max of the prerequisites, the current tick group, and the desired tick group
+		ETickingGroup MyActualTickGroup = FMath::Max<ETickingGroup>(MaxStartTickGroup, FMath::Max<ETickingGroup>(TickGroup.GetValue(), TickContext.TickGroup));
+		if (MyActualTickGroup != TickGroup)
+		{
+			// if the tick was "demoted", make sure it ends up in an ordinary tick group.
+			while (!CanDemoteIntoTickGroup(MyActualTickGroup))
 			{
-				// if the tick was "demoted", make sure it ends up in an ordinary tick group.
-				while (!CanDemoteIntoTickGroup(MyActualTickGroup))
-				{
-					MyActualTickGroup = ETickingGroup(MyActualTickGroup + 1);
-				}
+				MyActualTickGroup = ETickingGroup(MyActualTickGroup + 1);
 			}
-			InternalData->ActualStartTickGroup = MyActualTickGroup;
-			InternalData->ActualEndTickGroup = MyActualTickGroup;
+		}
+		InternalData->ActualStartTickGroup = MyActualTickGroup;
+		InternalData->ActualEndTickGroup = MyActualTickGroup;
 
-			// Also check to see if the end tick group needs to be extended separately
-			ETickingGroup MyActualEndTickGroup = FMath::Max<ETickingGroup>(MaxEndTickGroup, FMath::Max<ETickingGroup>(EndTickGroup.GetValue(), MyActualTickGroup));
+		// Also check to see if the end tick group needs to be extended separately
+		ETickingGroup MyActualEndTickGroup = FMath::Max<ETickingGroup>(MaxEndTickGroup, FMath::Max<ETickingGroup>(EndTickGroup.GetValue(), MyActualTickGroup));
 
-			if (MyActualEndTickGroup > MyActualTickGroup)
+		if (MyActualEndTickGroup > MyActualTickGroup)
+		{
+			check(MyActualEndTickGroup <= TG_NewlySpawned);
+			ETickingGroup TestTickGroup = ETickingGroup(MyActualTickGroup + 1);
+			while (TestTickGroup <= MyActualEndTickGroup)
 			{
-				check(MyActualEndTickGroup <= TG_NewlySpawned);
-				ETickingGroup TestTickGroup = ETickingGroup(MyActualTickGroup + 1);
-				while (TestTickGroup <= MyActualEndTickGroup)
+				if (CanDemoteIntoTickGroup(TestTickGroup))
 				{
-					if (CanDemoteIntoTickGroup(TestTickGroup))
-					{
-						InternalData->ActualEndTickGroup = TestTickGroup;
-					}
-					TestTickGroup = ETickingGroup(TestTickGroup + 1);
+					InternalData->ActualEndTickGroup = TestTickGroup;
 				}
+				TestTickGroup = ETickingGroup(TestTickGroup + 1);
 			}
+		}
 
-			if (TickState == FTickFunction::ETickState::Enabled)
-			{
-				TTS.QueueOrBatchTickTask(RawPrerequisites, this, TickContext);
-			}
+		if (TickState == FTickFunction::ETickState::Enabled)
+		{
+			TTS.QueueOrBatchTickTask(RawPrerequisites, this, TickContext);
 		}
-		InternalData->TickQueuedGFrameCounter.store(CurrentFrameCounter, std::memory_order_relaxed);
 	}
+
+	// Mark as queued (finished visiting), even if it was not turned into a real task
+	InternalData->TickQueuedGFrameCounter.store(CurrentFrameCounter, std::memory_order_relaxed);
 }
 
 void FTickFunction::QueueTickFunctionParallel(const struct FTickContext& TickContext, TArray<FTickFunction*, TInlineAllocator<8> >& StackForCycleDetection)
