Move mesh deformer to happen after the scene updates so that the mesh updater task doesn't immediately get synced.

diff --git a/Engine/Source/Runtime/Engine/Private/SkeletalMeshUpdater.cpp b/Engine/Source/Runtime/Engine/Private/SkeletalMeshUpdater.cpp
--- a/Engine/Source/Runtime/Engine/Private/SkeletalMeshUpdater.cpp
+++ b/Engine/Source/Runtime/Engine/Private/SkeletalMeshUpdater.cpp
@@ -309,23 +309,23 @@ struct FSkeletalMeshUpdater::FTaskData
 		Tasks.Filter.Trigger();
 
 		{
-			TRACE_CPUPROFILER_EVENT_SCOPE(MeshDeformer);
+			TRACE_CPUPROFILER_EVENT_SCOPE(Inline);
 			for (auto& UpdatePacket : Packets)
 			{
-				UpdatePacket->ProcessStage_MeshDeformer(*RHICmdList);
+				UpdatePacket->ProcessStage_Inline(*RHICmdList);
 			}
-			Tasks.MeshDeformer.Trigger();
+			Tasks.Inline.Trigger();
 		}
-	
+
 		{
-			TRACE_CPUPROFILER_EVENT_SCOPE(Inline);
+			TRACE_CPUPROFILER_EVENT_SCOPE(MeshDeformer);
 			for (auto& UpdatePacket : Packets)
 			{
-				UpdatePacket->ProcessStage_Inline(*RHICmdList);
+				UpdatePacket->ProcessStage_MeshDeformer(*RHICmdList);
 			}
-			Tasks.Inline.Trigger();
+			Tasks.MeshDeformer.Trigger();
 		}
-
+	
 		{
 			TRACE_CPUPROFILER_EVENT_SCOPE(GPUSkinCache);
 			for (auto& UpdatePacket : Packets)
@@ -354,8 +354,8 @@ struct FSkeletalMeshUpdater::FTaskData
 	struct
 	{
 		UE::Tasks::FTaskEvent Filter{ UE_SOURCE_LOCATION };
-		UE::Tasks::FTaskEvent MeshDeformer{ UE_SOURCE_LOCATION };
 		UE::Tasks::FTaskEvent Inline{ UE_SOURCE_LOCATION };
+		UE::Tasks::FTaskEvent MeshDeformer{ UE_SOURCE_LOCATION };
 
 	} Tasks;
 };
@@ -414,8 +414,8 @@ FSkeletalMeshUpdater::FSubmitTasks FSkeletalMeshUpdater::Submit(FRDGBuilder& Gra
 	return GraphBuilder.Blackboard.Create<FSubmitTasks>(FSubmitTasks
 	{
 		  .Filter       = TaskData.Tasks.Filter
-		, .MeshDeformer = TaskData.Tasks.MeshDeformer
 		, .Inline       = TaskData.Tasks.Inline
+		, .MeshDeformer = TaskData.Tasks.MeshDeformer
 		, .SkinCache    = MoveTemp(ForegroundTask)
 	});
 }
@@ -431,14 +431,14 @@ void FSkeletalMeshUpdater::WaitForStage(FRDGBuilder& GraphBuilder, ESkeletalMesh
 			SubmitTasks->Filter.Wait();
 			SubmitTasks->Filter = {};
 			break;
-		case ESkeletalMeshUpdateStage::MeshDeformer:
-			SubmitTasks->MeshDeformer.Wait();
-			SubmitTasks->MeshDeformer = {};
-			break;
 		case ESkeletalMeshUpdateStage::Inline:
 			SubmitTasks->Inline.Wait();
 			SubmitTasks->Inline = {};
 			break;
+		case ESkeletalMeshUpdateStage::MeshDeformer:
+			SubmitTasks->MeshDeformer.Wait();
+			SubmitTasks->MeshDeformer = {};
+			break;
 		case ESkeletalMeshUpdateStage::SkinCache:
 			SubmitTasks->SkinCache.Wait();
 			SubmitTasks->SkinCache = {};
