[Optim] Vlog serialization:
1) Use move instead of copy when flushing vlog buffers from threads
2) pre allocate serialization buffer

FVisualLogger::FlushThreadsEntries Average 2.498ms -> 1.321ms

diff --git a/Engine/Source/Runtime/Engine/Private/VisualLogger/VisualLoggerTraceDevice.cpp b/Engine/Source/Runtime/Engine/Private/VisualLogger/VisualLoggerTraceDevice.cpp
--- a/Engine/Source/Runtime/Engine/Private/VisualLogger/VisualLoggerTraceDevice.cpp
+++ b/Engine/Source/Runtime/Engine/Private/VisualLogger/VisualLoggerTraceDevice.cpp
@@ -56,6 +56,7 @@ void FVisualLoggerTraceDevice::Serialize(const UObject* InLogOwner, const FName&
 	if (UE_TRACE_CHANNELEXPR_IS_ENABLED(VisualLoggerChannel))
 	{
 		FBufferArchive Archive;
+		Archive.Reserve(1024);
 		Archive.UsingCustomVersion(EVisualLoggerVersion::GUID);
 		Archive.UsingCustomVersion(FUE5MainStreamObjectVersion::GUID);
 		Archive.SetCustomVersion(FUE5MainStreamObjectVersion::GUID, FUE5MainStreamObjectVersion::LatestVersion, "UE5MainStreamObjectVersion");
diff --git a/Engine/Source/Runtime/Engine/Private/VisualLogger/VisualLoggerTypes.cpp b/Engine/Source/Runtime/Engine/Private/VisualLogger/VisualLoggerTypes.cpp
--- a/Engine/Source/Runtime/Engine/Private/VisualLogger/VisualLoggerTypes.cpp
+++ b/Engine/Source/Runtime/Engine/Private/VisualLogger/VisualLoggerTypes.cpp
@@ -100,6 +100,21 @@ int32 FVisualLogEntry::AddEvent(const FVisualLogEventBase& Event)
 	return Events.Add(Event);
 }
 
+template<typename T>
+FORCEINLINE void AppendMoveArray(TArray<T>& DestArray, TArray<T>& SourceArray)
+{
+	if (DestArray.IsEmpty())
+	{
+		// Safe to do a swap
+		DestArray = MoveTemp(SourceArray);
+	}
+	else
+	{
+		// We have to copy memory from the source to the end of the dest
+		DestArray.Append(MoveTemp(SourceArray));
+	}
+}
+
 void FVisualLogEntry::MoveTo(FVisualLogEntry& Other)
 {
 	ensureMsgf(bIsInitialized && Other.bIsInitialized, TEXT("Both entries need to be initialized to move to the other"));
@@ -107,12 +122,13 @@ void FVisualLogEntry::MoveTo(FVisualLogEntry& Other)
 	ensureMsgf(bPassedClassAllowList == Other.bPassedClassAllowList, TEXT("Can only move similar entries"));
 	ensureMsgf(bPassedObjectAllowList == Other.bPassedObjectAllowList, TEXT("Can only move similar entries"));
 	ensureMsgf(bIsAllowedToLog == Other.bIsAllowedToLog, TEXT("Can only move similar entries"));
-	Other.Events.Append(Events);
-	Other.LogLines.Append(LogLines);
-	Other.Status.Append(Status);
-	Other.ElementsToDraw.Append(ElementsToDraw);
-	Other.HistogramSamples.Append(HistogramSamples);
-	Other.DataBlocks.Append(DataBlocks);
+
+	AppendMoveArray(Other.Events, Events);
+	AppendMoveArray(Other.LogLines, LogLines);
+	AppendMoveArray(Other.Status, Status);
+	AppendMoveArray(Other.ElementsToDraw, ElementsToDraw);
+	AppendMoveArray(Other.HistogramSamples, HistogramSamples);
+	AppendMoveArray(Other.DataBlocks, DataBlocks);
 	Reset();
 }
 
diff --git a/Engine/Source/Runtime/Engine/Public/VisualLogger/VisualLoggerTypes.h b/Engine/Source/Runtime/Engine/Public/VisualLogger/VisualLoggerTypes.h
--- a/Engine/Source/Runtime/Engine/Public/VisualLogger/VisualLoggerTypes.h
+++ b/Engine/Source/Runtime/Engine/Public/VisualLogger/VisualLoggerTypes.h
@@ -111,6 +111,8 @@ struct FVisualLogLine
 	FColor Color = FColor::White;
 	bool bMonospace = false;
 	FVisualLogLine() = default;
+	FVisualLogLine(FVisualLogLine&& Other) = default;
+	FVisualLogLine(const FVisualLogLine& Other) = default;
 	FVisualLogLine(const FName& InCategory, ELogVerbosity::Type InVerbosity, const FString& InLine);
 	FVisualLogLine(const FName& InCategory, ELogVerbosity::Type InVerbosity, const FString& InLine, int64 InUserData);
 	FVisualLogLine(const FName& InCategory, ELogVerbosity::Type InVerbosity, const FString& InLine, const FColor& InColor, bool bInMonospace);
@@ -136,6 +138,26 @@ struct FVisualLogStatusCategory
 
 struct FVisualLogShapeElement
 {
+	FVisualLogShapeElement(FVisualLogShapeElement&& Other)
+		:	Description(MoveTemp(Other.Description)),
+			Category(Other.Category),
+			Verbosity(Other.Verbosity),
+			Points(MoveTemp(Other.Points)),
+			TransformationMatrix(Other.TransformationMatrix),
+			UniqueId(Other.UniqueId),
+			Type(Other.Type),
+			Thickness(Other.Thickness)
+	{
+	}
+
+PRAGMA_DISABLE_DEPRECATION_WARNINGS
+	FVisualLogShapeElement(const FVisualLogShapeElement& Other) = default;
+	FVisualLogShapeElement& operator=(const FVisualLogShapeElement& Other) = default;
+PRAGMA_ENABLE_DEPRECATION_WARNINGS
+
+	ENGINE_API explicit FVisualLogShapeElement(EVisualLoggerShapeElement InType = EVisualLoggerShapeElement::Invalid);
+	FVisualLogShapeElement(const FString& InDescription, const FColor& InColor, uint16 InThickness, const FName& InCategory);
+
 	FString Description;
 	FName Category = UE::VisualLogger::NAME_UnnamedCategory;
 	TEnumAsByte<ELogVerbosity::Type> Verbosity = ELogVerbosity::All;
@@ -153,8 +175,6 @@ struct FVisualLogShapeElement
 		uint16 Mag;
 	};
 
-	ENGINE_API explicit FVisualLogShapeElement(EVisualLoggerShapeElement InType = EVisualLoggerShapeElement::Invalid);
-	FVisualLogShapeElement(const FString& InDescription, const FColor& InColor, uint16 InThickness, const FName& InCategory);
 	void SetColor(const FColor& InColor);
 	EVisualLoggerShapeElement GetType() const;
 	void SetType(EVisualLoggerShapeElement InType);
