[Engine]
* Fixed VisualLoggerTraceDevice singleton being instantated in source file

diff --git a/Engine/Source/Runtime/Engine/Classes/VisualLogger/VisualLoggerTraceDevice.h b/Engine/Source/Runtime/Engine/Classes/VisualLogger/VisualLoggerTraceDevice.h
--- a/Engine/Source/Runtime/Engine/Classes/VisualLogger/VisualLoggerTraceDevice.h
+++ b/Engine/Source/Runtime/Engine/Classes/VisualLogger/VisualLoggerTraceDevice.h
@@ -12,11 +12,7 @@ DECLARE_DELEGATE_TwoParams(FImmediateRenderDelegate, const UObject*, const FVisu
 class FVisualLoggerTraceDevice : public FVisualLogDevice
 {
 public:
-	static FVisualLoggerTraceDevice& Get()
-	{
-		static FVisualLoggerTraceDevice GDevice;
-		return GDevice;
-	}
+	static ENGINE_API FVisualLoggerTraceDevice& Get();
 
 	ENGINE_API FVisualLoggerTraceDevice();
 	virtual void Cleanup(bool bReleaseMemory = false) override;
diff --git a/Engine/Source/Runtime/Engine/Private/VisualLogger/VisualLoggerTraceDevice.cpp b/Engine/Source/Runtime/Engine/Private/VisualLogger/VisualLoggerTraceDevice.cpp
--- a/Engine/Source/Runtime/Engine/Private/VisualLogger/VisualLoggerTraceDevice.cpp
+++ b/Engine/Source/Runtime/Engine/Private/VisualLogger/VisualLoggerTraceDevice.cpp
@@ -18,6 +18,12 @@ UE_TRACE_EVENT_BEGIN(VisualLogger, VisualLogEntry)
 	UE_TRACE_EVENT_FIELD(uint8[], LogEntry)
 UE_TRACE_EVENT_END()
 
+FVisualLoggerTraceDevice& FVisualLoggerTraceDevice::Get()
+{
+	static FVisualLoggerTraceDevice GDevice;
+	return GDevice;
+}
+
 FVisualLoggerTraceDevice::FVisualLoggerTraceDevice()
 {
 
