[Engine - More dllstorage fixes]
* Fixed KismetMathLibrary debug target compile errors
* Removed dllstorage in chaos for methods that were inline
* Exported all methods in VisualLoggerTraceDevice

diff --git a/Engine/Source/Runtime/Engine/Classes/Kismet/KismetMathLibrary.h b/Engine/Source/Runtime/Engine/Classes/Kismet/KismetMathLibrary.h
--- a/Engine/Source/Runtime/Engine/Classes/Kismet/KismetMathLibrary.h
+++ b/Engine/Source/Runtime/Engine/Classes/Kismet/KismetMathLibrary.h
@@ -177,6 +177,12 @@ struct FQuaternionSpringState
 	}
 };
 
+#if KISMET_MATH_INLINE_ENABLED
+#define UE_INL_API
+#else
+#define UE_INL_API ENGINE_API
+#endif
+
 struct FRuntimeFloatCurve;
 
 UCLASS(meta=(BlueprintThreadSafe, ScriptName = "MathLibrary"), MinimalAPI)
@@ -190,7 +196,7 @@ class UKismetMathLibrary : public UBlueprintFunctionLibrary
 	
 	/** Returns a uniformly distributed random bool*/
 	UFUNCTION(BlueprintPure, Category="Math|Random", meta=(NotBlueprintThreadSafe))
-	static bool RandomBool();
+	static UE_INL_API bool RandomBool();
 
 	/** 
 	 * Get a random chance with the specified weight. Range of weight is 0.0 - 1.0 E.g.,
@@ -207,42 +213,42 @@ class UKismetMathLibrary : public UBlueprintFunctionLibrary
 	static ENGINE_API bool RandomBoolWithWeightFromStream(const FRandomStream& RandomStream, float Weight);
 
 	UE_DEPRECATED(5.2, "Use RandomBoolWithWeightFromStream taking the FRandomStream as the first argument.")
-	static bool RandomBoolWithWeightFromStream(float Weight, const FRandomStream& RandomStream)
+	static UE_INL_API bool RandomBoolWithWeightFromStream(float Weight, const FRandomStream& RandomStream)
 	{
 		return RandomBoolWithWeightFromStream(RandomStream, Weight);
 	}
 
 	/** Returns the logical complement of the Boolean value (NOT A) */
 	UFUNCTION(BlueprintPure, meta=(DisplayName = "NOT Boolean", CompactNodeTitle = "NOT", Keywords = "! not negate"), Category="Math|Boolean")
-	static bool Not_PreBool(bool A);
+	static UE_INL_API bool Not_PreBool(bool A);
 
 	/** Returns true if the values are equal (A == B) */
 	UFUNCTION(BlueprintPure, meta=(DisplayName = "Equal (Boolean)", CompactNodeTitle = "==", Keywords = "== equal"), Category="Math|Boolean")
-	static bool EqualEqual_BoolBool(bool A, bool B);
+	static UE_INL_API bool EqualEqual_BoolBool(bool A, bool B);
 
 	/** Returns true if the values are not equal (A != B) */
 	UFUNCTION(BlueprintPure, meta=(DisplayName = "Not Equal (Boolean)", CompactNodeTitle = "!=", Keywords = "!= not equal"), Category="Math|Boolean")
-	static bool NotEqual_BoolBool(bool A, bool B);
+	static UE_INL_API bool NotEqual_BoolBool(bool A, bool B);
 
 	/** Returns the logical AND of two values (A AND B) */
 	UFUNCTION(BlueprintPure, meta=(DisplayName = "AND Boolean", CompactNodeTitle = "AND", Keywords = "& and", CommutativeAssociativeBinaryOperator = "true"), Category="Math|Boolean")
-	static bool BooleanAND(bool A, bool B);
+	static UE_INL_API bool BooleanAND(bool A, bool B);
 
 	/** Returns the logical NAND of two values (A AND B) */
 	UFUNCTION(BlueprintPure, meta=(DisplayName = "NAND Boolean", CompactNodeTitle = "NAND", Keywords = "!& nand", CommutativeAssociativeBinaryOperator = "true"), Category="Math|Boolean")
-	static bool BooleanNAND(bool A, bool B);
+	static UE_INL_API bool BooleanNAND(bool A, bool B);
 
 	/** Returns the logical OR of two values (A OR B) */
 	UFUNCTION(BlueprintPure, meta=(DisplayName = "OR Boolean", CompactNodeTitle = "OR", Keywords = "| or", CommutativeAssociativeBinaryOperator = "true"), Category="Math|Boolean")
-	static bool BooleanOR(bool A, bool B);
+	static UE_INL_API bool BooleanOR(bool A, bool B);
 		
 	/** Returns the logical eXclusive OR of two values (A XOR B) */
 	UFUNCTION(BlueprintPure, meta=(DisplayName = "XOR Boolean", CompactNodeTitle = "XOR", Keywords = "^ xor"), Category="Math|Boolean")
-	static bool BooleanXOR(bool A, bool B);
+	static UE_INL_API bool BooleanXOR(bool A, bool B);
 
 	/** Returns the logical Not OR of two values (A NOR B) */
 	UFUNCTION(BlueprintPure, meta=(DisplayName = "NOR Boolean", CompactNodeTitle = "NOR", Keywords = "!^ nor"), Category="Math|Boolean")
-	static bool BooleanNOR(bool A, bool B);
+	static UE_INL_API bool BooleanNOR(bool A, bool B);
 
 	//
 	// Byte functions.
@@ -250,55 +256,55 @@ class UKismetMathLibrary : public UBlueprintFunctionLibrary
 
 	/** Multiplication (A * B) */
 	UFUNCTION(BlueprintPure, meta=(DisplayName = "Byte * Byte", CompactNodeTitle = "*", Keywords = "* multiply", CommutativeAssociativeBinaryOperator = "true"), Category="Math|Byte")
-	static uint8 Multiply_ByteByte(uint8 A, uint8 B);
+	static UE_INL_API uint8 Multiply_ByteByte(uint8 A, uint8 B);
 
 	/** Division (A / B) */
 	UFUNCTION(BlueprintPure, meta=(DisplayName = "Byte / Byte", CompactNodeTitle = "/", Keywords = "/ divide division"), Category="Math|Byte")
-	static uint8 Divide_ByteByte(uint8 A, uint8 B = 1);
+	static UE_INL_API uint8 Divide_ByteByte(uint8 A, uint8 B = 1);
 
 	/** Modulo (A % B) */
 	UFUNCTION(BlueprintPure, meta=(DisplayName = "% (Byte)", CompactNodeTitle = "%", Keywords = "% modulus"), Category="Math|Byte")
-	static uint8 Percent_ByteByte(uint8 A, uint8 B = 1);
+	static UE_INL_API uint8 Percent_ByteByte(uint8 A, uint8 B = 1);
 
 	/** Addition (A + B) */
 	UFUNCTION(BlueprintPure, meta=(DisplayName = "Byte + Byte", CompactNodeTitle = "+", Keywords = "+ add plus", CommutativeAssociativeBinaryOperator = "true"), Category="Math|Byte")
-	static uint8 Add_ByteByte(uint8 A, uint8 B = 1);
+	static UE_INL_API uint8 Add_ByteByte(uint8 A, uint8 B = 1);
 
 	/** Subtraction (A - B) */
 	UFUNCTION(BlueprintPure, meta=(DisplayName = "Byte - Byte", CompactNodeTitle = "-", Keywords = "- subtract minus"), Category="Math|Byte")
-	static uint8 Subtract_ByteByte(uint8 A, uint8 B = 1);
+	static UE_INL_API uint8 Subtract_ByteByte(uint8 A, uint8 B = 1);
 
 	/** Returns the minimum value of A and B */
 	UFUNCTION(BlueprintPure, meta = (DisplayName = "Min (Byte)", CompactNodeTitle = "MIN", CommutativeAssociativeBinaryOperator = "true"), Category = "Math|Byte")
-	static uint8 BMin(uint8 A, uint8 B);
+	static UE_INL_API uint8 BMin(uint8 A, uint8 B);
 
 	/** Returns the maximum value of A and B */
 	UFUNCTION(BlueprintPure, meta = (DisplayName = "Max (Byte)", CompactNodeTitle = "MAX", CommutativeAssociativeBinaryOperator = "true"), Category = "Math|Byte")
-	static uint8 BMax(uint8 A, uint8 B);
+	static UE_INL_API uint8 BMax(uint8 A, uint8 B);
 	
 	/** Returns true if A is less than B (A < B) */
 	UFUNCTION(BlueprintPure, meta=(DisplayName = "Byte < Byte", CompactNodeTitle = "<", Keywords = "< less"), Category="Math|Byte")
-	static bool Less_ByteByte(uint8 A, uint8 B);
+	static UE_INL_API bool Less_ByteByte(uint8 A, uint8 B);
 
 	/** Returns true if A is greater than B (A > B) */
 	UFUNCTION(BlueprintPure, meta=(DisplayName = "Byte > Byte", CompactNodeTitle = ">", Keywords = "> greater"), Category="Math|Byte")
-	static bool Greater_ByteByte(uint8 A, uint8 B);
+	static UE_INL_API bool Greater_ByteByte(uint8 A, uint8 B);
 
 	/** Returns true if A is less than or equal to B (A <= B) */
 	UFUNCTION(BlueprintPure, meta=(DisplayName = "Byte <= Byte", CompactNodeTitle = "<=", Keywords = "<= less"), Category="Math|Byte")
-	static bool LessEqual_ByteByte(uint8 A, uint8 B);
+	static UE_INL_API bool LessEqual_ByteByte(uint8 A, uint8 B);
 
 	/** Returns true if A is greater than or equal to B (A >= B) */
 	UFUNCTION(BlueprintPure, meta=(DisplayName = "Byte >= Byte", CompactNodeTitle = ">=", Keywords = ">= greater"), Category="Math|Byte")
-	static bool GreaterEqual_ByteByte(uint8 A, uint8 B);
+	static UE_INL_API bool GreaterEqual_ByteByte(uint8 A, uint8 B);
 
 	/** Returns true if A is equal to B (A == B) */
 	UFUNCTION(BlueprintPure, meta=(DisplayName = "Equal (Byte)", CompactNodeTitle = "==", Keywords = "== equal"), Category="Math|Byte")
-	static bool EqualEqual_ByteByte(uint8 A, uint8 B);
+	static UE_INL_API bool EqualEqual_ByteByte(uint8 A, uint8 B);
 
 	/** Returns true if A is not equal to B (A != B) */
 	UFUNCTION(BlueprintPure, meta=(DisplayName = "Not Equal (Byte)", CompactNodeTitle = "!=", Keywords = "!= not equal"), Category="Math|Byte")
-	static bool NotEqual_ByteByte(uint8 A, uint8 B);
+	static UE_INL_API bool NotEqual_ByteByte(uint8 A, uint8 B);
 
 	//
 	// Integer functions.
@@ -306,47 +312,47 @@ class UKismetMathLibrary : public UBlueprintFunctionLibrary
 
 	/** Multiplication (A * B) */
 	UFUNCTION(BlueprintPure, meta=(DisplayName = "int * int", CompactNodeTitle = "*", Keywords = "* multiply", CommutativeAssociativeBinaryOperator = "true"), Category="Math|Integer")
-	static int32 Multiply_IntInt(int32 A, int32 B);
+	static UE_INL_API int32 Multiply_IntInt(int32 A, int32 B);
 
 	/** Division (A / B) */
 	UFUNCTION(BlueprintPure, meta=(DisplayName = "int / int", CompactNodeTitle = "/", Keywords = "/ divide division"), Category="Math|Integer")
-	static int32 Divide_IntInt(int32 A, int32 B = 1);
+	static UE_INL_API int32 Divide_IntInt(int32 A, int32 B = 1);
 
 	/** Modulo (A % B) */
 	UFUNCTION(BlueprintPure, meta=(DisplayName = "% (Integer)", CompactNodeTitle = "%", Keywords = "% modulus"), Category="Math|Integer")
-	static int32 Percent_IntInt(int32 A, int32 B = 1);
+	static UE_INL_API int32 Percent_IntInt(int32 A, int32 B = 1);
 
 	/** Addition (A + B) */
 	UFUNCTION(BlueprintPure, meta=(DisplayName = "int + int", CompactNodeTitle = "+", Keywords = "+ add plus", CommutativeAssociativeBinaryOperator = "true"), Category="Math|Integer")
-	static int32 Add_IntInt(int32 A, int32 B = 1);
+	static UE_INL_API int32 Add_IntInt(int32 A, int32 B = 1);
 
 	/** Subtraction (A - B) */
 	UFUNCTION(BlueprintPure, meta=(DisplayName = "int - int", CompactNodeTitle = "-", Keywords = "- subtract minus"), Category="Math|Integer")
-	static int32 Subtract_IntInt(int32 A, int32 B = 1);
+	static UE_INL_API int32 Subtract_IntInt(int32 A, int32 B = 1);
 
 	/** Returns true if A is less than B (A < B) */
 	UFUNCTION(BlueprintPure, meta=(DisplayName = "integer < integer", CompactNodeTitle = "<", Keywords = "< less"), Category="Math|Integer")
-	static bool Less_IntInt(int32 A, int32 B);
+	static UE_INL_API bool Less_IntInt(int32 A, int32 B);
 
 	/** Returns true if A is greater than B (A > B) */
 	UFUNCTION(BlueprintPure, meta=(DisplayName = "integer > integer", CompactNodeTitle = ">", Keywords = "> greater"), Category="Math|Integer")
-	static bool Greater_IntInt(int32 A, int32 B);
+	static UE_INL_API bool Greater_IntInt(int32 A, int32 B);
 
 	/** Returns true if A is less than or equal to B (A <= B) */
 	UFUNCTION(BlueprintPure, meta=(DisplayName = "integer <= integer", CompactNodeTitle = "<=", Keywords = "<= less"), Category="Math|Integer")
-	static bool LessEqual_IntInt(int32 A, int32 B);
+	static UE_INL_API bool LessEqual_IntInt(int32 A, int32 B);
 
 	/** Returns true if A is greater than or equal to B (A >= B) */
 	UFUNCTION(BlueprintPure, meta=(DisplayName = "integer >= integer", CompactNodeTitle = ">=", Keywords = ">= greater"), Category="Math|Integer")
-	static bool GreaterEqual_IntInt(int32 A, int32 B);
+	static UE_INL_API bool GreaterEqual_IntInt(int32 A, int32 B);
 
 	/** Returns true if A is equal to B (A == B) */
 	UFUNCTION(BlueprintPure, meta=(DisplayName = "Equal (Integer)", CompactNodeTitle = "==", Keywords = "== equal"), Category="Math|Integer")
-	static bool EqualEqual_IntInt(int32 A, int32 B);
+	static UE_INL_API bool EqualEqual_IntInt(int32 A, int32 B);
 
 	/** Returns true if A is not equal to B (A != B) */
 	UFUNCTION(BlueprintPure, meta=(DisplayName = "Not Equal (Integer)", CompactNodeTitle = "!=", Keywords = "!= not equal"), Category="Math|Integer")
-	static bool NotEqual_IntInt(int32 A, int32 B);
+	static UE_INL_API bool NotEqual_IntInt(int32 A, int32 B);
 
 	/** Returns true if value is between Min and Max (V >= Min && V <= Max)
 	 * If InclusiveMin is true, value needs to be equal or larger than Min, else it needs to be larger
@@ -357,23 +363,23 @@ class UKismetMathLibrary : public UBlueprintFunctionLibrary
 
 	/** Bitwise AND (A & B) */
 	UFUNCTION(BlueprintPure, meta=(DisplayName = "Bitwise AND", CompactNodeTitle = "&", Keywords = "& and", CommutativeAssociativeBinaryOperator = "true"), Category="Math|Integer")
-	static int32 And_IntInt(int32 A, int32 B);
+	static UE_INL_API int32 And_IntInt(int32 A, int32 B);
 
 	/** Bitwise XOR (A ^ B) */
 	UFUNCTION(BlueprintPure, meta=(DisplayName = "Bitwise XOR", CompactNodeTitle = "^", Keywords = "^ xor"), Category="Math|Integer")
-	static int32 Xor_IntInt(int32 A, int32 B);
+	static UE_INL_API int32 Xor_IntInt(int32 A, int32 B);
 
 	/** Bitwise OR (A | B) */
 	UFUNCTION(BlueprintPure, meta=(DisplayName = "Bitwise OR", CompactNodeTitle = "|", Keywords = "| or", CommutativeAssociativeBinaryOperator = "true"), Category="Math|Integer")
-	static int32 Or_IntInt(int32 A, int32 B);
+	static UE_INL_API int32 Or_IntInt(int32 A, int32 B);
 
 	/** Bitwise NOT (~A) */
 	UFUNCTION(BlueprintPure, meta = (DisplayName = "Bitwise NOT", CompactNodeTitle = "~", Keywords = "~ not"), Category = "Math|Integer")
-	static int32 Not_Int(int32 A);
+	static UE_INL_API int32 Not_Int(int32 A);
 
 	/** Sign (integer, returns -1 if A < 0, 0 if A is zero, and +1 if A > 0) */
 	UFUNCTION(BlueprintPure, meta=(DisplayName = "Sign (Integer)"), Category="Math|Integer")
-	static int32 SignOfInteger(int32 A);
+	static UE_INL_API int32 SignOfInteger(int32 A);
 
 	/** Returns a uniformly distributed random number between 0 and Max - 1 */
 	UFUNCTION(BlueprintPure, Category="Math|Random", meta=(NotBlueprintThreadSafe))
@@ -381,15 +387,15 @@ class UKismetMathLibrary : public UBlueprintFunctionLibrary
 
 	/** Return a random integer between Min and Max (>= Min and <= Max) */
 	UFUNCTION(BlueprintPure, Category="Math|Random", meta = (NotBlueprintThreadSafe))
-	static int32 RandomIntegerInRange(int32 Min, int32 Max);
+	static UE_INL_API int32 RandomIntegerInRange(int32 Min, int32 Max);
 
 	/** Returns the minimum value of A and B */
 	UFUNCTION(BlueprintPure, meta=(DisplayName = "Min (Integer)", CompactNodeTitle = "MIN", CommutativeAssociativeBinaryOperator = "true"), Category="Math|Integer")
-	static int32 Min(int32 A, int32 B);
+	static UE_INL_API int32 Min(int32 A, int32 B);
 
 	/** Returns the maximum value of A and B */
 	UFUNCTION(BlueprintPure, meta=(DisplayName = "Max (Integer)", CompactNodeTitle = "MAX", CommutativeAssociativeBinaryOperator = "true"), Category="Math|Integer")
-	static int32 Max(int32 A, int32 B);
+	static UE_INL_API int32 Max(int32 A, int32 B);
 
 	/** Returns Value clamped to be between A and B (inclusive) */
 	UFUNCTION(BlueprintPure, meta=(DisplayName = "Clamp (Integer)"), Category="Math|Integer")
@@ -397,11 +403,11 @@ class UKismetMathLibrary : public UBlueprintFunctionLibrary
 
 	/** Returns Value between A and B (inclusive) that wraps around */
 	UFUNCTION(BlueprintPure, meta = (DisplayName = "Wrap (Integer)", Min = "0", Max = "100"), Category = "Math|Integer")
-	static int32 Wrap(int32 Value, int32 Min, int32 Max);
+	static UE_INL_API int32 Wrap(int32 Value, int32 Min, int32 Max);
 
 	/** Returns the absolute (positive) value of A */
 	UFUNCTION(BlueprintPure, meta=(DisplayName = "Absolute (Integer)", CompactNodeTitle = "ABS"), Category="Math|Integer")
-	static int32 Abs_Int(int32 A);
+	static UE_INL_API int32 Abs_Int(int32 A);
 
 	//
 	// Integer64 functions.
@@ -409,47 +415,47 @@ class UKismetMathLibrary : public UBlueprintFunctionLibrary
 	
 	/** Multiplication (A * B) */
 	UFUNCTION(BlueprintPure, meta=(DisplayName = "integer64 * integer64", CompactNodeTitle = "*", Keywords = "* multiply", CommutativeAssociativeBinaryOperator = "true"), Category="Math|Integer64")
-	static int64 Multiply_Int64Int64(int64 A, int64 B);
+	static UE_INL_API int64 Multiply_Int64Int64(int64 A, int64 B);
 
 	/** Division (A / B) */
 	UFUNCTION(BlueprintPure, meta=(DisplayName = "integer64 / integer64", CompactNodeTitle = "/", Keywords = "/ divide division"), Category="Math|Integer64")
-	static int64 Divide_Int64Int64(int64 A, int64 B = 1);
+	static UE_INL_API int64 Divide_Int64Int64(int64 A, int64 B = 1);
 	
 	/** Modulo (A % B) */
 	UFUNCTION(BlueprintPure, meta = (DisplayName = "% (integer64)", CompactNodeTitle = "%", Keywords = "% modulus"), Category = "Math|Integer64")
-	static int64 Percent_Int64Int64(int64 A, int64 B = 1);
+	static UE_INL_API int64 Percent_Int64Int64(int64 A, int64 B = 1);
 
 	/** Addition (A + B) */
 	UFUNCTION(BlueprintPure, meta=(DisplayName = "integer64 + integer64", CompactNodeTitle = "+", Keywords = "+ add plus", CommutativeAssociativeBinaryOperator = "true"), Category="Math|Integer64")
-	static int64 Add_Int64Int64(int64 A, int64 B = 1);
+	static UE_INL_API int64 Add_Int64Int64(int64 A, int64 B = 1);
 
 	/** Subtraction (A - B) */
 	UFUNCTION(BlueprintPure, meta=(DisplayName = "integer64 - integer64", CompactNodeTitle = "-", Keywords = "- subtract minus"), Category="Math|Integer64")
-	static int64 Subtract_Int64Int64(int64 A, int64 B = 1);
+	static UE_INL_API int64 Subtract_Int64Int64(int64 A, int64 B = 1);
 
 	/** Returns true if A is less than B (A < B) */
 	UFUNCTION(BlueprintPure, meta=(DisplayName = "integer64 < integer64", CompactNodeTitle = "<", Keywords = "< less"), Category="Math|Integer64")
-	static bool Less_Int64Int64(int64 A, int64 B);
+	static UE_INL_API bool Less_Int64Int64(int64 A, int64 B);
 
 	/** Returns true if A is greater than B (A > B) */
 	UFUNCTION(BlueprintPure, meta=(DisplayName = "integer64 > integer64", CompactNodeTitle = ">", Keywords = "> greater"), Category="Math|Integer64")
-	static bool Greater_Int64Int64(int64 A, int64 B);
+	static UE_INL_API bool Greater_Int64Int64(int64 A, int64 B);
 
 	/** Returns true if A is less than or equal to B (A <= B) */
 	UFUNCTION(BlueprintPure, meta=(DisplayName = "integer64 <= integer64", CompactNodeTitle = "<=", Keywords = "<= less"), Category="Math|Integer64")
-	static bool LessEqual_Int64Int64(int64 A, int64 B);
+	static UE_INL_API bool LessEqual_Int64Int64(int64 A, int64 B);
 
 	/** Returns true if A is greater than or equal to B (A >= B) */
 	UFUNCTION(BlueprintPure, meta=(DisplayName = "integer64 >= integer64", CompactNodeTitle = ">=", Keywords = ">= greater"), Category="Math|Integer64")
-	static bool GreaterEqual_Int64Int64(int64 A, int64 B);
+	static UE_INL_API bool GreaterEqual_Int64Int64(int64 A, int64 B);
 
 	/** Returns true if A is equal to B (A == B) */
 	UFUNCTION(BlueprintPure, meta=(DisplayName = "Equal (Integer64)", CompactNodeTitle = "==", Keywords = "== equal"), Category="Math|Integer64")
-	static bool EqualEqual_Int64Int64(int64 A, int64 B);
+	static UE_INL_API bool EqualEqual_Int64Int64(int64 A, int64 B);
 
 	/** Returns true if A is not equal to B (A != B) */
 	UFUNCTION(BlueprintPure, meta=(DisplayName = "Not Equal (Integer64)", CompactNodeTitle = "!=", Keywords = "!= not equal"), Category="Math|Integer64")
-	static bool NotEqual_Int64Int64(int64 A, int64 B);
+	static UE_INL_API bool NotEqual_Int64Int64(int64 A, int64 B);
 	
 	/** Returns true if value is between Min and Max (V >= Min && V <= Max)
 	 * If InclusiveMin is true, value needs to be equal or larger than Min, else it needs to be larger
@@ -460,23 +466,23 @@ class UKismetMathLibrary : public UBlueprintFunctionLibrary
 
 	/** Bitwise AND (A & B) */
 	UFUNCTION(BlueprintPure, meta=(DisplayName = "Bitwise AND", CompactNodeTitle = "&", Keywords = "& and", CommutativeAssociativeBinaryOperator = "true"), Category="Math|Integer64")
-	static int64 And_Int64Int64(int64 A, int64 B);
+	static UE_INL_API int64 And_Int64Int64(int64 A, int64 B);
 
 	/** Bitwise XOR (A ^ B) */
 	UFUNCTION(BlueprintPure, meta=(DisplayName = "Bitwise XOR", CompactNodeTitle = "^", Keywords = "^ xor"), Category="Math|Integer64")
-	static int64 Xor_Int64Int64(int64 A, int64 B);
+	static UE_INL_API int64 Xor_Int64Int64(int64 A, int64 B);
 
 	/** Bitwise OR (A | B) */
 	UFUNCTION(BlueprintPure, meta=(DisplayName = "Bitwise OR", CompactNodeTitle = "|", Keywords = "| or", CommutativeAssociativeBinaryOperator = "true"), Category="Math|Integer64")
-	static int64 Or_Int64Int64(int64 A, int64 B);
+	static UE_INL_API int64 Or_Int64Int64(int64 A, int64 B);
 
 	/** Bitwise NOT (~A) */
 	UFUNCTION(BlueprintPure, meta = (DisplayName = "Bitwise NOT", CompactNodeTitle = "~", Keywords = "~ not"), Category = "Math|Integer64")
-	static int64 Not_Int64(int64 A);
+	static UE_INL_API int64 Not_Int64(int64 A);
 
 	/** Sign (integer64, returns -1 if A < 0, 0 if A is zero, and +1 if A > 0) */
 	UFUNCTION(BlueprintPure, meta=(DisplayName = "Sign (Integer64)"), Category="Math|Integer64")
-	static int64 SignOfInteger64(int64 A);
+	static UE_INL_API int64 SignOfInteger64(int64 A);
 
 	/** Returns a uniformly distributed random number between 0 and Max - 1 */
 	UFUNCTION(BlueprintPure, Category="Math|Random", meta=(NotBlueprintThreadSafe))
@@ -484,15 +490,15 @@ class UKismetMathLibrary : public UBlueprintFunctionLibrary
 
 	/** Return a random integer64 between Min and Max (>= Min and <= Max) */
 	UFUNCTION(BlueprintPure, Category="Math|Random", meta = (NotBlueprintThreadSafe))
-	static int64 RandomInteger64InRange(int64 Min, int64 Max);
+	static UE_INL_API int64 RandomInteger64InRange(int64 Min, int64 Max);
 
 	/** Returns the minimum value of A and B */
 	UFUNCTION(BlueprintPure, meta=(DisplayName = "Min (Integer64)", CompactNodeTitle = "MIN", CommutativeAssociativeBinaryOperator = "true"), Category="Math|Integer64")
-	static int64 MinInt64(int64 A, int64 B);
+	static UE_INL_API int64 MinInt64(int64 A, int64 B);
 
 	/** Returns the maximum value of A and B */
 	UFUNCTION(BlueprintPure, meta=(DisplayName = "Max (Integer64)", CompactNodeTitle = "MAX", CommutativeAssociativeBinaryOperator = "true"), Category="Math|Integer64")
-	static int64 MaxInt64(int64 A, int64 B);
+	static UE_INL_API int64 MaxInt64(int64 A, int64 B);
 
 	/** Returns Value clamped to be between A and B (inclusive) */
 	UFUNCTION(BlueprintPure, meta=(DisplayName = "Clamp (Integer64)"), Category="Math|Integer64")
@@ -500,7 +506,7 @@ class UKismetMathLibrary : public UBlueprintFunctionLibrary
 
 	/** Returns the absolute (positive) value of A */
 	UFUNCTION(BlueprintPure, meta=(DisplayName = "Absolute (Integer64)", CompactNodeTitle = "ABS"), Category="Math|Integer64")
-	static int64 Abs_Int64(int64 A);
+	static UE_INL_API int64 Abs_Int64(int64 A);
 
 	//
 	// Float functions.
@@ -508,11 +514,11 @@ class UKismetMathLibrary : public UBlueprintFunctionLibrary
 
 	/** Power (Base to the Exp-th power) */
 	UFUNCTION(BlueprintPure, meta = (DisplayName = "Power"), Category = "Math|Float")
-	static double MultiplyMultiply_FloatFloat(double Base, double Exp);
+	static UE_INL_API double MultiplyMultiply_FloatFloat(double Base, double Exp);
 
 	/** Multiplication (A * B) */
 	UFUNCTION(BlueprintPure, meta = (DisplayName = "int * float", CompactNodeTitle = "*", Keywords = "* multiply"), Category = "Math|Float")
-	static double Multiply_IntFloat(int32 A, double B);
+	static UE_INL_API double Multiply_IntFloat(int32 A, double B);
 
 	/** Modulo (A % B) */
 	UFUNCTION(BlueprintPure, CustomThunk, meta = (DisplayName = "% (Float)", CompactNodeTitle = "%", Keywords = "% modulus"), Category = "Math|Float")
@@ -540,51 +546,51 @@ class UKismetMathLibrary : public UBlueprintFunctionLibrary
 
 	/** Returns the fractional part of a float. */
 	UFUNCTION(BlueprintPure, Category="Math|Float")
-	static double Fraction(double A);
+	static UE_INL_API double Fraction(double A);
 
 	/** Addition (A + B) */
 	UFUNCTION(BlueprintPure, meta = (DisplayName = "float + float", CompactNodeTitle = "+", Keywords = "+ add plus", CommutativeAssociativeBinaryOperator = "true"), Category = "Math|Float")
-	static double Add_DoubleDouble(double A, double B = 1.0);
+	static UE_INL_API double Add_DoubleDouble(double A, double B = 1.0);
 
 	/** Subtraction (A - B) */
 	UFUNCTION(BlueprintPure, meta = (DisplayName = "float - float", CompactNodeTitle = "-", Keywords = "- subtract minus"), Category = "Math|Float")
-	static double Subtract_DoubleDouble(double A, double B = 1.0);
+	static UE_INL_API double Subtract_DoubleDouble(double A, double B = 1.0);
 
 	/** Multiplication (A * B) */
 	UFUNCTION(BlueprintPure, meta = (DisplayName = "float * float", CompactNodeTitle = "*", Keywords = "* multiply", CommutativeAssociativeBinaryOperator = "true"), Category = "Math|Float")
-	static double Multiply_DoubleDouble(double A, double B);
+	static UE_INL_API double Multiply_DoubleDouble(double A, double B);
 
 	/** Division (A / B) */
 	UFUNCTION(BlueprintPure, meta = (DisplayName = "float / float", CompactNodeTitle = "/", Keywords = "/ divide division"), Category = "Math|Float")
-	static double Divide_DoubleDouble(double A, double B = 1.0);
+	static UE_INL_API double Divide_DoubleDouble(double A, double B = 1.0);
 
 	/** Returns true if A is Less than B (A < B) */
 	UFUNCTION(BlueprintPure, meta = (DisplayName = "float < float", CompactNodeTitle = "<", Keywords = "< less"), Category = "Math|Float")
-	static bool Less_DoubleDouble(double A, double B);
+	static UE_INL_API bool Less_DoubleDouble(double A, double B);
 
 	/** Returns true if A is greater than B (A > B) */
 	UFUNCTION(BlueprintPure, meta = (DisplayName = "float > float", CompactNodeTitle = ">", Keywords = "> greater"), Category = "Math|Float")
-	static bool Greater_DoubleDouble(double A, double B);
+	static UE_INL_API bool Greater_DoubleDouble(double A, double B);
 
 	/** Returns true if A is Less than or equal to B (A <= B) */
 	UFUNCTION(BlueprintPure, meta = (DisplayName = "float <= float", CompactNodeTitle = "<=", Keywords = "<= less"), Category = "Math|Float")
-	static bool LessEqual_DoubleDouble(double A, double B);
+	static UE_INL_API bool LessEqual_DoubleDouble(double A, double B);
 
 	/** Returns true if A is greater than or equal to B (A >= B) */
 	UFUNCTION(BlueprintPure, meta = (DisplayName = "float >= float", CompactNodeTitle = ">=", Keywords = ">= greater"), Category = "Math|Float")
-	static bool GreaterEqual_DoubleDouble(double A, double B);
+	static UE_INL_API bool GreaterEqual_DoubleDouble(double A, double B);
 
 	/** Returns true if A is exactly equal to B (A == B)*/
 	UFUNCTION(BlueprintPure, meta = (DisplayName = "Equal (Float)", CompactNodeTitle = "==", Keywords = "== equal"), Category = "Math|Float")
-	static bool EqualEqual_DoubleDouble(double A, double B);
+	static UE_INL_API bool EqualEqual_DoubleDouble(double A, double B);
 
 	/** Returns true if A is nearly equal to B (|A - B| < ErrorTolerance) */
 	UFUNCTION(BlueprintPure, meta = (DisplayName = "Nearly Equal (Float)", Keywords = "== equal"), Category = "Math|Float")
-	static bool NearlyEqual_FloatFloat(double A, double B, double ErrorTolerance = 1.e-6);
+	static UE_INL_API bool NearlyEqual_FloatFloat(double A, double B, double ErrorTolerance = 1.e-6);
 
 	/** Returns true if A does not equal B (A != B)*/
 	UFUNCTION(BlueprintPure, meta = (DisplayName = "Not Equal (Float)", CompactNodeTitle = "!=", Keywords = "!= not equal"), Category = "Math|Float")
-	static bool NotEqual_DoubleDouble(double A, double B);
+	static UE_INL_API bool NotEqual_DoubleDouble(double A, double B);
 
 	/** Returns true if value is between Min and Max (V >= Min && V <= Max)
 	 * If InclusiveMin is true, value needs to be equal or larger than Min, else it needs to be larger
@@ -603,43 +609,43 @@ class UKismetMathLibrary : public UBlueprintFunctionLibrary
 	 * if GridSize is very small precision issues may occur.
 	 */
 	UFUNCTION(BlueprintPure, meta = (DisplayName = "Snap To Grid (Float)"), Category = "Math|Float")
-	static double GridSnap_Float(double Location, double GridSize);
+	static UE_INL_API double GridSnap_Float(double Location, double GridSize);
 
 	/** Returns the absolute (positive) value of A */
 	UFUNCTION(BlueprintPure, meta=(DisplayName = "Absolute (Float)", CompactNodeTitle = "ABS"), Category="Math|Float")
-	static double Abs(double A);
+	static UE_INL_API double Abs(double A);
 
 	/** Returns the sine of A (expects Radians)*/
 	UFUNCTION(BlueprintPure, meta=(DisplayName = "Sin (Radians)", CompactNodeTitle = "SIN", Keywords = "sine"), Category="Math|Trig")
-	static double Sin(double A);
+	static UE_INL_API double Sin(double A);
 
 	/** Returns the inverse sine (arcsin) of A (result is in Radians) */
 	UFUNCTION(BlueprintPure, meta=(DisplayName = "Asin (Radians)", CompactNodeTitle = "ASIN", Keywords = "sine"), Category="Math|Trig")
-	static double Asin(double A);
+	static UE_INL_API double Asin(double A);
 
 	/** Returns the cosine of A (expects Radians)*/
 	UFUNCTION(BlueprintPure, meta=(DisplayName = "Cos (Radians)", CompactNodeTitle = "COS"), Category="Math|Trig")
-	static double Cos(double A);
+	static UE_INL_API double Cos(double A);
 
 	/** Returns the inverse cosine (arccos) of A (result is in Radians) */
 	UFUNCTION(BlueprintPure, meta=(DisplayName = "Acos (Radians)", CompactNodeTitle = "ACOS"), Category="Math|Trig")
-	static double Acos(double A);
+	static UE_INL_API double Acos(double A);
 
 	/** Returns the tan of A (expects Radians)*/
 	UFUNCTION(BlueprintPure, meta=(DisplayName = "Tan (Radians)", CompactNodeTitle = "TAN"), Category="Math|Trig")
-	static double Tan(double A);
+	static UE_INL_API double Tan(double A);
 
 	/** Returns the inverse tan (atan) (result is in Radians)*/
 	UFUNCTION(BlueprintPure, meta=(DisplayName = "Atan (Radians)"), Category="Math|Trig")
-	static double Atan(double A);
+	static UE_INL_API double Atan(double A);
 
 	/** Returns the inverse tan (atan2) of A/B (result is in Radians)*/
 	UFUNCTION(BlueprintPure, meta=(DisplayName = "Atan2 (Radians)"), Category="Math|Trig")
-	static double Atan2(double Y, double X);
+	static UE_INL_API double Atan2(double Y, double X);
 
 	/** Returns exponential(e) to the power A (e^A)*/
 	UFUNCTION(BlueprintPure, Category="Math|Float", meta=(CompactNodeTitle = "e"))
-	static double Exp(double A);
+	static UE_INL_API double Exp(double A);
 
 	/** Returns log of A base B (if B^R == A, returns R)*/
 	UFUNCTION(BlueprintPure, Category = "Math|Float")
@@ -647,67 +653,67 @@ class UKismetMathLibrary : public UBlueprintFunctionLibrary
 
 	/** Returns natural log of A (if e^R == A, returns R)*/
 	UFUNCTION(BlueprintPure, Category="Math|Float")
-	static double Loge(double A);
+	static UE_INL_API double Loge(double A);
 
 	/** Returns square root of A*/
 	UFUNCTION(BlueprintPure, Category="Math|Float", meta=(Keywords = "square root", CompactNodeTitle = "SQRT"))
-	static double Sqrt(double A);
+	static UE_INL_API double Sqrt(double A);
 
 	/** Returns square of A (A*A)*/
 	UFUNCTION(BlueprintPure, Category="Math|Float", meta=(CompactNodeTitle = "^2"))
-	static double Square(double A);
+	static UE_INL_API double Square(double A);
 
 	/** Returns a random float between 0 and 1 */
 	UFUNCTION(BlueprintPure, Category="Math|Random", meta=(NotBlueprintThreadSafe))
-	static double RandomFloat();
+	static UE_INL_API double RandomFloat();
 
 	/** Generate a random number between Min and Max */
 	UFUNCTION(BlueprintPure, Category="Math|Random", meta=(NotBlueprintThreadSafe))
-	static double RandomFloatInRange(double Min, double Max);
+	static UE_INL_API double RandomFloatInRange(double Min, double Max);
 
 	/** Returns the value of PI */
 	UFUNCTION(BlueprintPure, meta=(DisplayName = "Get PI", CompactNodeTitle = "PI"), Category="Math|Trig")
-	static double GetPI();
+	static UE_INL_API double GetPI();
 
 	/** Returns the value of TAU (= 2 * PI) */
 	UFUNCTION(BlueprintPure, meta=(DisplayName = "Get TAU", CompactNodeTitle = "TAU"), Category="Math|Trig")
-	static double GetTAU();
+	static UE_INL_API double GetTAU();
 
 	/** Returns radians value based on the input degrees */
 	UFUNCTION(BlueprintPure, meta=(DisplayName = "Degrees To Radians", CompactNodeTitle = "D2R"), Category="Math|Trig")
-	static double DegreesToRadians(double A);
+	static UE_INL_API double DegreesToRadians(double A);
 
 	/** Returns degrees value based on the input radians */
 	UFUNCTION(BlueprintPure, meta=(DisplayName = "Radians To Degrees", CompactNodeTitle = "R2D"), Category="Math|Trig")
-	static double RadiansToDegrees(double A);
+	static UE_INL_API double RadiansToDegrees(double A);
 
 	/** Returns the sin of A (expects Degrees)*/
 	UFUNCTION(BlueprintPure, meta=(DisplayName = "Sin (Degrees)", CompactNodeTitle = "SINd", Keywords = "sine"), Category="Math|Trig")
-	static double DegSin(double A);
+	static UE_INL_API double DegSin(double A);
 
 	/** Returns the inverse sin (arcsin) of A (result is in Degrees) */
 	UFUNCTION(BlueprintPure, meta=(DisplayName = "Asin (Degrees)", CompactNodeTitle = "ASINd", Keywords = "sine"), Category="Math|Trig")
-	static double DegAsin(double A);
+	static UE_INL_API double DegAsin(double A);
 
 	/** Returns the cos of A (expects Degrees)*/
 	UFUNCTION(BlueprintPure, meta=(DisplayName = "Cos (Degrees)", CompactNodeTitle = "COSd"), Category="Math|Trig")
-	static double DegCos(double A);
+	static UE_INL_API double DegCos(double A);
 
 	/** Returns the inverse cos (arccos) of A (result is in Degrees) */
 	UFUNCTION(BlueprintPure, meta=(DisplayName = "Acos (Degrees)", CompactNodeTitle = "ACOSd"), Category="Math|Trig")
-	static double DegAcos(double A);
+	static UE_INL_API double DegAcos(double A);
 
 	/** Returns the tan of A (expects Degrees)*/
 	UFUNCTION(BlueprintPure, meta=(DisplayName = "Tan (Degrees)", CompactNodeTitle = "TANd"), Category="Math|Trig")
-	static double DegTan(double A);
+	static UE_INL_API double DegTan(double A);
 
 	/** Returns the inverse tan (atan) (result is in Degrees)*/
 	UFUNCTION(BlueprintPure, meta=(DisplayName = "Atan (Degrees)"), Category="Math|Trig")
-	static double DegAtan(double A);
+	static UE_INL_API double DegAtan(double A);
 
 	/** Returns the inverse tan (atan2) of A/B (result is in Degrees)*/
 	UFUNCTION(BlueprintPure, meta=(DisplayName = "Atan2 (Degrees)"), Category="Math|Trig")
-	static double DegAtan2(double Y, double X);
+	static UE_INL_API double DegAtan2(double Y, double X);
 
 	/** 
 	 * Clamps an arbitrary angle to be between the given angles.  Will clamp to nearest boundary.
@@ -717,15 +723,15 @@ class UKismetMathLibrary : public UBlueprintFunctionLibrary
 	 * @return Returns clamped angle in the range -180..180.
 	 */
 	UFUNCTION(BlueprintPure, meta=(DisplayName = "Clamp Angle"), Category="Math|Float")
-	static double ClampAngle(double AngleDegrees, double MinAngleDegrees, double MaxAngleDegrees);
+	static UE_INL_API double ClampAngle(double AngleDegrees, double MinAngleDegrees, double MaxAngleDegrees);
 
 	/** Returns the minimum value of A and B */
 	UFUNCTION(BlueprintPure, meta=(DisplayName = "Min (Float)", CompactNodeTitle = "MIN", CommutativeAssociativeBinaryOperator = "true"), Category="Math|Float")
-	static double FMin(double A, double B);
+	static UE_INL_API double FMin(double A, double B);
 
 	/** Returns the maximum value of A and B */
 	UFUNCTION(BlueprintPure, meta=(DisplayName = "Max (Float)", CompactNodeTitle = "MAX", CommutativeAssociativeBinaryOperator = "true"), Category="Math|Float")
-	static double FMax(double A, double B);
+	static UE_INL_API double FMax(double A, double B);
 
 	/** Returns Value clamped between A and B (inclusive) */
 	UFUNCTION(BlueprintPure, meta=(DisplayName = "Clamp (Float)", Min="0.0", Max="1.0"), Category="Math|Float")
@@ -733,7 +739,7 @@ class UKismetMathLibrary : public UBlueprintFunctionLibrary
 
 	/** Returns Value wrapped from A and B (inclusive) */
 	UFUNCTION(BlueprintPure, meta = (DisplayName = "Wrap (Float)", Min = "0.0", Max = "1.0"), Category = "Math|Float")
-	static double FWrap(double Value, double Min, double Max);
+	static UE_INL_API double FWrap(double Value, double Min, double Max);
 
 	/** This functions returns 0 if B (the denominator) is zero */
 	UFUNCTION(BlueprintPure, Category = "Math|Float", meta = (Keywords = "percent"))
@@ -781,35 +787,35 @@ class UKismetMathLibrary : public UBlueprintFunctionLibrary
 
 	/** Rounds A to the nearest integer (e.g., -1.6 becomes -2 and 1.6 becomes 2) */
 	UFUNCTION(BlueprintPure, Category="Math|Float")
-	static int32 Round(double A);
+	static UE_INL_API int32 Round(double A);
 
 	/** Rounds A down towards negative infinity / down to the previous integer (e.g., -1.6 becomes -2 and 1.6 becomes 1) */
 	UFUNCTION(BlueprintPure, meta=(DisplayName = "Floor"), Category="Math|Float")
-	static int32 FFloor(double A);
+	static UE_INL_API int32 FFloor(double A);
 	
 	/** Rounds A towards zero, truncating the fractional part (e.g., -1.6 becomes -1 and 1.6 becomes 1) */
 	UFUNCTION(BlueprintPure, meta=(DisplayName = "Truncate", BlueprintAutocast), Category="Math|Float")
-	static int32 FTrunc(double A); // TODO: ok to change with BP Autocast?
+	static UE_INL_API int32 FTrunc(double A); // TODO: ok to change with BP Autocast?
 
 	/** Rounds A up towards positive infinity / up to the next integer (e.g., -1.6 becomes -1 and 1.6 becomes 2) */
 	UFUNCTION(BlueprintPure, Category = "Math|Float", meta=(DisplayName="Ceil"))
-	static int32 FCeil(double A);
+	static UE_INL_API int32 FCeil(double A);
 
 	/** Rounds A to the nearest integer (e.g., -1.6 becomes -2 and 1.6 becomes 2) */
 	UFUNCTION(BlueprintPure, meta = (DisplayName = "Round to Integer64"), Category = "Math|Float")
-	static int64 Round64(double A);
+	static UE_INL_API int64 Round64(double A);
 
 	/** Rounds A down towards negative infinity / down to the previous integer (e.g., -1.6 becomes -2 and 1.6 becomes 1) */
 	UFUNCTION(BlueprintPure, meta = (DisplayName = "Floor to Integer64"), Category = "Math|Float")
-	static int64 FFloor64(double A);
+	static UE_INL_API int64 FFloor64(double A);
 
 	/** Rounds A towards zero, truncating the fractional part (e.g., -1.6 becomes -1 and 1.6 becomes 1) */
 	UFUNCTION(BlueprintPure, meta = (DisplayName = "Truncate to Integer64", BlueprintAutocast), Category = "Math|Float")
-	static int64 FTrunc64(double A);
+	static UE_INL_API int64 FTrunc64(double A);
 
 	/** Rounds A up towards positive infinity / up to the next integer (e.g., -1.6 becomes -1 and 1.6 becomes 2) */
 	UFUNCTION(BlueprintPure, meta = (DisplayName = "Ceil to Integer64"), Category = "Math|Float")
-	static int64 FCeil64(double A);
+	static UE_INL_API int64 FCeil64(double A);
 
 	/** Returns the number of times Divisor will go into Dividend (i.e., Dividend divided by Divisor), as well as the remainder */
 	UFUNCTION(BlueprintPure, meta=(DisplayName = "Division (Whole and Remainder)"), Category="Math|Float")
@@ -821,7 +827,7 @@ class UKismetMathLibrary : public UBlueprintFunctionLibrary
 
 	/** Sign (float, returns -1 if A < 0, 0 if A is zero, and +1 if A > 0) */
 	UFUNCTION(BlueprintPure, meta=(DisplayName = "Sign (Float)"), Category="Math|Float")
-	static double SignOfFloat(double A);
+	static UE_INL_API double SignOfFloat(double A);
 
 	/** Returns Value normalized to the given range.  (e.g. 20 normalized to the range 10->50 would result in 0.25) */
 	UFUNCTION(BlueprintPure, Category="Math|Float")
@@ -837,7 +843,7 @@ class UKismetMathLibrary : public UBlueprintFunctionLibrary
 	
 	/** Multiplies the input value by pi. */
 	UFUNCTION(BlueprintPure, meta=(Keywords = "* multiply"), Category="Math|Float")
-	static double MultiplyByPi(double Value);
+	static UE_INL_API double MultiplyByPi(double Value);
 
 	/** Interpolate between A and B, applying an ease in/out function.  Exp controls the degree of the curve. */
 	UFUNCTION(BlueprintPure, Category = "Math|Float")
@@ -879,7 +885,7 @@ class UKismetMathLibrary : public UBlueprintFunctionLibrary
 	 * @return a new rotation component value clamped in the range (-360,360)
 	 */
 	UFUNCTION(BlueprintPure, Category="Math|Float")
-	static float FixedTurn(float InCurrent, float InDesired, float InDeltaRate);
+	static UE_INL_API float FixedTurn(float InCurrent, float InDesired, float InDeltaRate);
 
 
 	//
@@ -888,27 +894,27 @@ class UKismetMathLibrary : public UBlueprintFunctionLibrary
 	
 	/** Zero Int Point (0, 0) */
 	UFUNCTION(BlueprintPure, meta = (ScriptConstant = "Zero", ScriptConstantHost = "/Script/CoreUObject.IntPoint"), Category = "Math|IntPoint|Constants")
-	static FIntPoint IntPoint_Zero();
+	static UE_INL_API FIntPoint IntPoint_Zero();
 	
 	/** One Int Point (1, 1) */
 	UFUNCTION(BlueprintPure, meta = (ScriptConstant = "One", ScriptConstantHost = "/Script/CoreUObject.IntPoint"), Category = "Math|IntPoint|Constants")
-	static FIntPoint IntPoint_One();
+	static UE_INL_API FIntPoint IntPoint_One();
 	
 	/** Up Int Point (0, -1) */
 	UFUNCTION(BlueprintPure, meta = (ScriptConstant = "Up", ScriptConstantHost = "/Script/CoreUObject.IntPoint"), Category = "Math|IntPoint|Constants")
-	static FIntPoint IntPoint_Up();
+	static UE_INL_API FIntPoint IntPoint_Up();
 	
 	/** Left Int Point (-1, 0) */
 	UFUNCTION(BlueprintPure, meta = (ScriptConstant = "Left", ScriptConstantHost = "/Script/CoreUObject.IntPoint"), Category = "Math|IntPoint|Constants")
-	static FIntPoint IntPoint_Left();
+	static UE_INL_API FIntPoint IntPoint_Left();
 	
 	/** Right Int Point (1, 0) */
 	UFUNCTION(BlueprintPure, meta = (ScriptConstant = "Right", ScriptConstantHost = "/Script/CoreUObject.IntPoint"), Category = "Math|IntPoint|Constants")
-	static FIntPoint IntPoint_Right();
+	static UE_INL_API FIntPoint IntPoint_Right();
 	
 	/** Down Int Point (0, 1) */
 	UFUNCTION(BlueprintPure, meta = (ScriptConstant = "Down", ScriptConstantHost = "/Script/CoreUObject.IntPoint"), Category = "Math|IntPoint|Constants")
-	static FIntPoint IntPoint_Down();
+	static UE_INL_API FIntPoint IntPoint_Down();
 
 	//
 	// IntPoint functions
@@ -916,47 +922,47 @@ class UKismetMathLibrary : public UBlueprintFunctionLibrary
 
 	/** Converts an IntPoint to a Vector2D */
 	UFUNCTION(BlueprintPure, meta = (DisplayName = "To Vector2D (IntPoint)", CompactNodeTitle = "->", ScriptMethod = "Vector2D", Keywords = "cast convert", BlueprintAutocast), Category = "Math|Conversions")
-	static FVector2D Conv_IntPointToVector2D(FIntPoint InIntPoint);
+	static UE_INL_API FVector2D Conv_IntPointToVector2D(FIntPoint InIntPoint);
 
 	/** Returns IntPoint A added by B */
 	UFUNCTION(BlueprintPure, meta = (DisplayName = "IntPoint + IntPoint", CompactNodeTitle = "+", ScriptMethod = "Add", ScriptOperator = "+;+=", Keywords = "+ add plus"), Category = "Math|IntPoint")
-	static FIntPoint Add_IntPointIntPoint(FIntPoint A, FIntPoint B);
+	static UE_INL_API FIntPoint Add_IntPointIntPoint(FIntPoint A, FIntPoint B);
 
 	/** Addition (A - B) */
 	UFUNCTION(BlueprintPure, meta = (DisplayName = "IntPoint + Integer", CompactNodeTitle = "+", ScriptMethod = "AddInt", ScriptOperator = "+;+=", Keywords = "+ add plus"), Category = "Math|IntPoint")
-	static FIntPoint Add_IntPointInt(FIntPoint A, int32 B);
+	static UE_INL_API FIntPoint Add_IntPointInt(FIntPoint A, int32 B);
 
 	/** Returns IntPoint A subtracted by B */
 	UFUNCTION(BlueprintPure, meta = (DisplayName = "IntPoint - IntPoint", CompactNodeTitle = "-", ScriptMethod = "Subtract", ScriptOperator = "-;-=", Keywords = "- subtract minus"), Category = "Math|IntPoint")
-	static FIntPoint Subtract_IntPointIntPoint(FIntPoint A, FIntPoint B);
+	static UE_INL_API FIntPoint Subtract_IntPointIntPoint(FIntPoint A, FIntPoint B);
 
 	/** Subtraction (A - B) */
 	UFUNCTION(BlueprintPure, meta = (DisplayName = "IntPoint - Integer", CompactNodeTitle = "-", ScriptMethod = "SubtractInt", ScriptOperator = "-;-=", Keywords = "- subtract minus"), Category = "Math|IntPoint")
-	static FIntPoint Subtract_IntPointInt(FIntPoint A, int32 B);
+	static UE_INL_API FIntPoint Subtract_IntPointInt(FIntPoint A, int32 B);
 
 	/** Returns IntPoint A multiplied by B */
 	UFUNCTION(BlueprintPure, meta = (DisplayName = "IntPoint * IntPoint", CompactNodeTitle = "*", ScriptMethod = "Multiply", ScriptOperator = "*;*=", Keywords = "* multiply"), Category = "Math|IntPoint")
-	static FIntPoint Multiply_IntPointIntPoint(FIntPoint A, FIntPoint B);
+	static UE_INL_API FIntPoint Multiply_IntPointIntPoint(FIntPoint A, FIntPoint B);
 
 	/** Multiplication (A * B) */
 	UFUNCTION(BlueprintPure, meta = (DisplayName = "IntPoint * Integer", CompactNodeTitle = "*", ScriptMethod = "MultiplyInt", ScriptOperator = "*;*=", Keywords = "* multiply"), Category = "Math|IntPoint")
-	static FIntPoint Multiply_IntPointInt(FIntPoint A, int32 B);
+	static UE_INL_API FIntPoint Multiply_IntPointInt(FIntPoint A, int32 B);
 
 	/** Returns IntPoint A divided by B */
 	UFUNCTION(BlueprintPure, meta = (DisplayName = "IntPoint / IntPoint", CompactNodeTitle = "/", ScriptMethod = "Divide", ScriptOperator = "/;/=", Keywords = "/ divide"), Category = "Math|IntPoint")
-	static FIntPoint Divide_IntPointIntPoint(FIntPoint A, FIntPoint B);
+	static UE_INL_API FIntPoint Divide_IntPointIntPoint(FIntPoint A, FIntPoint B);
 
 	/** Division (A * B) */
 	UFUNCTION(BlueprintPure, meta = (DisplayName = "IntPoint / Integer", CompactNodeTitle = "/", ScriptMethod = "DivideInt", ScriptOperator = "/;/=", Keywords = "/ divide"), Category = "Math|IntPoint")
-	static FIntPoint Divide_IntPointInt(FIntPoint A, int32 B);
+	static UE_INL_API FIntPoint Divide_IntPointInt(FIntPoint A, int32 B);
 
 	/** Returns true if IntPoint A is equal to IntPoint B (A == B) */
 	UFUNCTION(BlueprintPure, meta = (DisplayName = "Equal (IntPoint)", CompactNodeTitle = "==", ScriptMethod = "Equals", ScriptOperator = "==", Keywords = "== equal"), Category = "Math|IntPoint")
-	static bool Equal_IntPointIntPoint(FIntPoint A, FIntPoint B);
+	static UE_INL_API bool Equal_IntPointIntPoint(FIntPoint A, FIntPoint B);
 
 	/** Returns true if IntPoint A is NOT equal to IntPoint B (A != B) */
 	UFUNCTION(BlueprintPure, meta = (DisplayName = "Not Equal (IntPoint)", CompactNodeTitle = "!=", ScriptMethod = "NotEqual", ScriptOperator = "!=", Keywords = "!= not equal"), Category = "Math|IntPoint")
-	static bool NotEqual_IntPointIntPoint(FIntPoint A, FIntPoint B);
+	static UE_INL_API bool NotEqual_IntPointIntPoint(FIntPoint A, FIntPoint B);
 
 	//
 	// Vector2D constants - exposed for scripting
@@ -964,15 +970,15 @@ class UKismetMathLibrary : public UBlueprintFunctionLibrary
 
 	/** 2D one vector constant (1,1) */
 	UFUNCTION(BlueprintPure, meta = (ScriptConstant = "One", ScriptConstantHost = "/Script/CoreUObject.Vector2D"), Category = "Math|Vector2D")
-	static FVector2D Vector2D_One();
+	static UE_INL_API FVector2D Vector2D_One();
 
 	/** 2D unit vector constant along the 45 degree angle or symmetrical positive axes (sqrt(.5),sqrt(.5)) or (.707,.707). https://en.wikipedia.org/wiki/Unit_vector */
 	UFUNCTION(BlueprintPure, meta = (ScriptConstant = "Unit45Deg", ScriptConstantHost = "/Script/CoreUObject.Vector2D"), Category = "Math|Vector2D")
-	static FVector2D Vector2D_Unit45Deg();
+	static UE_INL_API FVector2D Vector2D_Unit45Deg();
 
 	/** 2D zero vector constant (0,0) */
 	UFUNCTION(BlueprintPure, meta = (ScriptConstant = "Zero", ScriptConstantHost = "/Script/CoreUObject.Vector2D"), Category = "Math|Vector2D")
-	static FVector2D Vector2D_Zero();
+	static UE_INL_API FVector2D Vector2D_Zero();
 
 
 	//
@@ -981,11 +987,11 @@ class UKismetMathLibrary : public UBlueprintFunctionLibrary
 
 	/** Makes a 2d vector {X, Y} */
 	UFUNCTION(BlueprintPure, Category = "Math|Vector2D", meta = (Keywords = "construct build", NativeMakeFunc))
-	static FVector2D MakeVector2D(double X, double Y);
+	static UE_INL_API FVector2D MakeVector2D(double X, double Y);
 
 	/** Breaks a 2D vector apart into X, Y. */
 	UFUNCTION(BlueprintPure, Category = "Math|Vector2D", meta = (NativeBreakFunc))
-	static void BreakVector2D(FVector2D InVec, double& X, double& Y);
+	static UE_INL_API void BreakVector2D(FVector2D InVec, double& X, double& Y);
 
 	/** Converts a Vector2D to a Vector */
 	UFUNCTION(BlueprintPure, meta = (DisplayName = "To Vector (Vector2D)", CompactNodeTitle = "->", ScriptMethod = "Vector", Keywords = "cast convert", BlueprintAutocast), Category = "Math|Conversions")
@@ -997,55 +1003,55 @@ class UKismetMathLibrary : public UBlueprintFunctionLibrary
 
 	/** Returns addition of Vector A and Vector B (A + B) */
 	UFUNCTION(BlueprintPure, meta = (DisplayName = "vector2d + vector2d", CompactNodeTitle = "+", ScriptMethod = "Add", ScriptOperator = "+;+=", Keywords = "+ add plus", CommutativeAssociativeBinaryOperator = "true"), Category = "Math|Vector2D")
-	static FVector2D Add_Vector2DVector2D(FVector2D A, FVector2D B);
+	static UE_INL_API FVector2D Add_Vector2DVector2D(FVector2D A, FVector2D B);
 
 	/** Returns Vector A added by B */
 	UFUNCTION(BlueprintPure, meta = (DisplayName = "vector2d + float", CompactNodeTitle = "+", ScriptMethod = "AddFloat", ScriptOperator = "+;+=", Keywords = "+ add plus"), Category = "Math|Vector2D")
-	static FVector2D Add_Vector2DFloat(FVector2D A, double B);
+	static UE_INL_API FVector2D Add_Vector2DFloat(FVector2D A, double B);
 
 	/** Returns subtraction of Vector B from Vector A (A - B) */
 	UFUNCTION(BlueprintPure, meta = (DisplayName = "vector2d - vector2d", CompactNodeTitle = "-", ScriptMethod = "Subtract", ScriptOperator = "-;-=", Keywords = "- subtract minus"), Category = "Math|Vector2D")
-	static FVector2D Subtract_Vector2DVector2D(FVector2D A, FVector2D B);
+	static UE_INL_API FVector2D Subtract_Vector2DVector2D(FVector2D A, FVector2D B);
 
 	/** Returns Vector A subtracted by B */
 	UFUNCTION(BlueprintPure, meta = (DisplayName = "vector2d - float", CompactNodeTitle = "-", ScriptMethod = "SubtractFloat", ScriptOperator = "-;-=", Keywords = "- subtract minus"), Category = "Math|Vector2D")
-	static FVector2D Subtract_Vector2DFloat(FVector2D A, double B);
+	static UE_INL_API FVector2D Subtract_Vector2DFloat(FVector2D A, double B);
 
 	/** Element-wise Vector multiplication (Result = {A.x*B.x, A.y*B.y}) */
 	UFUNCTION(BlueprintPure, meta = (DisplayName = "vector2d * vector2d", CompactNodeTitle = "*", ScriptMethod = "Multiply", ScriptOperator = "*;*=", Keywords = "* multiply", CommutativeAssociativeBinaryOperator = "true"), Category = "Math|Vector2D")
-	static FVector2D Multiply_Vector2DVector2D(FVector2D A, FVector2D B);
+	static UE_INL_API FVector2D Multiply_Vector2DVector2D(FVector2D A, FVector2D B);
 
 	/** Returns Vector A scaled by B */
 	UFUNCTION(BlueprintPure, meta = (DisplayName = "vector2d * float", CompactNodeTitle = "*", ScriptMethod = "MultiplyFloat", ScriptOperator = "*;*=", Keywords = "* multiply"), Category = "Math|Vector2D")
-	static FVector2D Multiply_Vector2DFloat(FVector2D A, double B);
+	static UE_INL_API FVector2D Multiply_Vector2DFloat(FVector2D A, double B);
 
 	/** Element-wise Vector divide (Result = {A.x/B.x, A.y/B.y}) */
 	UFUNCTION(BlueprintPure, meta = (DisplayName = "vector2d / vector2d", CompactNodeTitle = "/", ScriptMethod = "Divide", ScriptOperator = "/;/=", Keywords = "/ divide division"), Category = "Math|Vector2D")
-	static FVector2D Divide_Vector2DVector2D(FVector2D A, FVector2D B);
+	static UE_INL_API FVector2D Divide_Vector2DVector2D(FVector2D A, FVector2D B);
 
 	/** Returns Vector A divided by B */
 	UFUNCTION(BlueprintPure, meta = (DisplayName = "vector2d / float", CompactNodeTitle = "/", ScriptMethod = "DivideFloat", ScriptOperator = "/;/=", Keywords = "/ divide division"), Category = "Math|Vector2D")
-	static FVector2D Divide_Vector2DFloat(FVector2D A, double B = 1.f);
+	static UE_INL_API FVector2D Divide_Vector2DFloat(FVector2D A, double B = 1.f);
 
 	/** Returns true if vector A is equal to vector B (A == B) */
 	UFUNCTION(BlueprintPure, meta=(DisplayName = "Equal Exactly (Vector2D)", CompactNodeTitle = "===", ScriptMethod = "Equals", ScriptOperator = "==", Keywords = "== equal"), Category="Math|Vector2D")
-	static bool EqualExactly_Vector2DVector2D(FVector2D A, FVector2D B);
+	static UE_INL_API bool EqualExactly_Vector2DVector2D(FVector2D A, FVector2D B);
 
 	/** Returns true if vector2D A is equal to vector2D B (A == B) within a specified error tolerance */
 	UFUNCTION(BlueprintPure, meta = (DisplayName = "Equal (Vector2D)", CompactNodeTitle = "==", ScriptMethod = "IsNearEqual", Keywords = "== equal"), Category = "Math|Vector2D")
-	static bool EqualEqual_Vector2DVector2D(FVector2D A, FVector2D B, float ErrorTolerance = 1.e-4f);
+	static UE_INL_API bool EqualEqual_Vector2DVector2D(FVector2D A, FVector2D B, float ErrorTolerance = 1.e-4f);
 
 	/** Returns true if vector2D A is not equal to vector2D B (A != B) within a specified error tolerance */
 	UFUNCTION(BlueprintPure, meta = (DisplayName = "Not Equal Exactly (Vector2D)", CompactNodeTitle = "!==", ScriptMethod = "NotEqual", ScriptOperator = "!=", Keywords = "!= not equal"), Category = "Math|Vector2D")
-	static bool NotEqualExactly_Vector2DVector2D(FVector2D A, FVector2D B);
+	static UE_INL_API bool NotEqualExactly_Vector2DVector2D(FVector2D A, FVector2D B);
 
 	/** Returns true if vector2D A is not equal to vector2D B (A != B) within a specified error tolerance */
 	UFUNCTION(BlueprintPure, meta = (DisplayName = "Not Equal (Vector2D)", CompactNodeTitle = "!=", ScriptMethod = "IsNotNearEqual", Keywords = "!= not equal"), Category = "Math|Vector2D")
-	static bool NotEqual_Vector2DVector2D(FVector2D A, FVector2D B, float ErrorTolerance = 1.e-4f);
+	static UE_INL_API bool NotEqual_Vector2DVector2D(FVector2D A, FVector2D B, float ErrorTolerance = 1.e-4f);
 
 	/** Gets a negated copy of the vector. */
 	UFUNCTION(BlueprintPure, meta = (ScriptMethod = "Negated", ScriptOperator = "neg"), Category = "Math|Vector2D")
-	static FVector2D Negated2D(const FVector2D& A);
+	static UE_INL_API FVector2D Negated2D(const FVector2D& A);
 
 	/**
 	 * Set the values of the vector directly.
@@ -1061,11 +1067,11 @@ class UKismetMathLibrary : public UBlueprintFunctionLibrary
 	 * @return New vector with clamped axes.
 	 */
 	UFUNCTION(BlueprintPure, meta = (ScriptMethod = "ClampedAxes"), Category = "Math|Vector2D")
-	static FVector2D ClampAxes2D(FVector2D A, double MinAxisVal, double MaxAxisVal);
+	static UE_INL_API FVector2D ClampAxes2D(FVector2D A, double MinAxisVal, double MaxAxisVal);
 
 	/** Returns the cross product of two 2d vectors - see  http://mathworld.wolfram.com/CrossProduct.html */
 	UFUNCTION(BlueprintPure, meta = (DisplayName = "Cross Product (2D)", CompactNodeTitle = "cross", ScriptMethod = "Cross", ScriptOperator = "^"), Category = "Math|Vector2D")
-	static double CrossProduct2D(FVector2D A, FVector2D B);
+	static UE_INL_API double CrossProduct2D(FVector2D A, FVector2D B);
 
 	/**
 	 * Distance between two 2D points.
@@ -1075,7 +1081,7 @@ class UKismetMathLibrary : public UBlueprintFunctionLibrary
 	 * @return The distance between two 2D points.
 	 */
 	UFUNCTION(BlueprintPure, meta = (Keywords = "magnitude", ScriptMethod = "Distance"), Category = "Math|Vector2D")
-	static double Distance2D(FVector2D V1, FVector2D V2);
+	static UE_INL_API double Distance2D(FVector2D V1, FVector2D V2);
 
 	/**
 	 * Squared distance between two 2D points.
@@ -1085,11 +1091,11 @@ class UKismetMathLibrary : public UBlueprintFunctionLibrary
 	 * @return The squared distance between two 2D points.
 	 */
 	UFUNCTION(BlueprintPure, meta = (Keywords = "magnitude", ScriptMethod = "DistanceSquared"), Category = "Math|Vector2D")
-	static double DistanceSquared2D(FVector2D V1, FVector2D V2);
+	static UE_INL_API double DistanceSquared2D(FVector2D V1, FVector2D V2);
 
 	/** Returns the dot product of two 2d vectors - see http://mathworld.wolfram.com/DotProduct.html */
 	UFUNCTION(BlueprintPure, meta = (DisplayName = "Dot Product (2D)", CompactNodeTitle = "dot", ScriptMethod = "Dot", ScriptOperator = "|"), Category = "Math|Vector2D")
-	static double DotProduct2D(FVector2D A, FVector2D B);
+	static UE_INL_API double DotProduct2D(FVector2D A, FVector2D B);
 
 	/**
 	* Get a copy of this vector with absolute value of each component.
@@ -1097,7 +1103,7 @@ class UKismetMathLibrary : public UBlueprintFunctionLibrary
 	* @return A copy of this vector with absolute value of each component.
 	*/
 	UFUNCTION(BlueprintPure, meta = (ScriptMethod = "GetAbs"), Category = "Math|Vector2D")
-	static FVector2D GetAbs2D(FVector2D A);
+	static UE_INL_API FVector2D GetAbs2D(FVector2D A);
 
 	/**
 	 * Get the maximum absolute value of the vector's components.
@@ -1105,7 +1111,7 @@ class UKismetMathLibrary : public UBlueprintFunctionLibrary
 	 * @return The maximum absolute value of the vector's components.
 	 */
 	UFUNCTION(BlueprintPure, meta = (ScriptMethod = "GetAbsMax"), Category = "Math|Vector2D")
-	static double GetAbsMax2D(FVector2D A);
+	static UE_INL_API double GetAbsMax2D(FVector2D A);
 
 	/**
 	 * Get the maximum value of the vector's components.
@@ -1113,7 +1119,7 @@ class UKismetMathLibrary : public UBlueprintFunctionLibrary
 	 * @return The maximum value of the vector's components.
 	 */
 	UFUNCTION(BlueprintPure, meta = (ScriptMethod = "GetMax"), Category = "Math|Vector2D")
-	static double GetMax2D(FVector2D A);
+	static UE_INL_API double GetMax2D(FVector2D A);
 
 	/**
 	 * Get the minimum value of the vector's components.
@@ -1121,7 +1127,7 @@ class UKismetMathLibrary : public UBlueprintFunctionLibrary
 	 * @return The minimum value of the vector's components.
 	 */
 	UFUNCTION(BlueprintPure, meta = (ScriptMethod = "GetMin"), Category = "Math|Vector2D")
-	static double GetMin2D(FVector2D A);
+	static UE_INL_API double GetMin2D(FVector2D A);
 
 	/**
 	 * Rotates around axis (0,0,1)
@@ -1130,7 +1136,7 @@ class UKismetMathLibrary : public UBlueprintFunctionLibrary
 	 * @return Rotated Vector
 	 */
 	UFUNCTION(BlueprintPure, meta = (ScriptMethod = "GetRotated"), Category = "Math|Vector2D")
-	static FVector2D GetRotated2D(FVector2D A, float AngleDeg);
+	static UE_INL_API FVector2D GetRotated2D(FVector2D A, float AngleDeg);
 
 	/**
 	 * Checks whether vector is near to zero within a specified tolerance.
@@ -1139,7 +1145,7 @@ class UKismetMathLibrary : public UBlueprintFunctionLibrary
 	 * @return true if vector is in tolerance to zero, otherwise false.
 	 */
 	UFUNCTION(BlueprintPure, meta = (ScriptMethod = "IsNearlyZero"), Category = "Math|Vector2D")
-	static bool IsNearlyZero2D(const FVector2D& A, float Tolerance = 1.e-4f);
+	static UE_INL_API bool IsNearlyZero2D(const FVector2D& A, float Tolerance = 1.e-4f);
 
 	/**
 	 * Checks whether all components of the vector are exactly zero.
@@ -1147,7 +1153,7 @@ class UKismetMathLibrary : public UBlueprintFunctionLibrary
 	 * @return true if vector is exactly zero, otherwise false.
 	 */
 	UFUNCTION(BlueprintPure, meta = (ScriptMethod = "IsZero"), Category = "Math|Vector2D")
-	static bool IsZero2D(const FVector2D& A);
+	static UE_INL_API bool IsZero2D(const FVector2D& A);
 
 	/**
 	 * Tries to reach Target based on distance from Current position, giving a nice smooth feeling when tracking a position.
@@ -1159,7 +1165,7 @@ class UKismetMathLibrary : public UBlueprintFunctionLibrary
 	 * @return		New interpolated position
 	 */
 	UFUNCTION(BlueprintPure, Category="Math|Interpolation", meta=(ScriptMethod="InterpTo", Keywords="position"))
-	static FVector2D Vector2DInterpTo(FVector2D Current, FVector2D Target, float DeltaTime, float InterpSpeed);
+	static UE_INL_API FVector2D Vector2DInterpTo(FVector2D Current, FVector2D Target, float DeltaTime, float InterpSpeed);
 	
 	/**
 	 * Tries to reach Target at a constant rate.
@@ -1171,7 +1177,7 @@ class UKismetMathLibrary : public UBlueprintFunctionLibrary
 	 * @return		New interpolated position
 	 */
 	UFUNCTION(BlueprintPure, Category="Math|Interpolation", meta=(ScriptMethod="InterpToConstant", Keywords="position"))
-	static FVector2D Vector2DInterpTo_Constant(FVector2D Current, FVector2D Target, float DeltaTime, float InterpSpeed);
+	static UE_INL_API FVector2D Vector2DInterpTo_Constant(FVector2D Current, FVector2D Target, float DeltaTime, float InterpSpeed);
 
 	/**
 	 * Gets a normalized copy of the vector, checking it is safe to do so based on the length.
@@ -1181,11 +1187,11 @@ class UKismetMathLibrary : public UBlueprintFunctionLibrary
 	 * @return A normalized copy of the vector if safe, (0,0) otherwise.
 	 */
 	UFUNCTION(BlueprintPure, meta = (DisplayName = "Normal Safe (Vector2D)", Keywords = "Unit Vector", ScriptMethod = "Normal"), Category = "Math|Vector2D")
-	static FVector2D NormalSafe2D(FVector2D A, float Tolerance = 1.e-8f);
+	static UE_INL_API FVector2D NormalSafe2D(FVector2D A, float Tolerance = 1.e-8f);
 
 	/** Returns a unit normal version of the 2D vector */
 	UFUNCTION(BlueprintPure, meta = (DisplayName = "Normalize 2D", Keywords = "Unit Vector", ScriptMethod = "NormalUnsafe"), Category = "Math|Vector2D")
-	static FVector2D Normal2D(FVector2D A);
+	static UE_INL_API FVector2D Normal2D(FVector2D A);
 
 	/**
 	 * Normalize this vector in-place if it is large enough, set it to (0,0) otherwise.
@@ -1198,7 +1204,7 @@ class UKismetMathLibrary : public UBlueprintFunctionLibrary
 
 	/** Converts spherical coordinates on the unit sphere into a Cartesian unit length vector. */
 	UFUNCTION(BlueprintPure, meta = (DisplayName = "Spherical2D To Unit Cartesian", Keywords = "Unit Vector", ScriptMethod = "SphericalToUnitCartesian"), Category = "Math|Vector2D")
-	static FVector Spherical2DToUnitCartesian(FVector2D A);
+	static UE_INL_API FVector Spherical2DToUnitCartesian(FVector2D A);
 
 	/**
 	 * Util to convert this vector into a unit direction vector and its original length.
@@ -1207,7 +1213,7 @@ class UKismetMathLibrary : public UBlueprintFunctionLibrary
 	 * @param OutLength Reference passed in to store length of the vector.
 	 */
 	UFUNCTION(BlueprintPure, meta = (DisplayName = "To Direction And Length", ScriptMethod = "ToDirectionAndLength"), Category = "Math|Vector2D")
-	static void ToDirectionAndLength2D(FVector2D A, FVector2D &OutDir, double &OutLength);
+	static UE_INL_API void ToDirectionAndLength2D(FVector2D A, FVector2D &OutDir, double &OutLength);
 
 	/**
 	 * Get this vector as a vector where each component has been rounded to the nearest int.
@@ -1215,7 +1221,7 @@ class UKismetMathLibrary : public UBlueprintFunctionLibrary
 	 * @return New FVector2D from this vector that is rounded.
 	 */
 	UFUNCTION(BlueprintPure, meta = (DisplayName = "To Rounded (Vector2D)", ScriptMethod = "ToRounded"), Category = "Math|Vector2D")
-	static FVector2D ToRounded2D(FVector2D A);
+	static UE_INL_API FVector2D ToRounded2D(FVector2D A);
 
 	/**
 	* Get a copy of the vector as sign only.
@@ -1224,15 +1230,15 @@ class UKismetMathLibrary : public UBlueprintFunctionLibrary
 	* @return A copy of the vector with each component set to +1 or -1
 	*/
 	UFUNCTION(BlueprintPure, meta = (DisplayName = "To Sign (+1/-1) 2D", ScriptMethod = "ToSign"), Category = "Math|Vector2D")
-	static FVector2D ToSign2D(FVector2D A);
+	static UE_INL_API FVector2D ToSign2D(FVector2D A);
 
 	/** Returns the length of a 2D Vector. */
 	UFUNCTION(BlueprintPure, meta = (DisplayName = "Vector2D Length", Keywords = "magnitude", ScriptMethod = "Length"), Category = "Math|Vector2D")
-	static double VSize2D(FVector2D A);
+	static UE_INL_API double VSize2D(FVector2D A);
 
 	/** Returns the squared length of a 2D Vector. */
 	UFUNCTION(BlueprintPure, meta = (DisplayName = "Vector2D Length Squared", Keywords = "magnitude", ScriptMethod = "LengthSquared"), Category = "Math|Vector2D")
-	static double VSize2DSquared(FVector2D A);
+	static UE_INL_API double VSize2DSquared(FVector2D A);
 
 
 	//
@@ -1241,35 +1247,35 @@ class UKismetMathLibrary : public UBlueprintFunctionLibrary
 
 	/** 3D vector zero constant (0,0,0) */
 	UFUNCTION(BlueprintPure, meta = (ScriptConstant = "Zero", ScriptConstantHost = "/Script/CoreUObject.Vector"), Category = "Math|Vector")
-	static FVector Vector_Zero();
+	static UE_INL_API FVector Vector_Zero();
 
 	/** 3D vector one constant (1,1,1) */
 	UFUNCTION(BlueprintPure, meta = (ScriptConstant = "One", ScriptConstantHost = "/Script/CoreUObject.Vector"), Category = "Math|Vector")
-	static FVector Vector_One();
+	static UE_INL_API FVector Vector_One();
 
 	/** 3D vector Unreal forward direction constant (1,0,0) */
 	UFUNCTION(BlueprintPure, meta = (ScriptConstant = "Forward", ScriptConstantHost = "/Script/CoreUObject.Vector"), Category = "Math|Vector")
-	static FVector Vector_Forward();
+	static UE_INL_API FVector Vector_Forward();
 
 	/** 3D vector Unreal backward direction constant (-1,0,0) */
 	UFUNCTION(BlueprintPure, meta = (ScriptConstant = "Backward", ScriptConstantHost = "/Script/CoreUObject.Vector"), Category = "Math|Vector")
-	static FVector Vector_Backward();
+	static UE_INL_API FVector Vector_Backward();
 
 	/** 3D vector Unreal up direction constant (0,0,1) */
 	UFUNCTION(BlueprintPure, meta = (ScriptConstant = "Up", ScriptConstantHost = "/Script/CoreUObject.Vector"), Category = "Math|Vector")
-	static FVector Vector_Up();
+	static UE_INL_API FVector Vector_Up();
 
 	/** 3D vector Unreal down direction constant (0,0,-1) */
 	UFUNCTION(BlueprintPure, meta = (ScriptConstant = "Down", ScriptConstantHost = "/Script/CoreUObject.Vector"), Category = "Math|Vector")
-	static FVector Vector_Down();
+	static UE_INL_API FVector Vector_Down();
 
 	/** 3D vector Unreal right direction constant (0,1,0) */
 	UFUNCTION(BlueprintPure, meta = (ScriptConstant = "Right", ScriptConstantHost = "/Script/CoreUObject.Vector"), Category = "Math|Vector")
-	static FVector Vector_Right();
+	static UE_INL_API FVector Vector_Right();
 
 	/** 3D vector Unreal left direction constant (0,-1,0) */
 	UFUNCTION(BlueprintPure, meta = (ScriptConstant = "Left", ScriptConstantHost = "/Script/CoreUObject.Vector"), Category = "Math|Vector")
-	static FVector Vector_Left();
+	static UE_INL_API FVector Vector_Left();
 
 	//
 	// Vector (3D) functions.
@@ -1277,7 +1283,7 @@ class UKismetMathLibrary : public UBlueprintFunctionLibrary
 
 	/** Makes a vector {X, Y, Z} */
 	UFUNCTION(BlueprintPure, Category = "Math|Vector", meta = (Keywords = "construct build", NativeMakeFunc))
-	static FVector MakeVector(double X, double Y, double Z);
+	static UE_INL_API FVector MakeVector(double X, double Y, double Z);
 
 	/** Creates a directional vector from rotation values {Pitch, Yaw} supplied in degrees with specified Length*/	
 	UFUNCTION(BlueprintPure, Category = "Math|Vector", meta = (Keywords = "rotation rotate"))
@@ -1303,14 +1309,14 @@ class UKismetMathLibrary : public UBlueprintFunctionLibrary
 
 	/** Breaks a vector apart into X, Y, Z */
 	UFUNCTION(BlueprintPure, Category="Math|Vector", meta=(NativeBreakFunc))
-	static void BreakVector(FVector InVec, double& X, double& Y, double& Z);
+	static UE_INL_API void BreakVector(FVector InVec, double& X, double& Y, double& Z);
 
 	UE_DEPRECATED(5.3, "This method has been deprecated and will be removed. Use the double version instead.")
-	static void BreakVector3f(FVector3f InVec, float& X, float& Y, float& Z);
+	static UE_INL_API void BreakVector3f(FVector3f InVec, float& X, float& Y, float& Z);
 
 	/** Converts a vector to LinearColor */
 	UFUNCTION(BlueprintPure, meta=(DisplayName = "To LinearColor (Vector)", CompactNodeTitle = "->", ScriptMethod = "LinearColor", Keywords="cast convert", BlueprintAutocast), Category="Math|Conversions")
-	static FLinearColor Conv_VectorToLinearColor(FVector InVec);
+	static UE_INL_API FLinearColor Conv_VectorToLinearColor(FVector InVec);
 
 	/** Converts a vector to a transform. Uses vector as location */
 	UFUNCTION(BlueprintPure, meta=(DisplayName = "To Transform (Vector)", CompactNodeTitle = "->", ScriptMethod = "Transform", Keywords="cast convert", BlueprintAutocast), Category="Math|Conversions")
@@ -1327,7 +1333,7 @@ class UKismetMathLibrary : public UBlueprintFunctionLibrary
 	 * @return FRotator from the Vector's direction, without any roll.
 	 */
 	UFUNCTION(BlueprintPure, meta=(DisplayName = "Rotation From X Vector", ScriptMethod = "Rotator", Keywords="rotation rotate cast convert", BlueprintAutocast), Category="Math|Conversions")
-	static FRotator Conv_VectorToRotator(FVector InVec);
+	static UE_INL_API FRotator Conv_VectorToRotator(FVector InVec);
 
 	/** Create a rotation from an axis and supplied angle (in degrees) */
 	UFUNCTION(BlueprintPure, meta=(ScriptMethod = "RotatorFromAxisAndAngle", Keywords="make construct build rotate rotation"), Category="Math|Vector")
@@ -1343,7 +1349,7 @@ class UKismetMathLibrary : public UBlueprintFunctionLibrary
 	 * @return Quaternion from the Vector's direction, without any roll.
 	 */
 	UFUNCTION(BlueprintPure, meta=(DisplayName = "To Quaternion (Vector)", ScriptMethod = "Quaternion", Keywords="rotation rotate cast convert", BlueprintAutocast), Category="Math|Conversions")
-	static FQuat Conv_VectorToQuaternion(FVector InVec);
+	static UE_INL_API FQuat Conv_VectorToQuaternion(FVector InVec);
 
 	/**
 	 * Interpolate from a vector to the direction of another vector along a spherical path.
@@ -1354,7 +1360,7 @@ class UKismetMathLibrary : public UBlueprintFunctionLibrary
 	 * @return Vector after interpolating between Vector and Direction along a spherical path. The magnitude will remain the length of the starting vector.
 	 */
 	UFUNCTION(BlueprintPure, meta = (DisplayName = "SlerpVectorToDirection", ScriptMethod = "SlerpVectors", Keywords = "spherical lerpvector vectorslerp interpolate"), Category = "Math|Vector")
-	static FVector Vector_SlerpVectorToDirection(FVector Vector, FVector Direction, double Alpha);
+	static UE_INL_API FVector Vector_SlerpVectorToDirection(FVector Vector, FVector Direction, double Alpha);
 
 	/**
 	 * Interpolate from normalized vector A to normalized vector B along a spherical path.
@@ -1365,43 +1371,43 @@ class UKismetMathLibrary : public UBlueprintFunctionLibrary
 	 * @return Vector after interpolating between A and B along a spherical path.
 	 */
 	UFUNCTION(BlueprintPure, meta = (DisplayName = "SlerpNormals", ScriptMethod = "SlerpNormals", Keywords = "spherical lerpnormal normalslerp interpolate"), Category = "Math|Vector")
-	static FVector Vector_SlerpNormals(FVector NormalA, FVector NormalB, double Alpha);
+	static UE_INL_API FVector Vector_SlerpNormals(FVector NormalA, FVector NormalB, double Alpha);
 
 	/** Vector addition */
 	UFUNCTION(BlueprintPure, meta=(DisplayName = "vector + vector", CompactNodeTitle = "+", ScriptMethod = "Add", ScriptOperator = "+;+=", Keywords = "+ add plus", CommutativeAssociativeBinaryOperator = "true"), Category="Math|Vector")
-	static FVector Add_VectorVector(FVector A, FVector B);
+	static UE_INL_API FVector Add_VectorVector(FVector A, FVector B);
 
 	/** Adds a float to each component of a vector */
 	UFUNCTION(BlueprintPure, meta=(DisplayName = "vector + float", CompactNodeTitle = "+", ScriptMethod = "AddFloat", Keywords = "+ add plus"), Category="Math|Vector")
-	static FVector Add_VectorFloat(FVector A, double B);
+	static UE_INL_API FVector Add_VectorFloat(FVector A, double B);
 	
 	/** Adds an integer to each component of a vector */
 	UFUNCTION(BlueprintPure, meta=(DisplayName = "vector + integer", CompactNodeTitle = "+", ScriptMethod = "AddInt", Keywords = "+ add plus"), Category="Math|Vector")
-	static FVector Add_VectorInt(FVector A, int32 B);
+	static UE_INL_API FVector Add_VectorInt(FVector A, int32 B);
 
 	/** Vector subtraction */
 	UFUNCTION(BlueprintPure, meta=(DisplayName = "vector - vector", CompactNodeTitle = "-", ScriptMethod = "Subtract", ScriptOperator = "-;-=", Keywords = "- subtract minus"), Category="Math|Vector")
-	static FVector Subtract_VectorVector(FVector A, FVector B);
+	static UE_INL_API FVector Subtract_VectorVector(FVector A, FVector B);
 
 	/** Subtracts a float from each component of a vector */
 	UFUNCTION(BlueprintPure, meta=(DisplayName = "vector - float", CompactNodeTitle = "-", ScriptMethod = "SubtractFloat", Keywords = "- subtract minus"), Category="Math|Vector")
-	static FVector Subtract_VectorFloat(FVector A, double B);
+	static UE_INL_API FVector Subtract_VectorFloat(FVector A, double B);
 
 	/** Subtracts an integer from each component of a vector */
 	UFUNCTION(BlueprintPure, meta=(DisplayName = "vector - integer", CompactNodeTitle = "-", ScriptMethod = "SubtractInt", Keywords = "- subtract minus"), Category="Math|Vector")
-	static FVector Subtract_VectorInt(FVector A, int32 B);
+	static UE_INL_API FVector Subtract_VectorInt(FVector A, int32 B);
 
 	/** Element-wise Vector multiplication (Result = {A.x*B.x, A.y*B.y, A.z*B.z}) */
 	UFUNCTION(BlueprintPure, meta=(DisplayName = "vector * vector", CompactNodeTitle = "*", ScriptMethod = "Multiply", ScriptOperator = "*;*=", Keywords = "* multiply", CommutativeAssociativeBinaryOperator = "true"), Category="Math|Vector")
-	static FVector Multiply_VectorVector(FVector A, FVector B);
+	static UE_INL_API FVector Multiply_VectorVector(FVector A, FVector B);
 
 	/** Scales Vector A by B */
 	UFUNCTION(BlueprintPure, meta=(DisplayName = "vector * float", CompactNodeTitle = "*", ScriptMethod = "MultiplyFloat", ScriptOperator = "*;*=", Keywords = "* multiply"), Category="Math|Vector")
-	static FVector Multiply_VectorFloat(FVector A, double B);
+	static UE_INL_API FVector Multiply_VectorFloat(FVector A, double B);
 	
 	/** Scales Vector A by B */
 	UFUNCTION(BlueprintPure, meta=(DisplayName = "vector * integer", CompactNodeTitle = "*", ScriptMethod = "MultiplyInt", Keywords = "* multiply"), Category="Math|Vector")
-	static FVector Multiply_VectorInt(FVector A, int32 B);
+	static UE_INL_API FVector Multiply_VectorInt(FVector A, int32 B);
 
 	/** Element-wise Vector division (Result = {A.x/B.x, A.y/B.y, A.z/B.z}) */
 	UFUNCTION(BlueprintPure, meta=(DisplayName = "vector / vector", CompactNodeTitle = "/", ScriptMethod = "Divide", ScriptOperator = "/;/=", Keywords = "/ divide division"), Category="Math|Vector")
@@ -1409,43 +1415,43 @@ class UKismetMathLibrary : public UBlueprintFunctionLibrary
 
 	/** Vector divide by a float */
 	UFUNCTION(BlueprintPure, meta=(DisplayName = "vector / float", CompactNodeTitle = "/", ScriptMethod = "DivideFloat", ScriptOperator = "/;/=", Keywords = "/ divide division"), Category="Math|Vector")
-	static FVector Divide_VectorFloat(FVector A, double B = 1.f);
+	static UE_INL_API FVector Divide_VectorFloat(FVector A, double B = 1.f);
 
 	/** Vector divide by an integer */
 	UFUNCTION(BlueprintPure, meta=(DisplayName = "vector / integer", CompactNodeTitle = "/", ScriptMethod = "DivideInt", Keywords = "/ divide division"), Category="Math|Vector")
-	static FVector Divide_VectorInt(FVector A, int32 B = 1);
+	static UE_INL_API FVector Divide_VectorInt(FVector A, int32 B = 1);
 
 	/** Negate a vector. */
 	UFUNCTION(BlueprintPure, meta = (ScriptMethod = "Negated", ScriptOperator = "neg"), Category="Math|Vector")
-	static FVector NegateVector(FVector A);
+	static UE_INL_API FVector NegateVector(FVector A);
 
 	/** Returns true if vector A is equal to vector B (A == B) */
 	UFUNCTION(BlueprintPure, meta=(DisplayName = "Equal Exactly (Vector)", CompactNodeTitle = "===", ScriptMethod = "Equals", ScriptOperator = "==", Keywords = "== equal"), Category="Math|Vector")
-	static bool EqualExactly_VectorVector(FVector A, FVector B);
+	static UE_INL_API bool EqualExactly_VectorVector(FVector A, FVector B);
 
 	/** Returns true if vector A is equal to vector B (A == B) within a specified error tolerance */
 	UFUNCTION(BlueprintPure, meta=(DisplayName = "Equal (Vector)", CompactNodeTitle = "==", ScriptMethod = "IsNearEqual", Keywords = "== equal"), Category="Math|Vector")
-	static bool EqualEqual_VectorVector(FVector A, FVector B, float ErrorTolerance = 1.e-4f);
+	static UE_INL_API bool EqualEqual_VectorVector(FVector A, FVector B, float ErrorTolerance = 1.e-4f);
 
 	/** Returns true if vector A is not equal to vector B (A != B) */
 	UFUNCTION(BlueprintPure, meta = (DisplayName = "Not Equal Exactly (Vector)", CompactNodeTitle = "!==", ScriptMethod = "NotEqual", ScriptOperator = "!=", Keywords = "!= not equal"), Category = "Math|Vector")
-	static bool NotEqualExactly_VectorVector(FVector A, FVector B);
+	static UE_INL_API bool NotEqualExactly_VectorVector(FVector A, FVector B);
 
 	/** Returns true if vector A is not equal to vector B (A != B) within a specified error tolerance */
 	UFUNCTION(BlueprintPure, meta=(DisplayName = "Not Equal (Vector)", CompactNodeTitle = "!=", ScriptMethod = "IsNotNearEqual"), Category="Math|Vector")
-	static bool NotEqual_VectorVector(FVector A, FVector B, float ErrorTolerance = 1.e-4f);
+	static UE_INL_API bool NotEqual_VectorVector(FVector A, FVector B, float ErrorTolerance = 1.e-4f);
 
 	/** Returns the dot product of two 3d vectors - see http://mathworld.wolfram.com/DotProduct.html */
 	UFUNCTION(BlueprintPure, meta=(DisplayName = "Dot Product", CompactNodeTitle = "dot", ScriptMethod = "Dot", ScriptOperator = "|"), Category="Math|Vector" )
-	static double Dot_VectorVector(FVector A, FVector B);
+	static UE_INL_API double Dot_VectorVector(FVector A, FVector B);
 
 	/** Returns the cross product of two 3d vectors - see http://mathworld.wolfram.com/CrossProduct.html */
 	UFUNCTION(BlueprintPure, meta=(DisplayName = "Cross Product", CompactNodeTitle = "cross", ScriptMethod = "Cross", ScriptOperator = "^"), Category="Math|Vector" )
-	static FVector Cross_VectorVector(FVector A, FVector B);
+	static UE_INL_API FVector Cross_VectorVector(FVector A, FVector B);
 
 	/** Returns result of vector A rotated by Rotator B */
 	UFUNCTION(BlueprintPure, meta=(DisplayName = "Rotate Vector", ScriptMethod = "Rotate"), Category="Math|Vector")
-	static FVector GreaterGreater_VectorRotator(FVector A, FRotator B);
+	static UE_INL_API FVector GreaterGreater_VectorRotator(FVector A, FRotator B);
 
 	/** Returns result of vector A rotated by AngleDeg around Axis */
 	UFUNCTION(BlueprintPure, meta=(DisplayName = "Rotate Vector Around Axis", ScriptMethod = "RotateAngleAxis"), Category="Math|Vector")
@@ -1453,7 +1459,7 @@ class UKismetMathLibrary : public UBlueprintFunctionLibrary
 
 	/** Returns result of vector A rotated by the inverse of Rotator B */
 	UFUNCTION(BlueprintPure, meta=(DisplayName = "Unrotate Vector", ScriptMethod = "Unrotate"), Category="Math|Vector")
-	static FVector LessLess_VectorRotator(FVector A, FRotator B);
+	static UE_INL_API FVector LessLess_VectorRotator(FVector A, FRotator B);
 
 	/** When this vector contains Euler angles (degrees), ensure that angles are between +/-180 */
 	UFUNCTION(BlueprintCallable, meta = (ScriptMethod = "UnwindEuler"), Category = "Math|Vector")
@@ -1461,35 +1467,35 @@ class UKismetMathLibrary : public UBlueprintFunctionLibrary
 
 	/** Create a copy of this vector, with its magnitude/size/length clamped between Min and Max. */
 	UFUNCTION(BlueprintPure, meta=(ScriptMethod = "ClampedSize"), Category="Math|Vector")
-	static FVector ClampVectorSize(FVector A, double Min, double Max);
+	static UE_INL_API FVector ClampVectorSize(FVector A, double Min, double Max);
 
 	/** Create a copy of this vector, with the 2D magnitude/size/length clamped between Min and Max. Z is unchanged. */
 	UFUNCTION(BlueprintPure, meta=(ScriptMethod = "ClampedSize2D"), Category="Math|Vector")
-	static FVector Vector_ClampSize2D(FVector A, double Min, double Max);
+	static UE_INL_API FVector Vector_ClampSize2D(FVector A, double Min, double Max);
 
 	/** Create a copy of this vector, with its maximum magnitude/size/length clamped to MaxSize. */
 	UFUNCTION(BlueprintPure, meta=(ScriptMethod = "ClampedSizeMax"), Category="Math|Vector")
-	static FVector Vector_ClampSizeMax(FVector A, double Max);
+	static UE_INL_API FVector Vector_ClampSizeMax(FVector A, double Max);
 
 	/** Create a copy of this vector, with the maximum 2D magnitude/size/length clamped to MaxSize. Z is unchanged. */
 	UFUNCTION(BlueprintPure, meta=(ScriptMethod = "ClampedSizeMax2D"), Category="Math|Vector")
-	static FVector Vector_ClampSizeMax2D(FVector A, double Max);
+	static UE_INL_API FVector Vector_ClampSizeMax2D(FVector A, double Max);
 
 	/** Find the minimum element (X, Y or Z) of a vector */
 	UFUNCTION(BlueprintPure, meta=(ScriptMethod = "GetMinElement"), Category="Math|Vector")
-	static double GetMinElement(FVector A);
+	static UE_INL_API double GetMinElement(FVector A);
 
 	/** Find the maximum element (X, Y or Z) of a vector */
 	UFUNCTION(BlueprintPure, meta=(ScriptMethod = "GetMaxElement"), Category="Math|Vector")
-	static double GetMaxElement(FVector A);
+	static UE_INL_API double GetMaxElement(FVector A);
 
 	/** Find the maximum absolute element (abs(X), abs(Y) or abs(Z)) of a vector */
 	UFUNCTION(BlueprintPure, meta=(ScriptMethod = "GetAbsMax"), Category="Math|Vector")
-	static double Vector_GetAbsMax(FVector A);
+	static UE_INL_API double Vector_GetAbsMax(FVector A);
 
 	/** Find the minimum absolute element (abs(X), abs(Y) or abs(Z)) of a vector */
 	UFUNCTION(BlueprintPure, meta=(ScriptMethod = "GetAbsMin"), Category="Math|Vector")
-	static double Vector_GetAbsMin(FVector A);
+	static UE_INL_API double Vector_GetAbsMin(FVector A);
 
 	/**
 	 * Get a copy of this vector with absolute value of each component.
@@ -1497,15 +1503,15 @@ class UKismetMathLibrary : public UBlueprintFunctionLibrary
 	 * @return A copy of this vector with absolute value of each component.
 	 */
 	UFUNCTION(BlueprintPure, meta=(ScriptMethod = "GetAbs"), Category="Math|Vector")
-	static FVector Vector_GetAbs(FVector A);
+	static UE_INL_API FVector Vector_GetAbs(FVector A);
 
 	/** Find the minimum elements (X, Y and Z) between the two vector's components */
 	UFUNCTION(BlueprintPure, meta=(ScriptMethod = "GetMin"), Category="Math|Vector")
-	static FVector Vector_ComponentMin(FVector A, FVector B);
+	static UE_INL_API FVector Vector_ComponentMin(FVector A, FVector B);
 
 	/** Find the maximum elements (X, Y and Z) between the two vector's components */
 	UFUNCTION(BlueprintPure, meta=(ScriptMethod = "GetMax"), Category="Math|Vector")
-	static FVector Vector_ComponentMax(FVector A, FVector B);
+	static UE_INL_API FVector Vector_ComponentMax(FVector A, FVector B);
 
 	/**
 	 * Get a copy of the vector as sign only.
@@ -1514,7 +1520,7 @@ class UKismetMathLibrary : public UBlueprintFunctionLibrary
 	 * @param A copy of the vector with each component set to +1 or -1
 	 */
 	UFUNCTION(BlueprintPure, meta=(ScriptMethod = "GetSignVector"), Category="Math|Vector")
-	static FVector Vector_GetSignVector(FVector A);
+	static UE_INL_API FVector Vector_GetSignVector(FVector A);
 
 	/**
 	 * Projects 2D components of vector based on Z.
@@ -1522,7 +1528,7 @@ class UKismetMathLibrary : public UBlueprintFunctionLibrary
 	 * @return Projected version of vector based on Z.
 	 */
 	UFUNCTION(BlueprintPure, meta=(ScriptMethod = "GetProjection"), Category="Math|Vector")
-	static FVector Vector_GetProjection(FVector A);
+	static UE_INL_API FVector Vector_GetProjection(FVector A);
 
 	/**
 	 * Convert a direction vector into a 'heading' angle.
@@ -1530,7 +1536,7 @@ class UKismetMathLibrary : public UBlueprintFunctionLibrary
 	 * @return 'Heading' angle between +/-PI radians. 0 is pointing down +X.
 	 */
 	UFUNCTION(BlueprintPure, meta=(ScriptMethod = "HeadingAngle"), Category="Math|Vector")
-	static double Vector_HeadingAngle(FVector A);
+	static UE_INL_API double Vector_HeadingAngle(FVector A);
 
 	/**
 	 * Returns the cosine of the angle between this vector and another projected onto the XY plane (no Z).
@@ -1539,7 +1545,7 @@ class UKismetMathLibrary : public UBlueprintFunctionLibrary
 	 * @return The cosine.
 	 */
 	UFUNCTION(BlueprintPure, meta=(ScriptMethod = "CosineAngle2D"), Category="Math|Vector")
-	static double Vector_CosineAngle2D(FVector A, FVector B);
+	static UE_INL_API double Vector_CosineAngle2D(FVector A, FVector B);
 
 	/**
 	 * Converts a vector containing degree values to a vector containing radian values.
@@ -1547,7 +1553,7 @@ class UKismetMathLibrary : public UBlueprintFunctionLibrary
 	 * @return Vector containing radian values
 	 */
 	UFUNCTION(BlueprintPure, meta=(ScriptMethod = "ToRadians"), Category="Math|Vector")
-	static FVector Vector_ToRadians(FVector A);
+	static UE_INL_API FVector Vector_ToRadians(FVector A);
 
 	/**
 	 * Converts a vector containing radian values to a vector containing degree values.
@@ -1555,18 +1561,18 @@ class UKismetMathLibrary : public UBlueprintFunctionLibrary
 	 * @return Vector  containing degree values
 	 */
 	UFUNCTION(BlueprintPure, meta=(ScriptMethod = "ToDegrees"), Category="Math|Vector")
-	static FVector Vector_ToDegrees(FVector A);
+	static UE_INL_API FVector Vector_ToDegrees(FVector A);
 
 	/** 
 	 * Converts a Cartesian unit vector into spherical coordinates on the unit sphere.
 	 * @return Output Theta will be in the range [0, PI], and output Phi will be in the range [-PI, PI]. 
 	 */
 	UFUNCTION(BlueprintPure, meta=(ScriptMethod = "UnitCartesianToSpherical"), Category="Math|Vector")
-	static FVector2D Vector_UnitCartesianToSpherical(FVector A);
+	static UE_INL_API FVector2D Vector_UnitCartesianToSpherical(FVector A);
 
 	/** Find the unit direction vector from one position to another or (0,0,0) if positions are the same. */
 	UFUNCTION(BlueprintPure, meta=(DisplayName = "Get Unit Direction (Vector)", ScriptMethod = "DirectionUnitTo", Keywords = "Unit Vector"), Category="Math|Vector")
-	static FVector GetDirectionUnitVector(FVector From, FVector To);
+	static UE_INL_API FVector GetDirectionUnitVector(FVector From, FVector To);
 
 	/** Breaks a vector apart into Yaw, Pitch rotation values given in degrees. (non-clamped) */
 	UFUNCTION(BlueprintPure, Category = "Math|Vector", meta = (ScriptMethod = "GetYawPitch", NativeBreakFunc))
@@ -1597,7 +1603,7 @@ class UKismetMathLibrary : public UBlueprintFunctionLibrary
 	 * @return The distance between two points.
 	 */
 	UFUNCTION(BlueprintPure, meta = (DisplayName = "Distance (Vector)", ScriptMethod = "Distance", Keywords = "magnitude"), Category = "Math|Vector")
-	static double Vector_Distance(FVector V1, FVector V2);
+	static UE_INL_API double Vector_Distance(FVector V1, FVector V2);
 
 	/**
 	 * Squared distance between two points.
@@ -1607,7 +1613,7 @@ class UKismetMathLibrary : public UBlueprintFunctionLibrary
 	 * @return The squared distance between two points.
 	 */
 	UFUNCTION(BlueprintPure, meta = (DisplayName = "Distance Squared (Vector)", ScriptMethod = "DistanceSquared", Keywords = "magnitude"), Category = "Math|Vector")
-	static double Vector_DistanceSquared(FVector V1, FVector V2);
+	static UE_INL_API double Vector_DistanceSquared(FVector V1, FVector V2);
 
 	/**
 	* Euclidean distance between two points in the XY plane (ignoring Z).
@@ -1617,7 +1623,7 @@ class UKismetMathLibrary : public UBlueprintFunctionLibrary
 	* @return The distance between two points in the XY plane.
 	*/
 	UFUNCTION(BlueprintPure, meta = (DisplayName = "Distance2D (Vector)", ScriptMethod = "Distance2D", Keywords = "magnitude"), Category = "Math|Vector")
-	static double Vector_Distance2D(FVector V1, FVector V2);
+	static UE_INL_API double Vector_Distance2D(FVector V1, FVector V2);
 
 	/**
 	* Squared euclidean distance between two points in the XY plane (ignoring Z).
@@ -1627,23 +1633,23 @@ class UKismetMathLibrary : public UBlueprintFunctionLibrary
 	* @return The distance between two points in the XY plane.
 	*/
 	UFUNCTION(BlueprintPure, meta = (DisplayName = "Distance2D Squared (Vector)", ScriptMethod = "Distance2DSquared", Keywords = "magnitude"), Category = "Math|Vector")
-	static double Vector_Distance2DSquared(FVector V1, FVector V2);
+	static UE_INL_API double Vector_Distance2DSquared(FVector V1, FVector V2);
 
 	/** Returns the length of the vector */
 	UFUNCTION(BlueprintPure, meta=(DisplayName = "Vector Length", ScriptMethod = "Length", Keywords="magnitude"), Category="Math|Vector")
-	static double VSize(FVector A);
+	static UE_INL_API double VSize(FVector A);
 
 	/** Returns the squared length of the vector */
 	UFUNCTION(BlueprintPure, meta=(DisplayName = "Vector Length Squared", ScriptMethod = "LengthSquared", Keywords="magnitude"), Category="Math|Vector")
-	static double VSizeSquared(FVector A);
+	static UE_INL_API double VSizeSquared(FVector A);
 
 	/** Returns the length of the vector's XY components. */
 	UFUNCTION(BlueprintPure, meta=(DisplayName = "Vector Length XY", ScriptMethod = "Length2D", Keywords="magnitude"), Category="Math|Vector")
-	static double VSizeXY(FVector A);
+	static UE_INL_API double VSizeXY(FVector A);
 
 	/** Returns the squared length of the vector's XY components. */
 	UFUNCTION(BlueprintPure, meta=(DisplayName = "Vector Length XY Squared", ScriptMethod = "Length2DSquared", Keywords="magnitude"), Category="Math|Vector")
-	static double VSizeXYSquared(FVector A);
+	static UE_INL_API double VSizeXYSquared(FVector A);
 
 	/**
 	 * Checks whether vector is near to zero within a specified tolerance.
@@ -1652,7 +1658,7 @@ class UKismetMathLibrary : public UBlueprintFunctionLibrary
 	 * @return true if vector is in tolerance to zero, otherwise false.
 	 */
 	UFUNCTION(BlueprintPure, meta = (ScriptMethod = "IsNearlyZero"), Category = "Math|Vector")
-	static bool Vector_IsNearlyZero(const FVector& A, float Tolerance = 1.e-4f);
+	static UE_INL_API bool Vector_IsNearlyZero(const FVector& A, float Tolerance = 1.e-4f);
 
 	/**
 	 * Checks whether all components of the vector are exactly zero.
@@ -1660,7 +1666,7 @@ class UKismetMathLibrary : public UBlueprintFunctionLibrary
 	 * @return true if vector is exactly zero, otherwise false.
 	 */
 	UFUNCTION(BlueprintPure, meta = (ScriptMethod = "IsZero"), Category = "Math|Vector")
-	static bool Vector_IsZero(const FVector& A);
+	static UE_INL_API bool Vector_IsZero(const FVector& A);
 
 	/**
 	 * Determines if any component is not a number (NAN)
@@ -1668,7 +1674,7 @@ class UKismetMathLibrary : public UBlueprintFunctionLibrary
 	 * @return true if one or more components is NAN, otherwise false.
 	 */
 	UFUNCTION(BlueprintPure, meta = (ScriptMethod = "IsNAN"), Category = "Math|Vector")
-	static bool Vector_IsNAN(const FVector& A);
+	static UE_INL_API bool Vector_IsNAN(const FVector& A);
 
 	/**
 	 * Checks whether all components of this vector are the same, within a tolerance.
@@ -1677,7 +1683,7 @@ class UKismetMathLibrary : public UBlueprintFunctionLibrary
 	 * @return true if the vectors are equal within tolerance limits, false otherwise.
 	 */
 	UFUNCTION(BlueprintPure, meta=(DisplayName = "Is Uniform (Vector)", ScriptMethod = "IsUniform"), Category="Math|Vector")
-	static bool Vector_IsUniform(const FVector& A, float Tolerance = 1.e-4f);
+	static UE_INL_API bool Vector_IsUniform(const FVector& A, float Tolerance = 1.e-4f);
 
 	/**
 	 * Determines if vector is normalized / unit (length 1) within specified squared tolerance.
@@ -1685,7 +1691,7 @@ class UKismetMathLibrary : public UBlueprintFunctionLibrary
 	 * @return true if unit, false otherwise.
 	 */
 	UFUNCTION(BlueprintPure, meta=(DisplayName = "Is Unit (Vector)", ScriptMethod = "IsUnit", Keywords="Unit Vector"), Category="Math|Vector")
-	static bool Vector_IsUnit(const FVector& A, float SquaredLenthTolerance = 1.e-4f);
+	static UE_INL_API bool Vector_IsUnit(const FVector& A, float SquaredLenthTolerance = 1.e-4f);
 
 	/**
 	 * Determines if vector is normalized / unit (length 1).
@@ -1693,7 +1699,7 @@ class UKismetMathLibrary : public UBlueprintFunctionLibrary
 	 * @return true if normalized, false otherwise.
 	 */
 	UFUNCTION(BlueprintPure, meta=(DisplayName = "Is Normal (Vector)", ScriptMethod = "IsNormal", Keywords="Unit Vector"), Category="Math|Vector")
-	static bool Vector_IsNormal(const FVector& A);
+	static UE_INL_API bool Vector_IsNormal(const FVector& A);
 
 	/**
 	 * Gets a normalized unit copy of the vector, ensuring it is safe to do so based on the length.
@@ -1703,7 +1709,7 @@ class UKismetMathLibrary : public UBlueprintFunctionLibrary
 	 * @return A normalized copy if safe, (0,0,0) otherwise.
 	 */
 	UFUNCTION(BlueprintPure, meta=(DisplayName = "Normalize", ScriptMethod = "Normal", Keywords="Unit Vector"), Category="Math|Vector")
-	static FVector Normal(FVector A, float Tolerance = 1.e-4f);
+	static UE_INL_API FVector Normal(FVector A, float Tolerance = 1.e-4f);
 
 	/**
 	 * Gets a normalized unit copy of the 2D components of the vector, ensuring it is safe to do so. Z is set to zero. 
@@ -1713,7 +1719,7 @@ class UKismetMathLibrary : public UBlueprintFunctionLibrary
 	 * @return Normalized copy if safe, (0,0,0) otherwise.
 	 */
 	UFUNCTION(BlueprintPure, meta=(DisplayName = "Normalize 2D (Vector)", ScriptMethod = "Normal2D", Keywords="Unit Vector"), Category="Math|Vector")
-	static FVector Vector_Normal2D(FVector A, float Tolerance = 1.e-4f);
+	static UE_INL_API FVector Vector_Normal2D(FVector A, float Tolerance = 1.e-4f);
 
 	/**
 	 * Calculates normalized unit version of vector without checking for zero length.
@@ -1721,7 +1727,7 @@ class UKismetMathLibrary : public UBlueprintFunctionLibrary
 	 * @return Normalized version of vector.
 	 */
 	UFUNCTION(BlueprintPure, meta=(DisplayName = "Normal Unsafe (Vector)", ScriptMethod = "NormalUnsafe", Keywords="Unit Vector"), Category="Math|Vector")
-	static FVector Vector_NormalUnsafe(const FVector& A);
+	static UE_INL_API FVector Vector_NormalUnsafe(const FVector& A);
 
 	/**
 	 * Normalize this vector in-place if it is large enough or set it to (0,0,0) otherwise.
@@ -1749,7 +1755,7 @@ class UKismetMathLibrary : public UBlueprintFunctionLibrary
 	 * @return		New interpolated position
 	 */
 	UFUNCTION(BlueprintPure, Category="Math|Interpolation", meta=(ScriptMethod = "InterpTo", Keywords="position"))
-	static FVector VInterpTo(FVector Current, FVector Target, float DeltaTime, float InterpSpeed);
+	static UE_INL_API FVector VInterpTo(FVector Current, FVector Target, float DeltaTime, float InterpSpeed);
 
 	/**
 	 * Tries to reach Target at a constant rate.
@@ -1761,7 +1767,7 @@ class UKismetMathLibrary : public UBlueprintFunctionLibrary
 	 * @return		New interpolated position
 	 */
 	UFUNCTION(BlueprintPure, Category = "Math|Interpolation", meta = (ScriptMethod = "InterpToConstant", Keywords = "position"))
-	static FVector VInterpTo_Constant(FVector Current, FVector Target, float DeltaTime, float InterpSpeed);
+	static UE_INL_API FVector VInterpTo_Constant(FVector Current, FVector Target, float DeltaTime, float InterpSpeed);
 
 	/**
 	 * Uses a simple spring model to interpolate a vector from Current to Target.
@@ -1812,7 +1818,7 @@ class UKismetMathLibrary : public UBlueprintFunctionLibrary
 	 * @return Reciprocal of this vector.
 	 */
 	UFUNCTION(BlueprintPure, meta=(DisplayName = "Reciprocal (Vector)", ScriptMethod = "Reciprocal"), Category="Math|Vector")
-	static FVector Vector_Reciprocal(const FVector& A);
+	static UE_INL_API FVector Vector_Reciprocal(const FVector& A);
 
 	/** 
 	 * Given a direction vector and a surface normal, returns the vector reflected across the surface normal.
@@ -1824,7 +1830,7 @@ class UKismetMathLibrary : public UBlueprintFunctionLibrary
 	 * @returns Reflected vector.
 	 */
 	UFUNCTION(BlueprintPure, meta=(ScriptMethod = "MirrorByVector", Keywords = "Reflection"), Category="Math|Vector")
-	static FVector GetReflectionVector(FVector Direction, FVector SurfaceNormal);
+	static UE_INL_API FVector GetReflectionVector(FVector Direction, FVector SurfaceNormal);
 
 	/** 
 	 * Given a direction vector and a surface normal, returns the vector reflected across the surface normal.
@@ -1836,7 +1842,7 @@ class UKismetMathLibrary : public UBlueprintFunctionLibrary
 	 * @returns Reflected vector.
 	 */
 	UFUNCTION(BlueprintPure, Category="Math|Vector")
-	static FVector MirrorVectorByNormal(FVector InVect, FVector InNormal);
+	static UE_INL_API FVector MirrorVectorByNormal(FVector InVect, FVector InNormal);
 
 	/**
 	 * Mirrors a vector about a plane.
@@ -1845,7 +1851,7 @@ class UKismetMathLibrary : public UBlueprintFunctionLibrary
 	 * @return Mirrored vector.
 	 */
 	UFUNCTION(BlueprintPure, meta=(ScriptMethod = "MirrorByPlane", Keywords = "Reflection"), Category="Math|Vector")
-	static FVector Vector_MirrorByPlane(FVector A, const FPlane& InPlane);
+	static UE_INL_API FVector Vector_MirrorByPlane(FVector A, const FPlane& InPlane);
 
 	/**
 	 * Gets a copy of this vector snapped to a grid.
@@ -1854,7 +1860,7 @@ class UKismetMathLibrary : public UBlueprintFunctionLibrary
 	 * @return A copy of this vector snapped to a grid.
 	 */
 	UFUNCTION(BlueprintPure, meta=(ScriptMethod = "SnappedToGrid", Keywords = "Bounding"), Category="Math|Vector")
-	static FVector Vector_SnappedToGrid(FVector InVect, float InGridSize);
+	static UE_INL_API FVector Vector_SnappedToGrid(FVector InVect, float InGridSize);
 
 	/**
 	 * Get a copy of this vector, clamped inside of an axis aligned cube centered at the origin.
@@ -1863,7 +1869,7 @@ class UKismetMathLibrary : public UBlueprintFunctionLibrary
 	 * @return A copy of this vector, bound by cube.
 	 */
 	UFUNCTION(BlueprintPure, meta=(ScriptMethod = "BoundedToCube", Keywords = "Bounding"), Category="Math|Vector")
-	static FVector Vector_BoundedToCube(FVector InVect, float InRadius);
+	static UE_INL_API FVector Vector_BoundedToCube(FVector InVect, float InRadius);
 
 	/**
 	 * Add a vector to this and clamp the result to an axis aligned cube centered at the origin.
@@ -1876,7 +1882,7 @@ class UKismetMathLibrary : public UBlueprintFunctionLibrary
 
 	/** Get a copy of this vector, clamped inside of the specified axis aligned cube. */
 	UFUNCTION(BlueprintPure, meta=(ScriptMethod = "BoundedToBox", Keywords = "Bounding"), Category="Math|Vector")
-	static FVector Vector_BoundedToBox(FVector InVect, FVector InBoxMin, FVector InBoxMax);
+	static UE_INL_API FVector Vector_BoundedToBox(FVector InVect, FVector InBoxMin, FVector InBoxMax);
 
 	/**
 	 * Gets a copy of this vector projected onto the input vector, which is assumed to be unit length.
@@ -1885,7 +1891,7 @@ class UKismetMathLibrary : public UBlueprintFunctionLibrary
 	 * @return Projected vector.
 	 */
 	UFUNCTION(BlueprintPure, meta=(ScriptMethod = "ProjectOnToNormal", Keywords = "Project"), Category="Math|Vector")
-	static FVector Vector_ProjectOnToNormal(FVector V, FVector InNormal);
+	static UE_INL_API FVector Vector_ProjectOnToNormal(FVector V, FVector InNormal);
 
 	/**
 	* Projects one vector (V) onto another (Target) and returns the projected vector.
@@ -1896,7 +1902,7 @@ class UKismetMathLibrary : public UBlueprintFunctionLibrary
 	* @return V projected on to Target.
 	*/
 	UFUNCTION(BlueprintPure, meta=(ScriptMethod = "ProjectOnTo", Keywords = "Project"), Category="Math|Vector")
-	static FVector ProjectVectorOnToVector(FVector V, FVector Target);
+	static UE_INL_API FVector ProjectVectorOnToVector(FVector V, FVector Target);
 
 	/**
 	 * Projects/snaps a point onto a plane defined by a point on the plane and a plane normal.
@@ -1907,7 +1913,7 @@ class UKismetMathLibrary : public UBlueprintFunctionLibrary
 	 * @return Point projected onto the plane.
 	 */
 	UFUNCTION(BlueprintPure, meta=(ScriptMethod = "ProjectPointOnToPlane", Keywords = "Project"), Category = "Math|Vector")
-	static FVector ProjectPointOnToPlane(FVector Point, FVector PlaneBase, FVector PlaneNormal);
+	static UE_INL_API FVector ProjectPointOnToPlane(FVector Point, FVector PlaneBase, FVector PlaneNormal);
 
 	/**
 	* Projects a vector onto a plane defined by a normalized vector (PlaneNormal).
@@ -1917,7 +1923,7 @@ class UKismetMathLibrary : public UBlueprintFunctionLibrary
 	* @return Vector projected onto the plane.
 	*/
 	UFUNCTION(BlueprintPure, meta=(ScriptMethod = "ProjectOnToPlane", Keywords = "Project"), Category="Math|Vector")
-	static FVector ProjectVectorOnToPlane(FVector V, FVector PlaneNormal);
+	static UE_INL_API FVector ProjectVectorOnToPlane(FVector V, FVector PlaneNormal);
 
 	/**
 	 * Find closest points between 2 segments.
@@ -1930,7 +1936,7 @@ class UKismetMathLibrary : public UBlueprintFunctionLibrary
 	 * @param	Segment2Point	Closest point on segment 2 to segment 1.
 	*/
 	UFUNCTION(BlueprintPure, Category = "Math|Vector")
-	static void FindNearestPointsOnLineSegments(FVector Segment1Start, FVector Segment1End, FVector Segment2Start, FVector Segment2End, FVector& Segment1Point, FVector& Segment2Point);
+	static UE_INL_API void FindNearestPointsOnLineSegments(FVector Segment1Start, FVector Segment1End, FVector Segment2Start, FVector Segment2End, FVector& Segment1Point, FVector& Segment2Point);
 	
 	/**
 	 * Find the closest point on a segment to a given point.
@@ -1941,7 +1947,7 @@ class UKismetMathLibrary : public UBlueprintFunctionLibrary
 	 * @return The closest point on the segment to the given point.
 	 */
 	UFUNCTION(BlueprintPure, Category = "Math|Vector")
-	static FVector FindClosestPointOnSegment(FVector Point, FVector SegmentStart, FVector SegmentEnd);
+	static UE_INL_API FVector FindClosestPointOnSegment(FVector Point, FVector SegmentStart, FVector SegmentEnd);
 
 	/**
 	 * Find the closest point on an infinite line to a given point.
@@ -1963,7 +1969,7 @@ class UKismetMathLibrary : public UBlueprintFunctionLibrary
 	* @return The distance from the given point to the closest point on the segment.
 	*/
 	UFUNCTION(BlueprintPure, Category = "Math|Vector")
-	static float GetPointDistanceToSegment(FVector Point, FVector SegmentStart, FVector SegmentEnd);
+	static UE_INL_API float GetPointDistanceToSegment(FVector Point, FVector SegmentStart, FVector SegmentEnd);
 
 	/**
 	* Find the distance from a point to the closest point on an infinite line.
@@ -1974,7 +1980,7 @@ class UKismetMathLibrary : public UBlueprintFunctionLibrary
 	* @return The distance from the given point to the closest point on the line.
 	*/
 	UFUNCTION(BlueprintPure, Category = "Math|Vector")
-	static float GetPointDistanceToLine(FVector Point, FVector LineOrigin, FVector LineDirection);
+	static UE_INL_API float GetPointDistanceToLine(FVector Point, FVector LineOrigin, FVector LineDirection);
 
 	/** Returns a random vector with length of 1 */
 	UFUNCTION(BlueprintPure, Category="Math|Random", meta=(NotBlueprintThreadSafe))
@@ -1982,11 +1988,11 @@ class UKismetMathLibrary : public UBlueprintFunctionLibrary
 
 	/** Returns a random point within the specified bounding box using the first vector as an origin and the second as the box extents. */
 	UFUNCTION(BlueprintPure, Category = "Math|Random", meta=(ScriptMethod = "RandomPointInBoxExtents", NotBlueprintThreadSafe))
-	static FVector RandomPointInBoundingBox(const FVector Center, const FVector HalfSize);
+	static UE_INL_API FVector RandomPointInBoundingBox(const FVector Center, const FVector HalfSize);
 
 	/** Returns a random point within the specified bounding box. */
 	UFUNCTION(BlueprintPure, Category = "Math|Random", meta=(DisplayName = "Random Point In Bounding Box (Box)", ScriptMethod = "RandomPointInBoxExtents", NotBlueprintThreadSafe))
-	static FVector RandomPointInBoundingBox_Box(const FBox Box);
+	static UE_INL_API FVector RandomPointInBoundingBox_Box(const FBox Box);
 
 	/** 
 	 * Returns a random vector with length of 1, within the specified cone, with uniform random distribution.
@@ -1994,7 +2000,7 @@ class UKismetMathLibrary : public UBlueprintFunctionLibrary
 	 * @param ConeHalfAngleInRadians	The half-angle of the cone (from ConeDir to edge), in radians.
 	 */
 	UFUNCTION(BlueprintPure, Category = "Math|Random", meta=(NotBlueprintThreadSafe))
-	static FVector RandomUnitVectorInConeInRadians(FVector ConeDir, float ConeHalfAngleInRadians);
+	static UE_INL_API FVector RandomUnitVectorInConeInRadians(FVector ConeDir, float ConeHalfAngleInRadians);
 
 	/** 
 	 * Returns a random vector with length of 1, within the specified cone, with uniform random distribution.
@@ -2037,7 +2043,7 @@ class UKismetMathLibrary : public UBlueprintFunctionLibrary
 
 	/** 4D vector zero constant (0,0,0) */
 	UFUNCTION(BlueprintPure, meta = (ScriptConstant = "Zero", ScriptConstantHost = "/Script/CoreUObject.Vector4"), Category = "Math|Vector4")
-	static FVector4 Vector4_Zero();
+	static UE_INL_API FVector4 Vector4_Zero();
 
 	//
 	// Vector4 functions
@@ -2045,15 +2051,15 @@ class UKismetMathLibrary : public UBlueprintFunctionLibrary
 
 	/** Makes a 4D vector {X, Y, Z, W} */
 	UFUNCTION(BlueprintPure, meta = (Keywords = "construct build", NativeMakeFunc), Category = "Math|Vector4")
-	static FVector4 MakeVector4(double X, double Y, double Z, double W);
+	static UE_INL_API FVector4 MakeVector4(double X, double Y, double Z, double W);
 
 	/** Breaks a 4D vector apart into X, Y, Z, W. */
 	UFUNCTION(BlueprintPure, meta = (NativeBreakFunc), Category = "Math|Vector4")
-	static void BreakVector4(const FVector4& InVec, double& X, double& Y, double& Z, double& W);
+	static UE_INL_API void BreakVector4(const FVector4& InVec, double& X, double& Y, double& Z, double& W);
 
 	/** Converts a Vector4 to a Vector (dropping the W element) */
 	UFUNCTION(BlueprintPure, meta = (DisplayName = "To Vector (Vector4)", CompactNodeTitle = "->", ScriptMethod = "Vector", Keywords = "cast convert", BlueprintAutocast), Category = "Math|Conversions")
-	static FVector Conv_Vector4ToVector(const FVector4& InVector4);
+	static UE_INL_API FVector Conv_Vector4ToVector(const FVector4& InVector4);
 
 	/**
 	 * Return the FRotator orientation corresponding to the direction in which the vector points.
@@ -2062,7 +2068,7 @@ class UKismetMathLibrary : public UBlueprintFunctionLibrary
 	 * @return FRotator from the Vector's direction, without any roll.
 	 */
 	UFUNCTION(BlueprintPure, meta=(DisplayName = "To Rotation (Vector4)", ScriptMethod = "Rotator", Keywords = "rotation rotate cast convert", BlueprintAutocast), Category = "Math|Conversions")
-	static FRotator Conv_Vector4ToRotator(const FVector4& InVec);
+	static UE_INL_API FRotator Conv_Vector4ToRotator(const FVector4& InVec);
 	
 	/**
 	 * Return the Quaternion orientation corresponding to the direction in which the vector points.
@@ -2074,43 +2080,43 @@ class UKismetMathLibrary : public UBlueprintFunctionLibrary
 	 * @return Quaternion from the Vector's direction, without any roll.
 	 */
 	UFUNCTION(BlueprintPure, meta=(DisplayName = "To Quaternion (Vector4)", ScriptMethod = "Quaternion", Keywords = "rotation rotate cast convert", BlueprintAutocast), Category = "Math|Conversions")
-	static FQuat Conv_Vector4ToQuaternion(const FVector4& InVec);
+	static UE_INL_API FQuat Conv_Vector4ToQuaternion(const FVector4& InVec);
 
 	/** Returns addition of Vector A and Vector B (A + B) */
 	UFUNCTION(BlueprintPure, meta = (DisplayName = "Vector4 + Vector4", CompactNodeTitle = "+", ScriptMethod = "Add", ScriptOperator = "+;+=", Keywords = "+ add plus", CommutativeAssociativeBinaryOperator = "true"), Category = "Math|Vector4")
-	static FVector4 Add_Vector4Vector4(const FVector4& A, const FVector4& B);
+	static UE_INL_API FVector4 Add_Vector4Vector4(const FVector4& A, const FVector4& B);
 
 	/** Returns subtraction of Vector B from Vector A (A - B) */
 	UFUNCTION(BlueprintPure, meta = (DisplayName = "Vector4 - Vector4", CompactNodeTitle = "-", ScriptMethod = "Subtract", ScriptOperator = "-;-=", Keywords = "- subtract minus"), Category = "Math|Vector4")
-	static FVector4 Subtract_Vector4Vector4(const FVector4& A, const FVector4& B);
+	static UE_INL_API FVector4 Subtract_Vector4Vector4(const FVector4& A, const FVector4& B);
 
 	/** Element-wise Vector multiplication (Result = {A.x*B.x, A.y*B.y, A.z*B.z, A.w*B.w}) */
 	UFUNCTION(BlueprintPure, meta = (DisplayName = "Vector4 * Vector4", CompactNodeTitle = "*", ScriptMethod = "Multiply", ScriptOperator = "*;*=", Keywords = "* multiply", CommutativeAssociativeBinaryOperator = "true"), Category = "Math|Vector4")
-	static FVector4 Multiply_Vector4Vector4(const FVector4& A, const FVector4& B);
+	static UE_INL_API FVector4 Multiply_Vector4Vector4(const FVector4& A, const FVector4& B);
 
 	/** Element-wise Vector divide (Result = {A.x/B.x, A.y/B.y, A.z/B.z, A.w/B.w}) */
 	UFUNCTION(BlueprintPure, meta = (DisplayName = "Vector4 / Vector4", CompactNodeTitle = "/", ScriptMethod = "Divide", ScriptOperator = "/;/=", Keywords = "/ divide division"), Category = "Math|Vector4")
-	static FVector4 Divide_Vector4Vector4(const FVector4& A, const FVector4& B);
+	static UE_INL_API FVector4 Divide_Vector4Vector4(const FVector4& A, const FVector4& B);
 
 	/** Returns true if vector A is equal to vector B (A == B) */
 	UFUNCTION(BlueprintPure, meta=(DisplayName = "Equal Exactly (Vector4)", CompactNodeTitle = "===", ScriptMethod = "Equals", ScriptOperator = "==", Keywords = "== equal"), Category = "Math|Vector4")
-	static bool EqualExactly_Vector4Vector4(const FVector4& A, const FVector4& B);
+	static UE_INL_API bool EqualExactly_Vector4Vector4(const FVector4& A, const FVector4& B);
 
 	/** Returns true if vector A is equal to vector B (A == B) within a specified error tolerance */
 	UFUNCTION(BlueprintPure, meta = (DisplayName = "Equal (Vector4)", CompactNodeTitle = "==", ScriptMethod = "IsNearEqual", Keywords = "== equal"), Category = "Math|Vector4")
-	static bool EqualEqual_Vector4Vector4(const FVector4& A, const FVector4& B, float ErrorTolerance = 1.e-4f);
+	static UE_INL_API bool EqualEqual_Vector4Vector4(const FVector4& A, const FVector4& B, float ErrorTolerance = 1.e-4f);
 
 	/** Returns true if vector A is not equal to vector B (A != B) within a specified error tolerance */
 	UFUNCTION(BlueprintPure, meta = (DisplayName = "Not Equal Exactly (Vector4)", CompactNodeTitle = "!==", ScriptMethod = "NotEqual", ScriptOperator = "!=", Keywords = "!= not equal"), Category = "Math|Vector4")
-	static bool NotEqualExactly_Vector4Vector4(const FVector4& A, const FVector4& B);
+	static UE_INL_API bool NotEqualExactly_Vector4Vector4(const FVector4& A, const FVector4& B);
 
 	/** Returns true if vector A is not equal to vector B (A != B) within a specified error tolerance */
 	UFUNCTION(BlueprintPure, meta = (DisplayName = "Not Equal (Vector4)", CompactNodeTitle = "!=", ScriptMethod = "IsNotNearEqual", Keywords = "!= not equal"), Category = "Math|Vector4")
-	static bool NotEqual_Vector4Vector4(const FVector4& A, const FVector4& B, float ErrorTolerance = 1.e-4f);
+	static UE_INL_API bool NotEqual_Vector4Vector4(const FVector4& A, const FVector4& B, float ErrorTolerance = 1.e-4f);
 
 	/** Gets a negated copy of the vector. Equivalent to -Vector for scripts. */
 	UFUNCTION(BlueprintPure, meta = (DisplayName = "Negated (Vector4)", ScriptMethod = "Negated", ScriptOperator = "neg"), Category = "Math|Vector4")
-	static FVector4 Vector4_Negated(const FVector4& A);
+	static UE_INL_API FVector4 Vector4_Negated(const FVector4& A);
 
 	/**
 	 * Assign the values of the supplied vector.
@@ -2133,15 +2139,15 @@ class UKismetMathLibrary : public UBlueprintFunctionLibrary
 
 	/** Returns the cross product of two vectors - see  http://mathworld.wolfram.com/CrossProduct.html */
 	UFUNCTION(BlueprintPure, meta = (DisplayName = "Cross Product XYZ (Vector4)", CompactNodeTitle = "cross3", ScriptMethod = "Cross3"), Category = "Math|Vector4")
-	static FVector4 Vector4_CrossProduct3(const FVector4& A, const FVector4& B);
+	static UE_INL_API FVector4 Vector4_CrossProduct3(const FVector4& A, const FVector4& B);
 
 	/** Returns the dot product of two vectors - see http://mathworld.wolfram.com/DotProduct.html */
 	UFUNCTION(BlueprintPure, meta = (DisplayName = "Dot Product (Vector4)", CompactNodeTitle = "dot", ScriptMethod = "Dot", ScriptOperator = "|"), Category = "Math|Vector4")
-	static double Vector4_DotProduct(const FVector4& A, const FVector4& B);
+	static UE_INL_API double Vector4_DotProduct(const FVector4& A, const FVector4& B);
 
 	/** Returns the dot product of two vectors - see http://mathworld.wolfram.com/DotProduct.html The W element is ignored.*/
 	UFUNCTION(BlueprintPure, meta = (DisplayName = "Dot Product XYZ (Vector4)", CompactNodeTitle = "dot3", ScriptMethod = "Dot3"), Category = "Math|Vector4")
-	static double Vector4_DotProduct3(const FVector4& A, const FVector4& B);
+	static UE_INL_API double Vector4_DotProduct3(const FVector4& A, const FVector4& B);
 
 	/**
 	 * Determines if any component is not a number (NAN)
@@ -2149,7 +2155,7 @@ class UKismetMathLibrary : public UBlueprintFunctionLibrary
 	 * @return true if one or more components is NAN, otherwise false.
 	 */
 	UFUNCTION(BlueprintPure, meta = (ScriptMethod = "IsNAN"), Category = "Math|Vector4")
-	static bool Vector4_IsNAN(const FVector4& A);
+	static UE_INL_API bool Vector4_IsNAN(const FVector4& A);
 
 	/**
 	 * Checks whether vector is near to zero within a specified tolerance. The W element is ignored.
@@ -2158,7 +2164,7 @@ class UKismetMathLibrary : public UBlueprintFunctionLibrary
 	 * @return true if vector is in tolerance to zero, otherwise false.
 	 */
 	UFUNCTION(BlueprintPure, meta = (ScriptMethod = "IsNearlyZero3"), Category = "Math|Vector4")
-	static bool Vector4_IsNearlyZero3(const FVector4& A, float Tolerance = 1.e-4f);
+	static UE_INL_API bool Vector4_IsNearlyZero3(const FVector4& A, float Tolerance = 1.e-4f);
 
 	/**
 	 * Checks whether all components of the vector are exactly zero.
@@ -2166,23 +2172,23 @@ class UKismetMathLibrary : public UBlueprintFunctionLibrary
 	 * @return true if vector is exactly zero, otherwise false.
 	 */
 	UFUNCTION(BlueprintPure, meta = (ScriptMethod = "IsZero"), Category = "Math|Vector4")
-	static bool Vector4_IsZero(const FVector4& A);
+	static UE_INL_API bool Vector4_IsZero(const FVector4& A);
 
 	/** Returns the length of the vector. */
 	UFUNCTION(BlueprintPure, meta = (DisplayName = "Length (Vector4)", ScriptMethod = "Length", Keywords = "magnitude"), Category = "Math|Vector4")
-	static double Vector4_Size(const FVector4& A);
+	static UE_INL_API double Vector4_Size(const FVector4& A);
 
 	/** Returns the squared length of the vector. */
 	UFUNCTION(BlueprintPure, meta = (DisplayName = "Length Squared (Vector4)", ScriptMethod = "LengthSquared", Keywords = "magnitude"), Category = "Math|Vector4")
-	static double Vector4_SizeSquared(const FVector4& A);
+	static UE_INL_API double Vector4_SizeSquared(const FVector4& A);
 
 	/** Returns the length of the vector. The W element is ignored. */
 	UFUNCTION(BlueprintPure, meta = (DisplayName = "Length XYZ (Vector4)", ScriptMethod = "Length3", Keywords = "magnitude"), Category = "Math|Vector4")
-	static double Vector4_Size3(const FVector4& A);
+	static UE_INL_API double Vector4_Size3(const FVector4& A);
 
 	/** Returns the squared length of the vector. The W element is ignored. */
 	UFUNCTION(BlueprintPure, meta = (DisplayName = "Length XYZ Squared (Vector4)", ScriptMethod = "LengthSquared3", Keywords = "magnitude"), Category = "Math|Vector4")
-	static double Vector4_SizeSquared3(const FVector4& A);
+	static UE_INL_API double Vector4_SizeSquared3(const FVector4& A);
 
 	/**
 	 * Determines if vector is normalized / unit (length 1) within specified squared tolerance. The W element is ignored.
@@ -2190,7 +2196,7 @@ class UKismetMathLibrary : public UBlueprintFunctionLibrary
 	 * @return true if unit, false otherwise.
 	 */
 	UFUNCTION(BlueprintPure, meta=(DisplayName = "Is Unit XYZ (Vector4)", ScriptMethod = "IsUnit3", Keywords = "Unit Vector"), Category = "Math|Vector4")
-	static bool Vector4_IsUnit3(const FVector4& A, float SquaredLenthTolerance = 1.e-4f);
+	static UE_INL_API bool Vector4_IsUnit3(const FVector4& A, float SquaredLenthTolerance = 1.e-4f);
 
 	/**
 	 * Determines if vector is normalized / unit (length 1). The W element is ignored.
@@ -2198,7 +2204,7 @@ class UKismetMathLibrary : public UBlueprintFunctionLibrary
 	 * @return true if normalized, false otherwise.
 	 */
 	UFUNCTION(BlueprintPure, meta=(DisplayName = "Is Normal XYZ (Vector4)", ScriptMethod = "IsNormal3", Keywords = "Unit Vector"), Category = "Math|Vector4")
-	static bool Vector4_IsNormal3(const FVector4& A);
+	static UE_INL_API bool Vector4_IsNormal3(const FVector4& A);
 
 	/**
 	 * Gets a normalized unit copy of the vector, ensuring it is safe to do so based on the length. The W element is ignored and the returned vector has W=0.
@@ -2208,7 +2214,7 @@ class UKismetMathLibrary : public UBlueprintFunctionLibrary
 	 * @return A normalized copy if safe, (0,0,0) otherwise.
 	 */
 	UFUNCTION(BlueprintPure, meta=(DisplayName = "Normalize XYZ (Vector 4)", ScriptMethod = "Normal3", Keywords = "Unit Vector"), Category = "Math|Vector4")
-	static FVector4 Vector4_Normal3(const FVector4& A, float Tolerance = 1.e-4f);
+	static UE_INL_API FVector4 Vector4_Normal3(const FVector4& A, float Tolerance = 1.e-4f);
 
 	/**
 	 * Calculates normalized unit version of vector without checking for zero length. The W element is ignored and the returned vector has W=0.
@@ -2216,7 +2222,7 @@ class UKismetMathLibrary : public UBlueprintFunctionLibrary
 	 * @return Normalized version of vector.
 	 */
 	UFUNCTION(BlueprintPure, meta=(DisplayName = "Normal Unsafe XYZ (Vector4)", ScriptMethod = "NormalUnsafe3", Keywords = "Unit Vector"), Category = "Math|Vector4")
-	static FVector4 Vector4_NormalUnsafe3(const FVector4& A);
+	static UE_INL_API FVector4 Vector4_NormalUnsafe3(const FVector4& A);
 
 	/**
 	 * Normalize this vector in-place if it is large enough or set it to (0,0,0,0) otherwise. The W element is ignored and the returned vector has W=0.
@@ -2237,7 +2243,7 @@ class UKismetMathLibrary : public UBlueprintFunctionLibrary
 	 * @returns Reflected vector.
 	 */
 	UFUNCTION(BlueprintPure, meta=(ScriptMethod = "MirrorByVector3", Keywords = "Reflection"), Category = "Math|Vector4")
-	static FVector4 Vector4_MirrorByVector3(const FVector4& Direction, const FVector4& SurfaceNormal);
+	static UE_INL_API FVector4 Vector4_MirrorByVector3(const FVector4& Direction, const FVector4& SurfaceNormal);
 
 	/**
 	 * Transform the input vector4 by a provided matrix4x4 and returns the resulting vector4.
@@ -2245,7 +2251,7 @@ class UKismetMathLibrary : public UBlueprintFunctionLibrary
 	 * @return Transformed vector4.
 	 */
 	UFUNCTION(BlueprintPure, meta = (DisplayName = "Transform Vector4 by Matrix"), Category = "Math|Vector4")
-	static FVector4 TransformVector4(const FMatrix& Matrix, const FVector4& Vec4);
+	static UE_INL_API FVector4 TransformVector4(const FMatrix& Matrix, const FVector4& Vec4);
 
 	//
 	// Rotator functions.
@@ -2260,39 +2266,39 @@ class UKismetMathLibrary : public UBlueprintFunctionLibrary
 
 	/** Builds a rotator given only a XAxis. Y and Z are unspecified but will be orthonormal. XAxis need not be normalized. */
 	UFUNCTION(BlueprintPure, Category="Math|Rotator", meta=(Keywords="construct build rotation rotate rotator makerotator"))
-	static FRotator MakeRotFromX(const FVector& X);
+	static UE_INL_API FRotator MakeRotFromX(const FVector& X);
 
 	/** Builds a rotation matrix given only a YAxis. X and Z are unspecified but will be orthonormal. YAxis need not be normalized. */
 	UFUNCTION(BlueprintPure, Category="Math|Rotator", meta=(Keywords="construct build rotation rotate rotator makerotator"))
-	static FRotator MakeRotFromY(const FVector& Y);
+	static UE_INL_API FRotator MakeRotFromY(const FVector& Y);
 
 	/** Builds a rotation matrix given only a ZAxis. X and Y are unspecified but will be orthonormal. ZAxis need not be normalized. */
 	UFUNCTION(BlueprintPure, Category="Math|Rotator", meta=(Keywords="construct build rotation rotate rotator makerotator"))
-	static FRotator MakeRotFromZ(const FVector& Z);
+	static UE_INL_API FRotator MakeRotFromZ(const FVector& Z);
 
 	/** Builds a matrix with given X and Y axes. X will remain fixed, Y may be changed minimally to enforce orthogonality. Z will be computed. Inputs need not be normalized. */
 	UFUNCTION(BlueprintPure, Category="Math|Rotator", meta=(Keywords="construct build rotation rotate rotator makerotator"))
-	static FRotator MakeRotFromXY(const FVector& X, const FVector& Y);
+	static UE_INL_API FRotator MakeRotFromXY(const FVector& X, const FVector& Y);
 
 	/** Builds a matrix with given X and Z axes. X will remain fixed, Z may be changed minimally to enforce orthogonality. Y will be computed. Inputs need not be normalized. */
 	UFUNCTION(BlueprintPure, Category="Math|Rotator", meta=(Keywords="construct build rotation rotate rotator makerotator"))
-	static FRotator MakeRotFromXZ(const FVector& X, const FVector& Z);
+	static UE_INL_API FRotator MakeRotFromXZ(const FVector& X, const FVector& Z);
 
 	/** Builds a matrix with given Y and X axes. Y will remain fixed, X may be changed minimally to enforce orthogonality. Z will be computed. Inputs need not be normalized. */
 	UFUNCTION(BlueprintPure, Category="Math|Rotator", meta=(Keywords="construct build rotation rotate rotator makerotator"))
-	static FRotator MakeRotFromYX(const FVector& Y, const FVector& X);
+	static UE_INL_API FRotator MakeRotFromYX(const FVector& Y, const FVector& X);
 
 	/** Builds a matrix with given Y and Z axes. Y will remain fixed, Z may be changed minimally to enforce orthogonality. X will be computed. Inputs need not be normalized. */
 	UFUNCTION(BlueprintPure, Category="Math|Rotator", meta=(Keywords="construct build rotation rotate rotator makerotator"))
-	static FRotator MakeRotFromYZ(const FVector& Y, const FVector& Z);
+	static UE_INL_API FRotator MakeRotFromYZ(const FVector& Y, const FVector& Z);
 
 	/** Builds a matrix with given Z and X axes. Z will remain fixed, X may be changed minimally to enforce orthogonality. Y will be computed. Inputs need not be normalized. */
 	UFUNCTION(BlueprintPure, Category="Math|Rotator", meta=(Keywords="construct build rotation rotate rotator makerotator"))
-	static FRotator MakeRotFromZX(const FVector& Z, const FVector& X);
+	static UE_INL_API FRotator MakeRotFromZX(const FVector& Z, const FVector& X);
 
 	/** Builds a matrix with given Z and Y axes. Z will remain fixed, Y may be changed minimally to enforce orthogonality. X will be computed. Inputs need not be normalized. */
 	UFUNCTION(BlueprintPure, Category="Math|Rotator", meta=(Keywords="construct build rotation rotate rotator makerotator"))
-	static FRotator MakeRotFromZY(const FVector& Z, const FVector& Y);
+	static UE_INL_API FRotator MakeRotFromZY(const FVector& Z, const FVector& Y);
 
 	// Build a reference frame from three axes
 	UFUNCTION(BlueprintPure, Category="Math|Rotator", meta=(Keywords="construct build rotation rotate"))
@@ -2300,7 +2306,7 @@ class UKismetMathLibrary : public UBlueprintFunctionLibrary
 
 	/** Find a rotation for an object at Start location to point at Target location. */
 	UFUNCTION(BlueprintPure, Category="Math|Rotator", meta=(Keywords="rotation rotate"))
-	static FRotator FindLookAtRotation(const FVector& Start, const FVector& Target);
+	static UE_INL_API FRotator FindLookAtRotation(const FVector& Start, const FVector& Target);
 
 	/** 
 	 * Find a local rotation (range of [-180, 180]) for an object with StartTransform to point at TargetLocation. 
@@ -2323,19 +2329,19 @@ class UKismetMathLibrary : public UBlueprintFunctionLibrary
 
 	/** Returns true if rotator A is equal to rotator B (A == B) within a specified error tolerance */
 	UFUNCTION(BlueprintPure, meta=(DisplayName = "Equal (Rotator)", CompactNodeTitle = "==", ScriptMethod = "IsNearEqual", ScriptOperator = "==", Keywords = "== equal"), Category="Math|Rotator")
-	static bool EqualEqual_RotatorRotator(FRotator A, FRotator B, float ErrorTolerance = 1.e-4f);
+	static UE_INL_API bool EqualEqual_RotatorRotator(FRotator A, FRotator B, float ErrorTolerance = 1.e-4f);
 
 	/** Returns true if rotator A is not equal to rotator B (A != B) within a specified error tolerance */
 	UFUNCTION(BlueprintPure, meta=(DisplayName = "Not Equal (Rotator)", CompactNodeTitle = "!=", ScriptMethod = "IsNotNearEqual", ScriptOperator = "!=", Keywords = "!= not equal"), Category="Math|Rotator")
-	static bool NotEqual_RotatorRotator(FRotator A, FRotator B, float ErrorTolerance = 1.e-4f);
+	static UE_INL_API bool NotEqual_RotatorRotator(FRotator A, FRotator B, float ErrorTolerance = 1.e-4f);
 
 	/** Returns rotator representing rotator A scaled by B */
 	UFUNCTION(BlueprintPure, meta=(DisplayName = "Scale Rotator (Float)", CompactNodeTitle = "*", ScriptMethod = "Scale", Keywords = "* multiply rotate rotation"), Category="Math|Rotator")
-	static FRotator Multiply_RotatorFloat(FRotator A, float B);
+	static UE_INL_API FRotator Multiply_RotatorFloat(FRotator A, float B);
 	
 	/** Returns rotator representing rotator A scaled by B */
 	UFUNCTION(BlueprintPure, meta=(DisplayName = "Scale Rotator (Integer)", CompactNodeTitle = "*", ScriptMethod = "ScaleInteger", Keywords = "* multiply rotate rotation"), Category="Math|Rotator")
-	static FRotator Multiply_RotatorInt(FRotator A, int32 B);
+	static UE_INL_API FRotator Multiply_RotatorInt(FRotator A, int32 B);
 
 	/** Combine 2 rotations to give you the resulting rotation of first applying A, then B. */
 	UFUNCTION(BlueprintPure, meta=(DisplayName = "Combine Rotators", ScriptMethod = "Combine", Keywords="rotate rotation add"), Category="Math|Rotator")
@@ -2343,7 +2349,7 @@ class UKismetMathLibrary : public UBlueprintFunctionLibrary
 
 	/** Negate a rotator*/
 	UFUNCTION(BlueprintPure, meta=(DisplayName="Invert Rotator", ScriptMethod = "Inversed", ScriptOperator = "neg", Keywords="rotate rotation"), Category="Math|Rotator")
-	static FRotator NegateRotator(FRotator A);
+	static UE_INL_API FRotator NegateRotator(FRotator A);
 
 	/** Rotate the world forward vector by the given rotation */
 	UFUNCTION(BlueprintPure, meta=(ScriptMethod = "GetForwardVector", Keywords="rotation rotate"), Category="Math|Vector")
@@ -2359,11 +2365,11 @@ class UKismetMathLibrary : public UBlueprintFunctionLibrary
 
 	/** Get the X direction vector after this rotation */
 	UFUNCTION(BlueprintPure, meta=(ScriptMethod = "ToVector", DisplayName = "Get Rotation X Vector", Keywords="rotation rotate cast convert", BlueprintAutocast), Category="Math|Rotator")
-	static FVector Conv_RotatorToVector(FRotator InRot);
+	static UE_INL_API FVector Conv_RotatorToVector(FRotator InRot);
 
 	/** Converts Rotator to Transform */
 	UFUNCTION(BlueprintPure, meta = (DisplayName = "To Transform (Rotator)", CompactNodeTitle = "->", ScriptMethod = "Transform", Keywords = "cast convert", BlueprintAutocast), Category = "Math|Conversions")
-	static FTransform Conv_RotatorToTransform(const FRotator& InRotator);
+	static UE_INL_API FTransform Conv_RotatorToTransform(const FRotator& InRotator);
 
 	/** Get the reference frame direction vectors (axes) described by this rotation */
 	UFUNCTION(BlueprintPure, meta=(ScriptMethod = "GetAxes", Keywords="rotate rotation"), Category="Math|Rotator")
@@ -2412,13 +2418,13 @@ class UKismetMathLibrary : public UBlueprintFunctionLibrary
 	* (Assumes Matrix represents a transform) 
 	*/
 	UFUNCTION(BlueprintPure, meta = (DisplayName = "To Transform (Matrix)", CompactNodeTitle = "->", ScriptMethod = "Transform", Keywords = "cast convert"), Category = "Math|Conversions")
-	static FTransform Conv_MatrixToTransform(const FMatrix& InMatrix);
+	static UE_INL_API FTransform Conv_MatrixToTransform(const FMatrix& InMatrix);
 
 	/** Converts a Matrix to a Rotator 
 	* (Assumes Matrix represents a transform) 
 	*/
 	UFUNCTION(BlueprintPure, meta = (DisplayName = "To Rotator (Matrix)", CompactNodeTitle = "->", ScriptMethod = "Rotator", Keywords = "cast convert"), Category = "Math|Conversions")
-	static FRotator Conv_MatrixToRotator(const FMatrix& InMatrix);
+	static UE_INL_API FRotator Conv_MatrixToRotator(const FMatrix& InMatrix);
 
 	/**
 	 * Get the origin of the co-ordinate system
@@ -2427,7 +2433,7 @@ class UKismetMathLibrary : public UBlueprintFunctionLibrary
 	 * @return co-ordinate system origin
 	 */
 	UFUNCTION(BlueprintPure, meta = (DisplayName = "Get Origin (Matrix)", ScriptMethod = "GetOrigin"), Category = "Math|Matrix")
-	static FVector Matrix_GetOrigin(const FMatrix& InMatrix);
+	static UE_INL_API FVector Matrix_GetOrigin(const FMatrix& InMatrix);
 
 	// Identity matrix
 	UFUNCTION(BlueprintPure, meta = (DisplayName = "Identity (Matrix)", ScriptConstant = "Identity", ScriptConstantHost = "/Script/CoreUObject.Matrix"), Category = "Math|Matrix")
@@ -2466,7 +2472,7 @@ class UKismetMathLibrary : public UBlueprintFunctionLibrary
 	 * @return true if two Matrix are equal, within specified tolerance, otherwise false.
 	 */
 	UFUNCTION(BlueprintPure, meta = (DisplayName = "Equal (Matrix)", CompactNodeTitle = "==", ScriptMethod = "Equals", ScriptOperator = "==", Keywords = "== equal"), Category = "Math|Matrix")
-	static bool EqualEqual_MatrixMatrix(const FMatrix& A, const FMatrix& B, float Tolerance = 1.e-4f);
+	static UE_INL_API bool EqualEqual_MatrixMatrix(const FMatrix& A, const FMatrix& B, float Tolerance = 1.e-4f);
 
 	/**
 	 * Checks whether another Matrix is not equal to this, within specified tolerance.
@@ -2475,26 +2481,26 @@ class UKismetMathLibrary : public UBlueprintFunctionLibrary
 	 * @return true if two Matrix are not equal, within specified tolerance, otherwise false.
 	 */
 	UFUNCTION(BlueprintPure, meta = (DisplayName = "Not Equal (Matrix)", CompactNodeTitle = "!=", ScriptMethod = "NotEqual", ScriptOperator = "!=", Keywords = "!= not equal"), Category = "Math|Matrix")
-	static bool NotEqual_MatrixMatrix(const FMatrix& A, const FMatrix& B, float Tolerance = 1.e-4f);
+	static UE_INL_API bool NotEqual_MatrixMatrix(const FMatrix& A, const FMatrix& B, float Tolerance = 1.e-4f);
 
 	/**
 	 * Transform a vector by the matrix.
 	 * (Assumes Matrix represents a transform) 
 	*/
 	UFUNCTION(BlueprintPure, meta = (DisplayName = "Transform Vector4 (Matrix)", ScriptMethod = "TransformVector4"), Category = "Math|Matrix")
-	static FVector4 Matrix_TransformVector4(const FMatrix& M, FVector4 V);
+	static UE_INL_API FVector4 Matrix_TransformVector4(const FMatrix& M, FVector4 V);
 
 	/** Transform a location - will take into account translation part of the FMatrix.
 	 * (Assumes Matrix represents a transform)
 	 */
 	UFUNCTION(BlueprintPure, meta = (DisplayName = "Transform Position (Matrix)", ScriptMethod = "TransformPosition"), Category = "Math|Matrix")
-	static FVector4 Matrix_TransformPosition(const FMatrix& M, FVector V);
+	static UE_INL_API FVector4 Matrix_TransformPosition(const FMatrix& M, FVector V);
 
 	/** Inverts the matrix and then transforms V - correctly handles scaling in this matrix.
 	* (Assumes Matrix represents a transform)
 	*/
 	UFUNCTION(BlueprintPure, meta = (DisplayName = "Inverse Transform Position (Matrix)", ScriptMethod = "InverseTransformPosition"), Category = "Math|Matrix")
-	static FVector Matrix_InverseTransformPosition(const FMatrix& M, FVector V);
+	static UE_INL_API FVector Matrix_InverseTransformPosition(const FMatrix& M, FVector V);
 
 	/**
 	 *	Transform a direction vector - will not take into account translation part of the FMatrix.
@@ -2502,7 +2508,7 @@ class UKismetMathLibrary : public UBlueprintFunctionLibrary
 	 * (Assumes Matrix represents a transform)
 	*/
 	UFUNCTION(BlueprintPure, meta = (DisplayName = "Transform Vector (Matrix)", ScriptMethod = "TransformVector"), Category = "Math|Matrix")
-	static FVector4 Matrix_TransformVector(const FMatrix& M, FVector V);
+	static UE_INL_API FVector4 Matrix_TransformVector(const FMatrix& M, FVector V);
 
 	/**
 	 *	Transform a direction vector by the inverse of this matrix - will not take into account translation part.
@@ -2510,7 +2516,7 @@ class UKismetMathLibrary : public UBlueprintFunctionLibrary
 	 * (Assumes Matrix represents a transform)
 	*/
 	UFUNCTION(BlueprintPure, meta = (DisplayName = "Inverse Transform Vector (Matrix)", ScriptMethod = "InverseTransformVector"), Category = "Math|Matrix")
-	static FVector Matrix_InverseTransformVector(const FMatrix& M, FVector V);
+	static UE_INL_API FVector Matrix_InverseTransformVector(const FMatrix& M, FVector V);
 
 	// Transpose.
 	UFUNCTION(BlueprintPure, meta = (DisplayName = "Get Transposed (Matrix)", ScriptMethod = "GetTransposed"), Category = "Math|Matrix")
@@ -2518,13 +2524,13 @@ class UKismetMathLibrary : public UBlueprintFunctionLibrary
 
 	// @return determinant of this matrix.
 	UFUNCTION(BlueprintPure, meta = (DisplayName = "Get Determinant (Matrix)", ScriptMethod = "GetDeterminant"), Category = "Math|Matrix")
-	static float Matrix_GetDeterminant(const FMatrix& M);
+	static UE_INL_API float Matrix_GetDeterminant(const FMatrix& M);
 
 	/** @return the determinant of rotation 3x3 matrix 
 	* (Assumes Top Left 3x3 Submatrix represents a Rotation)
 	*/
 	UFUNCTION(BlueprintPure, meta = (DisplayName = "Get Rotation Determinant (Matrix)", ScriptMethod = "GetRotDeterminant"), Category = "Math|Matrix")
-	static float Matrix_GetRotDeterminant(const FMatrix& M);
+	static UE_INL_API float Matrix_GetRotDeterminant(const FMatrix& M);
 
 	/** Get the inverse of the Matrix. Handles nil matrices. */
 	UFUNCTION(BlueprintPure, meta = (DisplayName = "GetInverse (Matrix)", ScriptMethod = "GetInverse"), Category = "Math|Matrix")
@@ -2550,7 +2556,7 @@ class UKismetMathLibrary : public UBlueprintFunctionLibrary
 	* (Assumes Matrix represents a transform)
 	*/
 	UFUNCTION(BlueprintPure, meta = (DisplayName = "Get Scale Vector (Matrix)", ScriptMethod = "GetScaleVector"), Category = "Math|Matrix")
-	static FVector Matrix_GetScaleVector(const FMatrix& M, float Tolerance = 1.e-8f);
+	static UE_INL_API FVector Matrix_GetScaleVector(const FMatrix& M, float Tolerance = 1.e-8f);
 
 	/** Remove any translation from this matrix
 	* (Assumes Matrix represents a transform)
@@ -2566,7 +2572,7 @@ class UKismetMathLibrary : public UBlueprintFunctionLibrary
 
 	/** Returns true if any element of this matrix is NaN */
 	UFUNCTION(BlueprintPure, meta = (DisplayName = "Contains NaN (Matrix)", ScriptMethod = "ContainsNaN"), Category = "Math|Matrix")
-	static bool Matrix_ContainsNaN(const FMatrix& M);
+	static UE_INL_API bool Matrix_ContainsNaN(const FMatrix& M);
 
 	/** Scale the translation part of the matrix by the supplied vector.
 	* (Assumes Matrix represents a transform)
@@ -2578,7 +2584,7 @@ class UKismetMathLibrary : public UBlueprintFunctionLibrary
 	* (Assumes Matrix represents a transform)
 	*/
 	UFUNCTION(BlueprintPure, meta = (DisplayName = "Get Maximum Axis Scale (Matrix)", ScriptMethod = "GetMaximumAxisScale"), Category = "Math|Matrix")
-	static float Matrix_GetMaximumAxisScale(const FMatrix& M);
+	static UE_INL_API float Matrix_GetMaximumAxisScale(const FMatrix& M);
 
 	/** Apply Scale to this matrix
 	* (Assumes Matrix represents a transform)
@@ -2594,7 +2600,7 @@ class UKismetMathLibrary : public UBlueprintFunctionLibrary
 	 * @ return vector of the axis
 	 */
 	UFUNCTION(BlueprintPure, meta = (DisplayName = "Get Scaled Axis (Matrix)", ScriptMethod = "GetScaledAxis"), Category = "Math|Matrix")
-	static FVector Matrix_GetScaledAxis(const FMatrix& M, TEnumAsByte<EAxis::Type> Axis);
+	static UE_INL_API FVector Matrix_GetScaledAxis(const FMatrix& M, TEnumAsByte<EAxis::Type> Axis);
 
 	/**
 	 * get axes of this matrix scaled by the scale of the matrix
@@ -2605,7 +2611,7 @@ class UKismetMathLibrary : public UBlueprintFunctionLibrary
 	 * @param Z axes returned to this param
 	 */
 	UFUNCTION(BlueprintPure, meta = (DisplayName = "Get Scaled Axes (Matrix)", ScriptMethod = "GetScaledAxes"), Category = "Math|Matrix")
-	static void Matrix_GetScaledAxes(const FMatrix& M, FVector &X, FVector &Y, FVector &Z);
+	static UE_INL_API void Matrix_GetScaledAxes(const FMatrix& M, FVector &X, FVector &Y, FVector &Z);
 
 	/**
 	 * get unit length axis of this matrix
@@ -2615,7 +2621,7 @@ class UKismetMathLibrary : public UBlueprintFunctionLibrary
 	 * @return vector of the axis
 	 */
 	UFUNCTION(BlueprintPure, meta = (DisplayName = "Get Unit Axis (Matrix)", ScriptMethod = "GetUnitAxis"), Category = "Math|Matrix")
-	static FVector Matrix_GetUnitAxis(const FMatrix& M, TEnumAsByte<EAxis::Type> Axis);
+	static UE_INL_API FVector Matrix_GetUnitAxis(const FMatrix& M, TEnumAsByte<EAxis::Type> Axis);
 
 	/**
 	 * get unit length axes of this matrix
@@ -2626,7 +2632,7 @@ class UKismetMathLibrary : public UBlueprintFunctionLibrary
 	 * @param Z axes returned to this param
 	 */
 	UFUNCTION(BlueprintPure, meta = (DisplayName = "Get Unit Axes (Matrix)", ScriptMethod = "GetUnitAxes"), Category = "Math|Matrix")
-	static void Matrix_GetUnitAxes(const FMatrix& M, FVector &X, FVector &Y, FVector &Z);
+	static UE_INL_API void Matrix_GetUnitAxes(const FMatrix& M, FVector &X, FVector &Y, FVector &Z);
 
 	/**
 	 * set an axis of this matrix
@@ -2651,7 +2657,7 @@ class UKismetMathLibrary : public UBlueprintFunctionLibrary
 	 * @return vector of the column
 	 */
 	UFUNCTION(BlueprintPure, meta = (DisplayName = "Get Column (Matrix)", ScriptMethod = "GetColumn"), Category = "Math|Matrix")
-	static FVector Matrix_GetColumn(const FMatrix& M, TEnumAsByte<EMatrixColumns::Type> Column);
+	static UE_INL_API FVector Matrix_GetColumn(const FMatrix& M, TEnumAsByte<EMatrixColumns::Type> Column);
 
 	UFUNCTION(BlueprintCallable, meta = (DisplayName = "Set Column (Matrix)", ScriptMethod = "SetColumn"), Category = "Math|Matrix")
 	static ENGINE_API void Matrix_SetColumn(UPARAM(Ref) FMatrix& M, TEnumAsByte<EMatrixColumns::Type> Column, FVector Value);
@@ -2662,7 +2668,7 @@ class UKismetMathLibrary : public UBlueprintFunctionLibrary
 	*@return rotator representation of this matrix
 	*/
 	UFUNCTION(BlueprintPure, meta = (DisplayName = "Get Rotator (Matrix)", ScriptMethod = "GetRotator"), Category = "Math|Matrix")
-	static FRotator Matrix_GetRotator(const FMatrix& M);
+	static UE_INL_API FRotator Matrix_GetRotator(const FMatrix& M);
 
 	/**
 	 * Transform a rotation matrix into a quaternion.
@@ -2671,7 +2677,7 @@ class UKismetMathLibrary : public UBlueprintFunctionLibrary
 	 * @warning rotation part will need to be unit length for this to be right!
 	 */
 	UFUNCTION(BlueprintPure, meta = (DisplayName = "To Quat (Matrix)", ScriptMethod = "ToQuat"), Category = "Math|Matrix")
-	static FQuat Matrix_ToQuat(const FMatrix& M);
+	static UE_INL_API FQuat Matrix_ToQuat(const FMatrix& M);
 
 	// Frustum plane extraction.
 
@@ -2680,42 +2686,42 @@ class UKismetMathLibrary : public UBlueprintFunctionLibrary
 	 * @param OutPlane the near plane of the Frustum of this matrix 
 	 */
 	UFUNCTION(BlueprintPure, meta = (DisplayName = "Get Frustum Near Plane (Matrix)", ScriptMethod = "GetFrustumNearPlane"), Category = "Math|Matrix")
-	static bool Matrix_GetFrustumNearPlane(const FMatrix& M, FPlane& OutPlane);
+	static UE_INL_API bool Matrix_GetFrustumNearPlane(const FMatrix& M, FPlane& OutPlane);
 
 	/** Get the far plane of the Frustum of this matrix
 	 * (Assumes Matrix represents a View Projection Matrix)
 	 * @param OutPlane the far plane of the Frustum of this matrix 
 	 */
 	UFUNCTION(BlueprintPure, meta = (DisplayName = "Get Frustum Far Plane (Matrix)", ScriptMethod = "GetFrustumFarPlane"), Category = "Math|Matrix")
-	static bool Matrix_GetFrustumFarPlane(const FMatrix& M, FPlane& OutPlane);
+	static UE_INL_API bool Matrix_GetFrustumFarPlane(const FMatrix& M, FPlane& OutPlane);
 
 	/** Get the left plane of the Frustum of this matrix
 	 * (Assumes Matrix represents a View Projection Matrix)
 	 * @param OutPlane the left plane of the Frustum of this matrix 
 	 */
 	UFUNCTION(BlueprintPure, meta = (DisplayName = "Get Frustum Left Plane (Matrix)", ScriptMethod = "GetFrustumLeftPlane"), Category = "Math|Matrix")
-	static bool Matrix_GetFrustumLeftPlane(const FMatrix& M, FPlane& OutPlane);
+	static UE_INL_API bool Matrix_GetFrustumLeftPlane(const FMatrix& M, FPlane& OutPlane);
 
 	/** Get the right plane of the Frustum of this matrix
 	 * (Assumes Matrix represents a View Projection Matrix)
 	 * @param OutPlane the right plane of the Frustum of this matrix 
 	 */
 	UFUNCTION(BlueprintPure, meta = (DisplayName = "Get Frustum Right Plane (Matrix)", ScriptMethod = "GetFrustumRightPlane"), Category = "Math|Matrix")
-	static bool Matrix_GetFrustumRightPlane(const FMatrix& M, FPlane& OutPlane);
+	static UE_INL_API bool Matrix_GetFrustumRightPlane(const FMatrix& M, FPlane& OutPlane);
 
 	/** Get the top plane of the Frustum of this matrix
 	 * (Assumes Matrix represents a View Projection Matrix)
 	 * @param OutPlane the top plane of the Frustum of this matrix 
 	 */
 	UFUNCTION(BlueprintPure, meta = (DisplayName = "Get Frustum Top Plane (Matrix)", ScriptMethod = "GetFrustumTopPlane"), Category = "Math|Matrix")
-	static bool Matrix_GetFrustumTopPlane(const FMatrix& M, FPlane& OutPlane);
+	static UE_INL_API bool Matrix_GetFrustumTopPlane(const FMatrix& M, FPlane& OutPlane);
 
 	/** Get the bottom plane of the Frustum of this matrix
 	 * (Assumes Matrix represents a View Projection Matrix)
 	 * @param OutPlane the bottom plane of the Frustum of this matrix 
 	 */
 	UFUNCTION(BlueprintPure, meta = (DisplayName = "Get Frustum Bottom Plane (Matrix)", ScriptMethod = "GetFrustumBottomPlane"), Category = "Math|Matrix")
-	static bool Matrix_GetFrustumBottomPlane(const FMatrix& M, FPlane& OutPlane);
+	static UE_INL_API bool Matrix_GetFrustumBottomPlane(const FMatrix& M, FPlane& OutPlane);
 
 	/**
 	 * Utility for mirroring this transform across a certain plane, and flipping one of the axis as well.
@@ -2730,7 +2736,7 @@ class UKismetMathLibrary : public UBlueprintFunctionLibrary
 
 	/** Identity quaternion constant */
 	UFUNCTION(BlueprintPure, meta = (ScriptConstant = "Identity", ScriptConstantHost = "/Script/CoreUObject.Quat"), Category = "Math|Quat")
-	static FQuat Quat_Identity();
+	static UE_INL_API FQuat Quat_Identity();
 
 	//
 	// Quat functions
@@ -2738,7 +2744,7 @@ class UKismetMathLibrary : public UBlueprintFunctionLibrary
 
 	/** Returns true if Quaternion A is equal to Quaternion B (A == B) within a specified error tolerance */
 	UFUNCTION(BlueprintPure, meta = (DisplayName = "Equal (Quat)", CompactNodeTitle = "==", ScriptMethod = "Equals", ScriptOperator = "==", Keywords = "== equal"), Category = "Math|Quat")
-	static bool EqualEqual_QuatQuat(const FQuat& A, const FQuat& B, float Tolerance = 1.e-4f);
+	static UE_INL_API bool EqualEqual_QuatQuat(const FQuat& A, const FQuat& B, float Tolerance = 1.e-4f);
 
 	/** Returns true if Quat A is not equal to Quat B (A != B) within a specified error tolerance */
 	UFUNCTION(BlueprintPure, meta = (DisplayName = "Not Equal (Quat)", CompactNodeTitle = "!=", ScriptMethod = "NotEqual", ScriptOperator = "!=", Keywords = "!= not equal"), Category = "Math|Quat")
@@ -2746,19 +2752,19 @@ class UKismetMathLibrary : public UBlueprintFunctionLibrary
 
 	/** Returns addition of Vector A and Vector B (A + B) */
 	UFUNCTION(BlueprintPure, meta = (DisplayName = "Quat + Quat", CompactNodeTitle = "+", ScriptMethod = "Add", ScriptOperator = "+;+=", Keywords = "+ add plus", CommutativeAssociativeBinaryOperator = "true"), Category = "Math|Quat")
-	static FQuat Add_QuatQuat(const FQuat& A, const FQuat& B);
+	static UE_INL_API FQuat Add_QuatQuat(const FQuat& A, const FQuat& B);
 
 	/** Returns subtraction of Vector B from Vector A (A - B) */
 	UFUNCTION(BlueprintPure, meta = (DisplayName = "Quat - Quat", CompactNodeTitle = "-", ScriptMethod = "Subtract", ScriptOperator = "-;-=", Keywords = "- subtract minus"), Category = "Math|Quat")
-	static FQuat Subtract_QuatQuat(const FQuat& A, const FQuat& B);
+	static UE_INL_API FQuat Subtract_QuatQuat(const FQuat& A, const FQuat& B);
 
 	/** Makes a quat {X, Y, Z, W} */
 	UFUNCTION(BlueprintPure, Category = "Math|Quat", meta = (Keywords = "construct build", NativeMakeFunc))
-	static FQuat MakeQuat(float X, float Y, float Z, float W);
+	static UE_INL_API FQuat MakeQuat(float X, float Y, float Z, float W);
 
 	/** Breaks a quat apart into X, Y, Z, W */
 	UFUNCTION(BlueprintPure, Category = "Math|Quat", meta = (NativeBreakFunc))
-	static void BreakQuat(const FQuat& InQuat, float& X, float& Y, float& Z, float& W);
+	static UE_INL_API void BreakQuat(const FQuat& InQuat, float& X, float& Y, float& Z, float& W);
 
 	/**
 	 * Gets the result of multiplying two quaternions (A * B).
@@ -2770,7 +2776,7 @@ class UKismetMathLibrary : public UBlueprintFunctionLibrary
 	 * @return The result of multiplication (A * B).
 	 */
 	UFUNCTION(BlueprintPure, meta = (DisplayName = "Quat * Quat", CompactNodeTitle = "*", ScriptMethod = "Multiply", ScriptOperator = "*;*=", Keywords = "* multiply", CommutativeAssociativeBinaryOperator = "true"), Category = "Math|Quat")
-	static FQuat Multiply_QuatQuat(const FQuat& A, const FQuat& B);
+	static UE_INL_API FQuat Multiply_QuatQuat(const FQuat& A, const FQuat& B);
 
 	/**
 	 * Checks whether this Quaternion is an Identity Quaternion.
@@ -2780,19 +2786,19 @@ class UKismetMathLibrary : public UBlueprintFunctionLibrary
 	 * @return true if Quaternion is a normalized Identity Quaternion.
 	 */
 	UFUNCTION(BlueprintPure, meta = (DisplayName = "Is Identity (Quat)", ScriptMethod = "IsIdentity"), Category = "Math|Quat")
-	static bool Quat_IsIdentity(const FQuat& Q, float Tolerance = 1.e-4f);
+	static UE_INL_API bool Quat_IsIdentity(const FQuat& Q, float Tolerance = 1.e-4f);
 
 	/**	Return true if this quaternion is normalized */
 	UFUNCTION(BlueprintPure, meta = (DisplayName = "Is Normalized (Quat)", ScriptMethod = "IsNormalized"), Category = "Math|Quat")
-	static bool Quat_IsNormalized(const FQuat& Q);
+	static UE_INL_API bool Quat_IsNormalized(const FQuat& Q);
 
 	/** Determine if all the values  are finite (not NaN nor Inf) in this Quat.	*/
 	UFUNCTION(BlueprintPure, meta = (DisplayName = "Is Finite (Quat)", ScriptMethod = "IsFinite"), Category = "Math|Quat")
-	static bool Quat_IsFinite(const FQuat& Q);
+	static UE_INL_API bool Quat_IsFinite(const FQuat& Q);
 
 	/** Determine if there are any non-finite values (NaN or Inf) in this Quat.	*/
 	UFUNCTION(BlueprintPure, meta = (DisplayName = "Is Non-Finite (Quat)", ScriptMethod = "IsNonFinite"), Category = "Math|Quat")
-	static bool Quat_IsNonFinite(const FQuat& Q);
+	static UE_INL_API bool Quat_IsNonFinite(const FQuat& Q);
 
 	/**
 	 * Find the angular distance/difference between two rotation quaternions.
@@ -2801,7 +2807,7 @@ class UKismetMathLibrary : public UBlueprintFunctionLibrary
 	 * @return angular distance in radians
 	 */
 	UFUNCTION(BlueprintPure, meta = (DisplayName = "Angular Distance (Quat)", ScriptMethod = "AngularDistance"), Category = "Math|Quat")
-	static float Quat_AngularDistance(const FQuat& A, const FQuat& B);
+	static UE_INL_API float Quat_AngularDistance(const FQuat& A, const FQuat& B);
 
 	/** Modify the quaternion to ensure that the delta between it and B represents the shortest possible rotation angle. */
 	UFUNCTION(BlueprintCallable, meta = (DisplayName = "Ensure Shortest Arc To (Quat)", ScriptMethod = "EnsureShortestArcTo"), Category = "Math|Quat")
@@ -2809,7 +2815,7 @@ class UKismetMathLibrary : public UBlueprintFunctionLibrary
 
 	/**	Convert a Quaternion into floating-point Euler angles (in degrees). */
 	UFUNCTION(BlueprintPure, meta = (DisplayName = "Euler (Quat)", ScriptMethod = "Euler"), Category = "Math|Quat")
-	static FVector Quat_Euler(const FQuat& Q);
+	static UE_INL_API FVector Quat_Euler(const FQuat& Q);
 
 	/**
 	 * Used in combination with Log().
@@ -2817,35 +2823,35 @@ class UKismetMathLibrary : public UBlueprintFunctionLibrary
 	 * Exp(q) = (sin(theta)*v, cos(theta))
 	 */
 	UFUNCTION(BlueprintPure, meta = (DisplayName = "Exp (Quat)", ScriptMethod = "Exp"), Category = "Math|Quat")
-	static FQuat Quat_Exp(const FQuat& Q);
+	static UE_INL_API FQuat Quat_Exp(const FQuat& Q);
 
 	/** Get the angle of this quaternion */
 	UFUNCTION(BlueprintPure, meta = (DisplayName = "Angle (Quat)", ScriptMethod = "GetAngle"), Category = "Math|Quat")
-	static float Quat_GetAngle(const FQuat& Q);
+	static UE_INL_API float Quat_GetAngle(const FQuat& Q);
 
 	/** Get the forward direction (X axis) after it has been rotated by this Quaternion. */
 	UFUNCTION(BlueprintPure, meta = (DisplayName = "Axis X (Quat)", ScriptMethod = "GetAxisX"), Category = "Math|Quat")
-	static FVector Quat_GetAxisX(const FQuat& Q);
+	static UE_INL_API FVector Quat_GetAxisX(const FQuat& Q);
 
 	/** Get the right direction (Y axis) after it has been rotated by this Quaternion. */
 	UFUNCTION(BlueprintPure, meta = (DisplayName = "Axis Y (Quat)", ScriptMethod = "GetAxisY"), Category = "Math|Quat")
-	static FVector Quat_GetAxisY(const FQuat& Q);
+	static UE_INL_API FVector Quat_GetAxisY(const FQuat& Q);
 
 	/** Get the up direction (Z axis) after it has been rotated by this Quaternion. */
 	UFUNCTION(BlueprintPure, meta = (DisplayName = "Axis Z (Quat)", ScriptMethod = "GetAxisZ"), Category = "Math|Quat")
-	static FVector Quat_GetAxisZ(const FQuat& Q);
+	static UE_INL_API FVector Quat_GetAxisZ(const FQuat& Q);
 
 	/** Get the forward direction (X axis) after it has been rotated by this Quaternion. */
 	UFUNCTION(BlueprintPure, meta = (DisplayName = "Vector Forward (Quat)", ScriptMethod = "VectorForward"), Category = "Math|Quat")
-	static FVector Quat_VectorForward(const FQuat& Q);
+	static UE_INL_API FVector Quat_VectorForward(const FQuat& Q);
 
 	/** Get the right direction (Y axis) after it has been rotated by this Quaternion. */
 	UFUNCTION(BlueprintPure, meta = (DisplayName = "Vector Right (Quat)", ScriptMethod = "VectorRight"), Category = "Math|Quat")
-	static FVector Quat_VectorRight(const FQuat& Q);
+	static UE_INL_API FVector Quat_VectorRight(const FQuat& Q);
 
 	/** Get the up direction (Z axis) after it has been rotated by this Quaternion. */
 	UFUNCTION(BlueprintPure, meta = (DisplayName = "Vector Up (Quat)", ScriptMethod = "VectorUp"), Category = "Math|Quat")
-	static FVector Quat_VectorUp(const FQuat& Q);
+	static UE_INL_API FVector Quat_VectorUp(const FQuat& Q);
 
 	/**
 	 * Normalize this quaternion if it is large enough as compared to the supplied tolerance.
@@ -2863,7 +2869,7 @@ class UKismetMathLibrary : public UBlueprintFunctionLibrary
 	 * @param Tolerance Minimum squared length of quaternion for normalization.
 	 */
 	UFUNCTION(BlueprintPure, meta = (DisplayName = "Normalized (Quat)", ScriptMethod = "Normalized"), Category = "Math|Quat")
-	static FQuat Quat_Normalized(const FQuat& Q, float Tolerance = 1.e-4f);
+	static UE_INL_API FQuat Quat_Normalized(const FQuat& Q, float Tolerance = 1.e-4f);
 
 	/**
 	 * Get the axis of rotation of the Quaternion.
@@ -2871,15 +2877,15 @@ class UKismetMathLibrary : public UBlueprintFunctionLibrary
 	 * For the identity Quaternion which has no such rotation, FVector(1,0,0) is returned.
 	 */
 	UFUNCTION(BlueprintPure, meta = (DisplayName = "Rotation Axis (Quat)", ScriptMethod = "GetRotationAxis"), Category = "Math|Quat")
-	static FVector Quat_GetRotationAxis(const FQuat& Q);
+	static UE_INL_API FVector Quat_GetRotationAxis(const FQuat& Q);
 
 	/**	Return an inversed copy of this quaternion. */
 	UFUNCTION(BlueprintPure, meta = (DisplayName = "Inversed (Quat)", ScriptMethod = "Inversed"), Category = "Math|Quat")
-	static FQuat Quat_Inversed(const FQuat& Q);
+	static UE_INL_API FQuat Quat_Inversed(const FQuat& Q);
 
 	/**	Quaternion with W=0 and V=theta*v. Used in combination with Exp(). */
 	UFUNCTION(BlueprintPure, meta = (DisplayName = "Log (Quat)", ScriptMethod = "Log"), Category = "Math|Quat")
-	static FQuat Quat_Log(const FQuat& Q);
+	static UE_INL_API FQuat Quat_Log(const FQuat& Q);
 
 	/** Set X, Y, Z, W components of Quaternion. */
 	UFUNCTION(BlueprintCallable, meta = (DisplayName = "Set Components (Quat)", ScriptMethod = "SetComponents"), Category = "Math|Quat")
@@ -2901,15 +2907,15 @@ class UKismetMathLibrary : public UBlueprintFunctionLibrary
 	 * @return constructed Quat
 	 */
 	UFUNCTION(BlueprintPure, meta = (DisplayName = "Make from Euler (Quat)"), Category = "Math|Quat")
-	static FQuat Quat_MakeFromEuler(const FVector& Euler);
+	static UE_INL_API FQuat Quat_MakeFromEuler(const FVector& Euler);
 
 	/** Converts to Rotator representation of this Quaternion. */
 	UFUNCTION(BlueprintPure, meta = (DisplayName = "ToRotator (Quat)", CompactNodeTitle = "->", ScriptMethod = "Rotator", Keywords = "cast convert", BlueprintAutocast), Category = "Math|Conversions")
-	static FRotator Quat_Rotator(const FQuat& Q);
+	static UE_INL_API FRotator Quat_Rotator(const FQuat& Q);
 
 	/** Converts to Quaternion representation of this Rotator. */
 	UFUNCTION(BlueprintPure, meta=(DisplayName = "ToQuaternion (Rotator)", CompactNodeTitle = "->", ScriptMethod = "Quaternion", Keywords="cast convert", BlueprintAutocast), Category="Math|Conversions")
-	static FQuat Conv_RotatorToQuaternion(FRotator InRot);
+	static UE_INL_API FQuat Conv_RotatorToQuaternion(FRotator InRot);
 
 	/**
 	 * Get the length of the quaternion.
@@ -2917,7 +2923,7 @@ class UKismetMathLibrary : public UBlueprintFunctionLibrary
 	 * @return The length of the quaternion.
 	 */
 	UFUNCTION(BlueprintPure, meta = (DisplayName = "Size (Quat)", ScriptMethod = "Size"), Category = "Math|Quat")
-	static float Quat_Size(const FQuat& Q);
+	static UE_INL_API float Quat_Size(const FQuat& Q);
 
 	/**
 	 * Get the squared length of the quaternion.
@@ -2925,7 +2931,7 @@ class UKismetMathLibrary : public UBlueprintFunctionLibrary
 	 * @return The squared length of the quaternion.
 	 */
 	UFUNCTION(BlueprintPure, meta = (DisplayName = "Size Squared (Quat)", ScriptMethod = "SizeSquared"), Category = "Math|Quat")
-	static float Quat_SizeSquared(const FQuat& Q);
+	static UE_INL_API float Quat_SizeSquared(const FQuat& Q);
 
 	/**
 	 * Rotate a vector by this quaternion.
@@ -2934,7 +2940,7 @@ class UKismetMathLibrary : public UBlueprintFunctionLibrary
 	 * @return vector after rotation
 	 */
 	UFUNCTION(BlueprintPure, meta = (DisplayName = "Rotate Vector (Quat)", ScriptMethod = "RotateVector"), Category = "Math|Quat")
-	static FVector Quat_RotateVector(const FQuat& Q, const FVector& V);
+	static UE_INL_API FVector Quat_RotateVector(const FQuat& Q, const FVector& V);
 
 	/**
 	 * Rotate a vector by the inverse of this quaternion.
@@ -2943,7 +2949,7 @@ class UKismetMathLibrary : public UBlueprintFunctionLibrary
 	 * @return vector after rotation by the inverse of this quaternion.
 	 */
 	UFUNCTION(BlueprintPure, meta = (DisplayName = "Unrotate Vector (Quat)", ScriptMethod = "UnrotateVector"), Category = "Math|Quat")
-	static FVector Quat_UnrotateVector(const FQuat& Q, const FVector& V);
+	static UE_INL_API FVector Quat_UnrotateVector(const FQuat& Q, const FVector& V);
 
 	/**
 	 * Spherical interpolation between Quaternions. Result is normalized.
@@ -2954,7 +2960,7 @@ class UKismetMathLibrary : public UBlueprintFunctionLibrary
 	 * @return Quat after spherical interpolation
 	 */
 	UFUNCTION(BlueprintPure, meta = (DisplayName = "Slerp (Quat)", ScriptMethod = "SlerpQuat", Keywords = "spherical interpolate"), Category = "Math|Quat")
-	static FQuat Quat_Slerp(const FQuat& A, const FQuat& B, double Alpha);
+	static UE_INL_API FQuat Quat_Slerp(const FQuat& A, const FQuat& B, double Alpha);
 
 	/**
 	 * Generates the 'smallest' (geodesic) rotation around a sphere between two vectors of arbitrary length.
@@ -2964,7 +2970,7 @@ class UKismetMathLibrary : public UBlueprintFunctionLibrary
 	 * @return Quat that will rotate from Start to End
 	 */
 	UFUNCTION(BlueprintPure, meta = (DisplayName = "Find Quat Between Vectors", ScriptMethod = "FindQuatBetweenVectors", Keywords = "FindQuat FindBetween"), Category = "Math|Quat")
-	static FQuat Quat_FindBetweenVectors(FVector Start, FVector End);
+	static UE_INL_API FQuat Quat_FindBetweenVectors(FVector Start, FVector End);
 
 	/**
 	 * Generates the 'smallest' (geodesic) rotation around a sphere between two normals (assumed to be unit length).
@@ -2974,7 +2980,7 @@ class UKismetMathLibrary : public UBlueprintFunctionLibrary
 	 * @return Quat that will rotate from Start to End
 	 */
 	UFUNCTION(BlueprintPure, meta = (DisplayName = "Find Quat Between Normals", ScriptMethod = "FindQuatBetweenNormals", Keywords = "FindQuat FindBetween"), Category = "Math|Quat")
-	static FQuat Quat_FindBetweenNormals(FVector StartNormal, FVector EndNormal);
+	static UE_INL_API FQuat Quat_FindBetweenNormals(FVector StartNormal, FVector EndNormal);
 
 	//
 	// LinearColor constants - exposed for scripting
@@ -2982,35 +2988,35 @@ class UKismetMathLibrary : public UBlueprintFunctionLibrary
 
 	/** White linear color */
 	UFUNCTION(BlueprintPure, meta = (ScriptConstant = "White", ScriptConstantHost = "/Script/CoreUObject.LinearColor"), Category = "Math|Color")
-	static FLinearColor LinearColor_White();
+	static UE_INL_API FLinearColor LinearColor_White();
 
 	/** Grey linear color */
 	UFUNCTION(BlueprintPure, meta = (ScriptConstant = "Gray", ScriptConstantHost = "/Script/CoreUObject.LinearColor"), Category = "Math|Color")
-	static FLinearColor LinearColor_Gray();
+	static UE_INL_API FLinearColor LinearColor_Gray();
 
 	/** Black linear color */
 	UFUNCTION(BlueprintPure, meta = (ScriptConstant = "Black", ScriptConstantHost = "/Script/CoreUObject.LinearColor"), Category = "Math|Color")
-	static FLinearColor LinearColor_Black();
+	static UE_INL_API FLinearColor LinearColor_Black();
 
 	/** Red linear color */
 	UFUNCTION(BlueprintPure, meta = (ScriptConstant = "Red", ScriptConstantHost = "/Script/CoreUObject.LinearColor"), Category = "Math|Color")
-	static FLinearColor LinearColor_Red();
+	static UE_INL_API FLinearColor LinearColor_Red();
 
 	/** Green linear color */
 	UFUNCTION(BlueprintPure, meta = (ScriptConstant = "Green", ScriptConstantHost = "/Script/CoreUObject.LinearColor"), Category = "Math|Color")
-	static FLinearColor LinearColor_Green();
+	static UE_INL_API FLinearColor LinearColor_Green();
 
 	/** Blue linear color */
 	UFUNCTION(BlueprintPure, meta = (ScriptConstant = "Blue", ScriptConstantHost = "/Script/CoreUObject.LinearColor"), Category = "Math|Color")
-	static FLinearColor LinearColor_Blue();
+	static UE_INL_API FLinearColor LinearColor_Blue();
 
 	/** Yellow linear color */
 	UFUNCTION(BlueprintPure, meta = (ScriptConstant = "Yellow", ScriptConstantHost = "/Script/CoreUObject.LinearColor"), Category = "Math|Color")
-	static FLinearColor LinearColor_Yellow();
+	static UE_INL_API FLinearColor LinearColor_Yellow();
 
 	/** Transparent linear color - black with 0 opacity/alpha */
 	UFUNCTION(BlueprintPure, meta = (ScriptConstant = "Transparent", ScriptConstantHost = "/Script/CoreUObject.LinearColor"), Category = "Math|Color")
-	static FLinearColor LinearColor_Transparent();
+	static UE_INL_API FLinearColor LinearColor_Transparent();
 
 
 	//
@@ -3019,11 +3025,11 @@ class UKismetMathLibrary : public UBlueprintFunctionLibrary
 
 	/** Make a color from individual color components (RGB space) */
 	UFUNCTION(BlueprintPure, Category = "Math|Color", meta = (Keywords = "construct build", NativeMakeFunc))
-	static FLinearColor MakeColor(float R, float G, float B, float A = 1.0f);
+	static UE_INL_API FLinearColor MakeColor(float R, float G, float B, float A = 1.0f);
 
 	/** Breaks apart a color into individual RGB components (as well as alpha) */
 	UFUNCTION(BlueprintPure, Category = "Math|Color")
-	static void BreakColor(FLinearColor InColor, float& R, float& G, float& B, float& A);
+	static UE_INL_API void BreakColor(FLinearColor InColor, float& R, float& G, float& B, float& A);
 
 	/** Assign contents of InColor */
 	UFUNCTION(BlueprintCallable, meta = (ScriptMethod = "Set"), Category = "Math|Color")
@@ -3061,39 +3067,39 @@ class UKismetMathLibrary : public UBlueprintFunctionLibrary
 
 	/** Converts a float into a LinearColor, where each element is a float */
 	UE_DEPRECATED(5.2, "This method has been deprecated and will be removed. Use the double version instead.")
-	static FLinearColor Conv_FloatToLinearColor(float InFloat);
+	static UE_INL_API FLinearColor Conv_FloatToLinearColor(float InFloat);
 
 	/** Converts a float into a LinearColor, where each RGB element is that float */
 	UFUNCTION(BlueprintPure, meta = (DisplayName = "To LinearColor (Float)", CompactNodeTitle = "->", Keywords = "cast convert", BlueprintAutocast), Category = "Math|Conversions")
-	static FLinearColor Conv_DoubleToLinearColor(double InDouble);
+	static UE_INL_API FLinearColor Conv_DoubleToLinearColor(double InDouble);
 
 	/** Make a color from individual color components (HSV space; Hue is [0..360) while Saturation and Value are 0..1) */
 	UFUNCTION(BlueprintPure, Category = "Math|Color", meta = (DisplayName = "HSV to RGB"))
-	static FLinearColor HSVToRGB(float H, float S, float V, float A = 1.0f);
+	static UE_INL_API FLinearColor HSVToRGB(float H, float S, float V, float A = 1.0f);
 
 	/** Converts a HSV linear color (where H is in R (0..360), S is in G (0..1), and V is in B (0..1)) to RGB */
 	UFUNCTION(BlueprintPure, meta = (DisplayName = "HSV to RGB (Vector)", ScriptMethod = "HSVIntoRGB", Keywords = "cast convert"), Category = "Math|Color")
-	static void HSVToRGB_Vector(FLinearColor HSV, FLinearColor& RGB);
+	static UE_INL_API void HSVToRGB_Vector(FLinearColor HSV, FLinearColor& RGB);
 
 	/** Converts a HSV linear color (where H is in R, S is in G, and V is in B) to linear RGB */
 	UFUNCTION(BlueprintPure, meta = (DisplayName = "HSV to RGB linear color", ScriptMethod = "HSVToRGB", Keywords = "cast convert"), Category = "Math|Color")
-	static FLinearColor HSVToRGBLinear(FLinearColor HSV);
+	static UE_INL_API FLinearColor HSVToRGBLinear(FLinearColor HSV);
 
 	/** Breaks apart a color into individual HSV components (as well as alpha) (Hue is [0..360) while Saturation and Value are 0..1) */
 	UFUNCTION(BlueprintPure, meta = (DisplayName = "RGB to HSV", ScriptMethod = "RGBIntoHSVComponents"), Category = "Math|Color")
-	static void RGBToHSV(FLinearColor InColor, float& H, float& S, float& V, float& A);
+	static UE_INL_API void RGBToHSV(FLinearColor InColor, float& H, float& S, float& V, float& A);
 
 	/** Converts a RGB linear color to HSV (where H is in R (0..360), S is in G (0..1), and V is in B (0..1)) */
 	UFUNCTION(BlueprintPure, meta = (DisplayName = "RGB to HSV (Vector)", ScriptMethod = "RGBIntoHSV", Keywords = "cast convert"), Category = "Math|Color")
-	static void RGBToHSV_Vector(FLinearColor RGB, FLinearColor& HSV);
+	static UE_INL_API void RGBToHSV_Vector(FLinearColor RGB, FLinearColor& HSV);
 
 	/** Converts a RGB linear color to HSV (where H is in R, S is in G, and V is in B) */
 	UFUNCTION(BlueprintPure, meta = (DisplayName = "RGB to HSV (LinearColor)", ScriptMethod = "RGBToHSV", Keywords = "cast convert"), Category = "Math|Color")
-	static FLinearColor RGBLinearToHSV(FLinearColor RGB);
+	static UE_INL_API FLinearColor RGBLinearToHSV(FLinearColor RGB);
 
 	/** Converts a LinearColor to a vector */
 	UFUNCTION(BlueprintPure, meta = (DisplayName = "To Vector (LinearColor)", ScriptMethod = "ToRGBVector", CompactNodeTitle = "->", Keywords = "cast convert", BlueprintAutocast), Category = "Math|Conversions")
-	static FVector Conv_LinearColorToVector(FLinearColor InLinearColor);
+	static UE_INL_API FVector Conv_LinearColorToVector(FLinearColor InLinearColor);
 
 	/** Converts from linear to 8-bit RGBE as outlined in Gregory Ward's Real Pixels article, Graphics Gems II, page 80. */
 	UFUNCTION(BlueprintPure, meta = (DisplayName = "To RGBE (LinearColor)", ScriptMethod = "ToRGBE", Keywords = "cast convert", BlueprintAutocast), Category = "Math|Color")
@@ -3118,19 +3124,19 @@ class UKismetMathLibrary : public UBlueprintFunctionLibrary
 	 * @return	Desaturated color
 	 */
 	UFUNCTION(BlueprintPure, meta = (DisplayName = "Desaturate (LinearColor)", ScriptMethod = "Desaturated"), Category = "Math|Color")
-	static FLinearColor LinearColor_Desaturated(FLinearColor InColor, float InDesaturation);
+	static UE_INL_API FLinearColor LinearColor_Desaturated(FLinearColor InColor, float InDesaturation);
 
 	/** Euclidean distance between two color points. */
 	UFUNCTION(BlueprintPure, meta = (DisplayName = "Distance (LinearColor)", ScriptMethod = "Distance", Keywords = "magnitude"), Category = "Math|Color")
-	static float LinearColor_Distance(FLinearColor C1, FLinearColor C2);
+	static UE_INL_API float LinearColor_Distance(FLinearColor C1, FLinearColor C2);
 	
 	/** Returns a copy of this color using the specified opacity/alpha.	 */
 	UFUNCTION(BlueprintPure, meta = (DisplayName = "New Opacity (LinearColor)", ScriptMethod = "ToNewOpacity"), Category = "Math|Color")
-	static FLinearColor LinearColor_ToNewOpacity(FLinearColor InColor, float InOpacity);
+	static UE_INL_API FLinearColor LinearColor_ToNewOpacity(FLinearColor InColor, float InOpacity);
 
 	/**	Returns the perceived brightness of a color on a display taking into account the impact on the human eye per color channel: green > red > blue. */
 	UFUNCTION(BlueprintPure, meta = (DisplayName = "Luminance (LinearColor)", ScriptMethod = "GetLuminance"), Category = "Math|Color")
-	static float LinearColor_GetLuminance(FLinearColor InColor);
+	static UE_INL_API float LinearColor_GetLuminance(FLinearColor InColor);
 
 	/**
 	 * Returns the maximum color channel value in this color structure
@@ -3138,7 +3144,7 @@ class UKismetMathLibrary : public UBlueprintFunctionLibrary
 	 * @return The maximum color channel value
 	 */
 	UFUNCTION(BlueprintPure, meta = (DisplayName = "Max (LinearColor)", ScriptMethod = "GetMax"), Category = "Math|Color")
-	static float LinearColor_GetMax(FLinearColor InColor);
+	static UE_INL_API float LinearColor_GetMax(FLinearColor InColor);
 
 	/**
 	 * Returns the minimum color channel value in this color structure
@@ -3146,7 +3152,7 @@ class UKismetMathLibrary : public UBlueprintFunctionLibrary
 	 * @return The minimum color channel value
 	 */
 	UFUNCTION(BlueprintPure, meta = (DisplayName = "Min (LinearColor)", ScriptMethod = "GetMin"), Category = "Math|Color")
-	static float LinearColor_GetMin(FLinearColor InColor);
+	static UE_INL_API float LinearColor_GetMin(FLinearColor InColor);
 
 	/**
 	 * Interpolate Linear Color from Current to Target. Scaled by distance to Target, so it has a strong start speed and ease out.
@@ -3158,11 +3164,11 @@ class UKismetMathLibrary : public UBlueprintFunctionLibrary
 	 * @return		New interpolated Color
 	 */
 	UFUNCTION(BlueprintPure, meta = (DisplayName = "Interpolate (LinearColor)", ScriptMethod = "InterpolateTo", Keywords = "color"), Category = "Math|Interpolation")
-	static FLinearColor CInterpTo(FLinearColor Current, FLinearColor Target, float DeltaTime, float InterpSpeed);
+	static UE_INL_API FLinearColor CInterpTo(FLinearColor Current, FLinearColor Target, float DeltaTime, float InterpSpeed);
 
 	/** Linearly interpolates between A and B based on Alpha (100% of A when Alpha=0 and 100% of B when Alpha=1) */
 	UFUNCTION(BlueprintPure, meta=(DisplayName = "Lerp (LinearColor)", ScriptMethod = "LerpTo"), Category="Math|Color")
-	static FLinearColor LinearColorLerp(FLinearColor A, FLinearColor B, float Alpha);
+	static UE_INL_API FLinearColor LinearColorLerp(FLinearColor A, FLinearColor B, float Alpha);
 
 	/**
 	 * Linearly interpolates between two colors by the specified Alpha amount (100% of A when Alpha=0 and 100% of B when Alpha=1).  The interpolation is performed in HSV color space taking the shortest path to the new color's hue.  This can give better results than a normal lerp, but is much more expensive.  The incoming colors are in RGB space, and the output color will be RGB.  The alpha value will also be interpolated.
@@ -3174,39 +3180,39 @@ class UKismetMathLibrary : public UBlueprintFunctionLibrary
 	 * @return	The interpolated color in linear RGB space along with the interpolated alpha value
 	 */
 	UFUNCTION(BlueprintPure, meta=(DisplayName = "Lerp Using HSV (LinearColor)", ScriptMethod = "LerpUsingHSVTo"), Category="Math|Color")
-	static FLinearColor LinearColorLerpUsingHSV(FLinearColor A, FLinearColor B, float Alpha);
+	static UE_INL_API FLinearColor LinearColorLerpUsingHSV(FLinearColor A, FLinearColor B, float Alpha);
 
 	/** Returns true if linear color A is equal to linear color B (A == B) within a specified error tolerance */
 	UFUNCTION(BlueprintPure, meta = (DisplayName = "Near Equal (LinearColor)", CompactNodeTitle = "~==", ScriptMethod = "IsNearEqual", ScriptOperator = "==", Keywords = "== equal"), Category = "Math|Color")
-	static bool LinearColor_IsNearEqual(FLinearColor A, FLinearColor B, float Tolerance = 1.e-4f);
+	static UE_INL_API bool LinearColor_IsNearEqual(FLinearColor A, FLinearColor B, float Tolerance = 1.e-4f);
 
 	/** Returns true if linear color A is equal to linear color B (A == B) within a specified error tolerance */
 	UFUNCTION(BlueprintPure, meta = (DisplayName = "Equal (LinearColor)", CompactNodeTitle = "==", ScriptMethod = "Equals", ScriptOperator = "==", Keywords = "== equal"), Category = "Math|Color")
-	static bool EqualEqual_LinearColorLinearColor(FLinearColor A, FLinearColor B);
+	static UE_INL_API bool EqualEqual_LinearColorLinearColor(FLinearColor A, FLinearColor B);
 
 	/** Returns true if linear color A is not equal to linear color B (A != B) within a specified error tolerance */
 	UFUNCTION(BlueprintPure, meta = (DisplayName = "Not Equal (LinearColor)", CompactNodeTitle = "!=", ScriptMethod = "NotEqual", ScriptOperator = "!=", Keywords = "!= not equal"), Category = "Math|Color")
-	static bool NotEqual_LinearColorLinearColor(FLinearColor A, FLinearColor B);
+	static UE_INL_API bool NotEqual_LinearColorLinearColor(FLinearColor A, FLinearColor B);
 
 	/** Element-wise addition of two linear colors (R+R, G+G, B+B, A+A) */
 	UFUNCTION(BlueprintPure, meta=(DisplayName = "LinearColor + LinearColor", CompactNodeTitle = "+", ScriptMethod = "Add", ScriptOperator = "+;+=", Keywords = "+ add plus"), Category="Math|Color")
-	static FLinearColor Add_LinearColorLinearColor(FLinearColor A, FLinearColor B);
+	static UE_INL_API FLinearColor Add_LinearColorLinearColor(FLinearColor A, FLinearColor B);
 
 	/** Element-wise subtraction of two linear colors (R-R, G-G, B-B, A-A) */
 	UFUNCTION(BlueprintPure, meta=(DisplayName = "LinearColor - LinearColor", CompactNodeTitle = "-", ScriptMethod = "Subtract", ScriptOperator = "-;-=", Keywords = "- subtract minus"), Category="Math|Color")
-	static FLinearColor Subtract_LinearColorLinearColor(FLinearColor A, FLinearColor B);
+	static UE_INL_API FLinearColor Subtract_LinearColorLinearColor(FLinearColor A, FLinearColor B);
 
 	/** Element-wise multiplication of two linear colors (R*R, G*G, B*B, A*A) */
 	UFUNCTION(BlueprintPure, meta=(DisplayName = "LinearColor * LinearColor", CompactNodeTitle = "*", ScriptMethod = "Multiply", ScriptOperator = "*;*=", Keywords = "* multiply"), Category="Math|Color")
-	static FLinearColor Multiply_LinearColorLinearColor(FLinearColor A, FLinearColor B);
+	static UE_INL_API FLinearColor Multiply_LinearColorLinearColor(FLinearColor A, FLinearColor B);
 
 	/** Element-wise multiplication of a linear color by a float (F*R, F*G, F*B, F*A) */
 	UFUNCTION(BlueprintPure, meta=(DisplayName = "LinearColor * Float", CompactNodeTitle = "*", ScriptMethod = "MultiplyFloat", Keywords = "* multiply"), Category="Math|Color")
-	static FLinearColor Multiply_LinearColorFloat(FLinearColor A, float B);
+	static UE_INL_API FLinearColor Multiply_LinearColorFloat(FLinearColor A, float B);
 
 	/** Element-wise multiplication of two linear colors (R/R, G/G, B/B, A/A) */
 	UFUNCTION(BlueprintPure, meta=(DisplayName = "LinearColor / LinearColor", CompactNodeTitle = "/", ScriptMethod = "Divide", ScriptOperator = "/;/=", Keywords = "/ divide"), Category="Math|Color")
-	static FLinearColor Divide_LinearColorLinearColor(FLinearColor A, FLinearColor B);
+	static UE_INL_API FLinearColor Divide_LinearColorLinearColor(FLinearColor A, FLinearColor B);
 	
 	/** Converts this color value to a hexadecimal string. The format of the string is RRGGBBAA. */
 	UFUNCTION(BlueprintPure, meta = (DisplayName = "ToHex"), Category = "Math|Color")
@@ -3241,15 +3247,15 @@ class UKismetMathLibrary : public UBlueprintFunctionLibrary
 
 	/** Addition (A + B) */
 	UFUNCTION(BlueprintPure, meta=(IgnoreTypePromotion, DisplayName="DateTime + Timespan", CompactNodeTitle="+", Keywords="+ add plus"), Category="Math|DateTime")
-	static FDateTime Add_DateTimeTimespan( FDateTime A, FTimespan B );
+	static UE_INL_API FDateTime Add_DateTimeTimespan( FDateTime A, FTimespan B );
 
 	/** Subtraction (A - B) */
 	UFUNCTION(BlueprintPure, meta=(IgnoreTypePromotion, DisplayName="DateTime - Timespan", CompactNodeTitle="-", Keywords="- subtract minus"), Category="Math|DateTime")
-	static FDateTime Subtract_DateTimeTimespan(FDateTime A, FTimespan B);
+	static UE_INL_API FDateTime Subtract_DateTimeTimespan(FDateTime A, FTimespan B);
 
 	/** Addition (A + B) */
 	UFUNCTION(BlueprintPure, meta = (IgnoreTypePromotion, DisplayName = "DateTime + DateTime", CompactNodeTitle = "+", Keywords = "+ add plus"), Category = "Math|DateTime")
-	static FDateTime Add_DateTimeDateTime(FDateTime A, FDateTime B);
+	static UE_INL_API FDateTime Add_DateTimeDateTime(FDateTime A, FDateTime B);
 
 	/** Subtraction (A - B) */
 	UFUNCTION(BlueprintPure, meta = (IgnoreTypePromotion, DisplayName = "DateTime - DateTime", CompactNodeTitle = "-", Keywords = "- subtract minus"), Category = "Math|DateTime")
@@ -3257,63 +3263,63 @@ class UKismetMathLibrary : public UBlueprintFunctionLibrary
 
 	/** Returns true if the values are equal (A == B) */
 	UFUNCTION(BlueprintPure, meta=(IgnoreTypePromotion, DisplayName="Equal (DateTime)", CompactNodeTitle="==", Keywords="== equal"), Category="Math|DateTime")
-	static bool EqualEqual_DateTimeDateTime( FDateTime A, FDateTime B );
+	static UE_INL_API bool EqualEqual_DateTimeDateTime( FDateTime A, FDateTime B );
 
 	/** Returns true if the values are not equal (A != B) */
 	UFUNCTION(BlueprintPure, meta=(IgnoreTypePromotion, DisplayName="Not Equal (DateTime)", CompactNodeTitle="!=", Keywords="!= not equal"), Category="Math|DateTime")
-	static bool NotEqual_DateTimeDateTime( FDateTime A, FDateTime B );
+	static UE_INL_API bool NotEqual_DateTimeDateTime( FDateTime A, FDateTime B );
 
 	/** Returns true if A is greater than B (A > B) */
 	UFUNCTION(BlueprintPure, meta=(IgnoreTypePromotion, DisplayName="DateTime > DateTime", CompactNodeTitle=">", Keywords="> greater"), Category="Math|DateTime")
-	static bool Greater_DateTimeDateTime( FDateTime A, FDateTime B );
+	static UE_INL_API bool Greater_DateTimeDateTime( FDateTime A, FDateTime B );
 
 	/** Returns true if A is greater than or equal to B (A >= B) */
 	UFUNCTION(BlueprintPure, meta=(IgnoreTypePromotion, DisplayName="DateTime >= DateTime", CompactNodeTitle=">=", Keywords=">= greater"), Category="Math|DateTime")
-	static bool GreaterEqual_DateTimeDateTime( FDateTime A, FDateTime B );
+	static UE_INL_API bool GreaterEqual_DateTimeDateTime( FDateTime A, FDateTime B );
 
 	/** Returns true if A is less than B (A < B) */
 	UFUNCTION(BlueprintPure, meta=(IgnoreTypePromotion, DisplayName="DateTime < DateTime", CompactNodeTitle="<", Keywords="< less"), Category="Math|DateTime")
-	static bool Less_DateTimeDateTime( FDateTime A, FDateTime B );
+	static UE_INL_API bool Less_DateTimeDateTime( FDateTime A, FDateTime B );
 
 	/** Returns true if A is less than or equal to B (A <= B) */
 	UFUNCTION(BlueprintPure, meta=(IgnoreTypePromotion, DisplayName="DateTime <= DateTime", CompactNodeTitle="<=", Keywords="<= less"), Category="Math|DateTime")
-	static bool LessEqual_DateTimeDateTime( FDateTime A, FDateTime B );
+	static UE_INL_API bool LessEqual_DateTimeDateTime( FDateTime A, FDateTime B );
 
 	/** Returns the date component of A */
 	UFUNCTION(BlueprintPure, meta=(DisplayName="Get Date"), Category="Math|DateTime")
-	static FDateTime GetDate( FDateTime A );
+	static UE_INL_API FDateTime GetDate( FDateTime A );
 
 	/** Returns the day component of A (1 to 31) */
 	UFUNCTION(BlueprintPure, meta=(DisplayName="Get Day"), Category="Math|DateTime")
-	static int32 GetDay( FDateTime A );
+	static UE_INL_API int32 GetDay( FDateTime A );
 
 	/** Returns the day of year of A */
 	UFUNCTION(BlueprintPure, meta=(DisplayName="Get Day Of Year"), Category="Math|DateTime")
-	static int32 GetDayOfYear( FDateTime A );
+	static UE_INL_API int32 GetDayOfYear( FDateTime A );
 
 	/** Returns the hour component of A (24h format) */
 	UFUNCTION(BlueprintPure, meta=(DisplayName="Get Hour"), Category="Math|DateTime")
-	static int32 GetHour( FDateTime A );
+	static UE_INL_API int32 GetHour( FDateTime A );
 
 	/** Returns the hour component of A (12h format) */
 	UFUNCTION(BlueprintPure, meta=(DisplayName="Get Hour 12"), Category="Math|DateTime")
-	static int32 GetHour12( FDateTime A );
+	static UE_INL_API int32 GetHour12( FDateTime A );
 
 	/** Returns the millisecond component of A */
 	UFUNCTION(BlueprintPure, meta=(DisplayName="Get Millisecond"), Category="Math|DateTime")
-	static int32 GetMillisecond( FDateTime A );
+	static UE_INL_API int32 GetMillisecond( FDateTime A );
 
 	/** Returns the minute component of A */
 	UFUNCTION(BlueprintPure, meta=(DisplayName="Get Minute"), Category="Math|DateTime")
-	static int32 GetMinute( FDateTime A );
+	static UE_INL_API int32 GetMinute( FDateTime A );
 
 	/** Returns the month component of A */
 	UFUNCTION(BlueprintPure, meta=(DisplayName="Get Month"), Category="Math|DateTime")
-	static int32 GetMonth( FDateTime A );
+	static UE_INL_API int32 GetMonth( FDateTime A );
 
 	/** Returns the second component of A */
 	UFUNCTION(BlueprintPure, meta=(DisplayName="Get Second"), Category="Math|DateTime")
-	static int32 GetSecond( FDateTime A );
+	static UE_INL_API int32 GetSecond( FDateTime A );
 
 	/** Returns the time elapsed since midnight of A */
 	UFUNCTION(BlueprintPure, meta=(DisplayName="Get Time Of Day"), Category="Math|DateTime")
@@ -3321,55 +3327,55 @@ class UKismetMathLibrary : public UBlueprintFunctionLibrary
 
 	/** Returns the year component of A */
 	UFUNCTION(BlueprintPure, meta=(DisplayName="Get Year"), Category="Math|DateTime")
-	static int32 GetYear( FDateTime A );
+	static UE_INL_API int32 GetYear( FDateTime A );
 
 	/** Returns whether A's time is in the afternoon */
 	UFUNCTION(BlueprintPure, meta=(DisplayName="Is Afternoon"), Category="Math|DateTime")
-	static bool IsAfternoon( FDateTime A );
+	static UE_INL_API bool IsAfternoon( FDateTime A );
 
 	/** Returns whether A's time is in the morning */
 	UFUNCTION(BlueprintPure, meta=(DisplayName="Is Morning"), Category="Math|DateTime")
-	static bool IsMorning( FDateTime A );
+	static UE_INL_API bool IsMorning( FDateTime A );
 
 	/** Returns the number of days in the given year and month */
 	UFUNCTION(BlueprintPure, meta=(DisplayName="Days In Month"), Category="Math|DateTime")
-	static int32 DaysInMonth( int32 Year, int32 Month );
+	static UE_INL_API int32 DaysInMonth( int32 Year, int32 Month );
 
 	/** Returns the number of days in the given year */
 	UFUNCTION(BlueprintPure, meta=(DisplayName="Days In Year"), Category="Math|DateTime")
-	static int32 DaysInYear( int32 Year );
+	static UE_INL_API int32 DaysInYear( int32 Year );
 
 	/** Returns whether given year is a leap year */
 	UFUNCTION(BlueprintPure, meta=(DisplayName="Is Leap Year"), Category="Math|DateTime")
-	static bool IsLeapYear( int32 Year );
+	static UE_INL_API bool IsLeapYear( int32 Year );
 
 	/** Returns the maximum date and time value */
 	UFUNCTION(BlueprintPure, meta=(DisplayName="Max Value (DateTime)"), Category="Math|DateTime")
-	static FDateTime DateTimeMaxValue( );
+	static UE_INL_API FDateTime DateTimeMaxValue( );
 
 	/** Returns the minimum date and time value */
 	UFUNCTION(BlueprintPure, meta=(DisplayName="MinValue (DateTime)"), Category="Math|DateTime")
-	static FDateTime DateTimeMinValue( );
+	static UE_INL_API FDateTime DateTimeMinValue( );
 
 	/** Returns the local date and time on this computer */
 	UFUNCTION(BlueprintPure, meta=(DisplayName="Now"), Category="Math|DateTime")
-	static FDateTime Now( );
+	static UE_INL_API FDateTime Now( );
 
 	/** Returns the local date on this computer */
 	UFUNCTION(BlueprintPure, meta=(DisplayName="Today"), Category="Math|DateTime")
-	static FDateTime Today( );
+	static UE_INL_API FDateTime Today( );
 
 	/** Returns the UTC date and time on this computer */
 	UFUNCTION(BlueprintPure, meta=(DisplayName="UTC Now"), Category="Math|DateTime")
-	static FDateTime UtcNow( );
+	static UE_INL_API FDateTime UtcNow( );
 
 	/** Converts a date string in ISO-8601 format to a DateTime object */
 	UFUNCTION(BlueprintPure, Category="Math|DateTime")
-	static bool DateTimeFromIsoString(FString IsoString, FDateTime& Result);
+	static UE_INL_API bool DateTimeFromIsoString(FString IsoString, FDateTime& Result);
 
 	/** Converts a date string to a DateTime object */
 	UFUNCTION(BlueprintPure, Category="Math|DateTime")
-	static bool DateTimeFromString(FString DateTimeString, FDateTime& Result);
+	static UE_INL_API bool DateTimeFromString(FString DateTimeString, FDateTime& Result);
 
 	/**
 	 * Returns this date as the number of seconds since the Unix Epoch (January 1st of 1970).
@@ -3378,7 +3384,7 @@ class UKismetMathLibrary : public UBlueprintFunctionLibrary
 	 * @see FromUnixTimestamp
 	 */
 	UFUNCTION(BlueprintPure, Category="Math|DateTime")
-	static int64 ToUnixTimestamp(const FDateTime& Time);
+	static UE_INL_API int64 ToUnixTimestamp(const FDateTime& Time);
 
 	/**
 	 * Returns this date as the number of seconds since the Unix Epoch (January 1st of 1970).
@@ -3387,7 +3393,7 @@ class UKismetMathLibrary : public UBlueprintFunctionLibrary
 	 * @see FromUnixTimestamp
 	 */
 	UFUNCTION(BlueprintPure, Category="Math|DateTime")
-	static double ToUnixTimestampDouble(const FDateTime& Time);
+	static UE_INL_API double ToUnixTimestampDouble(const FDateTime& Time);
 	
 	/**
 	 * Returns the date from Unix time (seconds from midnight 1970-01-01)
@@ -3397,7 +3403,7 @@ class UKismetMathLibrary : public UBlueprintFunctionLibrary
 	 * @see ToUnixTimestamp
 	 */
 	UFUNCTION(BlueprintPure, Category="Math|DateTime")
-	static FDateTime FromUnixTimestamp(const int64 UnixTime);
+	static UE_INL_API FDateTime FromUnixTimestamp(const int64 UnixTime);
 
 	//
 	// Timespan constants
@@ -3453,31 +3459,31 @@ class UKismetMathLibrary : public UBlueprintFunctionLibrary
 
 	/** Returns true if the values are equal (A == B) */
 	UFUNCTION(BlueprintPure, meta=(DisplayName="Equal (Timespan)", CompactNodeTitle="==", Keywords="== equal"), Category="Math|Timespan")
-	static bool EqualEqual_TimespanTimespan( FTimespan A, FTimespan B );
+	static UE_INL_API bool EqualEqual_TimespanTimespan( FTimespan A, FTimespan B );
 
 	/** Returns true if the values are not equal (A != B) */
 	UFUNCTION(BlueprintPure, meta=(DisplayName="Not Equal (Timespan)", CompactNodeTitle="!=", Keywords="!= not equal"), Category="Math|Timespan")
-	static bool NotEqual_TimespanTimespan( FTimespan A, FTimespan B );
+	static UE_INL_API bool NotEqual_TimespanTimespan( FTimespan A, FTimespan B );
 
 	/** Returns true if A is greater than B (A > B) */
 	UFUNCTION(BlueprintPure, meta=(DisplayName="Timespan > Timespan", CompactNodeTitle=">", Keywords="> greater"), Category="Math|Timespan")
-	static bool Greater_TimespanTimespan( FTimespan A, FTimespan B );
+	static UE_INL_API bool Greater_TimespanTimespan( FTimespan A, FTimespan B );
 
 	/** Returns true if A is greater than or equal to B (A >= B) */
 	UFUNCTION(BlueprintPure, meta=(DisplayName="Timespan >= Timespan", CompactNodeTitle=">=", Keywords=">= greater"), Category="Math|Timespan")
-	static bool GreaterEqual_TimespanTimespan( FTimespan A, FTimespan B );
+	static UE_INL_API bool GreaterEqual_TimespanTimespan( FTimespan A, FTimespan B );
 
 	/** Returns true if A is less than B (A < B) */
 	UFUNCTION(BlueprintPure, meta=(DisplayName="Timespan < Timespan", CompactNodeTitle="<", Keywords="< less"), Category="Math|Timespan")
-	static bool Less_TimespanTimespan( FTimespan A, FTimespan B );
+	static UE_INL_API bool Less_TimespanTimespan( FTimespan A, FTimespan B );
 
 	/** Returns true if A is less than or equal to B (A <= B) */
 	UFUNCTION(BlueprintPure, meta=(DisplayName="Timespan <= Timespan", CompactNodeTitle="<=", Keywords="<= less"), Category="Math|Timespan")
-	static bool LessEqual_TimespanTimespan( FTimespan A, FTimespan B );
+	static UE_INL_API bool LessEqual_TimespanTimespan( FTimespan A, FTimespan B );
 
 	/** Returns the days component of A */
 	UFUNCTION(BlueprintPure, meta=(DisplayName="Get Days"), Category="Math|Timespan")
-	static int32 GetDays( FTimespan A );
+	static UE_INL_API int32 GetDays( FTimespan A );
 
 	/** Returns the absolute value of A */
 	UFUNCTION(BlueprintPure, meta=(DisplayName="Get Duration"), Category="Math|Timespan")
@@ -3485,39 +3491,39 @@ class UKismetMathLibrary : public UBlueprintFunctionLibrary
 
 	/** Returns the hours component of A */
 	UFUNCTION(BlueprintPure, meta=(DisplayName="Get Hours"), Category="Math|Timespan")
-	static int32 GetHours( FTimespan A );
+	static UE_INL_API int32 GetHours( FTimespan A );
 
 	/** Returns the milliseconds component of A */
 	UFUNCTION(BlueprintPure, meta=(DisplayName="Get Milliseconds"), Category="Math|Timespan")
-	static int32 GetMilliseconds( FTimespan A );
+	static UE_INL_API int32 GetMilliseconds( FTimespan A );
 
 	/** Returns the minutes component of A */
 	UFUNCTION(BlueprintPure, meta=(DisplayName="Get Minutes"), Category="Math|Timespan")
-	static int32 GetMinutes( FTimespan A );
+	static UE_INL_API int32 GetMinutes( FTimespan A );
 
 	/** Returns the seconds component of A */
 	UFUNCTION(BlueprintPure, meta=(DisplayName="Get Seconds"), Category="Math|Timespan")
-	static int32 GetSeconds( FTimespan A );
+	static UE_INL_API int32 GetSeconds( FTimespan A );
 
 	/** Returns the total number of days in A */
 	UFUNCTION(BlueprintPure, meta=(DisplayName="Get Total Days"), Category="Math|Timespan")
-	static double GetTotalDays( FTimespan A );
+	static UE_INL_API double GetTotalDays( FTimespan A );
 
 	/** Returns the total number of hours in A */
 	UFUNCTION(BlueprintPure, meta=(DisplayName="Get Total Hours"), Category="Math|Timespan")
-	static double GetTotalHours( FTimespan A );
+	static UE_INL_API double GetTotalHours( FTimespan A );
 
 	/** Returns the total number of milliseconds in A */
 	UFUNCTION(BlueprintPure, meta=(DisplayName="Get Total Milliseconds"), Category="Math|Timespan")
-	static double GetTotalMilliseconds( FTimespan A );
+	static UE_INL_API double GetTotalMilliseconds( FTimespan A );
 
 	/** Returns the total number of minutes in A */
 	UFUNCTION(BlueprintPure, meta=(DisplayName="Get Total Minutes"), Category="Math|Timespan")
-	static double GetTotalMinutes( FTimespan A );
+	static UE_INL_API double GetTotalMinutes( FTimespan A );
 
 	/** Returns the total number of seconds in A */
 	UFUNCTION(BlueprintPure, meta=(DisplayName="Get Total Seconds"), Category="Math|Timespan")
-	static double GetTotalSeconds( FTimespan A );
+	static UE_INL_API double GetTotalSeconds( FTimespan A );
 
 	/** Returns a time span that represents the specified number of days */
 	UFUNCTION(BlueprintPure, meta=(DisplayName="From Days"), Category="Math|Timespan")
@@ -3541,11 +3547,11 @@ class UKismetMathLibrary : public UBlueprintFunctionLibrary
 
 	/** Returns the ratio between two time spans (A / B), handles zero values */
 	UFUNCTION(BlueprintPure, meta=(DisplayName="Timespan Ratio"), Category="Math|Timespan")
-	static float TimespanRatio( FTimespan A, FTimespan B );
+	static UE_INL_API float TimespanRatio( FTimespan A, FTimespan B );
 
 	/** Converts a time span string to a Timespan object */
 	UFUNCTION(BlueprintPure, Category="Math|Timespan")
-	static bool TimespanFromString(FString TimespanString, FTimespan& Result);
+	static UE_INL_API bool TimespanFromString(FString TimespanString, FTimespan& Result);
 
 	//
 	// Frame Time and Frame Rate Functions
@@ -3571,51 +3577,51 @@ class UKismetMathLibrary : public UBlueprintFunctionLibrary
 
 	/** Converts a byte to a float */
 	UE_DEPRECATED(5.2, "This method has been deprecated and will be removed. Use the double version instead.")
-	static float Conv_ByteToFloat(uint8 InByte);
+	static UE_INL_API float Conv_ByteToFloat(uint8 InByte);
 
 	/** Converts a byte to a float */
 	UFUNCTION(BlueprintPure, meta = (DisplayName = "To Float (Byte)", CompactNodeTitle = "->", Keywords = "cast convert", BlueprintAutocast), Category = "Math|Conversions")
-	static double Conv_ByteToDouble(uint8 InByte);
+	static UE_INL_API double Conv_ByteToDouble(uint8 InByte);
 
 	/** Converts an integer to a float */
 	UE_DEPRECATED(5.2, "This method has been deprecated and will be removed. Use the double version instead.")
-	static float Conv_IntToFloat(int32 InInt);
+	static UE_INL_API float Conv_IntToFloat(int32 InInt);
 
 	/** Converts an integer to a float */
 	UFUNCTION(BlueprintPure, meta = (DisplayName = "To Float (Integer)", CompactNodeTitle = "->", Keywords = "cast convert", BlueprintAutocast), Category = "Math|Conversions")
-	static double Conv_IntToDouble(int32 InInt);
+	static UE_INL_API double Conv_IntToDouble(int32 InInt);
 
 	/** Converts an integer to a 64 bit integer */
 	UFUNCTION(BlueprintPure, meta = (DisplayName = "To Integer64 (Integer)", CompactNodeTitle = "->", Keywords = "cast convert", BlueprintAutocast), Category = "Math|Conversions")
-	static int64 Conv_IntToInt64(int32 InInt);
+	static UE_INL_API int64 Conv_IntToInt64(int32 InInt);
 
 	/** Converts an integer to a byte (if the integer is too large, returns the low 8 bits) */
 	UFUNCTION(BlueprintPure, meta=(DisplayName = "To Byte (Integer)", CompactNodeTitle = "->", Keywords="cast convert", BlueprintAutocast), Category="Math|Conversions")
-	static uint8 Conv_IntToByte(int32 InInt);
+	static UE_INL_API uint8 Conv_IntToByte(int32 InInt);
 
 	/** Converts a 64 bit integer to a 32 bit integer (if the integer is too large, returns the low 32 bits) */
 	UFUNCTION(BlueprintPure, meta=(DisplayName = "To Integer (Integer64)", CompactNodeTitle = "->", Keywords="cast convert", BlueprintAutocast), Category="Math|Conversions")
-	static int32 Conv_Int64ToInt(int64 InInt);
+	static UE_INL_API int32 Conv_Int64ToInt(int64 InInt);
 
 	/** Converts a 64 bit floating point to a 32 bit floating point (if the float is too large, returns the low 32 bits) */
 	UFUNCTION(BlueprintPure, meta = (BlueprintInternalUseOnly = "true", DeprecatedFunction, DeprecationMessage = "Explicit conversions between floats and doubles are not necessary. Please remove node."))
-	static float Conv_DoubleToFloat(double InDouble);
+	static UE_INL_API float Conv_DoubleToFloat(double InDouble);
 
 	/** Converts a 32 bit floating point to a 64 bit floating point */
 	UFUNCTION(BlueprintPure, meta = (BlueprintInternalUseOnly = "true", DeprecatedFunction, DeprecationMessage = "Explicit conversions between floats and doubles are not necessary. Please remove node."))
-	static double Conv_FloatToDouble(float InFloat);
+	static UE_INL_API double Conv_FloatToDouble(float InFloat);
 
 	/** Converts a 64 bit integer to a byte (if the integer is too large, returns the low 8 bits) */
 	UFUNCTION(BlueprintPure, meta=(DisplayName = "To Byte (Integer64)", CompactNodeTitle = "->", Keywords="cast convert", BlueprintAutocast), Category="Math|Conversions")
-	static uint8 Conv_Int64ToByte(int64 InInt);
+	static UE_INL_API uint8 Conv_Int64ToByte(int64 InInt);
 
 	/** Converts a float to a 64 bit integer */
 	UFUNCTION(BlueprintPure, meta = (DisplayName = "To Integer64 (Float)", CompactNodeTitle = "->", Keywords = "cast convert", BlueprintAutocast), Category = "Math|Conversions")
-	static int64 Conv_DoubleToInt64(double InDouble);
+	static UE_INL_API int64 Conv_DoubleToInt64(double InDouble);
 
 	/** Converts a 64 bit integer to a float */
 	UFUNCTION(BlueprintPure, meta = (DisplayName = "To Float (Integer64)", CompactNodeTitle = "->", Keywords = "cast convert", BlueprintAutocast), Category = "Math|Conversions")
-	static double Conv_Int64ToDouble(int64 InInt);
+	static UE_INL_API double Conv_Int64ToDouble(int64 InInt);
 
 	/** Converts an integer to an IntVector*/
 	UFUNCTION(BlueprintPure, meta = (DisplayName = "To IntVector (Integer)", CompactNodeTitle = "->", Keywords = "cast convert", BlueprintAutocast), Category = "Math|Conversions")
@@ -3627,59 +3633,59 @@ class UKismetMathLibrary : public UBlueprintFunctionLibrary
 
 	/** Converts an integer to a FVector */
 	UFUNCTION(BlueprintPure, meta = (DisplayName = "To Vector (Integer)", CompactNodeTitle = "->", Keywords = "cast convert", BlueprintAutocast), Category = "Math|Conversions")
-	static FVector Conv_IntToVector(int32 InInt);
+	static UE_INL_API FVector Conv_IntToVector(int32 InInt);
 
 	/** Converts a int to a bool*/
 	UFUNCTION(BlueprintPure, meta=(DisplayName = "To Boolean (Integer)", CompactNodeTitle = "->", Keywords="cast convert", BlueprintAutocast), Category="Math|Conversions")
-	static bool Conv_IntToBool(int32 InInt);
+	static UE_INL_API bool Conv_IntToBool(int32 InInt);
 
 	/** Converts a bool to an int */
 	UFUNCTION(BlueprintPure, meta=(DisplayName = "To Integer (Boolean)", CompactNodeTitle = "->", Keywords="cast convert", BlueprintAutocast), Category="Math|Conversions")
-	static int32 Conv_BoolToInt(bool InBool);
+	static UE_INL_API int32 Conv_BoolToInt(bool InBool);
 
 	/** Converts a bool to a float (0.0f or 1.0f) */
 	UE_DEPRECATED(5.2, "This method has been deprecated and will be removed. Use the double version instead.")
-	static float Conv_BoolToFloat(bool InBool);
+	static UE_INL_API float Conv_BoolToFloat(bool InBool);
 
 	/** Converts a bool to a float (0.0 or 1.0) */
 	UFUNCTION(BlueprintPure, meta = (DisplayName = "To Float (Boolean)", CompactNodeTitle = "->", Keywords = "cast convert", BlueprintAutocast), Category = "Math|Conversions")
-	static double Conv_BoolToDouble(bool InBool);
+	static UE_INL_API double Conv_BoolToDouble(bool InBool);
 
 	/** Converts a bool to a byte */
 	UFUNCTION(BlueprintPure, meta=(DisplayName = "To Byte (Boolean)", CompactNodeTitle = "->", Keywords="cast convert", BlueprintAutocast), Category="Math|Conversions")
-	static uint8 Conv_BoolToByte(bool InBool);
+	static UE_INL_API uint8 Conv_BoolToByte(bool InBool);
 	
 	/** Converts a byte to an integer */
 	UFUNCTION(BlueprintPure, meta=(DisplayName = "To Integer (Byte)", CompactNodeTitle = "->", Keywords="cast convert", BlueprintAutocast), Category="Math|Conversions")
-	static int32 Conv_ByteToInt(uint8 InByte);
+	static UE_INL_API int32 Conv_ByteToInt(uint8 InByte);
 
 	/** Converts a byte to an integer */
 	UFUNCTION(BlueprintPure, meta = (DisplayName = "To Integer64 (Byte)", CompactNodeTitle = "->", Keywords = "cast convert", BlueprintAutocast), Category = "Math|Conversions")
-	static int64 Conv_ByteToInt64(uint8 InByte);
+	static UE_INL_API int64 Conv_ByteToInt64(uint8 InByte);
 
 	/** Converts a color to LinearColor */
 	UFUNCTION(BlueprintPure, meta = (DisplayName = "To LinearColor (Color)", CompactNodeTitle = "->", Keywords = "cast convert", BlueprintAutocast), Category = "Math|Conversions")
-	static FLinearColor Conv_ColorToLinearColor(FColor InColor);
+	static UE_INL_API FLinearColor Conv_ColorToLinearColor(FColor InColor);
 
 	/** Converts an IntVector to a vector */
 	UFUNCTION(BlueprintPure, meta = (DisplayName = "To Vector (IntVector)", CompactNodeTitle = "->", Keywords = "cast convert", BlueprintAutocast), Category = "Math|Conversions")
-	static FVector Conv_IntVectorToVector(const FIntVector& InIntVector);
+	static UE_INL_API FVector Conv_IntVectorToVector(const FIntVector& InIntVector);
 
 	/** Converts an IntVector2 to a vector2D */
 	UFUNCTION(BlueprintPure, meta = (DisplayName = "To Vector2D (IntVector2)", CompactNodeTitle = "->", Keywords = "cast convert", BlueprintAutocast), Category = "Math|Conversions")
-	static FVector2D Conv_IntVector2ToVector2D(const FIntVector2& InIntVector2);
+	static UE_INL_API FVector2D Conv_IntVector2ToVector2D(const FIntVector2& InIntVector2);
 
 	/** Converts a float into a vector, where each element is that float */
 	UE_DEPRECATED(5.2, "This method has been deprecated and will be removed. Use the double version instead.")
-	static FVector Conv_FloatToVector(float InFloat);
+	static UE_INL_API FVector Conv_FloatToVector(float InFloat);
 
 	/** Converts a double into a vector, where each element is that float */
 	UFUNCTION(BlueprintPure, meta = (DisplayName = "To Vector (Float)", CompactNodeTitle = "->", Keywords = "cast convert", BlueprintAutocast), Category = "Math|Conversions")
-	static FVector Conv_DoubleToVector(double InDouble);
+	static UE_INL_API FVector Conv_DoubleToVector(double InDouble);
 
 	/** Convert a float into a vector, where each element is that float */
 	UFUNCTION(BlueprintPure, meta=(DisplayName = "To Vector2D", CompactNodeTitle = "->", Keywords="cast convert", BlueprintAutocast), Category="Math|Conversions")
-	static FVector2D Conv_DoubleToVector2D(double InDouble);
+	static UE_INL_API FVector2D Conv_DoubleToVector2D(double InDouble);
 	
 	//
 	// Box functions
@@ -3798,7 +3804,7 @@ class UKismetMathLibrary : public UBlueprintFunctionLibrary
 
 	/** Breaks an FBoxSphereBounds into origin, extent, and radius */
 	UFUNCTION(BlueprintPure, meta = (NativeBreakFunc), Category = "Math|BoxSphereBounds")
-	static void BreakBoxSphereBounds(const FBoxSphereBounds& InBoxSphereBounds, FVector& Origin, FVector& BoxExtent, float& SphereRadius);
+	static UE_INL_API void BreakBoxSphereBounds(const FBoxSphereBounds& InBoxSphereBounds, FVector& Origin, FVector& BoxExtent, float& SphereRadius);
 
 
 	//
@@ -3874,11 +3880,11 @@ class UKismetMathLibrary : public UBlueprintFunctionLibrary
 	
 	/** Returns true if A and B are equal (A == B) */
 	UFUNCTION(BlueprintPure, meta=(DisplayName = "Equal (Object)", CompactNodeTitle = "==", Keywords = "== equal"), Category="Utilities")
-	static bool EqualEqual_ObjectObject(class UObject* A, class UObject* B);
+	static UE_INL_API bool EqualEqual_ObjectObject(class UObject* A, class UObject* B);
 
 	/** Returns true if A and B are not equal (A != B) */
 	UFUNCTION(BlueprintPure, meta=(DisplayName = "Not Equal (Object)", CompactNodeTitle = "!=", Keywords = "!= not equal"), Category="Utilities")
-	static bool NotEqual_ObjectObject(class UObject* A, class UObject* B);
+	static UE_INL_API bool NotEqual_ObjectObject(class UObject* A, class UObject* B);
 
 	//
 	// Class operators and functions.
@@ -3886,11 +3892,11 @@ class UKismetMathLibrary : public UBlueprintFunctionLibrary
 
 	/** Returns true if A and B are equal (A == B) */
 	UFUNCTION(BlueprintPure, meta=(DisplayName = "Equal (Class)", CompactNodeTitle = "==", Keywords = "== equal"), Category="Utilities")
-	static bool EqualEqual_ClassClass(class UClass* A, class UClass* B);
+	static UE_INL_API bool EqualEqual_ClassClass(class UClass* A, class UClass* B);
 
 	/** Returns true if A and B are not equal (A != B) */
 	UFUNCTION(BlueprintPure, meta=(DisplayName = "Not Equal (Class)", CompactNodeTitle = "!=", Keywords = "!= not equal"), Category="Utilities")
-	static bool NotEqual_ClassClass(class UClass* A, class UClass* B);
+	static UE_INL_API bool NotEqual_ClassClass(class UClass* A, class UClass* B);
 
 	/**
 	 * Determine if a class is a child of another class.
@@ -3908,11 +3914,11 @@ class UKismetMathLibrary : public UBlueprintFunctionLibrary
 	
 	/** Returns true if A and B are equal (A == B) */
 	UFUNCTION(BlueprintPure, meta=(DisplayName = "Equal (Name)", CompactNodeTitle = "==", Keywords = "== equal"), Category="Utilities|Name")
-	static bool EqualEqual_NameName(FName A, FName B);
+	static UE_INL_API bool EqualEqual_NameName(FName A, FName B);
 
 	/** Returns true if A and B are not equal (A != B) */
 	UFUNCTION(BlueprintPure, meta=(DisplayName = "Not Equal (Name)", CompactNodeTitle = "!=", Keywords = "!= not equal"), Category="Utilities|Name")
-	static bool NotEqual_NameName(FName A, FName B);
+	static UE_INL_API bool NotEqual_NameName(FName A, FName B);
 
 
 	//
@@ -3929,7 +3935,7 @@ class UKismetMathLibrary : public UBlueprintFunctionLibrary
 
 	/** Returns true if transform A is equal to transform B */
 	UFUNCTION(BlueprintPure, meta = (DisplayName = "Equal (Transform)", CompactNodeTitle = "==", ScriptMethod = "Equals", ScriptOperator = "==", Keywords = "== equal"), Category="Math|Transform")
-	static bool EqualEqual_TransformTransform(const FTransform& A, const FTransform& B);
+	static UE_INL_API bool EqualEqual_TransformTransform(const FTransform& A, const FTransform& B);
 
 	/** 
 	 *	Returns true if transform A is nearly equal to B 
@@ -3952,21 +3958,21 @@ class UKismetMathLibrary : public UBlueprintFunctionLibrary
 	 * @return New transform: A * B
 	 */
 	UFUNCTION(BlueprintPure, meta=(ScriptMethod = "Multiply", ScriptOperator = "*;*=", CompactNodeTitle = "*", Keywords="multiply *"), Category="Math|Transform")
-	static FTransform ComposeTransforms(const FTransform& A, const FTransform& B);
+	static UE_INL_API FTransform ComposeTransforms(const FTransform& A, const FTransform& B);
 
 	/** 
 	 *	Transform a position by the supplied transform.
 	 *	For example, if T was an object's transform, this would transform a position from local space to world space.
 	 */
 	UFUNCTION(BlueprintPure, meta=(ScriptMethod="TransformLocation", Keywords="location"), Category="Math|Transform")
-	static FVector TransformLocation(const FTransform& T, FVector Location);
+	static UE_INL_API FVector TransformLocation(const FTransform& T, FVector Location);
 
 	/** 
 	 *	Transform a direction vector by the supplied transform - will not change its length. 
 	 *	For example, if T was an object's transform, this would transform a direction from local space to world space.
 	 */
 	UFUNCTION(BlueprintPure, meta=(ScriptMethod="TransformDirection"), Category="Math|Transform")
-	static FVector TransformDirection(const FTransform& T, FVector Direction);
+	static UE_INL_API FVector TransformDirection(const FTransform& T, FVector Direction);
 
 	/** 
 	 *	Transform a rotator by the supplied transform. 
@@ -3980,14 +3986,14 @@ class UKismetMathLibrary : public UBlueprintFunctionLibrary
 	 *	For example, if T was an object's transform, this would transform a position from world space to local space.
 	 */
 	UFUNCTION(BlueprintPure, meta=(ScriptMethod="InverseTransformLocation", Keywords="location"), Category="Math|Transform")
-	static FVector InverseTransformLocation(const FTransform& T, FVector Location);
+	static UE_INL_API FVector InverseTransformLocation(const FTransform& T, FVector Location);
 
 	/** 
 	 *	Transform a direction vector by the inverse of the supplied transform - will not change its length.
 	 *	For example, if T was an object's transform, this would transform a direction from world space to local space.
 	 */
 	UFUNCTION(BlueprintPure, meta=(ScriptMethod="InverseTransformDirection"), Category="Math|Transform")
-	static FVector InverseTransformDirection(const FTransform& T, FVector Direction);
+	static UE_INL_API FVector InverseTransformDirection(const FTransform& T, FVector Direction);
 
 	/** 
 	 *	Transform a rotator by the inverse of the supplied transform. 
@@ -4007,7 +4013,7 @@ class UKismetMathLibrary : public UBlueprintFunctionLibrary
 	 * @return		The new relative transform
 	 */
 	UFUNCTION(BlueprintPure, meta=(ScriptMethod = "MakeRelative", Keywords="convert torelative"), Category="Math|Transform")
-	static FTransform MakeRelativeTransform(const FTransform& A, const FTransform& RelativeTo);
+	static UE_INL_API FTransform MakeRelativeTransform(const FTransform& A, const FTransform& RelativeTo);
 
 	/** 
 	 * Returns the inverse of the given transform T.
@@ -4018,7 +4024,7 @@ class UKismetMathLibrary : public UBlueprintFunctionLibrary
 	 * @return	The inverse of T.
 	 */
 	UFUNCTION(BlueprintPure, meta = (ScriptMethod = "Inverse", Keywords = "inverse"), Category = "Math|Transform")
-	static FTransform InvertTransform(const FTransform& T);
+	static UE_INL_API FTransform InvertTransform(const FTransform& T);
 
 	/** Linearly interpolates between A and B based on Alpha (100% of A when Alpha=0 and 100% of B when Alpha=1). */
 	UFUNCTION(BlueprintPure, meta=(DisplayName = "Lerp (Transform)", AdvancedDisplay = "3", ScriptMethod = "Lerp"), Category="Math|Transform")
@@ -4063,7 +4069,7 @@ class UKismetMathLibrary : public UBlueprintFunctionLibrary
 	 * @return		New interpolated position
 	 */
 	UFUNCTION(BlueprintPure, Category="Math|Interpolation")
-	static double FInterpTo(double Current, double Target, double DeltaTime, double InterpSpeed);
+	static UE_INL_API double FInterpTo(double Current, double Target, double DeltaTime, double InterpSpeed);
 
 	/**
 	 * Tries to reach Target at a constant rate.
@@ -4075,7 +4081,7 @@ class UKismetMathLibrary : public UBlueprintFunctionLibrary
 	 * @return		New interpolated position
 	 */
 	UFUNCTION(BlueprintPure, Category="Math|Interpolation")
-	static double FInterpTo_Constant(double Current, double Target, double DeltaTime, double InterpSpeed);
+	static UE_INL_API double FInterpTo_Constant(double Current, double Target, double DeltaTime, double InterpSpeed);
 
 	/**
 	 * Tries to reach Target rotation based on Current rotation, giving a nice smooth feeling when rotating to Target rotation.
@@ -4087,7 +4093,7 @@ class UKismetMathLibrary : public UBlueprintFunctionLibrary
 	 * @return		New interpolated position
 	 */
 	UFUNCTION(BlueprintPure, Category="Math|Interpolation", meta=(Keywords="rotation rotate"))
-	static FRotator RInterpTo(FRotator Current, FRotator Target, float DeltaTime, float InterpSpeed);
+	static UE_INL_API FRotator RInterpTo(FRotator Current, FRotator Target, float DeltaTime, float InterpSpeed);
 
 	/**
 	 * Tries to reach Target rotation at a constant rate.
@@ -4099,7 +4105,7 @@ class UKismetMathLibrary : public UBlueprintFunctionLibrary
 	 * @return		New interpolated position
 	 */
 	UFUNCTION(BlueprintPure, Category="Math|Interpolation", meta=(Keywords="rotation rotate"))
-	static FRotator RInterpTo_Constant(FRotator Current, FRotator Target, float DeltaTime, float InterpSpeed);
+	static UE_INL_API FRotator RInterpTo_Constant(FRotator Current, FRotator Target, float DeltaTime, float InterpSpeed);
 
 	/**
 	 * Uses a simple spring model to interpolate a float from Current to Target.
@@ -4157,7 +4163,7 @@ class UKismetMathLibrary : public UBlueprintFunctionLibrary
 	static ENGINE_API int32 RandomIntegerFromStream(const FRandomStream& Stream, int32 Max);
 
 	UE_DEPRECATED(5.2, "Use RandomIntegerFromStream taking the FRandomStream as the first argument.")
-	static int32 RandomIntegerFromStream(int32 Max, const FRandomStream& Stream)
+	static UE_INL_API int32 RandomIntegerFromStream(int32 Max, const FRandomStream& Stream)
 	{
 		return RandomIntegerFromStream(Stream, Max);
 	}
@@ -4167,7 +4173,7 @@ class UKismetMathLibrary : public UBlueprintFunctionLibrary
 	static ENGINE_API int32 RandomIntegerInRangeFromStream(const FRandomStream& Stream, int32 Min, int32 Max);
 
 	UE_DEPRECATED(5.2, "Use RandomIntegerInRangeFromStream taking the FRandomStream as the first argument.")
-	static int32 RandomIntegerInRangeFromStream(int32 Min, int32 Max, const FRandomStream& Stream)
+	static UE_INL_API int32 RandomIntegerInRangeFromStream(int32 Min, int32 Max, const FRandomStream& Stream)
 	{
 		return RandomIntegerInRangeFromStream(Stream, Min, Max);
 	}
@@ -4185,7 +4191,7 @@ class UKismetMathLibrary : public UBlueprintFunctionLibrary
 	static ENGINE_API float RandomFloatInRangeFromStream(const FRandomStream& Stream, float Min, float Max);
 
 	UE_DEPRECATED(5.2, "Use RandomFloatInRangeFromStream taking the FRandomStream as the first argument.")
-	static float RandomFloatInRangeFromStream(float Min, float Max, const FRandomStream& Stream)
+	static UE_INL_API float RandomFloatInRangeFromStream(float Min, float Max, const FRandomStream& Stream)
 	{
 		return RandomFloatInRangeFromStream(Stream, Min, Max);
 	}
@@ -4199,7 +4205,7 @@ class UKismetMathLibrary : public UBlueprintFunctionLibrary
 	static ENGINE_API FVector RandomPointInBoundingBoxFromStream(const FRandomStream& Stream, const FVector Center, const FVector HalfSize);
 
 	UE_DEPRECATED(5.2, "Use RandomPointInBoundingBoxFromStream taking the FRandomStream as the first argument.")
-	static FVector RandomPointInBoundingBoxFromStream(const FVector Center, const FVector HalfSize, const FRandomStream& Stream)
+	static UE_INL_API FVector RandomPointInBoundingBoxFromStream(const FVector Center, const FVector HalfSize, const FRandomStream& Stream)
 	{
 		return RandomPointInBoundingBoxFromStream(Stream, Center, HalfSize);
 	}
@@ -4209,7 +4215,7 @@ class UKismetMathLibrary : public UBlueprintFunctionLibrary
 	static ENGINE_API FVector RandomPointInBoundingBoxFromStream_Box(const FRandomStream& Stream, const FBox Box);
 
 	UE_DEPRECATED(5.2, "Use RandomPointInBoundingBoxFromStream_Box taking the FRandomStream as the first argument.")
-	static FVector RandomPointInBoundingBoxFromStream_Box(const FBox Box, const FRandomStream& Stream)
+	static UE_INL_API FVector RandomPointInBoundingBoxFromStream_Box(const FBox Box, const FRandomStream& Stream)
 	{
 		return RandomPointInBoundingBoxFromStream_Box(Stream, Box);
 	}
@@ -4219,7 +4225,7 @@ class UKismetMathLibrary : public UBlueprintFunctionLibrary
 	static ENGINE_API FRotator RandomRotatorFromStream(const FRandomStream& Stream, bool bRoll);
 
 	UE_DEPRECATED(5.2, "Use RandomRotatorFromStream taking the FRandomStream as the first argument.")
-	static FRotator RandomRotatorFromStream(bool bRoll, const FRandomStream& Stream)
+	static UE_INL_API FRotator RandomRotatorFromStream(bool bRoll, const FRandomStream& Stream)
 	{
 		return RandomRotatorFromStream(Stream, bRoll);
 	}
@@ -4246,7 +4252,7 @@ class UKismetMathLibrary : public UBlueprintFunctionLibrary
 	static ENGINE_API FVector RandomUnitVectorInConeInRadiansFromStream(const FRandomStream& Stream, const FVector& ConeDir, float ConeHalfAngleInRadians);
 
 	UE_DEPRECATED(5.2, "Use RandomUnitVectorInConeInRadiansFromStream taking the FRandomStream as the first argument.")
-	static FVector RandomUnitVectorInConeInRadiansFromStream(const FVector& ConeDir, float ConeHalfAngleInRadians, const FRandomStream& Stream)
+	static UE_INL_API FVector RandomUnitVectorInConeInRadiansFromStream(const FVector& ConeDir, float ConeHalfAngleInRadians, const FRandomStream& Stream)
 	{
 		return RandomUnitVectorInConeInRadiansFromStream(Stream, ConeDir, ConeHalfAngleInRadians);
 	}
@@ -4281,7 +4287,7 @@ class UKismetMathLibrary : public UBlueprintFunctionLibrary
 	static ENGINE_API FVector RandomUnitVectorInEllipticalConeInRadiansFromStream(const FRandomStream& Stream, const FVector& ConeDir, float MaxYawInRadians, float MaxPitchInRadians);
 
 	UE_DEPRECATED(5.2, "Use RandomUnitVectorInEllipticalConeInRadiansFromStream taking the FRandomStream as the first argument.")
-	static FVector RandomUnitVectorInEllipticalConeInRadiansFromStream(const FVector& ConeDir, float MaxYawInRadians, float MaxPitchInRadians, const FRandomStream& Stream)
+	static UE_INL_API FVector RandomUnitVectorInEllipticalConeInRadiansFromStream(const FVector& ConeDir, float MaxYawInRadians, float MaxPitchInRadians, const FRandomStream& Stream)
 	{
 		return RandomUnitVectorInEllipticalConeInRadiansFromStream(Stream, ConeDir, MaxYawInRadians, MaxPitchInRadians);
 	}
@@ -4346,7 +4352,7 @@ class UKismetMathLibrary : public UBlueprintFunctionLibrary
 	 * @return Whether the points are relatively coplanar, based on the tolerance
 	 */
 	UFUNCTION(BlueprintPure, Category = "Math|Geometry")
-	static bool PointsAreCoplanar(const TArray<FVector>& Points, float Tolerance = 0.1f);
+	static UE_INL_API bool PointsAreCoplanar(const TArray<FVector>& Points, float Tolerance = 0.1f);
 
 	/**
 	 * Determines whether the given point is in a box. Includes points on the box.
@@ -4589,21 +4595,21 @@ class UKismetMathLibrary : public UBlueprintFunctionLibrary
 
 	// NetQuantized vector make/breaks
 	UFUNCTION(BlueprintPure, Category = "Math|Vector", meta = (NativeMakeFunc))
-	static FVector_NetQuantize MakeVector_NetQuantize(double X, double Y, double Z) { return FVector_NetQuantize(X, Y, Z); }
+	static UE_INL_API FVector_NetQuantize MakeVector_NetQuantize(double X, double Y, double Z) { return FVector_NetQuantize(X, Y, Z); }
 	UFUNCTION(BlueprintPure, Category = "Math|Vector", meta = (NativeMakeFunc))
-	static FVector_NetQuantize10 MakeVector_NetQuantize10(double X, double Y, double Z) { return FVector_NetQuantize10(X, Y, Z); }
+	static UE_INL_API FVector_NetQuantize10 MakeVector_NetQuantize10(double X, double Y, double Z) { return FVector_NetQuantize10(X, Y, Z); }
 	UFUNCTION(BlueprintPure, Category = "Math|Vector", meta = (NativeMakeFunc))
-	static FVector_NetQuantize100 MakeVector_NetQuantize100(double X, double Y, double Z) { return FVector_NetQuantize100(X, Y, Z); }
+	static UE_INL_API FVector_NetQuantize100 MakeVector_NetQuantize100(double X, double Y, double Z) { return FVector_NetQuantize100(X, Y, Z); }
 	UFUNCTION(BlueprintPure, Category = "Math|Vector", meta = (NativeMakeFunc))
-	static FVector_NetQuantizeNormal MakeVector_NetQuantizeNormal(double X, double Y, double Z) { return FVector_NetQuantizeNormal(X, Y, Z); }
+	static UE_INL_API FVector_NetQuantizeNormal MakeVector_NetQuantizeNormal(double X, double Y, double Z) { return FVector_NetQuantizeNormal(X, Y, Z); }
 	UFUNCTION(BlueprintPure, Category = "Math|Vector", meta = (NativeBreakFunc))
-	static void BreakVector_NetQuantize(FVector_NetQuantize InVec, double& X, double& Y, double& Z) { BreakVector((FVector)InVec, X, Y, Z); }
+	static UE_INL_API void BreakVector_NetQuantize(FVector_NetQuantize InVec, double& X, double& Y, double& Z) { BreakVector((FVector)InVec, X, Y, Z); }
 	UFUNCTION(BlueprintPure, Category = "Math|Vector", meta = (NativeBreakFunc))
-	static void BreakVector_NetQuantize10(FVector_NetQuantize10 InVec, double& X, double& Y, double& Z) { BreakVector((FVector)InVec, X, Y, Z); }
+	static UE_INL_API void BreakVector_NetQuantize10(FVector_NetQuantize10 InVec, double& X, double& Y, double& Z) { BreakVector((FVector)InVec, X, Y, Z); }
 	UFUNCTION(BlueprintPure, Category = "Math|Vector", meta = (NativeBreakFunc))
-	static void BreakVector_NetQuantize100(FVector_NetQuantize100 InVec, double& X, double& Y, double& Z) { BreakVector((FVector)InVec, X, Y, Z); }
+	static UE_INL_API void BreakVector_NetQuantize100(FVector_NetQuantize100 InVec, double& X, double& Y, double& Z) { BreakVector((FVector)InVec, X, Y, Z); }
 	UFUNCTION(BlueprintPure, Category = "Math|Vector", meta = (NativeBreakFunc))
-	static void BreakVector_NetQuantizeNormal(FVector_NetQuantizeNormal InVec, double& X, double& Y, double& Z) { BreakVector((FVector)InVec, X, Y, Z); }
+	static UE_INL_API void BreakVector_NetQuantizeNormal(FVector_NetQuantizeNormal InVec, double& X, double& Y, double& Z) { BreakVector((FVector)InVec, X, Y, Z); }
 
 private:
 
@@ -4627,6 +4633,7 @@ private:
 	static ENGINE_API void ReportError_DaysInMonth();
 };
 
+#undef UE_INL_API
 
 // Conditionally inlined
 #if KISMET_MATH_INLINE_ENABLED
diff --git a/Engine/Source/Runtime/Engine/Classes/Sound/SoundEffectBase.h b/Engine/Source/Runtime/Engine/Classes/Sound/SoundEffectBase.h
--- a/Engine/Source/Runtime/Engine/Classes/Sound/SoundEffectBase.h
+++ b/Engine/Source/Runtime/Engine/Classes/Sound/SoundEffectBase.h
@@ -68,7 +68,7 @@ class USoundEffectPreset;
 		FCriticalSection SettingsCritSect; \
 		F##EFFECT_NAME##Settings SettingsCopy; \
 
-class ENGINE_API FSoundEffectBase
+class FSoundEffectBase
 {
 public:
 	virtual ~FSoundEffectBase() = default;
@@ -77,31 +77,31 @@ public:
 	virtual void OnPresetChanged() { }
 
 	/** Returns if the submix is active or bypassing audio. */
-	bool IsActive() const;
+	ENGINE_API bool IsActive() const;
 
 	/** Enables the submix effect. */
-	void SetEnabled(const bool bInIsEnabled);
+	ENGINE_API void SetEnabled(const bool bInIsEnabled);
 
 	/** Updates preset on audio render thread. Returns true if update processed a preset update, false if not. */
-	bool Update();
+	ENGINE_API bool Update();
 
-	USoundEffectPreset* GetPreset();
-	TWeakObjectPtr<USoundEffectPreset>& GetPresetPtr();
+	ENGINE_API USoundEffectPreset* GetPreset();
+	ENGINE_API TWeakObjectPtr<USoundEffectPreset>& GetPresetPtr();
 
 	/** Queries if the given preset object is the uobject preset for this preset instance, i.e. the preset which spawned this effect instance. */
-	bool IsPreset(USoundEffectPreset* InPreset) const;
+	ENGINE_API bool IsPreset(USoundEffectPreset* InPreset) const;
 
 	/** Enqueues a lambda command on a thread safe queue which is pumped from the audio render thread. */
-	void EffectCommand(TUniqueFunction<void()> Command);
+	ENGINE_API void EffectCommand(TUniqueFunction<void()> Command);
 
 	/** Returns the unique ID of the parent preset. */
 	uint32 GetParentPresetId() const { return ParentPresetUniqueId; }
 
 protected:
-	FSoundEffectBase();
+	ENGINE_API FSoundEffectBase();
 
 	/** Pumps messages awaiting execution on the audio render thread */
-	void PumpPendingMessages();
+	ENGINE_API void PumpPendingMessages();
 
 	FCriticalSection SettingsCritSect;
 	TArray<uint8> CurrentAudioThreadSettingsData;
diff --git a/Engine/Source/Runtime/Engine/Classes/VisualLogger/VisualLoggerTraceDevice.h b/Engine/Source/Runtime/Engine/Classes/VisualLogger/VisualLoggerTraceDevice.h
--- a/Engine/Source/Runtime/Engine/Classes/VisualLogger/VisualLoggerTraceDevice.h
+++ b/Engine/Source/Runtime/Engine/Classes/VisualLogger/VisualLoggerTraceDevice.h
@@ -15,12 +15,12 @@ public:
 	static ENGINE_API FVisualLoggerTraceDevice& Get();
 
 	ENGINE_API FVisualLoggerTraceDevice();
-	virtual void Cleanup(bool bReleaseMemory = false) override;
-	virtual void StartRecordingToFile(double TimeStamp) override;
-	virtual void StopRecordingToFile(double TimeStamp) override;
-	virtual void DiscardRecordingToFile() override;
-	virtual void SetFileName(const FString& InFileName) override;
-	virtual void Serialize(const UObject* InLogOwner, const FName& InOwnerName, const FName& InOwnerDisplayName, const FName& InOwnerClassName, const FVisualLogEntry& InLogEntry) override;
+	ENGINE_API virtual void Cleanup(bool bReleaseMemory = false) override;
+	ENGINE_API virtual void StartRecordingToFile(double TimeStamp) override;
+	ENGINE_API virtual void StopRecordingToFile(double TimeStamp) override;
+	ENGINE_API virtual void DiscardRecordingToFile() override;
+	ENGINE_API virtual void SetFileName(const FString& InFileName) override;
+	ENGINE_API virtual void Serialize(const UObject* InLogOwner, const FName& InOwnerName, const FName& InOwnerDisplayName, const FName& InOwnerClassName, const FVisualLogEntry& InLogEntry) override;
 	virtual bool HasFlags(int32 InFlags) const override { return !!(InFlags & (EVisualLoggerDeviceFlags::CanSaveToFile | EVisualLoggerDeviceFlags::StoreLogsLocally)); }
 
 	FImmediateRenderDelegate ImmediateRenderDelegate;
diff --git a/Engine/Source/Runtime/Experimental/Chaos/Public/Chaos/PBDConstraintBaseData.h b/Engine/Source/Runtime/Experimental/Chaos/Public/Chaos/PBDConstraintBaseData.h
--- a/Engine/Source/Runtime/Experimental/Chaos/Public/Chaos/PBDConstraintBaseData.h
+++ b/Engine/Source/Runtime/Experimental/Chaos/Public/Chaos/PBDConstraintBaseData.h
@@ -32,18 +32,18 @@ enum class EConstraintType
 	void Set##FNAME(TYPE InValue){ PROP.Modify(/*bInvalidate=*/true, DirtyFlags, Proxy, [&InValue](auto& Data) { Data.VNAME = InValue; }); }\
 	TYPE Get##FNAME() const{return PROP.Read().VNAME;}\
 
-class CHAOS_API FConstraintBase
+class FConstraintBase
 {
 public:
 	virtual ~FConstraintBase() {}
 
 	typedef TVector<TGeometryParticleHandle<FReal, 3>*, 2> FParticleHandlePair;
 
-	FConstraintBase(EConstraintType InType);
+	CHAOS_API FConstraintBase(EConstraintType InType);
 
 	EConstraintType GetType() const { return Type; }
 	bool IsType(EConstraintType InType) { return (Type == InType); }
-	bool IsValid() const;
+	CHAOS_API bool IsValid() const;
 
 	bool IsDirty() const { return DirtyFlags.IsDirty(); }
 	bool IsDirty(const EChaosPropertyFlags CheckBits) const { return DirtyFlags.IsDirty(CheckBits); }
@@ -55,7 +55,7 @@ public:
 
 	template<typename T = IPhysicsProxyBase> T* GetProxy() { return static_cast<T*>(Proxy); }
 
-	void SetProxy(IPhysicsProxyBase* InProxy);
+	CHAOS_API void SetProxy(IPhysicsProxyBase* InProxy);
 
 	void SyncRemoteData(FDirtyPropertiesManager& Manager, int32 DataIdx, FDirtyChaosProperties& RemoteData)
 	{
diff --git a/Engine/Source/Runtime/Experimental/Chaos/Public/Chaos/ShapeInstance.h b/Engine/Source/Runtime/Experimental/Chaos/Public/Chaos/ShapeInstance.h
--- a/Engine/Source/Runtime/Experimental/Chaos/Public/Chaos/ShapeInstance.h
+++ b/Engine/Source/Runtime/Experimental/Chaos/Public/Chaos/ShapeInstance.h
@@ -57,9 +57,9 @@ namespace Chaos
 	public:
 		// Downcasts exposed until we deprecate and remove FPerShapeData
 		FShapeInstanceProxy* AsShapeInstanceProxy();
-		CHAOS_API const FShapeInstanceProxy* AsShapeInstanceProxy() const;
+		const FShapeInstanceProxy* AsShapeInstanceProxy() const;
 		FShapeInstance* AsShapeInstance();
-		CHAOS_API const FShapeInstance* AsShapeInstance() const;
+		const FShapeInstance* AsShapeInstance() const;
 
 	public:
 		static constexpr bool AlwaysSerializable = true;
diff --git a/Engine/Source/Runtime/GeometryCore/Public/DynamicMesh/DynamicMesh3.h b/Engine/Source/Runtime/GeometryCore/Public/DynamicMesh/DynamicMesh3.h
--- a/Engine/Source/Runtime/GeometryCore/Public/DynamicMesh/DynamicMesh3.h
+++ b/Engine/Source/Runtime/GeometryCore/Public/DynamicMesh/DynamicMesh3.h
@@ -296,7 +296,7 @@ public:
 	 * @param Mesh Mesh to serialize.
 	 * @returns Passing down serializing archive.
 	 */
-	friend GEOMETRYCORE_API FArchive& operator<<(FArchive& Ar, FDynamicMesh3& Mesh)
+	friend FArchive& operator<<(FArchive& Ar, FDynamicMesh3& Mesh)
 	{
 		Mesh.Serialize(Ar);
 		return Ar;
