Remove timestamp-based dynamic resolution heuristic method
- Was enabled via r.DynamicRes.GPUTimingMeasureMethod = 1, but has been unused for a long time.
- The BeginDynamicResolutionRendering and EndDynamicResolutionRendering members of EDynamicResolutionStateEvent were also unused. The events were emitted by the engine, but the computed GPU time between those events was discarded.

Mark GGPUFrameTime as deprecated and replace use outside of platform RHIs with RHIGetGPUFrameCycles()
- Also removed RHIGetGPUFrameCycles() from platform RHI API. Only the global function remains, which simply returns GGPUFrameTime.
- Unifying the API will make it simpler to fix the thread race that occurs between the GT/RT and end-of-pipe threads which update the GGPUFrameTime field.

diff --git a/Engine/Plugins/Experimental/Avalanche/Source/Avalanche/Private/Framework/AvaGameViewportClient.cpp b/Engine/Plugins/Experimental/Avalanche/Source/Avalanche/Private/Framework/AvaGameViewportClient.cpp
--- a/Engine/Plugins/Experimental/Avalanche/Source/Avalanche/Private/Framework/AvaGameViewportClient.cpp
+++ b/Engine/Plugins/Experimental/Avalanche/Source/Avalanche/Private/Framework/AvaGameViewportClient.cpp
@@ -291,9 +291,6 @@ void UAvaGameViewportClient::Draw(FViewport* InViewport, FCanvas* InCanvas)
 		GetRendererModule().PerFrameCleanupIfSkipRenderer();
 	}
 
-	// Beyond this point, only UI rendering independent from dynamic resolution.
-	GEngine->EmitDynamicResolutionEvent(EDynamicResolutionStateEvent::EndDynamicResolutionRendering);
-
 	// Remove temporary debug lines.
 	constexpr const UWorld::ELineBatcherType LineBatchersToFlush[] = { UWorld::ELineBatcherType::World, UWorld::ELineBatcherType::Foreground };
 	World->FlushLineBatchers(LineBatchersToFlush);
diff --git a/Engine/Plugins/Performance/PerformanceMonitor/Source/PerformanceMonitor/Private/PerformanceMonitor.cpp b/Engine/Plugins/Performance/PerformanceMonitor/Source/PerformanceMonitor/Private/PerformanceMonitor.cpp
--- a/Engine/Plugins/Performance/PerformanceMonitor/Source/PerformanceMonitor/Private/PerformanceMonitor.cpp
+++ b/Engine/Plugins/Performance/PerformanceMonitor/Source/PerformanceMonitor/Private/PerformanceMonitor.cpp
@@ -366,7 +366,7 @@ void FPerformanceMonitorModule::RecordFrame()
 			ArrayForStatName.Add(FGenericPlatformTime::ToMilliseconds(GGameThreadTime));
 			GeneratedStats.Emplace(TEXT("GlobalGameThreadTime"), ArrayForStatName);
 			ArrayForStatName = GeneratedStats.FindOrAdd(TEXT("GlobalGPUFrameTime"));
-			ArrayForStatName.Add(FGenericPlatformTime::ToMilliseconds(GGPUFrameTime));
+			ArrayForStatName.Add(FGenericPlatformTime::ToMilliseconds(RHIGetGPUFrameCycles()));
 			GeneratedStats.Emplace(TEXT("GlobalGPUFrameTime"), ArrayForStatName);
 
 			break;
diff --git a/Engine/Plugins/VirtualProduction/StageMonitoring/Source/StageDataProvider/Private/FramePerformanceProvider.cpp b/Engine/Plugins/VirtualProduction/StageMonitoring/Source/StageDataProvider/Private/FramePerformanceProvider.cpp
--- a/Engine/Plugins/VirtualProduction/StageMonitoring/Source/StageDataProvider/Private/FramePerformanceProvider.cpp
+++ b/Engine/Plugins/VirtualProduction/StageMonitoring/Source/StageDataProvider/Private/FramePerformanceProvider.cpp
@@ -47,7 +47,7 @@ static float GAverageGPU = 0;
 
 void UpdateAverageGPUUsage()
 {
-	const float GPUTime = FPlatformTime::ToMilliseconds(GGPUFrameTime);
+	const float GPUTime = FPlatformTime::ToMilliseconds(RHIGetGPUFrameCycles());
 	GAverageGPU = 0.75*GAverageGPU + 0.25*GPUTime;
 }
 
@@ -310,7 +310,7 @@ void FFramePerformanceProvider::CheckHitches(int64 Frame)
 	{
 		const float GameThreadTime = FPlatformTime::ToMilliseconds(GGameThreadTime);
 		const float RenderThreadTime = FPlatformTime::ToMilliseconds(GRenderThreadTime);
-		const float GPUTime = FPlatformTime::ToMilliseconds(GGPUFrameTime);
+		const float GPUTime = FPlatformTime::ToMilliseconds(RHIGetGPUFrameCycles());
 		float HitchedFPS = CachedHitchSettings.MinimumFrameRate.AsDecimal();
 		if (!FMath::IsNearlyZero(FullFrameTime))
 		{
diff --git a/Engine/Plugins/VirtualProduction/StageMonitoring/Source/StageDataProvider/Private/FramePerformanceProvider.h b/Engine/Plugins/VirtualProduction/StageMonitoring/Source/StageDataProvider/Private/FramePerformanceProvider.h
--- a/Engine/Plugins/VirtualProduction/StageMonitoring/Source/StageDataProvider/Private/FramePerformanceProvider.h
+++ b/Engine/Plugins/VirtualProduction/StageMonitoring/Source/StageDataProvider/Private/FramePerformanceProvider.h
@@ -51,7 +51,7 @@ public:
 	UPROPERTY(VisibleAnywhere, Category = "Hitch", meta = (Unit = "ms"))
 	float RenderThreadMS = 0.f;
 
-	/** Current GPU time read from GGPUFrameTime in milliseconds */
+	/** Current GPU time read from RHIGetGPUFrameCycles() in milliseconds */
 	UPROPERTY(VisibleAnywhere, Category = "Hitch", meta = (Unit = "ms"))
 	float GPU_MS = 0.f;
 
diff --git a/Engine/Plugins/VirtualProduction/StageMonitoring/Source/StageMonitorCommon/Public/StageMonitorUtils.h b/Engine/Plugins/VirtualProduction/StageMonitoring/Source/StageMonitorCommon/Public/StageMonitorUtils.h
--- a/Engine/Plugins/VirtualProduction/StageMonitoring/Source/StageMonitorCommon/Public/StageMonitorUtils.h
+++ b/Engine/Plugins/VirtualProduction/StageMonitoring/Source/StageMonitorCommon/Public/StageMonitorUtils.h
@@ -79,7 +79,7 @@ public:
 	UPROPERTY(VisibleAnywhere, Category = "Performance", meta = (Unit = "ms"))
 	float RenderThreadWaitMS = 0.f;
 
-	/** Current GPU time read from GGPUFrameTime in milliseconds */
+	/** Current GPU time read from RHIGetGPUFrameCycles() in milliseconds */
 	UPROPERTY(VisibleAnywhere, Category = "Performance", meta = (Unit = "ms"))
 	float GPU_MS = 0.f;
 
diff --git a/Engine/Source/Developer/FunctionalTesting/Private/FunctionalTest.cpp b/Engine/Source/Developer/FunctionalTesting/Private/FunctionalTest.cpp
--- a/Engine/Source/Developer/FunctionalTesting/Private/FunctionalTest.cpp
+++ b/Engine/Source/Developer/FunctionalTesting/Private/FunctionalTest.cpp
@@ -1513,7 +1513,7 @@ void FPerfStatsRecord::Sample(UWorld* World, float DeltaSeconds, bool bBaseline)
 		Baseline.FrameTimeTracker.AddSample(StatUnitData->RawFrameTime);
 		Baseline.GameThreadTimeTracker.AddSample(FPlatformTime::ToMilliseconds(GGameThreadTime));
 		Baseline.RenderThreadTimeTracker.AddSample(FPlatformTime::ToMilliseconds(GRenderThreadTime));
-		Baseline.GPUTimeTracker.AddSample(FPlatformTime::ToMilliseconds(GGPUFrameTime));
+		Baseline.GPUTimeTracker.AddSample(FPlatformTime::ToMilliseconds(RHIGetGPUFrameCycles()));
 		Baseline.NumFrames++;
 		Baseline.SumTimeSeconds += DeltaSeconds;
 	}
@@ -1522,7 +1522,7 @@ void FPerfStatsRecord::Sample(UWorld* World, float DeltaSeconds, bool bBaseline)
 		Record.FrameTimeTracker.AddSample(StatUnitData->RawFrameTime);
 		Record.GameThreadTimeTracker.AddSample(FPlatformTime::ToMilliseconds(GGameThreadTime));
 		Record.RenderThreadTimeTracker.AddSample(FPlatformTime::ToMilliseconds(GRenderThreadTime));
-		Record.GPUTimeTracker.AddSample(FPlatformTime::ToMilliseconds(GGPUFrameTime));
+		Record.GPUTimeTracker.AddSample(FPlatformTime::ToMilliseconds(RHIGetGPUFrameCycles()));
 		Record.NumFrames++;
 		Record.SumTimeSeconds += DeltaSeconds;
 	}
diff --git a/Engine/Source/Runtime/Apple/MetalRHI/Private/MetalDynamicRHI.h b/Engine/Source/Runtime/Apple/MetalRHI/Private/MetalDynamicRHI.h
--- a/Engine/Source/Runtime/Apple/MetalRHI/Private/MetalDynamicRHI.h
+++ b/Engine/Source/Runtime/Apple/MetalRHI/Private/MetalDynamicRHI.h
@@ -169,7 +169,6 @@ public:
 	virtual FTextureRHIRef RHIGetViewportBackBuffer(FRHIViewport* Viewport) final override;
 	virtual void RHIAdvanceFrameForGetViewportBackBuffer(FRHIViewport* Viewport) final override;
 	virtual void RHIFlushResources() final override;
-	virtual uint32 RHIGetGPUFrameCycles(uint32 GPUIndex = 0) final override;
 	virtual FViewportRHIRef RHICreateViewport(void* WindowHandle, uint32 SizeX, uint32 SizeY, bool bIsFullscreen, EPixelFormat PreferredPixelFormat) final override;
 	virtual void RHIResizeViewport(FRHIViewport* Viewport, uint32 SizeX, uint32 SizeY, bool bIsFullscreen) final override;
 	virtual void RHIResizeViewport(FRHIViewport* Viewport, uint32 SizeX, uint32 SizeY, bool bIsFullscreen, EPixelFormat PreferredPixelFormat) final override;
diff --git a/Engine/Source/Runtime/Apple/MetalRHI/Private/MetalProfiler.cpp b/Engine/Source/Runtime/Apple/MetalRHI/Private/MetalProfiler.cpp
--- a/Engine/Source/Runtime/Apple/MetalRHI/Private/MetalProfiler.cpp
+++ b/Engine/Source/Runtime/Apple/MetalRHI/Private/MetalProfiler.cpp
@@ -141,7 +141,9 @@ void FMetalCommandBufferTimer::RecordFrame()
 	FrameEndGPUCycles = LastBufferTiming.EndTime * CyclesPerSecond;
 	
 	uint64 FrameGPUTimeCycles = uint64(CyclesPerSecond * RunningFrameTimeSeconds);
+PRAGMA_DISABLE_DEPRECATION_WARNINGS
 	FPlatformAtomics::AtomicStore_Relaxed((int32*)&GGPUFrameTime, int32(FrameGPUTimeCycles));
+PRAGMA_ENABLE_DEPRECATION_WARNINGS
 	
 #if STATS
 	FPlatformAtomics::AtomicStore_Relaxed(&GMetalGPUWorkTime, FrameGPUTimeCycles);
diff --git a/Engine/Source/Runtime/Apple/MetalRHI/Private/MetalRHI.cpp b/Engine/Source/Runtime/Apple/MetalRHI/Private/MetalRHI.cpp
--- a/Engine/Source/Runtime/Apple/MetalRHI/Private/MetalRHI.cpp
+++ b/Engine/Source/Runtime/Apple/MetalRHI/Private/MetalRHI.cpp
@@ -1594,12 +1594,6 @@ void FMetalDynamicRHI::RHIBlockUntilGPUIdle()
 	}
 }
 
-uint32 FMetalDynamicRHI::RHIGetGPUFrameCycles(uint32 GPUIndex)
-{
-	check(GPUIndex == 0);
-	return GGPUFrameTime;
-}
-
 IRHICommandContext* FMetalDynamicRHI::RHIGetDefaultContext()
 {
 	return &ImmediateContext;
diff --git a/Engine/Source/Runtime/D3D12RHI/Private/D3D12Commands.cpp b/Engine/Source/Runtime/D3D12RHI/Private/D3D12Commands.cpp
--- a/Engine/Source/Runtime/D3D12RHI/Private/D3D12Commands.cpp
+++ b/Engine/Source/Runtime/D3D12RHI/Private/D3D12Commands.cpp
@@ -2086,15 +2086,6 @@ void FD3D12DynamicRHI::RHIBlockUntilGPUIdle()
 	}
 }
 
-/*
-* Returns the total GPU time taken to render the last frame. Same metric as FPlatformTime::Cycles().
-*/
-uint32 FD3D12DynamicRHI::RHIGetGPUFrameCycles(uint32 GPUIndex)
-{
-	return GGPUFrameTime;
-}
-
-
 void FD3D12CommandContext::RHISetDepthBounds(float MinDepth, float MaxDepth)
 {
 	StateCache.SetDepthBounds(MinDepth, MaxDepth);
diff --git a/Engine/Source/Runtime/D3D12RHI/Private/D3D12RHIPrivate.h b/Engine/Source/Runtime/D3D12RHI/Private/D3D12RHIPrivate.h
--- a/Engine/Source/Runtime/D3D12RHI/Private/D3D12RHIPrivate.h
+++ b/Engine/Source/Runtime/D3D12RHI/Private/D3D12RHIPrivate.h
@@ -394,7 +394,6 @@ public:
 	virtual bool RHIMatchPrecachePSOInitializers(const FGraphicsPipelineStateInitializer& LHS, const FGraphicsPipelineStateInitializer& RHS) final override;
 	virtual void RHIAdvanceFrameForGetViewportBackBuffer(FRHIViewport* Viewport) final override;
 	virtual void RHIFlushResources() final override;
-	virtual uint32 RHIGetGPUFrameCycles(uint32 GPUIndex = 0) final override;
 	virtual FViewportRHIRef RHICreateViewport(void* WindowHandle, uint32 SizeX, uint32 SizeY, bool bIsFullscreen, EPixelFormat PreferredPixelFormat) final override;
 	virtual void RHIResizeViewport(FRHIViewport* Viewport, uint32 SizeX, uint32 SizeY, bool bIsFullscreen) final override;
 	virtual void RHIResizeViewport(FRHIViewport* ViewportRHI, uint32 SizeX, uint32 SizeY, bool bIsFullscreen, EPixelFormat PreferredPixelFormat) final override;
diff --git a/Engine/Source/Runtime/D3D12RHI/Private/D3D12Submission.cpp b/Engine/Source/Runtime/D3D12RHI/Private/D3D12Submission.cpp
--- a/Engine/Source/Runtime/D3D12RHI/Private/D3D12Submission.cpp
+++ b/Engine/Source/Runtime/D3D12RHI/Private/D3D12Submission.cpp
@@ -1445,12 +1445,16 @@ void FD3D12DynamicRHI::ProcessTimestamps(FD3D12TimingArray const& TimingPerQueue
 	if (GetHardwareGPUFrameTime(HardwareGPUTime) && CVarGPUTimeFromTimestamps.GetValueOnAnyThread() == 0)
 	{
 		SET_CYCLE_COUNTER(STAT_RHI_GPUTotalTimeHW, HardwareGPUTime);
+	PRAGMA_DISABLE_DEPRECATION_WARNINGS
 		GGPUFrameTime = HardwareGPUTime;
+	PRAGMA_ENABLE_DEPRECATION_WARNINGS
 	}
 	else
 	{
 		SET_CYCLE_COUNTER(STAT_RHI_GPUTotalTimeHW, 0);
+	PRAGMA_DISABLE_DEPRECATION_WARNINGS
 		GGPUFrameTime = FPlatformMath::TruncToInt(double(UnionBusyCycles) * Scale32);
+	PRAGMA_ENABLE_DEPRECATION_WARNINGS
 	}
 
 	for (auto const& Current : TimingPerQueue)
diff --git a/Engine/Source/Runtime/Engine/Private/ChartCreation.cpp b/Engine/Source/Runtime/Engine/Private/ChartCreation.cpp
--- a/Engine/Source/Runtime/Engine/Private/ChartCreation.cpp
+++ b/Engine/Source/Runtime/Engine/Private/ChartCreation.cpp
@@ -1233,7 +1233,7 @@ IPerformanceDataConsumer::FFrameData FPerformanceTrackingSystem::AnalyzeFrame(fl
 	// Copy these locally since the RT may update it between reads otherwise
 	const uint32 LocalRenderThreadTime = GRenderThreadTime;
 	const uint32 LocalRHIThreadTime = GRHIThreadTime;
-	const uint32 LocalGPUFrameTime = GGPUFrameTime;
+	const uint32 LocalGPUFrameTime = RHIGetGPUFrameCycles();
 
 	const double CurrentTime = FPlatformTime::Seconds();
 	if (LastTimeChartCreationTicked > 0)
@@ -1420,8 +1420,6 @@ void FPerformanceTrackingSystem::StartCharting()
 		GTargetFrameRatesForSummary.Add(FCString::Atoi(*FramerateString));
 	}
 
-	GGPUFrameTime = 0;
-
 	UE_LOG(LogChartCreation, Log, TEXT("Started creating FPS charts at %f seconds"), FPSChartStartTime);
 }
 
diff --git a/Engine/Source/Runtime/Engine/Private/DynamicResolution.cpp b/Engine/Source/Runtime/Engine/Private/DynamicResolution.cpp
--- a/Engine/Source/Runtime/Engine/Private/DynamicResolution.cpp
+++ b/Engine/Source/Runtime/Engine/Private/DynamicResolution.cpp
@@ -122,13 +122,6 @@ static TAutoConsoleVariable<int32> CVarMaxConsecutiveOverBudgetGPUFrameCount(
 	TEXT("Maximum number of consecutive frames tolerated over GPU budget."),
 	ECVF_RenderThreadSafe | ECVF_Default);
 
-static TAutoConsoleVariable<int32> CVarTimingMeasureModel(
-	TEXT("r.DynamicRes.GPUTimingMeasureMethod"),
-	0,
-	TEXT("Selects the method to use to measure GPU timings.\n")
-	TEXT(" 0: Same as stat unit (default);\n 1: Timestamp queries."),
-	ECVF_RenderThreadSafe | ECVF_Default);
-
 static TAutoConsoleVariable<int32> CVarUpperBoundQuantization(
 	TEXT("r.DynamicRes.UpperBoundQuantization"),
 	DynamicRenderScaling::FHeuristicSettings::kDefaultUpperBoundQuantization,
@@ -157,7 +150,7 @@ DynamicRenderScaling::FHeuristicSettings GetPrimaryDynamicResolutionSettings()
 DynamicRenderScaling::FBudget GDynamicPrimaryResolutionFraction(TEXT("DynamicPrimaryResolution"), &GetPrimaryDynamicResolutionSettings);
 
 
-static float TimeStampQueryResultToMiliSeconds(uint64 TimestampResult)
+static float TimeStampQueryResultToMilliSeconds(uint64 TimestampResult)
 {
 	return float(TimestampResult) / 1000.0f;
 }
@@ -243,11 +236,9 @@ uint64 FDynamicResolutionHeuristicProxy::CreateNewPreviousFrameTimings_RenderThr
 void FDynamicResolutionHeuristicProxy::CommitPreviousFrameGPUTimings_RenderThread(
 	uint64 HistoryFrameId,
 	float TotalFrameGPUBusyTimeMs,
-	float DynamicResolutionGPUBusyTimeMs,
 	const DynamicRenderScaling::TMap<float>& BudgetTimingMs)
 {
 	check(TotalFrameGPUBusyTimeMs >= 0.0f);
-	check(DynamicResolutionGPUBusyTimeMs >= 0.0f);
 
 	// Conveniently early return if invalid history frame id.
 	if (HistoryFrameId == FDynamicResolutionHeuristicProxy::kInvalidEntryId ||
@@ -262,7 +253,6 @@ void FDynamicResolutionHeuristicProxy::CommitPreviousFrameGPUTimings_RenderThrea
 	check(History[EntryId].TotalFrameGPUBusyTimeMs == -1.0f);
 
 	History[EntryId].TotalFrameGPUBusyTimeMs = TotalFrameGPUBusyTimeMs;
-	History[EntryId].GlobalDynamicResolutionTimeMs = DynamicResolutionGPUBusyTimeMs;
 	History[EntryId].BudgetTimingMs = BudgetTimingMs;
 }
 
@@ -708,15 +698,11 @@ public:
 	FDefaultDynamicResolutionStateProxy()
 	{
 		check(IsInGameThread());
-		InFlightFrames.SetNum(4);
-		CurrentFrameInFlightIndex = -1;
-		bUseTimeQueriesThisFrame = false;
 	}
 
 	~FDefaultDynamicResolutionStateProxy()
 	{
 		check(IsInRenderingThread());
-		checkf(InFlightFrames.Num()==0, TEXT("Ensure the object is properly deinitialized by Finalize call"));
 	}
 
 	void Reset()
@@ -725,17 +711,12 @@ public:
 
 		// Reset heuristic.
 		Heuristic.Reset_RenderThread();
-
-		// Set invalid heuristic's entry id on all inflight frames.
-		for (auto& InFlightFrame : InFlightFrames)
-		{
-			InFlightFrame.HeuristicHistoryEntry = FDynamicResolutionHeuristicProxy::kInvalidEntryId;
-		}
 	}
 
 	void BeginFrame(FRHICommandList& RHICmdList, float PrevFrameTimeMs, float PrevGameThreadTimeMs)
 	{
 		check(IsInRenderingThread());
+		ensure(GRHISupportsFrameCyclesBubblesRemoval);
 
 		if (DynamicRenderScaling::IsSupported())
 		{
@@ -755,95 +736,39 @@ public:
 
 		// Query render thread time Ms.
 		float PrevRenderThreadTimeMs = FPlatformTime::ToMilliseconds(GRenderThreadTime);
-		float PrevRHIThreadTimeMs = FPlatformTime::ToMilliseconds(GRHIThreadTime);
-
-		bUseTimeQueriesThisFrame = GSupportsTimestampRenderQueries && CVarTimingMeasureModel.GetValueOnRenderThread() == 1;
+		float PrevRHIThreadTimeMs    = FPlatformTime::ToMilliseconds(GRHIThreadTime);
+		float PrevFrameGPUTimeMs     = FPlatformTime::ToMilliseconds(RHIGetGPUFrameCycles());
 
-		if (bUseTimeQueriesThisFrame)
-		{
-			// Create the query pool
-			if (!QueryPool.IsValid())
-				QueryPool = RHICreateRenderQueryPool(RQT_AbsoluteTime);
-
-			// Process to the in-flight frames that had their queries fully landed.
-			HandLandedQueriesToHeuristic(/* bWait = */ false);
-
-			// Finds a new CurrentFrameInFlightIndex that does not have any pending queries.
-			FindNewInFlightIndex();
+		uint64 HistoryEntryId = Heuristic.CreateNewPreviousFrameTimings_RenderThread(
+			PrevFrameTimeMs,
+			PrevGameThreadTimeMs,
+			PrevRenderThreadTimeMs,
+			PrevRHIThreadTimeMs
+		);
 
-			InFlightFrameQueries& InFlightFrame = InFlightFrames[CurrentFrameInFlightIndex];
+		const DynamicRenderScaling::TMap<uint64>& LatestTimings = DynamicRenderScaling::GetLatestTimings();
 
-			// Feed the thread timings to the heuristic.
-			InFlightFrame.HeuristicHistoryEntry = Heuristic.CreateNewPreviousFrameTimings_RenderThread(
-				PrevFrameTimeMs, PrevGameThreadTimeMs, PrevRenderThreadTimeMs, PrevRHIThreadTimeMs);
-
-			check(QueryPool.IsValid());
-			InFlightFrame.BeginFrameQuery = QueryPool->AllocateQuery();
-			RHICmdList.EndRenderQuery(InFlightFrame.BeginFrameQuery.GetQuery());
-		}
-		else
+		DynamicRenderScaling::TMap<float> BudgetTimingMs;
+		for (TLinkedList<DynamicRenderScaling::FBudget*>::TIterator BudgetIt(DynamicRenderScaling::FBudget::GetGlobalList()); BudgetIt; BudgetIt.Next())
 		{
-			// If RHI does not support GPU busy time queries, fall back to what stat unit does.
-			ensure(GRHISupportsFrameCyclesBubblesRemoval);
-			float PrevFrameGPUTimeMs = FPlatformTime::ToMilliseconds(RHIGetGPUFrameCycles());
-
-			uint64 HistoryEntryId = Heuristic.CreateNewPreviousFrameTimings_RenderThread(
-				PrevFrameTimeMs, PrevGameThreadTimeMs, PrevRenderThreadTimeMs, PrevRHIThreadTimeMs);
-
-			const DynamicRenderScaling::TMap<uint64>& LatestTimings = DynamicRenderScaling::GetLatestTimings();
-			DynamicRenderScaling::TMap<float> BudgetTimingMs;
-			for (TLinkedList<DynamicRenderScaling::FBudget*>::TIterator BudgetIt(DynamicRenderScaling::FBudget::GetGlobalList()); BudgetIt; BudgetIt.Next())
-			{
-				const DynamicRenderScaling::FBudget& Budget = **BudgetIt;
-				BudgetTimingMs[Budget] = TimeStampQueryResultToMiliSeconds(LatestTimings[Budget]);
-			}
-			BudgetTimingMs[GDynamicPrimaryResolutionFraction] = PrevFrameGPUTimeMs;
-
-
-			Heuristic.CommitPreviousFrameGPUTimings_RenderThread(HistoryEntryId,
-				/* TotalFrameGPUBusyTimeMs = */ PrevFrameGPUTimeMs,
-				/* DynamicResolutionGPUBusyTimeMs = */ PrevFrameGPUTimeMs,
-				/* BudgetTimingMs = */ BudgetTimingMs);
+			const DynamicRenderScaling::FBudget& Budget = **BudgetIt;
+			BudgetTimingMs[Budget] = TimeStampQueryResultToMilliSeconds(LatestTimings[Budget]);
+		}
+		BudgetTimingMs[GDynamicPrimaryResolutionFraction] = PrevFrameGPUTimeMs;
 
-			Heuristic.RefreshCurrentFrameResolutionFraction_RenderThread();
+		Heuristic.CommitPreviousFrameGPUTimings_RenderThread(HistoryEntryId,
+			/* TotalFrameGPUBusyTimeMs = */ PrevFrameGPUTimeMs,
+			/* BudgetTimingMs = */ BudgetTimingMs);
 
-			// Set a non insane value for internal checks to pass as if GRHISupportsGPUBusyTimeQueries == true.
-			CurrentFrameInFlightIndex = 0;
-		}
+		Heuristic.RefreshCurrentFrameResolutionFraction_RenderThread();
 	}
 
 	void ProcessEvent(FRHICommandList& RHICmdList, EDynamicResolutionStateEvent Event)
 	{
 		check(IsInRenderingThread());
 
-		if (bUseTimeQueriesThisFrame)
-		{
-			InFlightFrameQueries& InFlightFrame = InFlightFrames[CurrentFrameInFlightIndex];
-
-			FRHIPooledRenderQuery* QueryPtr = nullptr;
-			switch (Event)
-			{
-			case EDynamicResolutionStateEvent::BeginDynamicResolutionRendering:
-				QueryPtr = &InFlightFrame.BeginDynamicResolutionQuery; break;
-			case EDynamicResolutionStateEvent::EndDynamicResolutionRendering:
-				QueryPtr = &InFlightFrame.EndDynamicResolutionQuery; break;
-			case EDynamicResolutionStateEvent::EndFrame:
-				QueryPtr = &InFlightFrame.EndFrameQuery; break;
-			default: check(0);
-			}
-
-			check(QueryPtr != nullptr);
-			check(QueryPool.IsValid());
-			*QueryPtr = QueryPool->AllocateQuery();
-			RHICmdList.EndRenderQuery(QueryPtr->GetQuery());
-		}
-
-		// Clobber CurrentFrameInFlightIndex for internal checks.
 		if (Event == EDynamicResolutionStateEvent::EndFrame)
 		{
-			CurrentFrameInFlightIndex = -1;
-			bUseTimeQueriesThisFrame = false;
-
 			DynamicRenderScaling::EndFrame();
 		}
 	}
@@ -852,160 +777,10 @@ public:
 	void Finalize()
 	{
 		check(IsInRenderingThread());
-
-		// Wait for all queries to land.
-		HandLandedQueriesToHeuristic(/* bWait = */ true);
-
-		// Reset the frame properties
-		InFlightFrames.Empty();
-		QueryPool = nullptr;
 	}
 
-
 	// Heuristic's proxy.
 	FDynamicResolutionHeuristicProxy Heuristic;
-
-
-private:
-	struct InFlightFrameQueries
-	{
-		// GPU queries.
-		FRHIPooledRenderQuery BeginFrameQuery;
-		FRHIPooledRenderQuery BeginDynamicResolutionQuery;
-		FRHIPooledRenderQuery EndDynamicResolutionQuery;
-		FRHIPooledRenderQuery EndFrameQuery;
-
-		// Heuristic's history 
-		uint64 HeuristicHistoryEntry;
-
-		InFlightFrameQueries()
-		{ 
-			ResetValues();
-		}
-
-		/// Reset values
-		void ResetValues()
-		{
-			HeuristicHistoryEntry = FDynamicResolutionHeuristicProxy::kInvalidEntryId;
-			BeginFrameQuery.ReleaseQuery();
-			BeginDynamicResolutionQuery.ReleaseQuery();
-			EndDynamicResolutionQuery.ReleaseQuery();
-			EndFrameQuery.ReleaseQuery();
-		}
-	};
-
-	// Shared query pool for the frames in flight
-	FRenderQueryPoolRHIRef QueryPool;
-
-	// List of frame queries in flight.
-	TArray<InFlightFrameQueries> InFlightFrames;
-
-	// Current frame's in flight index.
-	int32 CurrentFrameInFlightIndex;
-
-	// Uses GPU busy time queries.
-	bool bUseTimeQueriesThisFrame;
-
-
-	void HandLandedQueriesToHeuristic(bool bWait)
-	{
-		check(IsInRenderingThread());
-		check(GRHISupportsGPUTimestampBubblesRemoval || bWait);
-
-		bool ShouldRefreshHeuristic = false;
-
-		for (int32 i = 0; i < InFlightFrames.Num(); i++)
-		{
-			// If current in flight frame queries, ignore them since have not called EndRenderQuery().
-			if (i == CurrentFrameInFlightIndex)
-			{
-				continue;
-			}
-
-			InFlightFrameQueries& InFlightFrame = InFlightFrames[i];
-
-			// Results in microseconds.
-			uint64 BeginFrameResult = 0;
-			uint64 BeginDynamicResolutionResult = 0;
-			uint64 EndDynamicResolutionResult = 0;
-			uint64 EndFrameResult = 0;
-
-			int32 LandingCount = 0;
-			int32 QueryCount = 0;
-			if (InFlightFrame.BeginFrameQuery.IsValid())
-			{
-				LandingCount += RHIGetRenderQueryResult(
-					InFlightFrame.BeginFrameQuery.GetQuery(), BeginFrameResult, bWait) ? 1 : 0;
-				QueryCount += 1;
-			}
-
-			if (InFlightFrame.BeginDynamicResolutionQuery.IsValid())
-			{
-				LandingCount += RHIGetRenderQueryResult(
-					InFlightFrame.BeginDynamicResolutionQuery.GetQuery(), BeginDynamicResolutionResult, bWait) ? 1 : 0;
-				QueryCount += 1;
-			}
-
-			if (InFlightFrame.EndDynamicResolutionQuery.IsValid())
-			{
-				LandingCount += RHIGetRenderQueryResult(
-					InFlightFrame.EndDynamicResolutionQuery.GetQuery(), EndDynamicResolutionResult, bWait) ? 1 : 0;
-				QueryCount += 1;
-			}
-
-			if (InFlightFrame.EndFrameQuery.IsValid())
-			{
-				LandingCount += RHIGetRenderQueryResult(
-					InFlightFrame.EndFrameQuery.GetQuery(), EndFrameResult, bWait) ? 1 : 0;
-				QueryCount += 1;
-			}
-
-			check(QueryCount == 0 || QueryCount == 4);
-
-			// If all queries have landed, then hand the results to the heuristic.
-			if (LandingCount == 4)
-			{
-				Heuristic.CommitPreviousFrameGPUTimings_RenderThread(
-					InFlightFrame.HeuristicHistoryEntry,
-					/* TotalFrameGPUBusyTimeMs = */ TimeStampQueryResultToMiliSeconds(EndFrameResult - BeginFrameResult),
-					/* DynamicResolutionGPUBusyTimeMs = */ TimeStampQueryResultToMiliSeconds(EndDynamicResolutionResult - BeginDynamicResolutionResult),
-					/* bGPUTimingsHaveCPUBubbles = */ !GRHISupportsGPUTimestampBubblesRemoval);
-
-				// Reset this in-flight frame queries to be reused.
-				InFlightFrame.ResetValues();
-
-				ShouldRefreshHeuristic = true;
-			}
-		}
-
-		// Refresh the heuristic.
-		if (ShouldRefreshHeuristic)
-		{
-			Heuristic.RefreshCurrentFrameResolutionFraction_RenderThread();
-		}
-	}
-
-	void FindNewInFlightIndex()
-	{
-		check(IsInRenderingThread());
-		check(CurrentFrameInFlightIndex == -1);
-
-		for (int32 i = 0; i < InFlightFrames.Num(); i++)
-		{
-			auto& InFlightFrame = InFlightFrames[i];
-			if (!InFlightFrame.BeginFrameQuery.IsValid())
-			{
-				CurrentFrameInFlightIndex = i;
-				break;
-			}
-		}
-
-		// Allocate a new in-flight frame in the unlikely event.
-		if (CurrentFrameInFlightIndex == -1)
-		{
-			CurrentFrameInFlightIndex = InFlightFrames.Add(InFlightFrameQueries());
-		}
-	}
 };
 
 
diff --git a/Engine/Source/Runtime/Engine/Private/GameViewportClient.cpp b/Engine/Source/Runtime/Engine/Private/GameViewportClient.cpp
--- a/Engine/Source/Runtime/Engine/Private/GameViewportClient.cpp
+++ b/Engine/Source/Runtime/Engine/Private/GameViewportClient.cpp
@@ -1233,7 +1233,7 @@ void UGameViewportClient::SetDropDetail(float DeltaSeconds)
 		float FrameTime = 0.0f;
 		if (FPlatformProperties::SupportsWindowedMode() == false)
 		{
-			FrameTime	= FPlatformTime::ToSeconds(FMath::Max3<uint32>( GRenderThreadTime, GGameThreadTime, GGPUFrameTime ));
+			FrameTime = FPlatformTime::ToSeconds(FMath::Max3<uint32>(GRenderThreadTime, GGameThreadTime, RHIGetGPUFrameCycles()));
 			// If DeltaSeconds is bigger than 34 ms we can take it into account as we're not VSYNCing in that case.
 			if( DeltaSeconds > 0.034 )
 			{
@@ -1578,14 +1578,12 @@ void UGameViewportClient::Draw(FViewport* InViewport, FCanvas* SceneCanvas)
 			// game play experience on platforms that does not support it, but have it enabled by mistake.
 			if (DynamicResolutionStateInfos.Status == EDynamicResolutionStatus::Enabled)
 			{
-				GEngine->EmitDynamicResolutionEvent(EDynamicResolutionStateEvent::BeginDynamicResolutionRendering);
 				GEngine->GetDynamicResolutionState()->SetupMainViewFamily(ViewFamily);
 
 				bUsesDynamicResolution = ViewFamily.GetScreenPercentageInterface() != nullptr;
 			}
 			else if (DynamicResolutionStateInfos.Status == EDynamicResolutionStatus::DebugForceEnabled)
 			{
-				GEngine->EmitDynamicResolutionEvent(EDynamicResolutionStateEvent::BeginDynamicResolutionRendering);
 				ViewFamily.SetScreenPercentageInterface(new FLegacyScreenPercentageDriver(
 					ViewFamily,
 					DynamicResolutionStateInfos.ResolutionFractionApproximations[GDynamicPrimaryResolutionFraction],
@@ -1909,9 +1907,6 @@ void UGameViewportClient::Draw(FViewport* InViewport, FCanvas* SceneCanvas)
 		GetRendererModule().PerFrameCleanupIfSkipRenderer();
 	}
 
-	// Beyond this point, only UI rendering independent from dynamc resolution.
-	GEngine->EmitDynamicResolutionEvent(EDynamicResolutionStateEvent::EndDynamicResolutionRendering);
-
 	// Clear areas of the rendertarget (backbuffer) that aren't drawn over by the views.
 	if (!bBufferCleared)
 	{
diff --git a/Engine/Source/Runtime/Engine/Private/UnrealClient.cpp b/Engine/Source/Runtime/Engine/Private/UnrealClient.cpp
--- a/Engine/Source/Runtime/Engine/Private/UnrealClient.cpp
+++ b/Engine/Source/Runtime/Engine/Private/UnrealClient.cpp
@@ -387,8 +387,7 @@ int32 FStatUnitData::DrawStat(FViewport* InViewport, FCanvas* InCanvas, int32 In
 	for (uint32 GPUIndex : FRHIGPUMask::All())
 	{
 		/** Number of milliseconds the GPU was busy last frame. */
-		const uint32 GPUCycles = RHIGetGPUFrameCycles(GPUIndex);
-		RawGPUFrameTime[GPUIndex] = FPlatformTime::ToMilliseconds(GPUCycles);
+		RawGPUFrameTime[GPUIndex] = FPlatformTime::ToMilliseconds(RHIGetGPUFrameCycles(GPUIndex));
 		GPUFrameTime[GPUIndex] = 0.9 * GPUFrameTime[GPUIndex] + 0.1 * RawGPUFrameTime[GPUIndex];
 
 		if (GRHISupportsGPUUsage)
diff --git a/Engine/Source/Runtime/Engine/Private/UnrealEngine.cpp b/Engine/Source/Runtime/Engine/Private/UnrealEngine.cpp
--- a/Engine/Source/Runtime/Engine/Private/UnrealEngine.cpp
+++ b/Engine/Source/Runtime/Engine/Private/UnrealEngine.cpp
@@ -13885,14 +13885,6 @@ void UEngine::EmitDynamicResolutionEvent(EDynamicResolutionStateEvent Event)
 	{
 		EDynamicResolutionStateEvent ExpectedPreviousEvent = EDynamicResolutionStateEvent(int32(Event) - 1);
 
-		// When doing a window resize, the game viewport client end up being drawn twice between BeginFrame and EndFrame.
-		// In this case we ignore duplicated Begin and End rendering events.
-		if (Event == EDynamicResolutionStateEvent::BeginDynamicResolutionRendering &&
-			LastDynamicResolutionEvent == EDynamicResolutionStateEvent::EndDynamicResolutionRendering)
-		{
-			return;
-		}
-
 		// If the previous event is going missing, automatically fire it.
 		if (LastDynamicResolutionEvent != ExpectedPreviousEvent)
 		{
diff --git a/Engine/Source/Runtime/Engine/Public/DynamicResolutionProxy.h b/Engine/Source/Runtime/Engine/Public/DynamicResolutionProxy.h
--- a/Engine/Source/Runtime/Engine/Public/DynamicResolutionProxy.h
+++ b/Engine/Source/Runtime/Engine/Public/DynamicResolutionProxy.h
@@ -30,7 +30,6 @@ public:
 	ENGINE_API void CommitPreviousFrameGPUTimings_RenderThread(
 		uint64 HistoryFrameId,
 		float TotalFrameGPUBusyTimeMs,
-		float DynamicResolutionGPUBusyTimeMs,
 		const DynamicRenderScaling::TMap<float>& BudgetTimingMs);
 
 	/** Refresh resolution fraction's from history. */
@@ -71,9 +70,6 @@ private:
 		// Total GPU busy time for the entire frame in milliseconds.
 		float TotalFrameGPUBusyTimeMs = -1.0f;
 
-		// Total GPU busy time for the render thread commands that perform dynamic resolutions.
-		float GlobalDynamicResolutionTimeMs = -1.0f;
-
 		// Time for each individual timings
 		DynamicRenderScaling::TMap<float> BudgetTimingMs;
 
diff --git a/Engine/Source/Runtime/Engine/Public/DynamicResolutionState.h b/Engine/Source/Runtime/Engine/Public/DynamicResolutionState.h
--- a/Engine/Source/Runtime/Engine/Public/DynamicResolutionState.h
+++ b/Engine/Source/Runtime/Engine/Public/DynamicResolutionState.h
@@ -22,12 +22,6 @@ enum class EDynamicResolutionStateEvent : uint8
 	// Fired at the very begining of the frame.
 	BeginFrame,
 
-	// Fired when starting to render with dynamic resolution for the frame.
-	BeginDynamicResolutionRendering,
-
-	// Fired when finished to render with dynamic resolution for the frame.
-	EndDynamicResolutionRendering,
-
 	// Fired at the very end of the frame.
 	EndFrame
 };
diff --git a/Engine/Source/Runtime/Launch/Private/LaunchEngineLoop.cpp b/Engine/Source/Runtime/Launch/Private/LaunchEngineLoop.cpp
--- a/Engine/Source/Runtime/Launch/Private/LaunchEngineLoop.cpp
+++ b/Engine/Source/Runtime/Launch/Private/LaunchEngineLoop.cpp
@@ -1401,7 +1401,7 @@ static void UpdateCoreCsvStats_EndFrame()
 	{
 		CSV_CUSTOM_STAT_MINIMAL_GLOBAL(RenderThreadTime, FPlatformTime::ToMilliseconds(GRenderThreadTime), ECsvCustomStatOp::Set);
 		CSV_CUSTOM_STAT_MINIMAL_GLOBAL(GameThreadTime, FPlatformTime::ToMilliseconds(GGameThreadTime), ECsvCustomStatOp::Set);
-		CSV_CUSTOM_STAT_MINIMAL_GLOBAL(GPUTime, FPlatformTime::ToMilliseconds(GGPUFrameTime), ECsvCustomStatOp::Set);
+		CSV_CUSTOM_STAT_MINIMAL_GLOBAL(GPUTime, FPlatformTime::ToMilliseconds(RHIGetGPUFrameCycles()), ECsvCustomStatOp::Set);
 		CSV_CUSTOM_STAT_MINIMAL_GLOBAL(RenderThreadTime_CriticalPath, FPlatformTime::ToMilliseconds(GRenderThreadTimeCriticalPath), ECsvCustomStatOp::Set);
 		CSV_CUSTOM_STAT_MINIMAL_GLOBAL(GameThreadTime_CriticalPath, FPlatformTime::ToMilliseconds(GGameThreadTimeCriticalPath), ECsvCustomStatOp::Set);
 		if (IsRunningRHIInSeparateThread())
diff --git a/Engine/Source/Runtime/NullDrv/Public/NullRHI.h b/Engine/Source/Runtime/NullDrv/Public/NullRHI.h
--- a/Engine/Source/Runtime/NullDrv/Public/NullRHI.h
+++ b/Engine/Source/Runtime/NullDrv/Public/NullRHI.h
@@ -402,11 +402,6 @@ public:
 
 	}
 
-	virtual uint32 RHIGetGPUFrameCycles(uint32 GPUIndex = 0) final override
-	{ 
-		return 0; 
-	}
-
 	virtual FViewportRHIRef RHICreateViewport(void* WindowHandle, uint32 SizeX, uint32 SizeY, bool bIsFullscreen, EPixelFormat PreferredPixelFormat) final override
 	{ 
 		return new FRHIViewport(); 
diff --git a/Engine/Source/Runtime/OpenGLDrv/Private/OpenGLCommands.cpp b/Engine/Source/Runtime/OpenGLDrv/Private/OpenGLCommands.cpp
--- a/Engine/Source/Runtime/OpenGLDrv/Private/OpenGLCommands.cpp
+++ b/Engine/Source/Runtime/OpenGLDrv/Private/OpenGLCommands.cpp
@@ -2616,15 +2616,6 @@ void FOpenGLDynamicRHI::RHIBlockUntilGPUIdle()
 	// @todo dev-pr use a render query to determine when the GPU is done
 }
 
-/**
- * Returns the total GPU time taken to render the last frame. Same metric as FPlatformTime::Cycles().
- */
-uint32 FOpenGLDynamicRHI::RHIGetGPUFrameCycles(uint32 GPUIndex)
-{
-	check(GPUIndex == 0);
-	return GGPUFrameTime;
-}
-
 template <typename TRHIShader>
 void FOpenGLDynamicRHI::ApplyStaticUniformBuffers(TRHIShader* Shader)
 {
diff --git a/Engine/Source/Runtime/OpenGLDrv/Private/OpenGLDrv.cpp b/Engine/Source/Runtime/OpenGLDrv/Private/OpenGLDrv.cpp
--- a/Engine/Source/Runtime/OpenGLDrv/Private/OpenGLDrv.cpp
+++ b/Engine/Source/Runtime/OpenGLDrv/Private/OpenGLDrv.cpp
@@ -61,7 +61,11 @@ void FOpenGLDynamicRHI::RHISetExternalGPUTime(uint64 InExternalGPUTime)
 #if RHI_NEW_GPU_PROFILER
 	Profiler.ExternalGPUTime = InExternalGPUTime;
 #else
+
+PRAGMA_DISABLE_DEPRECATION_WARNINGS
 	GGPUFrameTime = InExternalGPUTime;
+PRAGMA_ENABLE_DEPRECATION_WARNINGS
+
 	GPUProfilingData->ExternalGPUTime = InExternalGPUTime;
 #endif
 }
@@ -325,7 +329,9 @@ void FOpenGLGPUProfiler::EndFrame()
 	{
 		uint64 GPUTiming = FrameTiming.GetTiming();
 		uint64 GPUFreq = FrameTiming.GetTimingFrequency();
+	PRAGMA_DISABLE_DEPRECATION_WARNINGS
 		GGPUFrameTime = FMath::TruncToInt( double(GPUTiming) / double(GPUFreq) / FPlatformTime::GetSecondsPerCycle() );
+	PRAGMA_ENABLE_DEPRECATION_WARNINGS
 	}
 	else if (FOpenGLDisjointTimeStampQuery::IsSupported())
 	{
@@ -335,18 +341,24 @@ void FOpenGLGPUProfiler::EndFrame()
 		int OldestQueryIndex = (CurrentGPUFrameQueryIndex + 1) % MAX_GPUFRAMEQUERIES;
 		if ( DisjointGPUFrameTimeQuery[OldestQueryIndex].IsResultValid() && DisjointGPUFrameTimeQuery[OldestQueryIndex].GetResult(&GPUTiming) )
 		{
+		PRAGMA_DISABLE_DEPRECATION_WARNINGS
 			GGPUFrameTime = FMath::TruncToInt( double(GPUTiming) / double(GPUFreq) / FPlatformTime::GetSecondsPerCycle() );
 			GLastGPUFrameTime = GGPUFrameTime;
+		PRAGMA_ENABLE_DEPRECATION_WARNINGS
 		}
 		else
 		{
 			// Keep the timing of the last frame if the query was disjoint (e.g. GPU frequency changed and the result is undefined)
+		PRAGMA_DISABLE_DEPRECATION_WARNINGS
 			GGPUFrameTime = GLastGPUFrameTime;
+		PRAGMA_ENABLE_DEPRECATION_WARNINGS
 		}
 	}
 	else
 	{
+	PRAGMA_DISABLE_DEPRECATION_WARNINGS
 		GGPUFrameTime = ExternalGPUTime;
+	PRAGMA_ENABLE_DEPRECATION_WARNINGS
 	}
 
 	// if we have a frame open, close it now.
diff --git a/Engine/Source/Runtime/OpenGLDrv/Private/OpenGLDrv.h b/Engine/Source/Runtime/OpenGLDrv/Private/OpenGLDrv.h
--- a/Engine/Source/Runtime/OpenGLDrv/Private/OpenGLDrv.h
+++ b/Engine/Source/Runtime/OpenGLDrv/Private/OpenGLDrv.h
@@ -375,7 +375,6 @@ public:
 	virtual void RHIAcquireThreadOwnership() final override;
 	virtual void RHIReleaseThreadOwnership() final override;
 	virtual void RHIFlushResources() final override;
-	virtual uint32 RHIGetGPUFrameCycles(uint32 GPUIndex = 0) final override;
 	virtual FViewportRHIRef RHICreateViewport(void* WindowHandle, uint32 SizeX, uint32 SizeY, bool bIsFullscreen, EPixelFormat PreferredPixelFormat) final override;
 	virtual void RHIResizeViewport(FRHIViewport* Viewport, uint32 SizeX, uint32 SizeY, bool bIsFullscreen) final override;
 	virtual EPixelFormat RHIPreferredPixelFormatHint(EPixelFormat PreferredPixelFormat) final override;
diff --git a/Engine/Source/Runtime/RHI/Private/GPUProfiler.cpp b/Engine/Source/Runtime/RHI/Private/GPUProfiler.cpp
--- a/Engine/Source/Runtime/RHI/Private/GPUProfiler.cpp
+++ b/Engine/Source/Runtime/RHI/Private/GPUProfiler.cpp
@@ -2404,7 +2404,9 @@ namespace UE::RHI::GPUProfiler
 			if (MaxQueueBusyCycles.IsSet())
 			{
 				// Set the total GPU time stat according to the value directly provided by the platform RHI
+			PRAGMA_DISABLE_DEPRECATION_WARNINGS
 				GGPUFrameTime = FPlatformMath::TruncToInt(FPlatformTime::ToSeconds64(*MaxQueueBusyCycles) / FPlatformTime::GetSecondsPerCycle());
+			PRAGMA_ENABLE_DEPRECATION_WARNINGS
 			}
 			else
 			{
@@ -2417,7 +2419,9 @@ namespace UE::RHI::GPUProfiler
 				uint64 WholeFrameUnion = FTimestampStream::ComputeUnion(StreamPointers);
 
 				// Update the global GPU frame time stats - need to convert to Cycles32 rather than Cycles64.
+			PRAGMA_DISABLE_DEPRECATION_WARNINGS
 				GGPUFrameTime = FPlatformMath::TruncToInt(FPlatformTime::ToSeconds64(WholeFrameUnion) / FPlatformTime::GetSecondsPerCycle());
+			PRAGMA_ENABLE_DEPRECATION_WARNINGS
 			}
 
 			// @todo set global csv GPU time
diff --git a/Engine/Source/Runtime/RHI/Public/DynamicRHI.h b/Engine/Source/Runtime/RHI/Public/DynamicRHI.h
--- a/Engine/Source/Runtime/RHI/Public/DynamicRHI.h
+++ b/Engine/Source/Runtime/RHI/Public/DynamicRHI.h
@@ -728,12 +728,6 @@ public:
 	// FlushType: Flush RHI Thread
 	virtual void RHIFlushResources() = 0;
 
-	/*
-	* Returns the total GPU time taken to render the last frame. Same metric as FPlatformTime::Cycles().
-	*/
-	// FlushType: Thread safe
-	virtual uint32 RHIGetGPUFrameCycles(uint32 GPUIndex = 0) = 0;
-
 	//  must be called from the main thread.
 	// FlushType: Thread safe
 	virtual FViewportRHIRef RHICreateViewport(void* WindowHandle, uint32 SizeX, uint32 SizeY, bool bIsFullscreen, EPixelFormat PreferredPixelFormat) = 0;
@@ -1418,7 +1412,9 @@ FORCEINLINE void RHIAdvanceFrameForGetViewportBackBuffer(FRHIViewport* Viewport)
 
 FORCEINLINE uint32 RHIGetGPUFrameCycles(uint32 GPUIndex = 0)
 {
-	return GDynamicRHI->RHIGetGPUFrameCycles(GPUIndex);
+PRAGMA_DISABLE_DEPRECATION_WARNINGS
+	return GGPUFrameTime;
+PRAGMA_ENABLE_DEPRECATION_WARNINGS
 }
 
 FORCEINLINE FViewportRHIRef RHICreateViewport(void* WindowHandle, uint32 SizeX, uint32 SizeY, bool bIsFullscreen, EPixelFormat PreferredPixelFormat)
diff --git a/Engine/Source/Runtime/RHI/Public/RHICommandList.h b/Engine/Source/Runtime/RHI/Public/RHICommandList.h
--- a/Engine/Source/Runtime/RHI/Public/RHICommandList.h
+++ b/Engine/Source/Runtime/RHI/Public/RHICommandList.h
@@ -5355,6 +5355,7 @@ public:
 	UE_DEPRECATED(5.5, "FlushPendingDeletes is deprecated and removed. Use RHICmdList.ImmediateFlush(EImmediateFlushType::FlushRHIThreadFlushResources)")
 	inline int32 FlushPendingDeletes() { return 0; }
 	
+	UE_DEPRECATED(5.6, "GetGPUFrameCycles is deprecated. Use the global scope RHIGetGPUFrameCycles() function.")
 	FORCEINLINE uint32 GetGPUFrameCycles()
 	{
 		return RHIGetGPUFrameCycles(GetGPUMask().ToIndex());
diff --git a/Engine/Source/Runtime/RHI/Public/RHIGlobals.h b/Engine/Source/Runtime/RHI/Public/RHIGlobals.h
--- a/Engine/Source/Runtime/RHI/Public/RHIGlobals.h
+++ b/Engine/Source/Runtime/RHI/Public/RHIGlobals.h
@@ -747,6 +747,7 @@ struct FRHIGlobals
 extern RHI_API FRHIGlobals GRHIGlobals;
 
 /** The GPU time taken to render the last frame. Same metric as FPlatformTime::Cycles(). */
+UE_DEPRECATED(5.6, "Direct use of GGPUFrameTime is deprecated. Call the global scope RHIGetGPUFrameCycles() function instead.")
 extern RHI_API uint32 GGPUFrameTime;
 
 //
diff --git a/Engine/Source/Runtime/RHI/Public/RHIValidation.h b/Engine/Source/Runtime/RHI/Public/RHIValidation.h
--- a/Engine/Source/Runtime/RHI/Public/RHIValidation.h
+++ b/Engine/Source/Runtime/RHI/Public/RHIValidation.h
@@ -805,15 +805,6 @@ public:
 		RHI->RHIFlushResources();
 	}
 
-	/*
-	* Returns the total GPU time taken to render the last frame. Same metric as FPlatformTime::Cycles().
-	*/
-	// FlushType: Thread safe
-	virtual uint32 RHIGetGPUFrameCycles(uint32 GPUIndex = 0) override final
-	{
-		return RHI->RHIGetGPUFrameCycles(GPUIndex);
-	}
-
 	//  must be called from the main thread.
 	// FlushType: Thread safe
 	virtual FViewportRHIRef RHICreateViewport(void* WindowHandle, uint32 SizeX, uint32 SizeY, bool bIsFullscreen, EPixelFormat PreferredPixelFormat) override final
diff --git a/Engine/Source/Runtime/VulkanRHI/Private/Android/VulkanAndroidPlatform.cpp b/Engine/Source/Runtime/VulkanRHI/Private/Android/VulkanAndroidPlatform.cpp
--- a/Engine/Source/Runtime/VulkanRHI/Private/Android/VulkanAndroidPlatform.cpp
+++ b/Engine/Source/Runtime/VulkanRHI/Private/Android/VulkanAndroidPlatform.cpp
@@ -36,7 +36,9 @@ namespace AndroidVulkan
 	{
 		const double GPUTimeInSeconds = (double)gpu_time_ns / 1000000000.0;
 
+	PRAGMA_DISABLE_DEPRECATION_WARNINGS
 		GGPUFrameTime = GPUTimeInSeconds / FPlatformTime::GetSecondsPerCycle64();
+	PRAGMA_ENABLE_DEPRECATION_WARNINGS
 	}
 
 	void SetSwappyPostWaitCallback()
diff --git a/Engine/Source/Runtime/VulkanRHI/Private/VulkanCommands.cpp b/Engine/Source/Runtime/VulkanRHI/Private/VulkanCommands.cpp
--- a/Engine/Source/Runtime/VulkanRHI/Private/VulkanCommands.cpp
+++ b/Engine/Source/Runtime/VulkanRHI/Private/VulkanCommands.cpp
@@ -869,12 +869,6 @@ bool FVulkanDynamicRHI::RHIIsRenderingSuspended()
 	return false;
 }
 
-uint32 FVulkanDynamicRHI::RHIGetGPUFrameCycles(uint32 GPUIndex)
-{
-	check(GPUIndex == 0);
-	return GGPUFrameTime;
-}
-
 void FVulkanCommandListContext::RHISetDepthBounds(float MinDepth, float MaxDepth)
 {
 	FVulkanCommandBuffer& CommandBuffer = GetCommandBuffer();
diff --git a/Engine/Source/Runtime/VulkanRHI/Private/VulkanDynamicRHI.h b/Engine/Source/Runtime/VulkanRHI/Private/VulkanDynamicRHI.h
--- a/Engine/Source/Runtime/VulkanRHI/Private/VulkanDynamicRHI.h
+++ b/Engine/Source/Runtime/VulkanRHI/Private/VulkanDynamicRHI.h
@@ -219,7 +219,6 @@ public:
 	virtual FTextureRHIRef RHICreateAliasedTexture(FTextureRHIRef& SourceTexture) final override;
 	virtual void RHIAdvanceFrameForGetViewportBackBuffer(FRHIViewport* Viewport) final override;
 	virtual void RHIFlushResources() final override;
-	virtual uint32 RHIGetGPUFrameCycles(uint32 GPUIndex = 0) final override;
 	virtual FViewportRHIRef RHICreateViewport(void* WindowHandle, uint32 SizeX, uint32 SizeY, bool bIsFullscreen, EPixelFormat PreferredPixelFormat) final override;
 	virtual void RHIResizeViewport(FRHIViewport* Viewport, uint32 SizeX, uint32 SizeY, bool bIsFullscreen) final override;
 	virtual void RHIResizeViewport(FRHIViewport* Viewport, uint32 SizeX, uint32 SizeY, bool bIsFullscreen, EPixelFormat PreferredPixelFormat) final override;
diff --git a/Engine/Source/Runtime/VulkanRHI/Private/VulkanSubmission.cpp b/Engine/Source/Runtime/VulkanRHI/Private/VulkanSubmission.cpp
--- a/Engine/Source/Runtime/VulkanRHI/Private/VulkanSubmission.cpp
+++ b/Engine/Source/Runtime/VulkanRHI/Private/VulkanSubmission.cpp
@@ -664,7 +664,9 @@ void FVulkanDynamicRHI::CompletePayload(FVulkanPayload* Payload)
 		{
 			const double SecondsPerCycle = FPlatformTime::GetSecondsPerCycle();
 			const double Frequency = double((uint64)((1000.0 * 1000.0 * 1000.0) / Device->GetLimits().timestampPeriod));
+		PRAGMA_DISABLE_DEPRECATION_WARNINGS
 			GGPUFrameTime = FMath::TruncToInt(double(BusyCycles) / Frequency / SecondsPerCycle);
+		PRAGMA_ENABLE_DEPRECATION_WARNINGS
 			BusyCycles = 0;
 		}
 	}
diff --git a/Engine/Source/Runtime/Windows/D3D11RHI/Private/D3D11Commands.cpp b/Engine/Source/Runtime/Windows/D3D11RHI/Private/D3D11Commands.cpp
--- a/Engine/Source/Runtime/Windows/D3D11RHI/Private/D3D11Commands.cpp
+++ b/Engine/Source/Runtime/Windows/D3D11RHI/Private/D3D11Commands.cpp
@@ -1321,15 +1321,6 @@ void FD3D11DynamicRHI::RHIBlockUntilGPUIdle()
 	}
 }
 
-/**
- * Returns the total GPU time taken to render the last frame. Same metric as FPlatformTime::Cycles().
- */
-uint32 FD3D11DynamicRHI::RHIGetGPUFrameCycles(uint32 GPUIndex)
-{
-	check(GPUIndex == 0);
-	return GGPUFrameTime;
-}
-
 // NVIDIA Depth Bounds Test interface
 void FD3D11DynamicRHI::EnableDepthBoundsTest(bool bEnable,float MinDepth,float MaxDepth)
 {
diff --git a/Engine/Source/Runtime/Windows/D3D11RHI/Private/D3D11RHI.cpp b/Engine/Source/Runtime/Windows/D3D11RHI/Private/D3D11RHI.cpp
--- a/Engine/Source/Runtime/Windows/D3D11RHI/Private/D3D11RHI.cpp
+++ b/Engine/Source/Runtime/Windows/D3D11RHI/Private/D3D11RHI.cpp
@@ -407,11 +407,15 @@ void FD3DGPUProfiler::EndFrame()
 	{
 		uint64 GPUTiming = FrameTiming.GetTiming();
 		uint64 GPUFreq = FrameTiming.GetTimingFrequency();
+	PRAGMA_DISABLE_DEPRECATION_WARNINGS
 		GGPUFrameTime = FMath::TruncToInt( double(GPUTiming) / double(GPUFreq) / FPlatformTime::GetSecondsPerCycle() );
+	PRAGMA_ENABLE_DEPRECATION_WARNINGS
 	}
 	else
 	{
+	PRAGMA_DISABLE_DEPRECATION_WARNINGS
 		GGPUFrameTime = 0;
+	PRAGMA_ENABLE_DEPRECATION_WARNINGS
 	}
 
 	// if we have a frame open, close it now.
diff --git a/Engine/Source/Runtime/Windows/D3D11RHI/Private/D3D11RHIPrivate.h b/Engine/Source/Runtime/Windows/D3D11RHI/Private/D3D11RHIPrivate.h
--- a/Engine/Source/Runtime/Windows/D3D11RHI/Private/D3D11RHIPrivate.h
+++ b/Engine/Source/Runtime/Windows/D3D11RHI/Private/D3D11RHIPrivate.h
@@ -499,7 +499,6 @@ public:
 	virtual FTextureRHIRef RHICreateAliasedTexture(FTextureRHIRef& SourceTexture) final override;
 	virtual void RHIAdvanceFrameForGetViewportBackBuffer(FRHIViewport* Viewport) final override;
 	virtual void RHIFlushResources() final override;
-	virtual uint32 RHIGetGPUFrameCycles(uint32 GPUIndex = 0) final override;
 	virtual FViewportRHIRef RHICreateViewport(void* WindowHandle, uint32 SizeX, uint32 SizeY, bool bIsFullscreen, EPixelFormat PreferredPixelFormat) override;
 	virtual void RHIResizeViewport(FRHIViewport* Viewport, uint32 SizeX, uint32 SizeY, bool bIsFullscreen) final override;
 	virtual void RHIResizeViewport(FRHIViewport* Viewport, uint32 SizeX, uint32 SizeY, bool bIsFullscreen, EPixelFormat PreferredPixelFormat) final override;
