[MIR] Fix typo and null pointer exception when emitting InlineHLSL in new material translator.

#rnx

diff --git a/Engine/Source/Runtime/Engine/Private/Materials/MaterialIREmitter.cpp b/Engine/Source/Runtime/Engine/Private/Materials/MaterialIREmitter.cpp
--- a/Engine/Source/Runtime/Engine/Private/Materials/MaterialIREmitter.cpp
+++ b/Engine/Source/Runtime/Engine/Private/Materials/MaterialIREmitter.cpp
@@ -2700,7 +2700,8 @@ static FValueRef EmitInlineHLSL(FEmitter& Emitter, const FType* Type, const FMat
 		Prototype.NumArguments = InputValues.Num();
 		for (int32 i = 0; i < InputValues.Num(); ++i)
 		{
-			Prototype.NumArguments = InputValues[i];
+			checkf(InputValues[i], TEXT("InputValues[%d] must not be null when InlineHLSL-instruction is emitted"), i);
+			Prototype.Arguments[i] = InputValues[i];
 		}
 	}
 
