Occlusion changes to ActiveSound & AudioComponent
ActiveSound:
+ Change OcclusionInterval's behavior.  Negative values only test for occlusion once on initial playback.  Non-negative values retest occlusion after the interval elapses.
+ Added OcclusionInterval as a per-platform float to audio settings.  Defaulted to 0.1 for all platforms.

AudioComponent:
+ Upgrade OcclusionInterval to be a uproperty w/ config
+ Send OcclusionInterval to active sound during creation (on playback)

diff --git a/Engine/Source/Runtime/Engine/Classes/Components/AudioComponent.h b/Engine/Source/Runtime/Engine/Classes/Components/AudioComponent.h
--- a/Engine/Source/Runtime/Engine/Classes/Components/AudioComponent.h
+++ b/Engine/Source/Runtime/Engine/Classes/Components/AudioComponent.h
@@ -392,7 +392,8 @@ public:
 	UPROPERTY(EditAnywhere, BlueprintReadWrite, Category = Concurrency)
 	TSet<TObjectPtr<USoundConcurrency>> ConcurrencySet;
 
-	/** While playing, this component will check for occlusion from its closest listener every this many seconds */
+	/** The interval at which to re-test occlusion (in seconds).  Values < 0.f will only check occlusion once when we start playback */
+	UPROPERTY(EditAnywhere, config, Category = Occlusion)
 	float OcclusionCheckInterval;
 
 	/** What time the audio component was told to play. Used to compute audio component state. */
diff --git a/Engine/Source/Runtime/Engine/Classes/Sound/AudioSettings.h b/Engine/Source/Runtime/Engine/Classes/Sound/AudioSettings.h
--- a/Engine/Source/Runtime/Engine/Classes/Sound/AudioSettings.h
+++ b/Engine/Source/Runtime/Engine/Classes/Sound/AudioSettings.h
@@ -4,6 +4,7 @@
 
 #include "CoreMinimal.h"
 #include "UObject/ObjectMacros.h"
+#include "UObject/PerPlatformProperties.h"
 #include "UObject/SoftObjectPath.h"
 #include "Engine/DeveloperSettings.h"
 #include "AudioMixerTypes.h"
@@ -176,6 +177,13 @@ class UAudioSettings : public UDeveloperSettings
 	UPROPERTY(config, EditAnywhere, Category = "Audio")
 	EDefaultAudioCompressionType DefaultAudioCompressionType;
 
+	/**
+	 * The default interval at which to re-test occlusion (in seconds).  Values < 0.f will only check occlusion once.
+	 *   NOTE! Will be overridden by sounds coming from an audio component!
+	 */
+	UPROPERTY(config, EditAnywhere, Category = Occlusion)
+	FPerPlatformFloat DefaultOcclusionCheckInterval = .1f;
+
 	/** The default compression quality (e.g. for new SoundWaves) */
 	UPROPERTY(config)
 	int32 DefaultCompressionQuality_DEPRECATED = 80;
diff --git a/Engine/Source/Runtime/Engine/Private/ActiveSound.cpp b/Engine/Source/Runtime/Engine/Private/ActiveSound.cpp
--- a/Engine/Source/Runtime/Engine/Private/ActiveSound.cpp
+++ b/Engine/Source/Runtime/Engine/Private/ActiveSound.cpp
@@ -21,13 +21,6 @@ FAutoConsoleVariableRef CVarAudioOcclusionEnabled(
 	TEXT("Disables (1) or enables (0) audio occlusion.\n"),
 	ECVF_Default);
 
-static float AudioOcclusionCheckIntervalCvar = 0.f;
-FAutoConsoleVariableRef CVarAudioOcclusionCheckInterval(
-	TEXT("au.OcclusionCheckInterval"),
-	AudioOcclusionCheckIntervalCvar,
-	TEXT("Time interval to recheck occlusion from the closest listener (in seconds).\n"),
-	ECVF_Default);
-
 static bool bUseAsyncOcclusionTracesCVar = false;
 FAutoConsoleVariableRef CVarUseAsyncOcclusionTraces(
 	TEXT("au.UseAsyncOcclusionTraces"),
@@ -159,7 +152,6 @@ FActiveSound::FActiveSound()
 	, Priority(1.0f)
 	, VolumeConcurrency(0.0f)
 	, MaxSourceModulationValue(-1.0f)
-	, OcclusionCheckInterval(0.0f)
 	, LastOcclusionCheckTime(TNumericLimits<float>::Lowest())
 	, MaxDistance(FAudioDevice::GetMaxWorldDistance())
 	, LastLocation(FVector::ZeroVector)
@@ -184,7 +176,9 @@ FActiveSound::FActiveSound()
 		ActiveSoundTraceDelegate.BindStatic(&OcclusionTraceDone);
 	}
 
-	OcclusionCheckInterval = AudioOcclusionCheckIntervalCvar;
+	const UAudioSettings* AudioSettings = GetDefault<UAudioSettings>();
+	check(AudioSettings);
+	OcclusionCheckInterval = AudioSettings->DefaultOcclusionCheckInterval.GetValue();
 }
 
 FActiveSound::FActiveSound(const FActiveSound&) = default;
@@ -1455,7 +1449,11 @@ void FActiveSound::CheckOcclusion(const FVector ListenerLocation, const FVector
 	}
 	else
 	{
-		if (!bAsyncOcclusionPending && (PlaybackTime - LastOcclusionCheckTime) > OcclusionCheckInterval)
+		// If our occlusion check interval is less than 0, we only want to check occlusion once.  Otherwise, retest when the interval passes.
+		const bool bOcclusionCheckDue = OcclusionCheckInterval >= 0.f && (PlaybackTime - LastOcclusionCheckTime) > OcclusionCheckInterval;
+		const bool bNeedsInitialOcclusionCheck = LastOcclusionCheckTime < 0.f;
+
+		if (!bAsyncOcclusionPending && (bNeedsInitialOcclusionCheck || bOcclusionCheckDue))
 		{
 			LastOcclusionCheckTime = PlaybackTime;
 
diff --git a/Engine/Source/Runtime/Engine/Private/Components/AudioComponent.cpp b/Engine/Source/Runtime/Engine/Private/Components/AudioComponent.cpp
--- a/Engine/Source/Runtime/Engine/Private/Components/AudioComponent.cpp
+++ b/Engine/Source/Runtime/Engine/Private/Components/AudioComponent.cpp
@@ -755,6 +755,7 @@ void UAudioComponent::PlayInternal(const PlayInternalRequestData& InPlayRequestD
 	NewActiveSound.LowPassFilterFrequency = LowPassFilterFrequency;
 	NewActiveSound.bEnableHighPassFilter = bEnableHighPassFilter;
 	NewActiveSound.HighPassFilterFrequency = HighPassFilterFrequency;
+	NewActiveSound.OcclusionCheckInterval = OcclusionCheckInterval;
 	NewActiveSound.RequestedStartTime = FMath::Max(0.f, InPlayRequestData.StartTime);
 
 	if (bOverrideSubtitlePriority)
diff --git a/Engine/Source/Runtime/Engine/Public/ActiveSound.h b/Engine/Source/Runtime/Engine/Public/ActiveSound.h
--- a/Engine/Source/Runtime/Engine/Public/ActiveSound.h
+++ b/Engine/Source/Runtime/Engine/Public/ActiveSound.h
@@ -842,11 +842,11 @@ private:
 	/** Updates ramping concurrency volume scalars */
 	ENGINE_API void UpdateConcurrencyVolumeScalars(const float DeltaTime);
 
-	/** if OcclusionCheckInterval > 0.0, checks if the sound has become (un)occluded during playback
-	 * and calls eventOcclusionChanged() if so
-	 * primarily used for gameplay-relevant ambient sounds
-	 * CurrentLocation is the location of this component that will be used for playback
-	 * @param ListenerLocation location of the closest listener to the sound
+	/** if OcclusionCheckInterval >= 0.0, periodically checks if the sound has become (un)occluded during playback
+	 * When the interval is < 0.0, perform the initial occlusion check, but do not retest occlusion.
+	 * @param ListenerLocation - location of the closest listener to the sound
+	 * @param SoundLocation - location of the sound
+	 * @param AttenuationSettingsPtr - Attenuation settings used (subject to various overrides)
 	 */
 	ENGINE_API void CheckOcclusion(const FVector ListenerLocation, const FVector SoundLocation, const FSoundAttenuationSettings* AttenuationSettingsPtr);
 
