[AutoRTFM] Replace include of "AutoRTFM/AutoRTFM.h" with "AutoRTFM.h"

"AutoRTFM/AutoRTFM.h" is a temporary header that is used to aid in the migration of AutoRTFM out of core.
We want to remove all references to this so we can delete the header.

diff --git a/Engine/Source/Runtime/Engine/Private/Player.cpp b/Engine/Source/Runtime/Engine/Private/Player.cpp
--- a/Engine/Source/Runtime/Engine/Private/Player.cpp
+++ b/Engine/Source/Runtime/Engine/Private/Player.cpp
@@ -17,7 +17,7 @@
 #include "GameFramework/CheatManager.h"
 #include "GameFramework/GameStateBase.h"
 
-#include "AutoRTFM/AutoRTFM.h"
+#include "AutoRTFM.h"
 
 #include UE_INLINE_GENERATED_CPP_BY_NAME(Player)
 
@@ -177,4 +177,4 @@ void UPlayer::SwitchController(class APlayerController* PC)
 
 void UPlayer::ReceivedPlayerController(class APlayerController* NewController)
 {
-}
\ No newline at end of file
+}
