Fix for install debug command.

diff --git a/Engine/Source/Runtime/Engine/Classes/Engine/AssetManager.h b/Engine/Source/Runtime/Engine/Classes/Engine/AssetManager.h
--- a/Engine/Source/Runtime/Engine/Classes/Engine/AssetManager.h
+++ b/Engine/Source/Runtime/Engine/Classes/Engine/AssetManager.h
@@ -49,6 +49,8 @@ struct FAssetManagerDownloadParams
 	FStreamableDownloadPriority Priority = UE::StreamableManager::Private::DownloadDefaultPriority;
 	/** If set, any soft references will also be installed */
 	bool bInstallSoftReferences = false;
+
+	bool operator==(const FAssetManagerDownloadParams& Right) const = default;
 };
 
 struct FAssetManagerLoadParams
diff --git a/Engine/Source/Runtime/Engine/Private/AssetManager.cpp b/Engine/Source/Runtime/Engine/Private/AssetManager.cpp
--- a/Engine/Source/Runtime/Engine/Private/AssetManager.cpp
+++ b/Engine/Source/Runtime/Engine/Private/AssetManager.cpp
@@ -63,6 +63,12 @@ struct FPrimaryAssetLoadState
 
 	/** The set of bundles to be loaded by the handle */
 	TArray<FName> BundleNames;
+	TOptional<FAssetManagerDownloadParams> DownloadParams;
+
+	bool IsSame(TConstArrayView<FName> InBundleNames, const TOptional<FAssetManagerDownloadParams>& InDownloadParams) const
+	{
+		return Algo::Compare(InBundleNames, BundleNames) && (InDownloadParams == DownloadParams);
+	}
 
 	/** If this state is keeping things in memory */
 	bool IsValid() const { return Handle.IsValid() && Handle->IsActive(); }
@@ -81,6 +87,7 @@ struct FPrimaryAssetLoadState
 			Handle = nullptr;
 		}
 		BundleNames.Reset();
+		DownloadParams.Reset();
 	}
 };
 
@@ -1853,7 +1860,7 @@ TSharedPtr<FStreamableHandle> UAssetManager::ChangeBundleStateForPrimaryAssets(
 			// If the pending state is valid, check if it is different
 			if (NameData->PendingState.IsValid())
 			{
-				if (Algo::Compare(NameData->PendingState.BundleNames, NewBundleState))
+				if (NameData->PendingState.IsSame(NewBundleState, LoadParams.DownloadParams))
 				{
 					// This will wait on any existing handles to finish
 					AllHandles.Add(NameData->PendingState.Handle);
@@ -1863,7 +1870,7 @@ TSharedPtr<FStreamableHandle> UAssetManager::ChangeBundleStateForPrimaryAssets(
 				// Clear pending state
 				NameData->PendingState.Reset(true);
 			}
-			else if (NameData->CurrentState.IsValid() && Algo::Compare(NameData->CurrentState.BundleNames, NewBundleState))
+			else if (NameData->CurrentState.IsValid() && NameData->CurrentState.IsSame(NewBundleState, LoadParams.DownloadParams))
 			{
 				// If no pending, compare with current
 				continue;
@@ -1901,6 +1908,7 @@ TSharedPtr<FStreamableHandle> UAssetManager::ChangeBundleStateForPrimaryAssets(
 			{
 				// New state has no assets to load. Set the CurrentState's bundles and clear the handle
 				NameData->CurrentState.BundleNames = NewBundleState;
+				NameData->CurrentState.DownloadParams = CopyTemp(*LoadParams.DownloadParams);
 				NameData->CurrentState.Handle.Reset();
 				continue;
 			}
@@ -1929,14 +1937,12 @@ TSharedPtr<FStreamableHandle> UAssetManager::ChangeBundleStateForPrimaryAssets(
 			// Set pending bundle state before issuing the request, because it can be resolved immediately.
 			// OnAssetStateChangeCompleted requires this data to handle state transition correctly.
 			NameData->PendingState.BundleNames = NewBundleState;
+			NameData->PendingState.DownloadParams = CopyTemp(LoadParams.DownloadParams);
 
 			FAssetManagerLoadParams AssetListLoadParams;
 			AssetListLoadParams.Priority = LoadParams.Priority;
-			if (LoadParams.DownloadParams)
-			{
-				AssetListLoadParams.DownloadParams = CopyTemp(*LoadParams.DownloadParams);
-			}
-
+			AssetListLoadParams.DownloadParams = CopyTemp(LoadParams.DownloadParams);
+			
 			NewHandle = LoadAssetList(MoveTemp(PathsToLoad), MoveTemp(AssetListLoadParams), FString(DebugName), CopyTemp(Location));
 			if (NewHandle && !NewHandle->HasLoadCompleted())
 			{
@@ -2125,6 +2131,7 @@ void UAssetManager::OnAssetStateChangeCompleted(TSharedPtr<FStreamableHandle> Bo
 			else
 			{
 				NameData->CurrentState.Handle = BoundHandle;
+				NameData->CurrentState.DownloadParams = NameData->PendingState.DownloadParams;
 				NameData->CurrentState.BundleNames = NameData->PendingState.BundleNames;
 				
 				NotifyOnAssetStateChangeCompleted(PrimaryAssetId);
