Greatly reduce the number of Quartz Subscribers created at runtime. Unsubscribe handles to deleted clocks as early as possible.
#swarm https://p4-swarm.epicgames.net/reviews/41854345

diff --git a/Engine/Source/Runtime/AudioMixer/Private/Quartz/AudioMixerClockHandle.cpp b/Engine/Source/Runtime/AudioMixer/Private/Quartz/AudioMixerClockHandle.cpp
--- a/Engine/Source/Runtime/AudioMixer/Private/Quartz/AudioMixerClockHandle.cpp
+++ b/Engine/Source/Runtime/AudioMixer/Private/Quartz/AudioMixerClockHandle.cpp
@@ -27,6 +27,11 @@ void UQuartzClockHandle::BeginDestroy()
 	RawHandle.SendCommandToClock([Subscriber](Audio::FQuartzClock* InClock) { InClock->UnsubscribeFromAllTimeDivisions(Subscriber); });
 }
 
+bool UQuartzClockHandle::ShouldUnsubscribe()
+{
+	return !RawHandle.IsValid();
+}
+
 void UQuartzClockHandle::StartClock(const UObject* WorldContextObject, UQuartzClockHandle*& ClockHandle)
 {
 	ClockHandle = this;
diff --git a/Engine/Source/Runtime/AudioMixer/Public/Quartz/AudioMixerClockHandle.h b/Engine/Source/Runtime/AudioMixer/Public/Quartz/AudioMixerClockHandle.h
--- a/Engine/Source/Runtime/AudioMixer/Public/Quartz/AudioMixerClockHandle.h
+++ b/Engine/Source/Runtime/AudioMixer/Public/Quartz/AudioMixerClockHandle.h
@@ -39,6 +39,13 @@ public:
 	AUDIOMIXER_API void BeginDestroy() override;
 	// end UObject interface
 
+private:
+	// begin FQuartzTickableObject interface
+	AUDIOMIXER_API virtual bool ShouldUnsubscribe() override;
+	// end FQuartzTickableObject interface
+	
+public:
+
 // Begin Blueprint Interface
 
 	// Clock manipulation
diff --git a/Engine/Source/Runtime/Engine/Classes/Components/AudioComponent.h b/Engine/Source/Runtime/Engine/Classes/Components/AudioComponent.h
--- a/Engine/Source/Runtime/Engine/Classes/Components/AudioComponent.h
+++ b/Engine/Source/Runtime/Engine/Classes/Components/AudioComponent.h
@@ -975,8 +975,6 @@ protected:
 
 	// Used for processing queue commands
 	//~ Begin FQuartzTickableObject
-	ENGINE_API virtual void ProcessCommand(const Audio::FQuartzQuantizedCommandDelegateData& Data) override;
-	virtual void ProcessCommand(const Audio::FQuartzMetronomeDelegateData& Data) override {};
 	ENGINE_API virtual void ProcessCommand(const Audio::FQuartzQueueCommandData& InQueueCommandData) override;
 	//~ End FQuartzTickableObject
 
diff --git a/Engine/Source/Runtime/Engine/Classes/Sound/QuartzSubscription.cpp b/Engine/Source/Runtime/Engine/Classes/Sound/QuartzSubscription.cpp
--- a/Engine/Source/Runtime/Engine/Classes/Sound/QuartzSubscription.cpp
+++ b/Engine/Source/Runtime/Engine/Classes/Sound/QuartzSubscription.cpp
@@ -154,6 +154,11 @@ Audio::FQuartzGameThreadSubscriber FQuartzTickableObject::GetQuartzSubscriber()
 void FQuartzTickableObject::QuartzTick(float DeltaTime)
 {
 	CommandQueuePtr->PumpCommandQueue(*this);
+	
+	if (ShouldUnsubscribe())
+	{
+		QuartzUnsubscribe();	
+	}
 }
 
 bool FQuartzTickableObject::QuartzIsTickable() const
diff --git a/Engine/Source/Runtime/Engine/Classes/Sound/QuartzSubscription.h b/Engine/Source/Runtime/Engine/Classes/Sound/QuartzSubscription.h
--- a/Engine/Source/Runtime/Engine/Classes/Sound/QuartzSubscription.h
+++ b/Engine/Source/Runtime/Engine/Classes/Sound/QuartzSubscription.h
@@ -90,6 +90,7 @@ public:
 	virtual void ProcessCommand(const Audio::FQuartzQuantizedCommandDelegateData& Data) {}
 	virtual void ProcessCommand(const Audio::FQuartzMetronomeDelegateData& Data) {}
 	virtual void ProcessCommand(const Audio::FQuartzQueueCommandData& Data) {}
+	
 
 	const Audio::FQuartzOffset& GetQuartzOffset() const { return NotificationOffset; }
 
@@ -98,6 +99,7 @@ protected:
 	UE_API void SetNotificationAnticipationAmountMusicalDuration(const EQuartzCommandQuantization Duration,  const double Multiplier);
 
 	UE_API void QuartzUnsubscribe();
+	UE_API virtual bool ShouldUnsubscribe() { return false; } 
 
 private:
 	struct FMetronomeDelegateGameThreadData
diff --git a/Engine/Source/Runtime/Engine/Private/Components/AudioComponent.cpp b/Engine/Source/Runtime/Engine/Private/Components/AudioComponent.cpp
--- a/Engine/Source/Runtime/Engine/Private/Components/AudioComponent.cpp
+++ b/Engine/Source/Runtime/Engine/Private/Components/AudioComponent.cpp
@@ -463,10 +463,6 @@ void UAudioComponent::Play(float StartTime)
 	PlayInternal(InternalRequestData);
 }
 
-void UAudioComponent::ProcessCommand(const Audio::FQuartzQuantizedCommandDelegateData& Data)
-{
-}
-
 void UAudioComponent::ProcessCommand(const Audio::FQuartzQueueCommandData& InQueueCommandData)
 {
 	PlayQueuedQuantizedInternal(GetWorld(), InQueueCommandData.AudioComponentCommandInfo);
@@ -483,13 +479,13 @@ void UAudioComponent::PlayQuantized(
 	, EAudioFaderCurve InFadeCurve)
 {
 	//Initialize the tickable object portion of the Audio Component, if it hasn't been initialized already
-	if (!FQuartzTickableObject::IsInitialized())
+	if (!bAlwaysTakeVoiceSlotCVar && !FQuartzTickableObject::IsInitialized())
 	{
 		// if the WorldContextObject is null, attempt to fallback on our world
 		// (call to Init() will ensure on a valid world)
 		Init(WorldContextObject ? WorldContextObject->GetWorld() : GetWorld());
+		check(FQuartzTickableObject::IsInitialized());
 	}
-	check(FQuartzTickableObject::IsInitialized());
 
 	int32 TimeCVarVal = FMath::Clamp(CVarTimeToTakeUpVoiceSlot->GetInt(), 0, (int32)EQuartzCommandQuantization::Count - 1);
 	EQuartzCommandQuantization MinimumQuantization = EQuartzCommandQuantization(TimeCVarVal);
@@ -551,9 +547,6 @@ void UAudioComponent::PlayQuantized(
 
 void UAudioComponent::PlayQueuedQuantizedInternal(const UObject* WorldContextObject, FAudioComponentCommandInfo InCommandInfo)
 {
-	// confirm the FQuartzTickableObject has been initialized.
-	ensure(FQuartzTickableObject::IsInitialized());
-
 	bool bFoundQuantizedCommand = false;
 	bool bIsValidCommand = true;
 	UQuartzClockHandle* Handle = nullptr; // will be retrieved from PendingQuartzCommandData
