New Virtualization Mode which Seeks to Estimated Start Time when Realizing

diff --git a/Engine/Source/Runtime/Engine/Classes/Sound/SoundBase.h b/Engine/Source/Runtime/Engine/Classes/Sound/SoundBase.h
--- a/Engine/Source/Runtime/Engine/Classes/Sound/SoundBase.h
+++ b/Engine/Source/Runtime/Engine/Classes/Sound/SoundBase.h
@@ -68,7 +68,12 @@ enum class EVirtualizationMode : uint8
 	PlayWhenSilent,
 
 	/** If sound is looping, sound restarts from beginning upon realization from being virtual */
-	Restart
+	Restart,
+
+	/** Sound fully virtualizes, but attempts to seek to the proper part of its looping duration when realized.
+	 * The playback behavior is similar to PlayWhenSilent, but does not take up a voice slot.
+	 */
+	SeekRestart
 };
 
 /**
diff --git a/Engine/Source/Runtime/Engine/Private/ActiveSound.cpp b/Engine/Source/Runtime/Engine/Private/ActiveSound.cpp
--- a/Engine/Source/Runtime/Engine/Private/ActiveSound.cpp
+++ b/Engine/Source/Runtime/Engine/Private/ActiveSound.cpp
@@ -75,6 +75,14 @@ FAutoConsoleVariableRef CVarAudioLinkSpatializeAllSounds(
 	TEXT("0 (off), !=0 (on)"),
 	ECVF_Default);
 
+static int32 OverrideSeekStartTimeCVar = 0;
+FAutoConsoleVariableRef CVarOverrideSeekStartTime(
+	TEXT("au.Virtualization.OverrideSeekStartTimeMode"),
+	OverrideSeekStartTimeCVar,
+	TEXT("Optionally force all sounds with virtualization mode set to 'Seek Start Time' to another Virtualization Mode.\n")
+	TEXT("0: Leave as-is at 'Seek Start Time', 1: Set to 'Disabled', 2: Set to 'Play When Silent', 3: Set to 'Restart'"),
+	ECVF_Default);
+
 FTraceDelegate FActiveSound::ActiveSoundTraceDelegate;
 TMap<FTraceHandle, FActiveSound::FAsyncTraceDetails> FActiveSound::TraceToActiveSoundMap;
 
@@ -398,13 +406,8 @@ bool FActiveSound::IsPlayWhenSilent() const
 	{
 		return false;
 	}
-
-	if (bEnableVirtualizationModeOverride)
-	{
-		return VirtualizationModeOverride == EVirtualizationMode::PlayWhenSilent;
-	}
 	
-	return Sound && Sound->IsPlayWhenSilent();
+	return GetVirtualizationMode() == EVirtualizationMode::PlayWhenSilent;
 }
 
 float FActiveSound::GetConcurrencyPriority() const
@@ -421,17 +424,36 @@ float FActiveSound::GetConcurrencyPriority() const
 
 EVirtualizationMode FActiveSound::GetVirtualizationMode() const
 {
+	EVirtualizationMode Mode = EVirtualizationMode::Disabled;
+
+	if (Sound)
+	{
+		Mode = Sound->VirtualizationMode;
+	}
+	
 	if (bEnableVirtualizationModeOverride)
 	{
-		return VirtualizationModeOverride;
+		Mode = VirtualizationModeOverride;
 	}
 
-	if (Sound)
+	if (Mode == EVirtualizationMode::SeekRestart)
 	{
-		return Sound->VirtualizationMode;
+		switch (OverrideSeekStartTimeCVar)
+		{
+		case 1:
+			Mode = EVirtualizationMode::Disabled;
+			break;
+		case 2:
+			Mode = EVirtualizationMode::PlayWhenSilent;
+			break;
+		case 3:
+			Mode = EVirtualizationMode::Restart;
+			break;
+		default:
+			break;
+		}
 	}
-
-	return EVirtualizationMode::Disabled;
+	return Mode;
 }
 
 void FActiveSound::OverrideVirtualizationMode(EVirtualizationMode Mode)
diff --git a/Engine/Source/Runtime/Engine/Private/AudioDevice.cpp b/Engine/Source/Runtime/Engine/Private/AudioDevice.cpp
--- a/Engine/Source/Runtime/Engine/Private/AudioDevice.cpp
+++ b/Engine/Source/Runtime/Engine/Private/AudioDevice.cpp
@@ -5378,6 +5378,12 @@ void FAudioDevice::AddNewActiveSoundInternal(const FActiveSound& InNewActiveSoun
 
 		// Clear must be called after AddSoundToStop to ensure AudioComponent is properly removed from AudioComponentIDToActiveSoundMap
 		VirtualActiveSound.ClearAudioComponent();
+
+		// Attempt to seek the ActiveSound based on virtualization time
+		if (ActiveSound->GetVirtualizationMode() == EVirtualizationMode::SeekRestart)
+		{
+			ActiveSound->RequestedStartTime = ActiveSound->PlaybackTime;
+		}
 	}
 
 	InitSoundParams(*ActiveSound, MoveTemp(InDefaultParams), InVirtualLoopToRetrigger != nullptr /*bIsVirtualLoopRealzing*/);
diff --git a/Engine/Source/Runtime/Engine/Private/SoundWave.cpp b/Engine/Source/Runtime/Engine/Private/SoundWave.cpp
--- a/Engine/Source/Runtime/Engine/Private/SoundWave.cpp
+++ b/Engine/Source/Runtime/Engine/Private/SoundWave.cpp
@@ -3512,7 +3512,17 @@ void USoundWave::Parse(FAudioDevice* AudioDevice, const UPTRINT NodeWaveInstance
 	WaveInstance->AttenuationHighpassFilterFrequency = ParseParams.AttenuationHighpassFilterFrequency;
 	WaveInstance->AmbientZoneFilterFrequency = ParseParams.AmbientZoneFilterFrequency;
 	WaveInstance->bApplyRadioFilter = ActiveSound.bApplyRadioFilter;
-	WaveInstance->StartTime = ParseParams.StartTime;
+
+	// Procedural sounds do not have a duration that  makes sense to fmod with
+	if (!bProcedural && IsLooping())
+	{
+		WaveInstance->StartTime = FMath::Fmod(ParseParams.StartTime, Duration);
+	}
+	else
+	{
+		WaveInstance->StartTime = ParseParams.StartTime;
+	}
+
 	WaveInstance->UserIndex = ActiveSound.UserIndex;
 	WaveInstance->NonSpatializedRadiusStart = ParseParams.NonSpatializedRadiusStart;
 	WaveInstance->NonSpatializedRadiusEnd = ParseParams.NonSpatializedRadiusEnd;
