RemoteObjects:  Quick fix for compilation warning (uint8 bitfield converted to bool).

#rnx

diff --git a/Engine/Source/Runtime/Engine/Private/Components/ActorComponent.cpp b/Engine/Source/Runtime/Engine/Private/Components/ActorComponent.cpp
--- a/Engine/Source/Runtime/Engine/Private/Components/ActorComponent.cpp
+++ b/Engine/Source/Runtime/Engine/Private/Components/ActorComponent.cpp
@@ -706,8 +706,8 @@ void UActorComponent::PostMigrate(const struct FUObjectMigrationContext& Migrati
 	};
 	static FUObjectAnnotationSparse<FLoanedState,/*bAutoRemove=*/true> RemoteObjects_LoanedStates;
 	FLoanedState NewState {
-		.bWasRegistered = bRegistered, .bWasDestroyed = IsBeingDestroyed(), .bWasPhysicsStateCreated = bPhysicsStateCreated,
-		.bWasTickFunctionRegistered = bTickFunctionsRegistered, .bWasComponentTickEnabled = IsComponentTickEnabled()
+		.bWasRegistered = !!bRegistered, .bWasDestroyed = IsBeingDestroyed(), .bWasPhysicsStateCreated = !!bPhysicsStateCreated,
+		.bWasTickFunctionRegistered = !!bTickFunctionsRegistered, .bWasComponentTickEnabled = IsComponentTickEnabled()
 	};
 
 	if (MigrationContext.MigrationSide == EObjectMigrationSide::Send)
@@ -800,8 +800,8 @@ void UActorComponent::PostMigrate(const struct FUObjectMigrationContext& Migrati
 				}
 
 				FLoanedState AppliedState {
-					.bWasRegistered = bRegistered, .bWasDestroyed = IsBeingDestroyed(), .bWasPhysicsStateCreated = bPhysicsStateCreated,
-					.bWasTickFunctionRegistered = bTickFunctionsRegistered, .bWasComponentTickEnabled = IsComponentTickEnabled()
+					.bWasRegistered = !!bRegistered, .bWasDestroyed = IsBeingDestroyed(), .bWasPhysicsStateCreated = !!bPhysicsStateCreated,
+					.bWasTickFunctionRegistered = !!bTickFunctionsRegistered, .bWasComponentTickEnabled = IsComponentTickEnabled()
 				};
 
 				ensureMsgf(AppliedState == NewState, TEXT("%hs: ReturnedLoan State Application code path was not correct for %s"), __func__, *GetName());
