World Partition: Added enabled flag on WP Cell Transformers to skip the transformation without losing the settings.

diff --git a/Engine/Plugins/Experimental/FastGeoStreaming/Source/FastGeoStreaming/Private/FastGeoWorldPartitionRuntimeCellTransformer.cpp b/Engine/Plugins/Experimental/FastGeoStreaming/Source/FastGeoStreaming/Private/FastGeoWorldPartitionRuntimeCellTransformer.cpp
--- a/Engine/Plugins/Experimental/FastGeoStreaming/Source/FastGeoStreaming/Private/FastGeoWorldPartitionRuntimeCellTransformer.cpp
+++ b/Engine/Plugins/Experimental/FastGeoStreaming/Source/FastGeoStreaming/Private/FastGeoWorldPartitionRuntimeCellTransformer.cpp
@@ -199,7 +199,7 @@ bool UFastGeoWorldPartitionRuntimeCellTransformer::IsDebugMode() const
 
 void UFastGeoWorldPartitionRuntimeCellTransformer::OnSelectionChanged(UObject* Object)
 {
-	if (bDebugModeOnEditorSelection)
+	if (bDebugModeOnEditorSelection && IsEnabled())
 	{
 		UWorldPartition* OuterWorldPartition = GetTypedOuter<UWorldPartition>();
 		UWorld* World = OuterWorldPartition ? OuterWorldPartition->GetWorld() : nullptr;
diff --git a/Engine/Source/Runtime/Engine/Public/WorldPartition/WorldPartition.h b/Engine/Source/Runtime/Engine/Public/WorldPartition/WorldPartition.h
--- a/Engine/Source/Runtime/Engine/Public/WorldPartition/WorldPartition.h
+++ b/Engine/Source/Runtime/Engine/Public/WorldPartition/WorldPartition.h
@@ -122,9 +122,9 @@ struct FRuntimeCellTransformerInstance
 	GENERATED_USTRUCT_BODY()
 
 #if WITH_EDITORONLY_DATA
-	inline void PreTransform(ULevel* InLevel) const { if (Instance) { Instance->PreTransform(InLevel); } }
-	inline void Transform(ULevel* InLevel) const { if (Instance) { Instance->Transform(InLevel); } }
-	inline void PostTransform(ULevel* InLevel) const { if (Instance) { Instance->PostTransform(InLevel); } }
+	inline void PreTransform(ULevel* InLevel) const { if (Instance && Instance->IsEnabled()) { Instance->PreTransform(InLevel); } }
+	inline void Transform(ULevel* InLevel) const { if (Instance && Instance->IsEnabled()) { Instance->Transform(InLevel); } }
+	inline void PostTransform(ULevel* InLevel) const { if (Instance && Instance->IsEnabled()) { Instance->PostTransform(InLevel); } }
 
 	/** Runtime cell transformer class */
 	UPROPERTY(EditAnywhere, Category = WorldPartitionSetup, AdvancedDisplay)
@@ -233,10 +233,10 @@ public:
 	bool CanGenerateStreaming() const { return !StreamingPolicy; }
 
 	UE_DEPRECATED(5.3, "GenerateStreaming is deprecated, use GenerateStreaming with a param struct instead")
-	bool GenerateStreaming(TArray<FString>* OutPackagesToGenerate = nullptr) { return false; }
+	ENGINE_API bool GenerateStreaming(TArray<FString>* OutPackagesToGenerate = nullptr) { return false; }
 
 	UE_DEPRECATED(5.3, "GenerateContainerStreaming is deprecated, use GenerateContainerStreaming with a param struct instead")
-	bool GenerateContainerStreaming(const UActorDescContainer* ActorDescContainer, TArray<FString>* OutPackagesToGenerate = nullptr) { return false; }
+	ENGINE_API bool GenerateContainerStreaming(const UActorDescContainer* ActorDescContainer, TArray<FString>* OutPackagesToGenerate = nullptr) { return false; }
 
 	PRAGMA_DISABLE_DEPRECATION_WARNINGS
 	struct FGenerateStreamingParams
@@ -399,7 +399,7 @@ public:
 	PRAGMA_ENABLE_DEPRECATION_WARNINGS
 
 	UE_DEPRECATED(5.2, "CheckForErrors is deprecated, CheckForErrors with FCheckForErrorsParams should be used instead.")
-	static void CheckForErrors(IStreamingGenerationErrorHandler* ErrorHandler, const UActorDescContainer* ActorDescContainer, bool bEnableStreaming, bool bIsChangelistValidation) {}
+	static ENGINE_API void CheckForErrors(IStreamingGenerationErrorHandler* ErrorHandler, const UActorDescContainer* ActorDescContainer, bool bEnableStreaming, bool bIsChangelistValidation) {}
 
 	static ENGINE_API void CheckForErrors(const FCheckForErrorsParams& Params);
 
@@ -428,10 +428,10 @@ public:
 	FActorDescContainerInstancePreInitializeDelegate OnActorDescContainerInstancePreInitialize;
 
 	UE_DEPRECATED(5.4, "Use RegisterActorDescContainerInstance")
-	UActorDescContainer* RegisterActorDescContainer(const FContainerRegistrationParams& InRegistrationParameters) { return nullptr; }
+	ENGINE_API UActorDescContainer* RegisterActorDescContainer(const FContainerRegistrationParams& InRegistrationParameters) { return nullptr; }
 
 	UE_DEPRECATED(5.4, "Use UnregisterActorDescContainerInstance")
-	bool UnregisterActorDescContainer(UActorDescContainer* Container) { return false; }
+	ENGINE_API bool UnregisterActorDescContainer(UActorDescContainer* Container) { return false; }
 
 	ENGINE_API void UninitializeActorDescContainers();
 
@@ -488,7 +488,7 @@ public:
 	ENGINE_API bool CanStream() const;
 	ENGINE_API bool IsServer() const;
 	ENGINE_API bool IsServerStreamingEnabled() const;
-	bool IsContentBundleEnabled() const { return !bDisableContentBundles; }
+	ENGINE_API bool IsContentBundleEnabled() const { return !bDisableContentBundles; }
 	ENGINE_API bool IsServerStreamingOutEnabled() const;
 	ENGINE_API bool UseMakingVisibleTransactionRequests() const;
 	ENGINE_API bool UseMakingInvisibleTransactionRequests() const;
@@ -851,4 +851,4 @@ private:
 #if WITH_EDITOR
 	friend class FScopedCookingExternalStreamingObject;
 #endif
-};
+};
\ No newline at end of file
diff --git a/Engine/Source/Runtime/Engine/Public/WorldPartition/WorldPartitionRuntimeCellTransformer.h b/Engine/Source/Runtime/Engine/Public/WorldPartition/WorldPartitionRuntimeCellTransformer.h
--- a/Engine/Source/Runtime/Engine/Public/WorldPartition/WorldPartitionRuntimeCellTransformer.h
+++ b/Engine/Source/Runtime/Engine/Public/WorldPartition/WorldPartitionRuntimeCellTransformer.h
@@ -15,4 +15,10 @@ public:
 	virtual void Transform(ULevel* InLevel) {}
 	virtual void PostTransform(ULevel* InLevel) {}
 #endif
-};
+
+	bool IsEnabled() const { return bEnabled; }
+
+private:
+	UPROPERTY(EditAnywhere, Category = Transformer)
+	bool bEnabled = true;
+};
\ No newline at end of file
