Swap the dynamically-generated DialogueWave subtitles to use the new subtitle plugin.

#virtualized

diff --git a/Engine/Source/Runtime/Engine/Classes/Sound/DialogueSoundWaveProxy.h b/Engine/Source/Runtime/Engine/Classes/Sound/DialogueSoundWaveProxy.h
--- a/Engine/Source/Runtime/Engine/Classes/Sound/DialogueSoundWaveProxy.h
+++ b/Engine/Source/Runtime/Engine/Classes/Sound/DialogueSoundWaveProxy.h
@@ -22,6 +22,7 @@ class UDialogueSoundWaveProxy : public USoundBase
 
 	USoundWave* SoundWave;
 	TArray<struct FSubtitleCue> Subtitles;
+	FText SubtitleText;
 
 public:
 	/** Returns whether the sound base is set up in a playable manner */
diff --git a/Engine/Source/Runtime/Engine/Private/DialogueWave.cpp b/Engine/Source/Runtime/Engine/Private/DialogueWave.cpp
--- a/Engine/Source/Runtime/Engine/Private/DialogueWave.cpp
+++ b/Engine/Source/Runtime/Engine/Private/DialogueWave.cpp
@@ -9,6 +9,9 @@
 #include "Sound/DialogueSoundWaveProxy.h"
 #include "Sound/DialogueVoice.h"
 #include "SubtitleManager.h"
+#include "Subtitles/ISubtitlesAndClosedCaptionsModule.h"
+#include "Subtitles/SubtitlesAndClosedCaptionsDelegates.h"
+
 #include "UObject/AssetRegistryTagsContext.h"
 
 #include UE_INLINE_GENERATED_CPP_BY_NAME(DialogueWave)
@@ -30,6 +33,15 @@ const FString FDialogueConstants::PropertyName_DialogueContext			= TEXT("Context
 const FString FDialogueConstants::PropertyName_IsMature					= FString(FLocMetadataObject::COMPARISON_MODIFIER_PREFIX) + TEXT("IsMature");
 #endif //WITH_EDITORONLY_DATA
 
+static TAutoConsoleVariable<int32> CVarUseNewSubtitles(
+	TEXT("au.UseNewSubtitles"),
+	0,
+	TEXT("Defines whether to force DialogueWaves to use the new subtitles plugin.\n")
+	TEXT("0: Use old subtitles\n")
+	TEXT("1: Use new subtitles\n"),
+	ECVF_Default
+);
+
 #if WITH_EDITORONLY_DATA
 namespace
 {
@@ -450,8 +462,10 @@ float UDialogueSoundWaveProxy::GetPitchMultiplier()
 
 void UDialogueSoundWaveProxy::Parse(class FAudioDevice* AudioDevice, const UPTRINT NodeWaveInstanceHash, FActiveSound& ActiveSound, const FSoundParseParameters& ParseParams, TArray<FWaveInstance*>& WaveInstances)
 {
+	const bool UseNewSubtitles = (CVarUseNewSubtitles.GetValueOnGameThread() == 1);
+
 	int OldWaveInstanceCount = WaveInstances.Num();
-	bool bHasSubtitles = (Subtitles.Num() > 0);
+	const bool bHasSubtitles = UseNewSubtitles ? (!SubtitleText.IsEmpty()) : (Subtitles.Num() > 0);
 
 	// Check if the wave instance exists before we try to add it...
 	const bool bWaveInstanceAlreadyExisted = ActiveSound.FindWaveInstance(NodeWaveInstanceHash) != nullptr;
@@ -472,19 +486,41 @@ void UDialogueSoundWaveProxy::Parse(class FAudioDevice* AudioDevice, const UPTRI
 		// Add in the subtitle if they exist
 		if (ActiveSound.bHandleSubtitles && bHasSubtitles)
 		{
-			FQueueSubtitleParams QueueSubtitleParams(Subtitles);
+			if (UseNewSubtitles)
 			{
-				QueueSubtitleParams.AudioComponentID = ActiveSound.GetAudioComponentID();
-				QueueSubtitleParams.WorldPtr = ActiveSound.GetWeakWorld();
-				QueueSubtitleParams.WaveInstance = (PTRINT)NewWaveInstance;
-				QueueSubtitleParams.SubtitlePriority = ActiveSound.SubtitlePriority;
-				QueueSubtitleParams.Duration = GetDuration();
-				QueueSubtitleParams.bManualWordWrap = false;
-				QueueSubtitleParams.bSingleLine = false;
-				QueueSubtitleParams.RequestedStartTime = ActiveSound.RequestedStartTime;
+				IModularFeatures& ModularFeatures = IModularFeatures::Get();
+				if (ModularFeatures.IsModularFeatureAvailable(ISubtitlesAndClosedCaptionsModule::GetModularFeatureName()))
+				{
+					USubtitleAssetUserData* SubtitleCue = NewObject<USubtitleAssetUserData>(GetTransientPackage(), NAME_None, RF_Transient);
+					SubtitleCue->Text = SubtitleText;
+					SubtitleCue->Priority = ActiveSound.SubtitlePriority;
+					SubtitleCue->Duration = GetDuration();
+					SubtitleCue->StartOffset = ActiveSound.RequestedStartTime;
+
+					const ISubtitlesAndClosedCaptionsModule& SubtitlesModule = ModularFeatures.GetModularFeature<ISubtitlesAndClosedCaptionsModule>(ISubtitlesAndClosedCaptionsModule::GetModularFeatureName());
+					FSubtitlesAndClosedCaptionsDelegates::QueueSubtitle.ExecuteIfBound(FQueueSubtitleParameters{ *SubtitleCue }, ESubtitleTiming::InternallyTimed);
+				}
+				else
+				{
+					UE_LOG(LogAudio, Warning, TEXT("The SubtitlesAndClosedCaptions plugin needs to be enabled for subtitles to be displayed."), *GetPathName());
+				}
 			}
+			else
+			{
+				FQueueSubtitleParams QueueSubtitleParams(Subtitles);
+				{
+					QueueSubtitleParams.AudioComponentID = ActiveSound.GetAudioComponentID();
+					QueueSubtitleParams.WorldPtr = ActiveSound.GetWeakWorld();
+					QueueSubtitleParams.WaveInstance = (PTRINT)NewWaveInstance;
+					QueueSubtitleParams.SubtitlePriority = ActiveSound.SubtitlePriority;
+					QueueSubtitleParams.Duration = GetDuration();
+					QueueSubtitleParams.bManualWordWrap = false;
+					QueueSubtitleParams.bSingleLine = false;
+					QueueSubtitleParams.RequestedStartTime = ActiveSound.RequestedStartTime;
+				}
 
-			FSubtitleManager::QueueSubtitles(QueueSubtitleParams);
+				FSubtitleManager::QueueSubtitles(QueueSubtitleParams);
+			}
 		}
 	}
 }
@@ -816,12 +852,20 @@ void UDialogueWave::UpdateMappingProxy(FDialogueContextMapping& ContextMapping)
 		ContextMapping.Proxy->SoundWave = ContextMapping.SoundWave;
 		CopySoundBasePropertiesToProxy(ContextMapping.SoundWave, ContextMapping.Proxy);
 
-		FSubtitleCue NewSubtitleCue;
-		NewSubtitleCue.Text = GetLocalizedSubtitle(ContextMapping);
-		NewSubtitleCue.Time = 0.0f;
+		const bool UseNewSubtitles = (CVarUseNewSubtitles.GetValueOnGameThread() == 1);
+		if (UseNewSubtitles)
+		{
+			ContextMapping.Proxy->SubtitleText = GetLocalizedSubtitle(ContextMapping);
+		}
+		else
+		{
+			FSubtitleCue NewSubtitleCue;
+			NewSubtitleCue.Text = GetLocalizedSubtitle(ContextMapping);
+			NewSubtitleCue.Time = 0.0f;
 
-		ContextMapping.Proxy->Subtitles.Empty();
-		ContextMapping.Proxy->Subtitles.Add(NewSubtitleCue);
+			ContextMapping.Proxy->Subtitles.Empty();
+			ContextMapping.Proxy->Subtitles.Add(NewSubtitleCue);
+		}
 	}
 }
 
diff --git a/Engine/Source/Runtime/Engine/Public/Subtitles/SubtitlesAndClosedCaptionsDelegates.h b/Engine/Source/Runtime/Engine/Public/Subtitles/SubtitlesAndClosedCaptionsDelegates.h
--- a/Engine/Source/Runtime/Engine/Public/Subtitles/SubtitlesAndClosedCaptionsDelegates.h
+++ b/Engine/Source/Runtime/Engine/Public/Subtitles/SubtitlesAndClosedCaptionsDelegates.h
@@ -2,7 +2,7 @@
 
 #pragma once
 #include "Delegates/Delegate.h"
-
+#include "Engine/AssetUserData.h"
 #include "SubtitlesAndClosedCaptionsDelegates.generated.h"
 
 #define UE_API ENGINE_API
@@ -19,6 +19,78 @@ enum class ESubtitleTiming : uint8
 	ExternallyTimed
 };
 
+// ESRB rating categories.
+UENUM(BlueprintType)
+enum class ESRB : uint8
+{
+	Everyone,
+	Everyone10Plus,
+	Teen,
+	Mature,
+	AdultsOnly,
+	RatingPending,
+	RatingPending17Plus
+};
+
+// Subtitle type for type-specific rendering.
+UENUM()
+enum class ESubtitleType : uint8
+{
+	Subtitle,
+	ClosedCaption,
+	AudioDescription
+};
+
+// Minimum duration to display subtitle.
+static constexpr float SubtitleMinDuration = 0.05f;
+
+/**
+ *  Base class for subtitle data being attached to assets
+ */
+UCLASS(BlueprintType)
+class ENGINE_API USubtitleAssetUserData : public UAssetUserData
+{
+	GENERATED_BODY()
+
+public:
+
+	// #SUBTITLES_PRD: carried over from FSubtitleCue, still required
+	//
+	// The text to appear in the subtitle.
+	UPROPERTY(EditAnywhere, BlueprintReadWrite, Category = Subtitles, meta = (Multiline = true))
+	FText Text;
+
+	// #SUBTITLES_PRD: carried over from FSubtitleCue, still required
+	//
+	// Time to display in seconds. 
+	// Defaulted to 3 seconds so when adding new subtitles it's not required to enter a placeholder Duration.
+	// Duration can be be set by ingestion pipelines when importing Subtitles in bulk
+	UPROPERTY(EditAnywhere, BlueprintReadWrite, Category = Subtitles, meta = (ClampMin = 0.05f))
+	float Duration = 3.f;
+
+	// Some subtitles have a delay before they're allowed to be displayed (primarily from the legacy system).
+	// StartOffset measures how long in Seconds, after queuing, before the subtitle is allowed to enter the active subtitles queue.
+	// ESubtitleTiming::ExternallyTimed does not effect this initial delay.
+	UPROPERTY(EditAnywhere, BlueprintReadWrite, Category = Subtitles, meta = (ClampMin = 0.f))
+	float StartOffset = 0.f;
+
+	// #SUBTITLES_PRD: Priority comes from USoundBase::GetSubtitlePriority, USoundCue::GetSubtitlePriority and USoundWave::GetSubtitlePriority
+	// Consolidate various subtitle properties throughout sound/audio code into this new subtitles plugin.
+	// Kept the 10000 default in case that's what users are already used to.
+	// 
+	// The priority of the subtitle.  Defaults to 10000.  Higher values will play instead of lower values.
+	UPROPERTY(EditAnywhere, BlueprintReadWrite, Category = Subtitles)
+	float Priority = 10000.f;
+
+	// ESRB rating category.
+	UPROPERTY(EditAnywhere, BlueprintReadWrite, Category = Subtitles)
+	ESRB Esrb = ESRB::Everyone;
+
+	// Subtitle type for type-specific rendering.
+	UPROPERTY(EditAnywhere, BlueprintReadWrite, Category = Subtitles)
+	ESubtitleType SubtitleType = ESubtitleType::Subtitle;
+};
+
 
 struct FQueueSubtitleParameters
 {
