Fix simple curve table usage in cooked games not preserving interpolation mode, resulting in cooked exe behavior differing from the editor

#tests Tested in local cooked build
#rnx

diff --git a/Engine/Source/Runtime/Engine/Private/CurveTable.cpp b/Engine/Source/Runtime/Engine/Private/CurveTable.cpp
--- a/Engine/Source/Runtime/Engine/Private/CurveTable.cpp
+++ b/Engine/Source/Runtime/Engine/Private/CurveTable.cpp
@@ -19,7 +19,7 @@ DECLARE_CYCLE_STAT(TEXT("CurveTableRowHandle Eval"),STAT_CurveTableRowHandleEval
 
 int32 UCurveTable::GlobalCachedCurveID = 1;
 
-// MaxZ where we will create per-connection dormancy lists within the spatialization grid. This prevents lists being created while flying over in battle bus/parachutes.
+// Enable/Disable optimizing curve tables by removing redundant keys from rows
 int32 CVar_CurveTable_RemoveRedundantKeys = 1;
 static FAutoConsoleVariableRef CVarCurveTableRemoveRedundantKeys(TEXT("CurveTable.RemoveRedundantKeys"), CVar_CurveTable_RemoveRedundantKeys, TEXT(""), ECVF_Default);
 
@@ -766,6 +766,18 @@ TArray<FString> UCurveTable::CreateTableFromCSVString(const FString& InString, E
 	return OutProblems;
 }
 
+template<class CurveType>
+void CopySharedInterpolationMode(const CurveType* SourceCurve, CurveType* TargetCurve) { }
+
+template<>
+void CopySharedInterpolationMode<FSimpleCurve>(const FSimpleCurve* SourceCurve, FSimpleCurve* TargetCurve)
+{
+	if (SourceCurve && TargetCurve)
+	{
+		TargetCurve->SetKeyInterpMode(SourceCurve->GetKeyInterpMode());
+	}
+}
+
 template<class CurveType, class CurveKeyType>
 void CopyRowsToTable(const TMap<FName, CurveType*>& SourceRows, TMap<FName, FRealCurve*>& RowMap, TArray<FString>& OutProblems)
 {
@@ -773,8 +785,11 @@ void CopyRowsToTable(const TMap<FName, CurveType*>& SourceRows, TMap<FName, FRea
 	{
 		CurveType* NewCurve = new CurveType();
 		CurveType* InCurve = CurveRow.Value;
-		TArray<CurveKeyType> CurveKeys = InCurve->GetCopyOfKeys();
+		const TArray<CurveKeyType>& CurveKeys = InCurve->GetConstRefOfKeys();
 		NewCurve->SetKeys(CurveKeys);
+		
+		// Simple curves need to setup their InterpMode
+		CopySharedInterpolationMode(InCurve, NewCurve);
 
 #if WITH_EDITOR
 		// check for duplicate key entries
