ISMComponentDescriptor - Added support for a few more shadows properties
* bCastVolumetricTranslucentShadow
* bSelfShadowOnly
* bCastFarShadow
* bCastInsetShadow
* bCastCinematicShadow

diff --git a/Engine/Source/Runtime/Engine/Private/ISMPartition/ISMComponentDescriptor.cpp b/Engine/Source/Runtime/Engine/Private/ISMPartition/ISMComponentDescriptor.cpp
--- a/Engine/Source/Runtime/Engine/Private/ISMPartition/ISMComponentDescriptor.cpp
+++ b/Engine/Source/Runtime/Engine/Private/ISMPartition/ISMComponentDescriptor.cpp
@@ -106,14 +106,19 @@ void FISMComponentDescriptorBase::InitFrom(const UStaticMeshComponent* Template,
 	CustomDepthStencilValue = Template->CustomDepthStencilValue;
 	bCastShadow = Template->CastShadow;
 	bEmissiveLightSource = Template->bEmissiveLightSource;
-	bCastStaticShadow = Template->bCastStaticShadow;
-	bCastDynamicShadow = Template->bCastDynamicShadow;
-	bCastContactShadow = Template->bCastContactShadow;
-	bCastShadowAsTwoSided = Template->bCastShadowAsTwoSided;
-	bCastHiddenShadow = Template->bCastHiddenShadow;
 	bAffectDynamicIndirectLighting = Template->bAffectDynamicIndirectLighting;
 	bAffectDynamicIndirectLightingWhileHidden = Template->bAffectIndirectLightingWhileHidden;
 	bAffectDistanceFieldLighting = Template->bAffectDistanceFieldLighting;
+	bCastDynamicShadow = Template->bCastDynamicShadow;
+	bCastStaticShadow = Template->bCastStaticShadow;
+	bCastVolumetricTranslucentShadow = Template->bCastVolumetricTranslucentShadow;
+	bCastContactShadow = Template->bCastContactShadow;
+	bSelfShadowOnly = Template->bSelfShadowOnly;
+	bCastFarShadow = Template->bCastFarShadow;
+	bCastInsetShadow = Template->bCastInsetShadow;
+	bCastCinematicShadow = Template->bCastCinematicShadow;
+	bCastHiddenShadow = Template->bCastHiddenShadow;
+	bCastShadowAsTwoSided = Template->bCastShadowAsTwoSided;
 	bReceivesDecals = Template->bReceivesDecals;
 	bOverrideLightMapRes = Template->bOverrideLightMapRes;
 	bUseAsOccluder = Template->bUseAsOccluder;
@@ -252,14 +257,19 @@ bool FISMComponentDescriptorBase::operator==(const FISMComponentDescriptorBase&
 	CustomDepthStencilValue == Other.CustomDepthStencilValue &&
 	bCastShadow == Other.bCastShadow &&
 	bEmissiveLightSource == Other.bEmissiveLightSource &&
-	bCastStaticShadow == Other.bCastStaticShadow &&
-	bCastDynamicShadow == Other.bCastDynamicShadow &&
-	bCastContactShadow == Other.bCastContactShadow &&
-	bCastShadowAsTwoSided == Other.bCastShadowAsTwoSided &&
-	bCastHiddenShadow == Other.bCastHiddenShadow &&
 	bAffectDynamicIndirectLighting == Other.bAffectDynamicIndirectLighting &&
 	bAffectDynamicIndirectLightingWhileHidden == Other.bAffectDynamicIndirectLightingWhileHidden &&
 	bAffectDistanceFieldLighting == Other.bAffectDistanceFieldLighting &&
+	bCastDynamicShadow == Other.bCastDynamicShadow &&
+	bCastStaticShadow == Other.bCastStaticShadow &&
+	bCastVolumetricTranslucentShadow == Other.bCastVolumetricTranslucentShadow &&
+	bCastContactShadow == Other.bCastContactShadow &&
+	bSelfShadowOnly == Other.bSelfShadowOnly &&
+	bCastFarShadow == Other.bCastFarShadow &&
+	bCastInsetShadow == Other.bCastInsetShadow &&
+	bCastCinematicShadow == Other.bCastCinematicShadow &&
+	bCastHiddenShadow == Other.bCastHiddenShadow &&
+	bCastShadowAsTwoSided == Other.bCastShadowAsTwoSided &&
 	bReceivesDecals == Other.bReceivesDecals &&
 	bOverrideLightMapRes == Other.bOverrideLightMapRes &&
 	bUseAsOccluder == Other.bUseAsOccluder &&
@@ -389,14 +399,19 @@ void FISMComponentDescriptorBase::InitComponent(UInstancedStaticMeshComponent* I
 	ISMComponent->CustomDepthStencilValue = CustomDepthStencilValue;
 	ISMComponent->CastShadow = bCastShadow;
 	ISMComponent->bEmissiveLightSource = bEmissiveLightSource;
-	ISMComponent->bCastStaticShadow = bCastStaticShadow;
-	ISMComponent->bCastDynamicShadow = bCastDynamicShadow;
-	ISMComponent->bCastContactShadow = bCastContactShadow;
-	ISMComponent->bCastShadowAsTwoSided = bCastShadowAsTwoSided;
-	ISMComponent->bCastHiddenShadow = bCastHiddenShadow;
 	ISMComponent->bAffectDynamicIndirectLighting = bAffectDynamicIndirectLighting;
 	ISMComponent->bAffectIndirectLightingWhileHidden = bAffectDynamicIndirectLightingWhileHidden;
 	ISMComponent->bAffectDistanceFieldLighting = bAffectDistanceFieldLighting;
+	ISMComponent->bCastDynamicShadow = bCastDynamicShadow;
+	ISMComponent->bCastStaticShadow = bCastStaticShadow;
+	ISMComponent->bCastVolumetricTranslucentShadow = bCastVolumetricTranslucentShadow;
+	ISMComponent->bCastContactShadow = bCastContactShadow;
+	ISMComponent->bSelfShadowOnly = bSelfShadowOnly;
+	ISMComponent->bCastFarShadow = bCastFarShadow;
+	ISMComponent->bCastInsetShadow = bCastInsetShadow;
+	ISMComponent->bCastCinematicShadow = bCastCinematicShadow;
+	ISMComponent->bCastHiddenShadow = bCastHiddenShadow;
+	ISMComponent->bCastShadowAsTwoSided = bCastShadowAsTwoSided;	
 	ISMComponent->bReceivesDecals = bReceivesDecals;
 	ISMComponent->bOverrideLightMapRes = bOverrideLightMapRes;
 	ISMComponent->bUseAsOccluder = bUseAsOccluder;
diff --git a/Engine/Source/Runtime/Engine/Public/ISMPartition/ISMComponentDescriptor.h b/Engine/Source/Runtime/Engine/Public/ISMPartition/ISMComponentDescriptor.h
--- a/Engine/Source/Runtime/Engine/Public/ISMPartition/ISMComponentDescriptor.h
+++ b/Engine/Source/Runtime/Engine/Public/ISMPartition/ISMComponentDescriptor.h
@@ -119,6 +119,15 @@ public:
 
 	UPROPERTY(EditAnywhere, Category = "Component Settings")
 	uint8 bEmissiveLightSource : 1;
+
+	UPROPERTY(EditAnywhere, Category = "Component Settings")
+	uint8 bAffectDynamicIndirectLighting : 1;
+
+	UPROPERTY(EditAnywhere, Category = "Component Settings")
+	uint8 bAffectDynamicIndirectLightingWhileHidden : 1;
+
+	UPROPERTY(EditAnywhere, Category = "Component Settings")
+	uint8 bAffectDistanceFieldLighting : 1;
 		
 	UPROPERTY(EditAnywhere, Category = "Component Settings")
 	uint8 bCastDynamicShadow : 1;
@@ -126,23 +135,29 @@ public:
 	UPROPERTY(EditAnywhere, Category = "Component Settings")
 	uint8 bCastStaticShadow : 1;
 
+	UPROPERTY(EditAnywhere, Category = "Component Settings")
+	uint8 bCastVolumetricTranslucentShadow : 1;
+
 	UPROPERTY(EditAnywhere, Category = "Component Settings")
 	uint8 bCastContactShadow : 1;
 
 	UPROPERTY(EditAnywhere, Category = "Component Settings")
-	uint8 bCastShadowAsTwoSided : 1;
+	uint8 bSelfShadowOnly : 1;
 
 	UPROPERTY(EditAnywhere, Category = "Component Settings")
-	uint8 bCastHiddenShadow : 1;
+	uint8 bCastFarShadow : 1;
 
 	UPROPERTY(EditAnywhere, Category = "Component Settings")
-	uint8 bAffectDynamicIndirectLighting : 1;
+	uint8 bCastInsetShadow : 1;
 
 	UPROPERTY(EditAnywhere, Category = "Component Settings")
-	uint8 bAffectDynamicIndirectLightingWhileHidden : 1;
+	uint8 bCastCinematicShadow : 1;
 
 	UPROPERTY(EditAnywhere, Category = "Component Settings")
-	uint8 bAffectDistanceFieldLighting : 1;
+	uint8 bCastHiddenShadow : 1;
+
+	UPROPERTY(EditAnywhere, Category = "Component Settings")
+	uint8 bCastShadowAsTwoSided : 1;
 
 	UPROPERTY(EditAnywhere, Category = "Component Settings")
 	uint8 bReceivesDecals : 1;
