[AsyncCompilation/StaticMesh]
- Make mesh cards compilation entirely async in editor as it was already the case for cooking
- This reduces hitches from the game thread after mesh building in editor

diff --git a/Engine/Source/Runtime/Engine/Private/DistanceFieldAtlas.cpp b/Engine/Source/Runtime/Engine/Private/DistanceFieldAtlas.cpp
--- a/Engine/Source/Runtime/Engine/Private/DistanceFieldAtlas.cpp
+++ b/Engine/Source/Runtime/Engine/Private/DistanceFieldAtlas.cpp
@@ -826,22 +826,27 @@ void FDistanceFieldAsyncQueue::Build(FAsyncDistanceFieldTask* Task, FQueuedThrea
 			Task->bGenerateDistanceFieldAsIfTwoSided,
 			*Task->GeneratedVolumeData);
 
-		// If cooking, build the dependent mesh card representation data immediately in the existing worker task
-		if (IsRunningCookCommandlet())
+		// Begin mesh card representation while we're async
+		FAsyncCardRepresentationTask* CardTask = BeginCacheMeshCardRepresentationInternal(
+			Task->StaticMesh,
+			*Task->StaticMesh->GetRenderData(),
+			Task->DDCKey,
+			MoveTemp(Task->SourceMeshData));
+
+		if (CardTask)
 		{
-			FAsyncCardRepresentationTask* CardTask = BeginCacheMeshCardRepresentationInternal(
-				Task->StaticMesh,
-				*Task->StaticMesh->GetRenderData(),
-				Task->DDCKey,
-				MoveTemp(Task->SourceMeshData));
-
-			if (CardTask)
+			// If cooking, build the dependent mesh card representation data immediately in the existing worker task
+			if (IsRunningCookCommandlet())
 			{
 				GCardRepresentationAsyncQueue->RunImmediateTask(CardTask);
 			}
-
-			Task->bGeneratedCardRepresentation = true;
+			else
+			{
+				GCardRepresentationAsyncQueue->AddTask(CardTask);
+			}
 		}
+
+		Task->bGeneratedCardRepresentation = true;
 	}
 
 	{
