[AsyncCompilation/StaticMesh]
- Make nav collision compilation async during static mesh build since it was already async in the postload case
- This reduces hitches from the game thread after mesh building in editor

diff --git a/Engine/Source/Runtime/Engine/Private/StaticMeshBuild.cpp b/Engine/Source/Runtime/Engine/Private/StaticMeshBuild.cpp
--- a/Engine/Source/Runtime/Engine/Private/StaticMeshBuild.cpp
+++ b/Engine/Source/Runtime/Engine/Private/StaticMeshBuild.cpp
@@ -14,6 +14,7 @@
 #include "PhysicsEngine/BodySetup.h"
 #include "DistanceFieldAtlas.h"
 #include "MeshCardBuild.h"
+#include "AI/Navigation/NavCollisionBase.h"
 #include "SceneInterface.h"
 
 #if WITH_EDITOR
@@ -418,6 +419,15 @@ void UStaticMesh::BeginBuildInternal(FStaticMeshBuildContext* Context)
 	// Ensure we have a bodysetup.
 	CreateBodySetup();
 	check(GetBodySetup() != nullptr);
+	
+	// Make sure the object is created on the game-thread before going async
+	if (bHasNavigationData && GetBodySetup() != nullptr)
+	{
+		if (GetNavCollision() == nullptr)
+		{
+			SetNavCollision(UNavCollisionBase::ConstructNew(*this));
+		}
+	}
 
 	// Release the static mesh's resources.
 	ReleaseResources();
@@ -521,6 +531,9 @@ bool UStaticMesh::ExecuteBuildInternal(const FBuildParameters& BuildParameters)
 	GetBodySetup()->InvalidatePhysicsData();
 	GetBodySetup()->CreatePhysicsMeshes();
 
+	// Update nav collision 
+	CreateNavCollision(/*bIsUpdate=*/true);
+
 	// Compare the derived data keys to see if renderable mesh data has actually changed.
 	check(GetRenderData());
 	bool bHasRenderDataChanged = GetRenderData()->DerivedDataKey != ExistingDerivedDataKey;
@@ -644,11 +657,10 @@ void UStaticMesh::FinishBuildInternal(const TArray<IStaticMeshComponent*>& InAff
 		InitResources();
 	}
 
-
 	for (IStaticMeshComponent* Component : InAffectedComponents)
 	{
-		Component->OnMeshRebuild(bHasRenderDataChanged);			
-	}	
+		Component->OnMeshRebuild(bHasRenderDataChanged);
+	}
 
 	// If extended bounds were already calculated in the PreBuild step and are unlocked, 
 	// we will only validate here to avoid modifying the value while another thread could read it.
@@ -672,9 +684,6 @@ void UStaticMesh::FinishBuildInternal(const TArray<IStaticMeshComponent*>& InAff
 		}
 	}
 
-	// Update nav collision 
-	CreateNavCollision(/*bIsUpdate=*/true);
-
 #if WITH_EDITORONLY_DATA
 	// We can ditch the assembly references now
 	ClearNaniteAssemblyReferenceCache();
