Prevent skeletal vertex factory reinitialization on skin weight updates

diff --git a/Engine/Source/Runtime/Engine/Private/SkeletalRenderGPUSkin.cpp b/Engine/Source/Runtime/Engine/Private/SkeletalRenderGPUSkin.cpp
--- a/Engine/Source/Runtime/Engine/Private/SkeletalRenderGPUSkin.cpp
+++ b/Engine/Source/Runtime/Engine/Private/SkeletalRenderGPUSkin.cpp
@@ -2111,7 +2111,8 @@ void UpdateVertexFactory(
 	{
 		FGPUBaseSkinVertexFactory* VertexFactory = FactoryPtr.Get();
 
-		if (VertexFactory != nullptr)
+		// Do not recreate the factory if it's been released, given the loose scheduling this may result in dangling factories
+		if (VertexFactory != nullptr && VertexFactory->IsInitialized())
 		{
 			// Setup the update data for enqueue
 			FDynamicUpdateVertexFactoryData VertexUpdateData(VertexFactory, InVertexBuffers);
