Remove test ensure that leaked into last commit.

#rnx

diff --git a/Engine/Source/Runtime/Engine/Private/Components/InstancedSkinnedMeshComponent.cpp b/Engine/Source/Runtime/Engine/Private/Components/InstancedSkinnedMeshComponent.cpp
--- a/Engine/Source/Runtime/Engine/Private/Components/InstancedSkinnedMeshComponent.cpp
+++ b/Engine/Source/Runtime/Engine/Private/Components/InstancedSkinnedMeshComponent.cpp
@@ -850,8 +850,6 @@ void UInstancedSkinnedMeshComponent::BeginDestroy()
 void UInstancedSkinnedMeshComponent::OnRegister()
 {
 	Super::OnRegister();
-
-	ensure(!PrimitiveLocalBoundsOverride.IsValid);
 }
 
 void UInstancedSkinnedMeshComponent::TickComponent(float DeltaTime, ELevelTick TickType, FActorComponentTickFunction* ThisTickFunction)
