PR #13142: Fix SVT tile range streaming request priority selection

diff --git a/Engine/Source/Runtime/Engine/Private/SparseVolumeTexture/SparseVolumeTextureStreamingManager.cpp b/Engine/Source/Runtime/Engine/Private/SparseVolumeTexture/SparseVolumeTextureStreamingManager.cpp
--- a/Engine/Source/Runtime/Engine/Private/SparseVolumeTexture/SparseVolumeTextureStreamingManager.cpp
+++ b/Engine/Source/Runtime/Engine/Private/SparseVolumeTexture/SparseVolumeTextureStreamingManager.cpp
@@ -1066,7 +1066,7 @@ void FStreamingManager::FilterRequests()
 			// When the priority changes, we need to create tile ranges based on the currently set bets in RequestedTilesInCurrentPriority.
 			if (Priority != PrevPriority)
 			{
-				MakeTileRanges(Priority);
+				MakeTileRanges(PrevPriority);
 				RequestedTilesInCurrentPriority.SetRange(0, NumTiles, false);
 			}
 			PrevPriority = Priority;
