EngineNetTest
* Have ScopeCVar use the highest priority so unit tests can override even user typed cvar settings

#rnx

diff --git a/Engine/Source/Runtime/Engine/Private/Tests/NetTestHelpers.cpp b/Engine/Source/Runtime/Engine/Private/Tests/NetTestHelpers.cpp
--- a/Engine/Source/Runtime/Engine/Private/Tests/NetTestHelpers.cpp
+++ b/Engine/Source/Runtime/Engine/Private/Tests/NetTestHelpers.cpp
@@ -576,7 +576,7 @@ FScopedCVarOverrideInt::FScopedCVarOverrideInt(const TCHAR* VariableName, int32
 	if (Variable)
 	{
 		SavedValue = Variable->GetInt();
-		Variable->Set(Value);
+		Variable->Set(Value, ECVF_SetByConsole /* Use highest priority so we override even console setters */);
 	}
 }
 
@@ -584,7 +584,7 @@ FScopedCVarOverrideInt::~FScopedCVarOverrideInt()
 {
 	if (Variable)
 	{
-		Variable->Set(SavedValue);
+		Variable->Set(SavedValue, ECVF_SetByConsole);
 	}
 }
 
