allow transactional execution of batched tick tasks - we wrap the execution with UE::RemoteExecutor::EnqueueTransactional to queue them all up, then flush it with UE::RemoteExecutor::ExecutePendingWork();

- exposed public EnqueueTransactional for queuing up transactional work without immediately executing it

diff --git a/Engine/Source/Runtime/CoreUObject/Private/UObject/RemoteExecutor.cpp b/Engine/Source/Runtime/CoreUObject/Private/UObject/RemoteExecutor.cpp
--- a/Engine/Source/Runtime/CoreUObject/Private/UObject/RemoteExecutor.cpp
+++ b/Engine/Source/Runtime/CoreUObject/Private/UObject/RemoteExecutor.cpp
@@ -104,6 +104,11 @@ public:
 	{
 		check(ExecutingWork == nullptr);
 
+		if (PendingWorks.Num() > 0)
+		{
+			UE_LOG(LogRemoteExec, VeryVerbose, TEXT("ExecutePendingWork BEGIN with %d enqueued work"), PendingWorks.Num());
+		}
+
 		double LastStallPrintTime = FPlatformTime::Seconds();
 
 		for (int32 LocalIterationNumber = 0; ; LocalIterationNumber++)
@@ -295,7 +300,7 @@ public:
 								const uint32 ExecutionAttempts = ExecutingWork->ExecutionAttempts;
 								UE_AUTORTFM_ONABORT(ExecutionAttempts)
 								{
-									constexpr uint32 MaxExecutionAttemptsBeforeLivelockWarning = 200;
+									constexpr uint32 MaxExecutionAttemptsBeforeLivelockWarning = 500;
 									if ((ExecutionAttempts >= MaxExecutionAttemptsBeforeLivelockWarning) && (ExecutionAttempts % MaxExecutionAttemptsBeforeLivelockWarning) == 0)
 									{
 										UE_LOG(LogRemoteExec, Display, TEXT("vvv Transaction Aborted %u Times. Dumping Callstack. vvv"), ExecutionAttempts);
@@ -889,6 +894,17 @@ void ExecuteTransactional(FName WorkName, const TFunctionRef<void(void)>& Work)
 	ExecuteTransactionalWithExplicitPriority(WorkName, RootWorkPriority, Work);
 }
 
+void EnqueueTransactional(FName WorkName, const TFunctionRef<void(void)>& InWork)
+{
+#if UE_WITH_REMOTE_OBJECT_HANDLE
+	check(!AutoRTFM::IsClosed());
+	FRemoteWorkPriority RootWorkPriority = FRemoteWorkPriority::CreateRootWorkPriority(UE::RemoteObject::GetGlobalServerId(), GRemoteExecutor.GenerateNextTransactionId());
+	GRemoteExecutor.EnqueueWork(WorkName, RootWorkPriority, true, [&InWork]() { InWork(); });
+#else
+	InWork();
+#endif
+}
+
 } // namespace UE::RemoteExecutor
 
 FString FRemoteWorkPriority::ToString() const
diff --git a/Engine/Source/Runtime/CoreUObject/Public/UObject/RemoteExecutor.h b/Engine/Source/Runtime/CoreUObject/Public/UObject/RemoteExecutor.h
--- a/Engine/Source/Runtime/CoreUObject/Public/UObject/RemoteExecutor.h
+++ b/Engine/Source/Runtime/CoreUObject/Public/UObject/RemoteExecutor.h
@@ -216,6 +216,7 @@ namespace UE::RemoteExecutor
 
 	COREUOBJECT_API void ExecuteTransactional(FName WorkName, const TFunctionRef<void(void)>& Work);
 	COREUOBJECT_API void ExecuteTransactionalWithExplicitPriority(FName WorkName, FRemoteWorkPriority WorkPriority, const TFunctionRef<void(void)>& Work);
+	COREUOBJECT_API void EnqueueTransactional(FName WorkName, const TFunctionRef<void(void)>& Work);
 	COREUOBJECT_API void ExecutePendingWork();
 
 	COREUOBJECT_API void AbortTransactionRequiresDependencies(FStringView Description);
diff --git a/Engine/Source/Runtime/Engine/Private/TickTaskManager.cpp b/Engine/Source/Runtime/Engine/Private/TickTaskManager.cpp
--- a/Engine/Source/Runtime/Engine/Private/TickTaskManager.cpp
+++ b/Engine/Source/Runtime/Engine/Private/TickTaskManager.cpp
@@ -61,6 +61,11 @@ static TAutoConsoleVariable<int32> CVarAllowBatchedTicks(
 	0,
 	TEXT("If true, tick functions with bAllowTickBatching will be automatically combined into a single tick task."));
 
+static TAutoConsoleVariable<int32> CVarAllowBatchedTicksUnordered(
+	TEXT("tick.AllowBatchedTicksUnordered"),
+	0,
+	TEXT("If true, tick functions within a batch can be run out of order."));
+
 static TAutoConsoleVariable<int32> CVarAllowOptimizedPrerequisites(
 	TEXT("tick.AllowOptimizedPrerequisites"),
 	1,
@@ -407,6 +412,10 @@ public:
 			}
 			if (Target->IsTickFunctionEnabled())
 			{
+#if UE_WITH_REMOTE_OBJECT_HANDLE
+				auto ExecuteTickWork = [this, Target, CurrentThread, &MyCompletionGraphEvent]()
+				{
+#endif
 #if DO_TIMEGUARD
 				FTimerNameDelegate NameFunction = FTimerNameDelegate::CreateLambda( [&]{ return FString::Printf(TEXT("Slowtick %s "), *Target->DiagnosticMessage()); } );
 				SCOPE_TIME_GUARD_DELEGATE_MS(NameFunction, 4);
@@ -414,9 +423,31 @@ public:
 				LIGHTWEIGHT_TIME_GUARD_BEGIN(FBatchTickFunctionTask, GTimeguardThresholdMS);
 				Target->ExecuteTick(Target->CalculateDeltaTime(Context.DeltaSeconds, Context.World), Context.TickType, CurrentThread, MyCompletionGraphEvent);
 				LIGHTWEIGHT_TIME_GUARD_END(FBatchTickFunctionTask, Target->DiagnosticMessage());
+
+				Target->ClearTaskInformation();  // This is stale and a good time to clear it for safety
+#if UE_WITH_REMOTE_OBJECT_HANDLE
+				};
+
+				if (Target->bRunTransactionally)
+				{
+					static FName TransactionalWorkName(TEXT("FTickFunction"));
+					UE::RemoteExecutor::EnqueueTransactional(TransactionalWorkName, ExecuteTickWork);
+				}
+				else
+				{
+					ExecuteTickWork();
+				}
+#endif
+			}
+			else
+			{
+				Target->ClearTaskInformation();  // This is stale and a good time to clear it for safety
 			}
-			Target->ClearTaskInformation();  // This is stale and a good time to clear it for safety
 		}
+
+#if UE_WITH_REMOTE_OBJECT_HANDLE
+		UE::RemoteExecutor::ExecutePendingWork();
+#endif
 	}
 };
 
@@ -525,6 +556,8 @@ class FTickTaskSequencer
 	bool				bAllowConcurrentTicks;
 	/** If true, multiple ticks can be combined into a single task */
 	bool				bAllowBatchedTicksForFrame;
+	/** If true, ticks within a batch can run out of order */
+	bool				bAllowBatchedTicksUnorderedForFrame;
 	/** If true, some prerequisites will be ignored */
 	bool				bAllowOptimizedPrerequisites;
 
@@ -730,11 +763,29 @@ public:
 
 						for (FTickFunction* Prereq : Prerequisites)
 						{
-							// Ignore prerequisites that are already in this batch
-							if (Prereq->GetTaskPointer(FTickFunction::ETickTaskState::HasTask) != PossibleBatch->TickTask && !PossibleBatch->TickPrerequisites.Contains(Prereq))
+							const bool bBatchRequiresPrerequisite = PossibleBatch->TickPrerequisites.Contains(Prereq);
+							const bool bBatchExecutesPrerequisite = Prereq->GetTaskPointer(FTickFunction::ETickTaskState::HasTask) == PossibleBatch->TickTask;
+
+							if (!bAllowBatchedTicksUnorderedForFrame)
 							{
-								bPrerequisitesMatch = false;
-								break;
+								// if batches must execute in-order (default), then we can accept this batch if either the batch
+								// requires the same prerequisite, or if the prerequisite is already in this batch
+								if (!bBatchRequiresPrerequisite && !bBatchExecutesPrerequisite)
+								{
+									bPrerequisitesMatch = false;
+									break;
+								}
+							}
+							else
+							{
+								// if batches may execute their functions out of order, then we can't accept a batch
+								// that itself executes the prerequisite, we can only accept batches that have the prerequisite
+								// as a prerequisite of the batch itself
+								if (!bBatchRequiresPrerequisite)
+								{
+									bPrerequisitesMatch = false;
+									break;
+								}
 							}
 						}
 						if (bPrerequisitesMatch)
@@ -1038,6 +1089,7 @@ public:
 		}
 
 		bAllowBatchedTicksForFrame = !!CVarAllowBatchedTicks.GetValueOnGameThread();
+		bAllowBatchedTicksUnorderedForFrame = !!CVarAllowBatchedTicksUnordered.GetValueOnGameThread();
 		bAllowOptimizedPrerequisites = !!CVarAllowOptimizedPrerequisites.GetValueOnGameThread();
 
 		WaitForCleanup();
