[AsyncCompilation/StaticMesh]
- Fix cancelation that could still show errors when canceled

diff --git a/Engine/Source/Runtime/Engine/Private/DistanceFieldAtlas.cpp b/Engine/Source/Runtime/Engine/Private/DistanceFieldAtlas.cpp
--- a/Engine/Source/Runtime/Engine/Private/DistanceFieldAtlas.cpp
+++ b/Engine/Source/Runtime/Engine/Private/DistanceFieldAtlas.cpp
@@ -272,6 +272,12 @@ void FDistanceFieldVolumeData::CacheDerivedData(UStaticMesh* Mesh, UStaticMesh*
 			IMeshBuilderModule& MeshBuilderModule = IMeshBuilderModule::GetForPlatform(NewTask->TargetPlatform);
 			if (!MeshBuilderModule.BuildMeshVertexPositions(Mesh, NewTask->SourceMeshData.TriangleIndices, NewTask->SourceMeshData.VertexPositions, NewTask->SourceMeshData.Sections))
 			{
+				if (UE::Tasks::FCancellationTokenScope::IsCurrentWorkCanceled())
+				{
+					delete NewTask;
+					return;
+				}
+
 				UE_LOG(LogStaticMesh, Error, TEXT("Failed to build static mesh. See previous line(s) for details."));
 			}
 		}
diff --git a/Engine/Source/Runtime/Engine/Private/StaticMesh.cpp b/Engine/Source/Runtime/Engine/Private/StaticMesh.cpp
--- a/Engine/Source/Runtime/Engine/Private/StaticMesh.cpp
+++ b/Engine/Source/Runtime/Engine/Private/StaticMesh.cpp
@@ -4218,6 +4218,11 @@ void FStaticMeshRenderData::Cache(const ITargetPlatform* TargetPlatform, UStatic
 		// Do not launch any distance field or mesh card representation tasks for this mesh
 		return;
 	}
+	
+	if (UE::Tasks::FCancellationTokenScope::IsCurrentWorkCanceled())
+	{
+		return;
+	}
 
 	static const auto CVar = IConsoleManager::Get().FindTConsoleVariableDataInt(TEXT("r.GenerateMeshDistanceFields"));
 
