Fixed: Broken normals on a skinned mesh with a SkinCache enabled on a platforms that do not support MVF
VF tangents vertex format was not matching SkinCache format

diff --git a/Engine/Source/Runtime/Engine/Private/GPUSkinVertexFactory.cpp b/Engine/Source/Runtime/Engine/Private/GPUSkinVertexFactory.cpp
--- a/Engine/Source/Runtime/Engine/Private/GPUSkinVertexFactory.cpp
+++ b/Engine/Source/Runtime/Engine/Private/GPUSkinVertexFactory.cpp
@@ -1393,19 +1393,35 @@ void FGPUSkinPassthroughVertexFactory::InitRHI(FRHICommandListBase& RHICmdList)
 		Data.ColorIndexMask = ~0u;
 	}
 
+	if (Data.TangentsSRV)
+	{
+		TangentFormat = Data.TangentsSRV->GetDesc().Common.Format;
+		PreSkinnedTangentSRV = Data.TangentsSRV;
+		
+		if (IsOpenGLPlatform(GMaxRHIShaderPlatform))
+		{
+			// See SkinCache TangentFormat
+			TangentFormat = PF_R16G16B16A16_SNORM;
+		}
+	}
+
 	if (EnumHasAnyFlags(VertexAttributesRequested, EVertexAttributeFlags::Tangent) && !IsManualVertexFetch(Data.TangentBasisComponents[0]))
 	{
+		EVertexElementType VtxType = (TangentFormat == PF_R8G8B8A8_SNORM) ? VET_PackedNormal : VET_Short4N;
+		uint32 Stride = (TangentFormat == PF_R8G8B8A8_SNORM) ? 8 : 16;
+		uint32 Offset = (TangentFormat == PF_R8G8B8A8_SNORM) ? 4 : 8;
+		
 		FRHIBuffer* Buffer = GetVertexBufferRHI(Data.TangentBasisComponents[0].VertexBuffer);
 		SourceStreamBuffers[EVertexAttribute::VertexTangent] = Buffer;
 		StreamSourceSlots[EVertexAttribute::VertexTangent] = FRHIStreamSourceSlot::Create(Buffer);
 		Data.TangentBasisComponents[0].VertexStreamUsage |= EVertexStreamUsage::Overridden;
 		Data.TangentBasisComponents[0].Offset = 0;
-		Data.TangentBasisComponents[0].Type = VET_Short4N;
-		Data.TangentBasisComponents[0].Stride = 16;
+		Data.TangentBasisComponents[0].Type = VtxType;
+		Data.TangentBasisComponents[0].Stride = Stride;
 		Data.TangentBasisComponents[1].VertexStreamUsage |= EVertexStreamUsage::Overridden;
-		Data.TangentBasisComponents[1].Offset = 8;
-		Data.TangentBasisComponents[1].Type = VET_Short4N;
-		Data.TangentBasisComponents[1].Stride = 16;
+		Data.TangentBasisComponents[1].Offset = Offset;
+		Data.TangentBasisComponents[1].Type = VtxType;
+		Data.TangentBasisComponents[1].Stride = Stride;
 		VertexAttributesToBind |= EVertexAttributeFlags::Tangent;
 	}
 
@@ -1416,12 +1432,6 @@ void FGPUSkinPassthroughVertexFactory::InitRHI(FRHICommandListBase& RHICmdList)
 	LooseParameters.PreSkinnedTangentBuffer = GNullVertexBuffer.VertexBufferSRV;
 	LooseParametersUniformBuffer = TUniformBufferRef<FGPUSkinPassThroughFactoryLooseParameters>::CreateUniformBufferImmediate(LooseParameters, UniformBuffer_MultiFrame);
 
-	if (Data.TangentsSRV)
-	{
-		TangentFormat = Data.TangentsSRV->GetDesc().Common.Format;
-		PreSkinnedTangentSRV = Data.TangentsSRV;
-	}
-
 	FLocalVertexFactory::InitRHI(RHICmdList);
 }
 
