Iris
* Downgraded error when a subobject's outer is not net-referenceable to a warning.
* Set the fallback to use the RootObject as the Outer on the client explicitely so the clients don't log errors too

#rnx

diff --git a/Engine/Source/Runtime/Engine/Private/Net/Iris/ReplicationSystem/NetSubObjectFactory.cpp b/Engine/Source/Runtime/Engine/Private/Net/Iris/ReplicationSystem/NetSubObjectFactory.cpp
--- a/Engine/Source/Runtime/Engine/Private/Net/Iris/ReplicationSystem/NetSubObjectFactory.cpp
+++ b/Engine/Source/Runtime/Engine/Private/Net/Iris/ReplicationSystem/NetSubObjectFactory.cpp
@@ -75,9 +75,15 @@ TUniquePtr<UE::Net::FNetObjectCreationHeader> UNetSubObjectFactory::CreateAndFil
 		{
 			SubObjectHeader->OuterReference = Bridge->GetOrCreateObjectReference(OuterObject);
 
+			// If the Outer is not net-referenceable, use the RootObject instead
 			if (!SubObjectHeader->OuterReference.IsValid())
 			{
-				UE_LOG(LogIris, Error, TEXT("UNetSubObjectFactory::CreateAndFillHeader %s could not create NetReference to Outer %s"), *Bridge->PrintObjectFromNetRefHandle(Handle), *GetNameSafe(OuterObject));
+				UE_LOG(LogIris, Warning, TEXT("UNetSubObjectFactory::CreateAndFillHeader subobject: %s has an Outer: %s that is not stable or replicated. Clients will use RootObject: %s as the Outer instead"), 
+					*Bridge->PrintObjectFromNetRefHandle(Handle), 
+					*GetNameSafe(OuterObject),
+					*GetNameSafe(RootObject));
+
+				SubObjectHeader->bOuterIsRootObject = true;
 			}
 		}
 	}
