Make sure touch is considered when checking if a key is from keyboard in UKismetInputLibrary::Key_IsKeyboardKey.

diff --git a/Engine/Source/Runtime/Engine/Private/KismetInputLibrary.cpp b/Engine/Source/Runtime/Engine/Private/KismetInputLibrary.cpp
--- a/Engine/Source/Runtime/Engine/Private/KismetInputLibrary.cpp
+++ b/Engine/Source/Runtime/Engine/Private/KismetInputLibrary.cpp
@@ -61,7 +61,7 @@ bool UKismetInputLibrary::Key_IsMouseButton(const FKey& Key)
 
 bool UKismetInputLibrary::Key_IsKeyboardKey(const FKey& Key)
 {
-	return Key.IsBindableInBlueprints() && (Key.IsGamepadKey() == false && Key.IsMouseButton() == false);
+	return Key.IsBindableInBlueprints() && (Key.IsGamepadKey() == false && Key.IsMouseButton() == false && Key.IsTouch() == false);
 }
 
 bool UKismetInputLibrary::Key_IsAxis1D(const FKey& Key)
