Added Type traits for defining operations like Size(), Distance(), Dot(), etc. so arithmetic, vector or custom types can work with our splines. This also contains some WIP for recovering attributes support.

diff --git a/Engine/Source/Runtime/Engine/Classes/Curves/Spline.h b/Engine/Source/Runtime/Engine/Classes/Curves/Spline.h
--- a/Engine/Source/Runtime/Engine/Classes/Curves/Spline.h
+++ b/Engine/Source/Runtime/Engine/Classes/Curves/Spline.h
@@ -70,20 +70,18 @@ struct FSpline
 	
 	bool SupportsAttributes() const { return false;/*IsNew();*/ }
 	ENGINE_API bool HasAttributeChannel(FName AttributeName) const;
-	ENGINE_API int32 NumAttributeValues(FName AttributeName) const;
-	ENGINE_API float GetAttributeParameter(int32 Index, const FName& Name) const;
-	ENGINE_API int32 SetAttributeParameter(int32 Index, float Parameter, const FName& Name);
-	ENGINE_API void RemoveAttributeValue(int32 Index, FName AttributeName);
 	ENGINE_API TArray<FName> GetFloatPropertyChannels() const;
 	ENGINE_API TArray<FName> GetVectorPropertyChannels() const;
-	
-	template<typename T> ENGINE_API T GetAttributeValue(int32 Index, const FName& Name) const;
-	template<typename T> ENGINE_API void SetAttributeValue(int32 Index, const T& Value, const FName& Name);
-	template<typename T> ENGINE_API bool CreateAttributeChannel(FName AttributeName) const;
-	template<typename T> ENGINE_API int32 AddAttributeValue(float Param, const T& Value, FName AttributeName) const;
-	
-	
-	template<typename T> ENGINE_API T EvaluateAttribute(float Param, FName AttributeName) const;
+
+	template <typename AttrType> ENGINE_API float GetAttributeParameter(int32 Index, const FName& Name) const;
+	template <typename AttrType> ENGINE_API int32 SetAttributeParameter(int32 Index, float Parameter, const FName& Name);
+	template <typename AttrType> ENGINE_API int32 NumAttributeValues(FName AttributeName) const;
+	template <typename AttrType> ENGINE_API AttrType GetAttributeValue(int32 Index, const FName& Name) const;
+	template <typename AttrType> ENGINE_API void SetAttributeValue(int32 Index, const AttrType& Value, const FName& Name);
+	template <typename AttrType> ENGINE_API bool CreateAttributeChannel(FName AttributeName) const;
+	template <typename AttrType> ENGINE_API int32 AddAttributeValue(float Param, const AttrType& Value, FName AttributeName) const;
+	template <typename AttrType> ENGINE_API void RemoveAttributeValue(int32 Index, FName AttributeName);
+	template <typename AttrType> ENGINE_API AttrType EvaluateAttribute(float Param, FName AttributeName) const;
 
 	ENGINE_API float FindNearest(const FVector& InLocation, float& OutSquaredDist) const;
 	ENGINE_API float FindNearestOnSegment(const FVector& InLocation, int32 SegmentIndex, float& OutSquaredDist) const;
diff --git a/Engine/Source/Runtime/Engine/Private/Components/SplineComponent.cpp b/Engine/Source/Runtime/Engine/Private/Components/SplineComponent.cpp
--- a/Engine/Source/Runtime/Engine/Private/Components/SplineComponent.cpp
+++ b/Engine/Source/Runtime/Engine/Private/Components/SplineComponent.cpp
@@ -649,10 +649,10 @@ float USplineComponent::GetDistanceAlongSplineAtLocation(const FVector& InLocati
 
 int32 USplineComponent::GetNumberOfPropertyValues(FName PropertyName) const
 {
-	if (UE::SplineComponent::GUseSplineCurves == false && Spline.SupportsAttributes())
-	{
-		return Spline.NumAttributeValues(PropertyName);
-	}
+	// if (UE::SplineComponent::GUseSplineCurves == false && Spline.SupportsAttributes())
+	// {
+	// 	return Spline.NumAttributeValues(PropertyName);
+	// }
 
 	return 0;
 }
@@ -737,20 +737,20 @@ int32 SetPropertyAtSplineInputKey(const FSpline& Spline, float InKey, const T& I
 
 float GetInputKeyAtIndex(const FSpline& Spline, int32 Index, FName PropertyName)
 {
-	if (UE::SplineComponent::GUseSplineCurves == false && Spline.SupportsAttributes())
-	{
-		return Spline.GetAttributeParameter(Index, PropertyName);
-	}
+	// if (UE::SplineComponent::GUseSplineCurves == false && Spline.SupportsAttributes())
+	// {
+	// 	return Spline.GetAttributeParameter(Index, PropertyName);
+	// }
 
 	return 0.f;
 }
 
 int32 SetInputKeyAtIndex(FSpline& Spline, int32 Index, float InKey, FName PropertyName)
 {
-	if (UE::SplineComponent::GUseSplineCurves == false && Spline.SupportsAttributes())
-	{
-		return Spline.SetAttributeParameter(Index, InKey , PropertyName);
-	}
+	// if (UE::SplineComponent::GUseSplineCurves == false && Spline.SupportsAttributes())
+	// {
+	// 	return Spline.SetAttributeParameter(Index, InKey , PropertyName);
+	// }
 
 	return INDEX_NONE;
 }
@@ -779,7 +779,7 @@ void USplineComponent::RemovePropertyAtIndex(int32 Index, FName PropertyName)
 {
 	if (UE::SplineComponent::GUseSplineCurves == false && Spline.SupportsAttributes())
 	{
-		Spline.RemoveAttributeValue(Index, PropertyName);
+	//	Spline.RemoveAttributeValue(Index, PropertyName);
 	}
 }
 
diff --git a/Engine/Source/Runtime/Engine/Private/Curves/Spline.cpp b/Engine/Source/Runtime/Engine/Private/Curves/Spline.cpp
--- a/Engine/Source/Runtime/Engine/Private/Curves/Spline.cpp
+++ b/Engine/Source/Runtime/Engine/Private/Curves/Spline.cpp
@@ -7,6 +7,7 @@
 #include "Serialization/MemoryReader.h"
 #include "Serialization/MemoryWriter.h"
 #include "Splines/TangentBezierSpline.h"
+#include "Splines/MultiSpline.h"
 #include "UObject/FortniteMainBranchObjectVersion.h"
 
 #include UE_INLINE_GENERATED_CPP_BY_NAME(Spline)
@@ -115,8 +116,9 @@ inline EInterpCurveMode ConvertTangentModeToInterpCurveMode(ETangentMode Mode)
  * tangent control and automatic tangent computation.
  */
 
-class FNewSpline: public FTangentBezierSpline3d
+class FNewSpline: public TMultiSpline<FTangentBezierSpline3d>
 {
+	using FTangentBezierControlPoint = TTangentBezierControlPoint<FVector3d>;
 	
 public:
 	
@@ -124,12 +126,13 @@ public:
     FNewSpline(const FNewSpline& Other);
 	FNewSpline(const FLegacySpline& Other);
 	FNewSpline(const FSplineCurves& Other);
+	virtual ~FNewSpline() override = default;
     FNewSpline& operator=(const FNewSpline& Other);
 
 	bool operator==(const FNewSpline& Other) const;
 	bool operator!=(const FNewSpline& Other) const;
 
-    bool Serialize(FArchive& Ar);
+    virtual bool Serialize(FArchive& Ar) override;
     friend FArchive& operator<<(FArchive& Ar, FNewSpline& Spline)
     {
         Spline.Serialize(Ar);
@@ -139,7 +142,7 @@ public:
 	float GetSegmentLength(const int32 Index, const float Param, const FVector& Scale3D = FVector(1.0f)) const;
 	float GetSplineLength() const;
 	
-    void SetClosedLoop(bool bInClosedLoop);
+    virtual void SetClosedLoop(bool bInClosedLoop) override;
 
 	void Reset();
 	void ResetToDefault();
@@ -175,8 +178,8 @@ public:
 	float GetDistanceAtParameter(float Parameter) const;
 
     virtual float FindNearest(const FVector& InLocation, float& OutSquaredDistance) const override;
-    virtual float FindNearestOnSegment(const FVector& InLocation, int32 SegmentIndex, float& OutSquaredDistance) const override;
-
+    float FindNearestOnSegment(const FVector& InLocation, int32 SegmentIndex, float& OutSquaredDistance) const;
+	
 	/** Parameter based evaluation API. */
     virtual FVector EvaluateImpl(float Parameter) const override;
     FVector EvaluateDerivative(float Parameter) const;
@@ -186,8 +189,12 @@ public:
 	template <typename AttrType> TArray<FName> GetAttributeChannelNames() const;
 	template <typename AttrType> AttrType EvaluateAttribute(FName Name, float Parameter) const;
 	template <typename AttrType> int32 AddAttributeValue(FName Name, const AttrType& Value, float Parameter);
-	int32 SetAttributeParameter(const FName& Name, int32 Index, float ParentSpaceParameter);
-	float GetAttributeParameter(const FName& Name, int32 Index) const;
+	template <typename AttrType> void SetAttributeValue(FName Name, const AttrType& Value, int32 Index);
+	template <typename AttrType> void RemoveAttributeValue(FName Name, int32 Index);
+	template <typename AttrType> AttrType GetAttributeValue(const FName& Name, int32 Index) const;
+	template <typename AttrType> int32 NumAttributeValues(FName Name) const;
+	template <typename AttrType> int32 SetAttributeParameter(const FName& Name, int32 Index, float ParentSpaceParameter);
+	template <typename AttrType> float GetAttributeParameter(const FName& Name, int32 Index) const;
 	
 	/** Updates the spline using the current configuration. */
 	void UpdateSpline();
@@ -265,6 +272,8 @@ private:
 template <typename AttrType>
 TArray<FName> FNewSpline::GetAttributeChannelNames() const
 {
+	using AttrSplineType = TTangentBezierSpline<AttrType>;
+
 	TArray<FName> ChannelNames; /*= InternalSpline.GetAttributeChannelNames<AttrType>();
 	ChannelNames.Remove(RotationAttrName);
 	ChannelNames.Remove(ScaleAttrName);
@@ -276,6 +285,8 @@ TArray<FName> FNewSpline::GetAttributeChannelNames() const
 template <typename AttrType>
 AttrType FNewSpline::EvaluateAttribute(FName Name, float Parameter) const
 {
+	using AttrSplineType = TTangentBezierSpline<AttrType>;
+
 	Parameter = ToInternalSplineSpace(Parameter);
 	return AttrType();//InternalSpline.EvaluateAttribute<AttrType>(Name, Parameter);
 }
@@ -283,15 +294,67 @@ AttrType FNewSpline::EvaluateAttribute(FName Name, float Parameter) const
 template <typename AttrType>
 int32 FNewSpline::AddAttributeValue(FName Name, const AttrType& Value, float Parameter)
 {
+	using AttrSplineType = TTangentBezierSpline<AttrType>;
+
 	Parameter = ToInternalSplineSpace(Parameter);
 	return -1;//InternalSpline.AddAttributeValue<AttrType>(Name, Value, Parameter);
 }
 
+template <typename AttrType>
+void FNewSpline::SetAttributeValue(FName Name, const AttrType& Value, int32 Index)
+{
+	using AttrSplineType = TTangentBezierSpline<AttrType>;
+
+	if (AttrSplineType* Channel = GetAttributeChannelImpl<AttrSplineType>(Name))
+	{
+		Channel->SetValue(Index, Value);
+	}
+}
+
+template <typename AttrType>
+void FNewSpline::RemoveAttributeValue(FName Name, int32 Index)
+{
+	using AttrSplineType = TTangentBezierSpline<AttrType>;
+
+	if (AttrSplineType* Channel = GetAttributeChannelImpl<AttrSplineType>(Name))
+	{
+		Channel->RemovePoint(Index);
+	}
+}
+
+template <typename AttrType>
+AttrType FNewSpline::GetAttributeValue(const FName& Name, int32 Index) const
+{
+	using AttrSplineType = TTangentBezierSpline<AttrType>;
+	
+	if (AttrSplineType* Channel = GetAttributeChannelImpl<AttrSplineType>(Name))
+	{
+		return Channel->GetValue(Index);
+	}
+
+	return AttrType();
+}
+
+template <typename AttrType>
+int32 FNewSpline::NumAttributeValues(FName Name) const
+{
+	using AttrSplineType = TTangentBezierSpline<AttrType>;
+	
+	if (AttrSplineType* Channel = GetAttributeChannelImpl<AttrSplineType>(Name))
+	{
+		return Channel->GetNumPoints();
+	}
+
+	return 0;
+}
+
+template <typename AttrType>
 int32 FNewSpline::SetAttributeParameter(const FName& Name, int32 Index, float Parameter)
 {
 	return -1;//InternalSpline.SetAttributeParameter(Name, Index, ToInternalSplineSpace(Parameter));
 }
 
+template <typename AttrType>
 float FNewSpline::GetAttributeParameter(const FName& Name, int32 Index) const
 {
 	return -1;//FromInternalSplineSpace(InternalSpline.GetAttributeParameter(Name, Index));
@@ -741,93 +804,124 @@ FVector FSpline::EvaluateScale(float Param) const
 
 bool FSpline::HasAttributeChannel(FName AttributeName) const
 {
-	// if (NewData)
-	// {
-	// 	return NewData->HasAttributeChannel(AttributeName);
-	// }
+	if (NewData)
+	{
+		return NewData->HasAttributeChannel(AttributeName);
+	}
 
 	return false;
 }
 
-int32 FSpline::NumAttributeValues(FName AttributeName) const
+TArray<FName> FSpline::GetFloatPropertyChannels() const
 {
-	// if (NewData)
-	// {
-	// 	return NewData->NumAttributeValues(AttributeName);
-	// }
+	TArray<FName> ChannelNames;
+	if (NewData)
+	{
+		ChannelNames = NewData->GetAttributeChannelNames<float>();
+	}
+	return ChannelNames;
+}
 
-	return 0;
+TArray<FName> FSpline::GetVectorPropertyChannels() const
+{
+	TArray<FName> ChannelNames;
+	if (NewData)
+	{
+		ChannelNames = NewData->GetAttributeChannelNames<FVector>();
+	}
+	return ChannelNames;
 }
 
+template <typename AttrType>
 float FSpline::GetAttributeParameter(int32 Index, const FName& Name) const
 {
-	// if (NewData)
-	// {
-	// 	return NewData->GetAttributeParameter(Name, Index);
-	// }
+	if (NewData)
+	{
+		return NewData->GetAttributeParameter<AttrType>(Name, Index);
+	}
 
 	return 0.f;
 }
+template ENGINE_API float FSpline::GetAttributeParameter<float>(int32 Index, const FName& Name) const;
+template ENGINE_API float FSpline::GetAttributeParameter<FVector>(int32 Index, const FName& Name) const;
 
+template <typename AttrType>
 int32 FSpline::SetAttributeParameter(int32 Index, float Parameter, const FName& Name)
 {
-	// if (NewData)
-	// {
-	// 	return NewData->SetAttributeParameter(Name, Index, Parameter);
-	// }
+	if (NewData)
+	{
+		return NewData->SetAttributeParameter<AttrType>(Name, Index, Parameter);
+	}
+	
 	return Index;
 }
+template ENGINE_API int32 FSpline::SetAttributeParameter<float>(int32 Index, float Parameter, const FName& Name);
+template ENGINE_API int32 FSpline::SetAttributeParameter<FVector>(int32 Index, float Parameter, const FName& Name);
 
-template<typename T>
-T FSpline::GetAttributeValue(int32 Index, const FName& Name) const
+template <typename AttrType>
+int32 FSpline::NumAttributeValues(FName AttributeName) const
 {
-	// if (NewData)
-	// {
-	// 	return NewData->GetAttributeValue<T>(Name, Index);
-	// }
+	if (NewData)
+	{
+		return NewData->NumAttributeValues<AttrType>(AttributeName);
+	}
+
+	return 0;
+}
+template ENGINE_API int32 FSpline::NumAttributeValues<float>(FName AttributeName) const;
+template ENGINE_API int32 FSpline::NumAttributeValues<FVector>(FName AttributeName) const;
+
+template<typename AttrType>
+AttrType FSpline::GetAttributeValue(int32 Index, const FName& Name) const
+{
+	if (NewData)
+	{
+		return NewData->GetAttributeValue<AttrType>(Name, Index);
+	}
 	
-	return T();
+	return AttrType();
 }
 template ENGINE_API float FSpline::GetAttributeValue<float>(int32 Index, const FName& Name) const;
 template ENGINE_API FVector FSpline::GetAttributeValue<FVector>(int32 Index, const FName& Name) const;
 
-template<typename T>
-void FSpline::SetAttributeValue(int32 Index, const T& Value, const FName& Name)
+template<typename AttrType>
+void FSpline::SetAttributeValue(int32 Index, const AttrType& Value, const FName& Name)
 {
-	// if (NewData)
-	// {
-	// 	NewData->SetAttributeValue<T>(Name, Index, Value);
-	// }
+	if (NewData)
+	{
+	//	NewData->SetAttributeValue<AttrType>(Name, Index, Value);
+	}
 }
 template ENGINE_API void FSpline::SetAttributeValue<float>(int32 Index, const float& Value, const FName& Name);
 template ENGINE_API void FSpline::SetAttributeValue<FVector>(int32 Index, const FVector& Value, const FName& Name);
 
-
-template<typename T>
+template<typename AttrType>
 bool FSpline::CreateAttributeChannel(FName AttributeName) const
 {
-	// if (NewData)
-	// {
-	// 	return NewData->CreateAttributeChannel<T>(AttributeName);
-	// }
+	if (NewData)
+	{
+	//	return NewData->CreateAttributeChannel<AttrType>(AttributeName);
+	}
+	
 	return false;
 }
 template ENGINE_API bool FSpline::CreateAttributeChannel<float>(FName AttributeName) const;
 template ENGINE_API bool FSpline::CreateAttributeChannel<FVector>(FName AttributeName) const;
 
-template<typename T>
-int32 FSpline::AddAttributeValue(float Param, const T& Value, FName AttributeName) const
+template<typename AttrType>
+int32 FSpline::AddAttributeValue(float Param, const AttrType& Value, FName AttributeName) const
 {
-	// if (NewData)
-	// {
-	// 	return NewData->AddAttributeValue(AttributeName, Value, Param);
-	// }
+	if (NewData)
+	{
+		return NewData->AddAttributeValue(AttributeName, Value, Param);
+	}
 	
 	return INDEX_NONE;
 }
 template ENGINE_API int32 FSpline::AddAttributeValue<float>(float Param, const float& Value, FName AttributeName) const;
 template ENGINE_API int32 FSpline::AddAttributeValue<FVector>(float Param, const FVector& Value, FName AttributeName) const;
 
+template <typename AttrType>
 void FSpline::RemoveAttributeValue(int32 Index, FName AttributeName)
 {
 	// if (NewData)
@@ -835,39 +929,21 @@ void FSpline::RemoveAttributeValue(int32 Index, FName AttributeName)
 	// 	NewData->RemoveAttributeValue(AttributeName, Index);
 	// }
 }
+template ENGINE_API void FSpline::RemoveAttributeValue<float>(int32 Index, FName AttributeName);
+template ENGINE_API void FSpline::RemoveAttributeValue<FVector>(int32 Index, FName AttributeName);
 
-TArray<FName> FSpline::GetFloatPropertyChannels() const
+template<typename AttrType>
+AttrType FSpline::EvaluateAttribute(float Param, FName AttributeName) const
 {
-	TArray<FName> ChannelNames;
 	if (NewData)
 	{
-		ChannelNames = NewData->GetAttributeChannelNames<float>();
+		return NewData->EvaluateAttribute<AttrType>(AttributeName, Param);
 	}
-	return ChannelNames;
+	return AttrType(0);
 }
-
-TArray<FName> FSpline::GetVectorPropertyChannels() const
-{
-	TArray<FName> ChannelNames;
-	if (NewData)
-	{
-		ChannelNames = NewData->GetAttributeChannelNames<FVector>();
-	}
-	return ChannelNames;
-}
-
-template<typename T>
-T FSpline::EvaluateAttribute(float Param, FName AttributeName) const
-{
-	if (NewData)
-	{
-		return NewData->EvaluateAttribute<T>(AttributeName, Param);
-	}
-	return T(0);
-}
-
 template ENGINE_API float FSpline::EvaluateAttribute<float>(float Param, FName AttributeName) const;
 template ENGINE_API FVector FSpline::EvaluateAttribute<FVector>(float Param, FName AttributeName) const;
+
 float FSpline::FindNearest(const FVector& InLocation, float& OutSquaredDist) const
 {
 	switch (CurrentImplementation)
@@ -1113,7 +1189,7 @@ int32 FSpline::GetNumControlPoints() const
 	{
 	default: return 0;
 	case 1: return LegacyData->GetNumControlPoints();
-	case 2: return NewData->GetNumPoints();
+	case 2: return NewData->GetSpline().GetNumPoints();
 	}
 }
 
@@ -1220,10 +1296,10 @@ void FSpline::OnSplineImplementationChanged()
 
 FNewSpline::FNewSpline()
 {
-    CreateAttributeChannel<FQuat>(RotationAttrName);
-    CreateAttributeChannel<FVector>(ScaleAttrName);
-    CreateAttributeChannel<float>(ParamToLengthAttrName);
-    CreateAttributeChannel<float>(LengthToParamAttrName);
+	CreateAttributeChannel<FQuat>(RotationAttrName, TTangentBezierSpline<FQuat>::GetSplineTypeName());
+	CreateAttributeChannel<FVector>(ScaleAttrName, TTangentBezierSpline<FVector>::GetSplineTypeName());
+	CreateAttributeChannel<float>(ParamToLengthAttrName, TTangentBezierSpline<float>::GetSplineTypeName());
+	CreateAttributeChannel<float>(LengthToParamAttrName, TTangentBezierSpline<float>::GetSplineTypeName());
     
     ResetToDefault();
 }
@@ -1235,10 +1311,10 @@ FNewSpline::FNewSpline(const FNewSpline& Other)
 
 FNewSpline::FNewSpline(const FLegacySpline& Other)
 {
-	CreateAttributeChannel<FQuat>(RotationAttrName);
-	CreateAttributeChannel<FVector>(ScaleAttrName);
-	CreateAttributeChannel<float>(ParamToLengthAttrName);
-	CreateAttributeChannel<float>(LengthToParamAttrName);
+	CreateAttributeChannel<FQuat>(RotationAttrName, TTangentBezierSpline<FQuat>::GetSplineTypeName());
+	CreateAttributeChannel<FVector>(ScaleAttrName, TTangentBezierSpline<FVector>::GetSplineTypeName());
+	CreateAttributeChannel<float>(ParamToLengthAttrName, TTangentBezierSpline<float>::GetSplineTypeName());
+	CreateAttributeChannel<float>(LengthToParamAttrName, TTangentBezierSpline<float>::GetSplineTypeName());
 	
 	TArray<FTangentBezierControlPoint> Points;
     for (int32 Idx = 0; Idx < Other.GetNumControlPoints(); Idx++)
@@ -1252,13 +1328,13 @@ FNewSpline::FNewSpline(const FLegacySpline& Other)
     }
 
 	// initialize spline
-	SetControlPoints(Points);
+	GetSpline().SetControlPoints(Points);
 	
 	// setup rotation/scale attributes
-	for (int32 Idx = 0; Idx < GetNumPoints(); Idx++)
+	for (int32 Idx = 0; Idx < GetSpline().GetNumPoints(); Idx++)
 	{    		
-		AddAttributeValue(RotationAttrName, Other.GetRotation(Idx), static_cast<float>(Idx));
-		AddAttributeValue(ScaleAttrName, Other.GetScale(Idx), static_cast<float>(Idx));
+		AddAttributeValue<FQuat>(RotationAttrName, Other.GetRotation(Idx), static_cast<float>(Idx));
+		AddAttributeValue<FVector>(ScaleAttrName, Other.GetScale(Idx), static_cast<float>(Idx));
 	}
 
 	FSpline::FUpdateSplineParams UpdateParams;
@@ -1269,10 +1345,10 @@ FNewSpline::FNewSpline(const FLegacySpline& Other)
 
 FNewSpline::FNewSpline(const FSplineCurves& Other)
 {
-	CreateAttributeChannel<FQuat>(RotationAttrName);
-	CreateAttributeChannel<FVector>(ScaleAttrName);
-	CreateAttributeChannel<float>(ParamToLengthAttrName);
-	CreateAttributeChannel<float>(LengthToParamAttrName);
+	CreateAttributeChannel<FQuat>(RotationAttrName, TTangentBezierSpline<FQuat>::GetSplineTypeName());
+	CreateAttributeChannel<FVector>(ScaleAttrName, TTangentBezierSpline<FVector>::GetSplineTypeName());
+	CreateAttributeChannel<float>(ParamToLengthAttrName, TTangentBezierSpline<float>::GetSplineTypeName());
+	CreateAttributeChannel<float>(LengthToParamAttrName, TTangentBezierSpline<float>::GetSplineTypeName());
 	
 	TArray<FTangentBezierControlPoint> Points;
     for (int32 Idx = 0; Idx < Other.Position.Points.Num(); Idx++)
@@ -1288,16 +1364,16 @@ FNewSpline::FNewSpline(const FSplineCurves& Other)
     }
 	
 	// initialize spline
-	SetControlPoints(Points);
+	GetSpline().SetControlPoints(Points);
 	
     // setup rotation/scale attributes
-    for (int32 Idx = 0; Idx < GetNumPoints(); Idx++)
+    for (int32 Idx = 0; Idx < GetSpline().GetNumPoints(); Idx++)
     {
     	const FQuat& Rotation = Other.Rotation.Points[Idx].OutVal;
         const FVector& Scale = Other.Scale.Points[Idx].OutVal;
     	
-    	AddAttributeValue(RotationAttrName, Rotation, static_cast<float>(Idx));
-    	AddAttributeValue(ScaleAttrName, Scale, static_cast<float>(Idx));
+    	AddAttributeValue<FQuat>(RotationAttrName, Rotation, static_cast<float>(Idx));
+    	AddAttributeValue<FVector>(ScaleAttrName, Scale, static_cast<float>(Idx));
     }
 
 	FSpline::FUpdateSplineParams UpdateParams;
@@ -1312,10 +1388,10 @@ FNewSpline& FNewSpline::operator=(const FNewSpline& Other)
     if (this != &Other)
     {
         InternalSpline = Other.InternalSpline;
-    	TangentModes = Other.TangentModes;
+    	GetSpline().SetTangentModes(Other.GetSpline().GetTangentModes());
     	ReparamStepsPerSegment = Other.ReparamStepsPerSegment;
-        Tension = Other.Tension;
-    	bStationaryEndpoint = Other.bStationaryEndpoint;
+    	GetSpline().SetTension(Other.GetSpline().GetTension());
+    	GetSpline().SetStationaryEndpoints(Other.GetSpline().IsStationaryEndpoints());
 
     	UpdateSpline();
     }
@@ -1325,7 +1401,7 @@ FNewSpline& FNewSpline::operator=(const FNewSpline& Other)
 
 bool FNewSpline::operator==(const FNewSpline& Other) const
 {
-    return InternalSpline == Other.InternalSpline && Tension == Other.Tension;
+    return GetSpline() == Other.GetSpline() && GetSpline().GetTension() == Other.GetSpline().GetTension();
 }
 
 bool FNewSpline::operator!=(const FNewSpline& Other) const
@@ -1335,7 +1411,7 @@ bool FNewSpline::operator!=(const FNewSpline& Other) const
 
 bool FNewSpline::Serialize(FArchive& Ar)
 {
-    FTangentBezierSpline3d::Serialize(Ar);
+    TMultiSpline::Serialize(Ar);
 	
     Ar << ReparamStepsPerSegment;
 	
@@ -1349,7 +1425,7 @@ bool FNewSpline::Serialize(FArchive& Ar)
 
 float FNewSpline::GetSegmentLength(const int32 Index, const float Param, const FVector& Scale3D) const
 {
-	const int32 NumPoints = GetNumPoints();
+	const int32 NumPoints = GetSpline().GetNumPoints();
 	const int32 LastPoint = NumPoints - 1;
 
 	check(Index >= 0 && ((IsClosedLoop() && Index < NumPoints) || (!IsClosedLoop() && Index < LastPoint)));
@@ -1376,13 +1452,13 @@ float FNewSpline::GetSegmentLength(const int32 Index, const float Param, const F
 		{ 0.90617985f, 0.23692688f }
 	};
 
-	const auto& P0 = GetPosition(Index);
-	const auto& T0 = GetTangentOut(Index);
-	const auto& P1 = GetPosition(Index == LastPoint ? 0 : Index + 1);
-	const auto& T1 = GetTangentOut(Index == LastPoint ? 0 : Index + 1);
+	const auto& P0 = GetSpline().GetValue(Index);
+	const auto& T0 = GetSpline().GetTangentOut(Index);
+	const auto& P1 = GetSpline().GetValue(Index == LastPoint ? 0 : Index + 1);
+	const auto& T1 = GetSpline().GetTangentOut(Index == LastPoint ? 0 : Index + 1);
 
 	// Special cases for linear
-	if (TangentModes[Index] == ETangentMode::Linear)
+	if (GetSpline().GetTangentModes()[Index] == ETangentMode::Linear)
 	{
 		return ((P1 - P0) * Scale3D).Size() * Param;
 	}
@@ -1415,14 +1491,14 @@ float FNewSpline::GetSplineLength() const
 
 void FNewSpline::SetClosedLoop(bool bInClosedLoop)
 {
-    FTangentBezierSpline3d::SetClosedLoop(bInClosedLoop);
+    GetSpline().SetClosedLoop(bInClosedLoop);
     UpdateSpline();
 }
 
 void FNewSpline::Reset()
 {
 	InternalSpline.Clear();
-	TangentModes.Empty();
+	GetSpline().SetTangentModes(TArray<ETangentMode>());
 	PositionCurve.Reset();
 	RotationCurve.Reset();
 	ScaleCurve.Reset();
@@ -1430,7 +1506,7 @@ void FNewSpline::Reset()
 
 void FNewSpline::ResetToDefault()
 {
-    SetControlPoints( { FTangentBezierControlPoint(FVector(0.0, 0.0, 0.0)),
+    GetSpline().SetControlPoints( { FTangentBezierControlPoint(FVector(0.0, 0.0, 0.0)),
     					FTangentBezierControlPoint(FVector(100.0, 0.0, 0.0)) } );
 
 	ResetRotation();
@@ -1443,9 +1519,9 @@ void FNewSpline::ResetRotation()
 {
     ClearAttributeChannel(RotationAttrName);
 
-    for (int32 Idx = 0; Idx < GetNumPoints(); Idx++)
+    for (int32 Idx = 0; Idx < GetSpline().GetNumPoints(); Idx++)
     {    		
-    	AddAttributeValue(RotationAttrName, FQuat::Identity, static_cast<float>(Idx));
+    	AddAttributeValue<FQuat>(RotationAttrName, FQuat::Identity, static_cast<float>(Idx));
     }
 }
 
@@ -1453,21 +1529,21 @@ void FNewSpline::ResetScale()
 {
     ClearAttributeChannel(ScaleAttrName);
 	
-    for (int32 Idx = 0; Idx < GetNumPoints(); Idx++)
+    for (int32 Idx = 0; Idx < GetSpline().GetNumPoints(); Idx++)
     {    		
-        AddAttributeValue(ScaleAttrName, FVector::OneVector, static_cast<float>(Idx));
+        AddAttributeValue<FVector>(ScaleAttrName, FVector::OneVector, static_cast<float>(Idx));
     }
 }
 
 void FNewSpline::AddPoint(const FSplinePoint& Point)
 {
 	FTangentBezierControlPoint ControlPoint = ConvertToTangentBezierControlPoint(Point);
-	const int32 NumPoints = GetNumPoints();
+	const int32 NumPoints = GetSpline().GetNumPoints();
     
 	// Special case for empty spline
 	if (NumPoints == 0)
 	{
-		AppendPoint(ControlPoint);
+		GetSpline().AppendPoint(ControlPoint);
 		UpdatePointAttributes(Point, 0);
 		return;
 	}
@@ -1476,7 +1552,7 @@ void FNewSpline::AddPoint(const FSplinePoint& Point)
 	float Fraction = Point.InputKey - BaseIndex;
 	    
 	// Insert at calculated parameter
-	int32 InsertedIndex = FTangentBezierSpline3d::InsertPointAtSegmentParam(BaseIndex, Fraction, ControlPoint);
+	int32 InsertedIndex = GetSpline().InsertPointAtSegmentParam(BaseIndex, Fraction, ControlPoint);
 	
 	UpdatePointAttributes(Point, InsertedIndex);
 }
@@ -1484,25 +1560,25 @@ void FNewSpline::AddPoint(const FSplinePoint& Point)
 void FNewSpline::InsertPoint(const FSplinePoint& Point, int32 Index)
 {
 	FTangentBezierControlPoint ControlPoint = ConvertToTangentBezierControlPoint(Point);
-	const int32 NumPoints = GetNumPoints();
+	const int32 NumPoints = GetSpline().GetNumPoints();
     
 	// Special case for empty spline
 	if (NumPoints == 0 || (Index >= NumPoints && !IsClosedLoop()))
 	{
-		AppendPoint(ControlPoint);
+		GetSpline().AppendPoint(ControlPoint);
 		UpdatePointAttributes(Point, NumPoints);
 		return;
 	}
 	else if (Index == 0)
 	{
-		PrependPoint(ControlPoint);
+		GetSpline().PrependPoint(ControlPoint);
 		UpdatePointAttributes(Point, 0);
 		return;
 	}
 	
-	Index = FMath::Clamp(Index, 0, GetNumPoints());
+	Index = FMath::Clamp(Index, 0, GetSpline().GetNumPoints());
 	
-	int32 InsertedSegmentIndex = FTangentBezierSpline3d::InsertPointAtPosition(Index, ControlPoint);
+	int32 InsertedSegmentIndex = GetSpline().InsertPointAtPosition(Index, ControlPoint);
 	
 	UpdatePointAttributes(Point, InsertedSegmentIndex);
 }
@@ -1511,17 +1587,17 @@ FSplinePoint FNewSpline::GetPoint(int32 Index) const
 {
     FSplinePoint Point;
     
-    if (Index < 0 || Index >= GetNumPoints())
+    if (Index < 0 || Index >= GetSpline().GetNumPoints())
     {
     	return Point;
     }
 
-    Point.Position = GetPosition(Index);
-	Point.ArriveTangent = GetTangentIn(Index);
-    Point.LeaveTangent = GetTangentOut(Index);
+    Point.Position = GetSpline().GetValue(Index);
+	Point.ArriveTangent = GetSpline().GetTangentIn(Index);
+    Point.LeaveTangent = GetSpline().GetTangentOut(Index);
     Point.Rotation = GetAttributeValue<FQuat>(RotationAttrName, Index).Rotator();
     Point.Scale = GetAttributeValue<FVector>(ScaleAttrName, Index);
-	EInterpCurveMode InterpCurveMode = ConvertTangentModeToInterpCurveMode(TangentModes[Index]);
+	EInterpCurveMode InterpCurveMode = ConvertTangentModeToInterpCurveMode(GetSpline().GetTangentModes()[Index]);
     Point.Type = ConvertInterpCurveModeToSplinePointType(InterpCurveMode);
 
     return Point;
@@ -1530,95 +1606,95 @@ FSplinePoint FNewSpline::GetPoint(int32 Index) const
 /** Removes a point from the spline */
 void FNewSpline::RemovePoint(int32 Index)
 {
-	if (Index < 0 || Index >= GetNumPoints())
+	if (Index < 0 || Index >= GetSpline().GetNumPoints())
 	{
 		return;
 	}
 
-	FTangentBezierSpline3d::RemovePoint(Index);
+	GetSpline().RemovePoint(Index);
     RebuildLegacyCurves();
 }
 
 void FNewSpline::SetLocation(int32 Index, const FVector& InLocation)
 {
-    if (Index < 0 || Index >= GetNumPoints())
+    if (Index < 0 || Index >= GetSpline().GetNumPoints())
     {
     	return;
     }
-	//SetValueAt(Index * 4, InLocation);
-	ModifyPoint(Index, {InLocation, GetTangentIn(Index), GetTangentOut(Index), GetTangentMode(Index)});
+	
+	GetSpline().ModifyPoint(Index, {InLocation, GetSpline().GetTangentIn(Index), GetSpline().GetTangentOut(Index), GetSpline().GetTangentMode(Index)});
     RebuildLegacyCurves();
 }
 
 FVector FNewSpline::GetLocation(const int32 Index) const
 {
-	if (Index < 0 || Index >= GetNumPoints())
+	if (Index < 0 || Index >= GetSpline().GetNumPoints())
     {
     	return FVector();
     }
 
-    return GetPosition(Index);
+    return GetSpline().GetValue(Index);
 }
 
 void FNewSpline::SetInTangent(const int32 Index, const FVector& InTangent)
 {
-	if (Index < 0 || Index >= GetNumPoints())
+	if (Index < 0 || Index >= GetSpline().GetNumPoints())
     {
     	return;
     }
 
-	SetPointTangentMode(Index, ETangentMode::User);
-    this->SetTangentIn(Index, InTangent);
+	GetSpline().SetPointTangentMode(Index, ETangentMode::User);
+    GetSpline().SetTangentIn(Index, InTangent);
     RebuildLegacyCurves();
 }
 
 FVector FNewSpline::GetInTangent(const int32 Index) const
 {
-	if (Index < 0 || Index >= GetNumPoints())
+	if (Index < 0 || Index >= GetSpline().GetNumPoints())
     {
     	return FVector();
     }
 
-    return GetTangentIn(Index);
+    return GetSpline().GetTangentIn(Index);
 }
 
 void FNewSpline::SetOutTangent(const int32 Index, const FVector& OutTangent)
 {
-	if (Index < 0 || Index >= GetNumPoints())
+	if (Index < 0 || Index >= GetSpline().GetNumPoints())
     {
     	return;
     }
-	SetPointTangentMode(Index, ETangentMode::User);
-	this->SetTangentOut(Index, OutTangent);
-    
+	
+	GetSpline().SetPointTangentMode(Index, ETangentMode::User);
+	GetSpline().SetTangentOut(Index, OutTangent);
     RebuildLegacyCurves();
 }
 
 FVector FNewSpline::GetOutTangent(const int32 Index) const
 {
-	if (Index < 0 || Index >= GetNumPoints())
+	if (Index < 0 || Index >= GetSpline().GetNumPoints())
     {
     	return FVector();
     }
 
-    return GetTangentOut(Index);
+    return GetSpline().GetTangentOut(Index);
 }
 
 void FNewSpline::SetRotation(int32 Index, const FQuat& InRotation)
 {
-	if (Index < 0 || Index >= GetNumPoints())
+	if (Index < 0 || Index >= GetSpline().GetNumPoints())
     {
     	return;
     }
 
 	// if rotation value not added to attribute channel, add it
-	if (Index >= NumAttributeValues(RotationAttrName))
+	if (Index >= NumAttributeValues<FQuat>(RotationAttrName))
 	{
-		AddAttributeValue(RotationAttrName, InRotation, GetParameterAt(Index));
+		AddAttributeValue<FQuat>(RotationAttrName, InRotation, GetSpline().GetParameterAt(Index));
 	}
 	else
 	{
-		SetAttributeValue(RotationAttrName, Index, InRotation);
+		SetAttributeValue<FQuat>(RotationAttrName, InRotation, Index);
 	}
 	
     RebuildLegacyCurves();
@@ -1626,7 +1702,7 @@ void FNewSpline::SetRotation(int32 Index, const FQuat& InRotation)
 
 FQuat FNewSpline::GetRotation(const int32 Index) const
 {
-	if (Index < 0 || Index >= GetNumPoints())
+	if (Index < 0 || Index >= GetSpline().GetNumPoints())
     {
     	return FQuat::Identity;
     }
@@ -1636,19 +1712,19 @@ FQuat FNewSpline::GetRotation(const int32 Index) const
 
 void FNewSpline::SetScale(int32 Index, const FVector& InScale)
 {
-	if (Index < 0 || Index >= GetNumPoints())
+	if (Index < 0 || Index >= GetSpline().GetNumPoints())
     {
     	return;
     }
 
 	// if scale value not added to attribute channel, add it
-	if (Index >= NumAttributeValues(ScaleAttrName))
+	if (Index >= NumAttributeValues<float>(ScaleAttrName))
 	{
-		AddAttributeValue(ScaleAttrName, InScale, GetParameterAt(Index));
+		AddAttributeValue<FVector>(ScaleAttrName, InScale, GetSpline().GetParameterAt(Index));
 	}
 	else
 	{
-		SetAttributeValue(ScaleAttrName, Index, InScale);
+		SetAttributeValue<FVector>(ScaleAttrName, InScale, Index);
 	}
     
     RebuildLegacyCurves();
@@ -1656,7 +1732,7 @@ void FNewSpline::SetScale(int32 Index, const FVector& InScale)
 
 FVector FNewSpline::GetScale(const int32 Index) const
 {
-    if (Index < 0 || Index >= GetNumPoints())
+    if (Index < 0 || Index >= GetSpline().GetNumPoints())
     {
     	return FVector::OneVector;
     }
@@ -1666,29 +1742,29 @@ FVector FNewSpline::GetScale(const int32 Index) const
 
 void FNewSpline::SetSplinePointType(int32 Index, EInterpCurveMode Type)
 {
-    if (Index < 0 || Index >= GetNumPoints())
+    if (Index < 0 || Index >= GetSpline().GetNumPoints())
     {
     	return;
     }
 	
-	SetPointTangentMode(Index, ConvertInterpCurveModeToTangentMode(Type));
+	GetSpline().SetPointTangentMode(Index, ConvertInterpCurveModeToTangentMode(Type));
     
     RebuildLegacyCurves();
 }
 
 EInterpCurveMode FNewSpline::GetSplinePointType(int32 Index) const
 {
-    if (Index < 0 || Index >= GetNumPoints())
+    if (Index < 0 || Index >= GetSpline().GetNumPoints())
     {
     	return CIM_Unknown;
     }
-	ETangentMode TangentMode = GetTangentMode(Index);
+	ETangentMode TangentMode = GetSpline().GetTangentMode(Index);
     return ConvertTangentModeToInterpCurveMode(TangentMode);
 }
 
 float FNewSpline::GetParameterAtIndex(int32 Index) const
 {
-    if (Index < 0 || Index >= GetNumPoints())
+    if (Index < 0 || Index >= GetSpline().GetNumPoints())
     {
     	return 0.f;
     }
@@ -1699,7 +1775,7 @@ float FNewSpline::GetParameterAtIndex(int32 Index) const
 float FNewSpline::GetParameterAtDistance(float Distance) const
 {
 	const float NormalizedDistance = FMath::Clamp(Distance / FMath::Max(GetSplineLength(), UE_SMALL_NUMBER), 0.f, 1.f);
-	return FTangentBezierSpline3d::EvaluateAttribute<float>(LengthToParamAttrName, NormalizedDistance);
+	return EvaluateAttribute<float>(LengthToParamAttrName, NormalizedDistance);
 }
 
 float FNewSpline::GetDistanceAtParameter(float Parameter) const
@@ -1708,7 +1784,7 @@ float FNewSpline::GetDistanceAtParameter(float Parameter) const
 	// the resulting value as if it were in internal parent space.
 	// see the todo in UpdateReparamTable
 	Parameter = ToNormalizedSpace(Parameter);
-	return FTangentBezierSpline3d::EvaluateAttribute<float>(ParamToLengthAttrName, Parameter);
+	return EvaluateAttribute<float>(ParamToLengthAttrName, Parameter);
 }
 
 float FNewSpline::FindNearest(const FVector& InLocation, float& OutSquaredDistance) const
@@ -1719,7 +1795,7 @@ float FNewSpline::FindNearest(const FVector& InLocation, float& OutSquaredDistan
     }
     else
     {
-	    return FromInternalSplineSpace(FTangentBezierSpline3d::FindNearest(InLocation, OutSquaredDistance));
+	    return FromInternalSplineSpace(GetSpline().FindNearest(InLocation, OutSquaredDistance));
     }
 }
 
@@ -1736,20 +1812,20 @@ float FNewSpline::FindNearestOnSegment(const FVector& InLocation, int32 SegmentI
     }
     else
     {
-	   return FromInternalSplineSpace(FTangentBezierSpline3d::FindNearestOnSegment(InLocation, SegmentIndex, OutSquaredDistance));
+	   return FromInternalSplineSpace(GetSpline().FindNearestOnSegment(InLocation, SegmentIndex, OutSquaredDistance));
     }
 }
 
 FVector FNewSpline::EvaluateImpl(float Parameter) const
 {
 	Parameter = ToInternalSplineSpace(Parameter);
-    return FVector(GetPosition(Parameter));
+    return GetSpline().GetValue(Parameter);
 }
 
 FVector FNewSpline::EvaluateDerivative(float Parameter) const
 {
 	Parameter = ToInternalSplineSpace(Parameter);
-    return FVector(GetTangent(Parameter));
+    return GetSpline().GetTangent(Parameter);
 }
 
 FQuat FNewSpline::EvaluateRotation(float Parameter) const
@@ -1764,7 +1840,7 @@ FVector FNewSpline::EvaluateScale(float Parameter) const
 
 void FNewSpline::UpdateSpline()
 {
-    UpdateTangents();			// updates the tangents on our version of the points
+    GetSpline().UpdateTangents();			// updates the tangents on our version of the points
     UpdateReparamTable();		// updates both reparam tables
     RebuildLegacyCurves();		// updates legacy curves based on points, never evaluates internal spline.
 }
@@ -1774,7 +1850,7 @@ void FNewSpline::UpdateSpline(const FSpline::FUpdateSplineParams& Params)
 	// todo: handled unhandled params
     
     SetClosedLoop(Params.bClosedLoop);
-    bStationaryEndpoint = Params.bStationaryEndpoints;
+	GetSpline().SetStationaryEndpoints(Params.bStationaryEndpoints);
     ReparamStepsPerSegment = Params.ReparamStepsPerSegment;
 
     UpdateSpline();
@@ -1782,32 +1858,32 @@ void FNewSpline::UpdateSpline(const FSpline::FUpdateSplineParams& Params)
 
 float FNewSpline::FromNormalizedSpace(float Parameter) const
 {
-	return GetNumSegments() * Parameter;
+	return GetSpline().GetNumSegments() * Parameter;
 }
 
 float FNewSpline::ToNormalizedSpace(float Parameter) const
 {
-	return Parameter / GetNumSegments();
+	return Parameter / GetSpline().GetNumSegments();
 }
 
 float FNewSpline::FromInternalSplineSpace(float Parameter) const
 {
 	float OutLocalParam;
-	const int32 SegmentIndex = this->FindSegmentIndex(Parameter, OutLocalParam);
+	const int32 SegmentIndex = GetSpline().FindSegmentIndex(Parameter, OutLocalParam);
 	return static_cast<float>(SegmentIndex) + OutLocalParam;
 }
 float FNewSpline::ToInternalSplineSpace(float Parameter) const
 {
-	const int32 NumSegments = GetNumSegments();
+	const int32 NumSegments = GetSpline().GetNumSegments();
 	float ClampedParam = FMath::Clamp(Parameter, 0.0f, static_cast<float>(NumSegments));
 	const int32 SegmentIndex = FMath::Min(FMath::FloorToInt(ClampedParam), NumSegments - 1);
 	const float LocalT = ClampedParam - static_cast<float>(SegmentIndex); // Parameter [0,1] within segment
-	return GetSegmentParameterRange(SegmentIndex).Interpolate(LocalT);
+	return GetSpline().GetSegmentParameterRange(SegmentIndex).Interpolate(LocalT);
 }
 
 void FNewSpline::UpdateReparamTable()
 {
-    const int32 NumPoints = GetNumPoints();
+    const int32 NumPoints = GetSpline().GetNumPoints();
   	const int32 NumSegments = NumPoints - 1;
     const int32 TotalSteps = NumSegments * ReparamStepsPerSegment + 1;
 	
@@ -1832,20 +1908,20 @@ void FNewSpline::UpdateReparamTable()
     		const float SplineParam = (static_cast<float>(SegmentIndex * ReparamStepsPerSegment + Step) / TotalSteps);
     		const float CurrentLength = OffsetLength + AccumulatedLength;
 
-    		FTangentBezierSpline3d::AddAttributeValue(ParamToLengthAttrName, CurrentLength, SplineParam);
+    		AddAttributeValue<float>(ParamToLengthAttrName, CurrentLength, SplineParam);
     	}
 
     	AccumulatedLength += GetSegmentLength(SegmentIndex, 1.0f);
     }
-    FTangentBezierSpline3d::AddAttributeValue(ParamToLengthAttrName, AccumulatedLength, 1.0f);
+    AddAttributeValue<float>(ParamToLengthAttrName, AccumulatedLength, 1.0f);
 
     // Here we get Length -> Param mapping
     for (int32 Step = 0; Step <= NumSegments * ReparamStepsPerSegment + 1; ++Step)
     {
     	const float SplineParam = FMath::Clamp(static_cast<float>(Step) / TotalSteps, 0.0f, 1.0f);
-    	const float CurrentLength = FTangentBezierSpline3d::EvaluateAttribute<float>(ParamToLengthAttrName, SplineParam);
+    	const float CurrentLength = EvaluateAttribute<float>(ParamToLengthAttrName, SplineParam);
     	
-    	FTangentBezierSpline3d::AddAttributeValue(LengthToParamAttrName, FromNormalizedSpace(SplineParam), CurrentLength / AccumulatedLength);
+    	AddAttributeValue<float>(LengthToParamAttrName, FromNormalizedSpace(SplineParam), CurrentLength / AccumulatedLength);
     }
 }
 
@@ -1863,18 +1939,18 @@ void FNewSpline::RebuildLegacyCurves()
 	// if this fails, some other function which mutates InternalSpline is failing to keep Points appropriately sized
 	//check(Points.Num() == GetNumPoints());
 	
-    for (int32 i = 0; i < GetNumPoints(); ++i)
+    for (int32 i = 0; i < GetSpline().GetNumPoints(); ++i)
     {
     	const float AttributeParameter = static_cast<float>(i);
     	
     	// Additionally populate interp curves here so we don't double loop
-    	PositionCurve.Points.Emplace(AttributeParameter, GetPosition(i), GetTangentIn(i), GetTangentOut(i), ConvertTangentModeToInterpCurveMode(TangentModes[i]));
+    	PositionCurve.Points.Emplace(AttributeParameter, GetSpline().GetValue(i), GetSpline().GetTangentIn(i), GetSpline().GetTangentOut(i), ConvertTangentModeToInterpCurveMode(GetSpline().GetTangentModes()[i]));
     	RotationCurve.AddPoint(AttributeParameter, GetAttributeValue<FQuat>(RotationAttrName, i));
     	ScaleCurve.AddPoint(AttributeParameter, GetAttributeValue<FVector>(ScaleAttrName, i));
     }
 }
 
-FTangentBezierControlPoint FNewSpline::ConvertToTangentBezierControlPoint(const FSplinePoint& Point) const
+FNewSpline::FTangentBezierControlPoint FNewSpline::ConvertToTangentBezierControlPoint(const FSplinePoint& Point) const
 {
 	FTangentBezierControlPoint NewPoint;
 	NewPoint.Position = Point.Position;
@@ -1901,20 +1977,20 @@ void FNewSpline::UpdatePointAttributes(const FSplinePoint& Point, int32 PointInd
 float FNewSpline::ConvertIndexToParameter(int32 Index, float Fraction) const
 {
 	// Handle empty spline
-	if (GetNumPoints() <= 1)
+	if (GetSpline().GetNumPoints() <= 1)
 		return 0.0f;
     
 	// Clamp to valid range
-	Index = FMath::Clamp(Index, 0, GetNumPoints() - 1);
+	Index = FMath::Clamp(Index, 0, GetSpline().GetNumPoints() - 1);
     
 	// For exact index, just return parameter at that index
 	if (FMath::IsNearlyZero(Fraction))
-		return GetParameterAt(Index);
+		return GetSpline().GetParameterAt(Index);
     
 	// For fractional indices, interpolate
-	int32 NextIndex = FMath::Min(Index + 1, GetNumPoints() - 1);
-	float StartParam = GetParameterAt(Index);
-	float EndParam = GetParameterAt(NextIndex);
+	int32 NextIndex = FMath::Min(Index + 1, GetSpline().GetNumPoints() - 1);
+	float StartParam = GetSpline().GetParameterAt(Index);
+	float EndParam = GetSpline().GetParameterAt(NextIndex);
     
 	return FMath::Lerp(StartParam, EndParam, Fraction);
 }
@@ -1922,14 +1998,14 @@ float FNewSpline::ConvertIndexToParameter(int32 Index, float Fraction) const
 int32 FNewSpline::ConvertParameterToIndex(float Parameter) const
 {
 	// Handle empty spline
-	if (GetNumPoints() <= 1)
+	if (GetSpline().GetNumPoints() <= 1)
 		return 0;
     
 	// Find segment containing this parameter
-	for (int32 i = 0; i < GetNumPoints() - 1; i++)
+	for (int32 i = 0; i < GetSpline().GetNumPoints() - 1; i++)
 	{
-		float StartParam = GetParameterAt(i);
-		float EndParam = GetParameterAt(i + 1);
+		float StartParam = GetSpline().GetParameterAt(i);
+		float EndParam = GetSpline().GetParameterAt(i + 1);
         
 		if (Parameter >= StartParam && Parameter <= EndParam)
 		{
@@ -1940,7 +2016,7 @@ int32 FNewSpline::ConvertParameterToIndex(float Parameter) const
 	}
     
 	// Out of range
-	return (Parameter < GetParameterAt(0)) ? 0 : (GetNumPoints() - 1);
+	return (Parameter < GetSpline().GetParameterAt(0)) ? 0 : (GetSpline().GetNumPoints() - 1);
 }
 
 
diff --git a/Engine/Source/Runtime/GeometryCore/Public/Splines/BSpline.h b/Engine/Source/Runtime/GeometryCore/Public/Splines/BSpline.h
--- a/Engine/Source/Runtime/GeometryCore/Public/Splines/BSpline.h
+++ b/Engine/Source/Runtime/GeometryCore/Public/Splines/BSpline.h
@@ -386,7 +386,7 @@ public:
 	    
 	    for (int32 i = 1; i < Points.Num(); ++i)
 	    {
-	        const float Length = static_cast<float>((Points[i] - Points[i-1]).Size());
+	        const float Length = static_cast<float>(Math::Distance(Points[i], Points[i-1]));
 	        Chords.Add(Length);
 	        TotalLength += Length;
 	    }
@@ -394,7 +394,7 @@ public:
 	    // For closed loop, add the closing segment length
 	    if (this->IsClosedLoop() && Points.Num() > 0)
 	    {
-	        const float ClosureLength = static_cast<float>((Points[0] - Points.Last()).Size());
+	        const float ClosureLength = static_cast<float>(Math::Distance(Points[0], Points.Last()));
 	        Chords.Add(ClosureLength);
 	        TotalLength += ClosureLength;
 	    }
@@ -498,7 +498,7 @@ public:
 	    for (int32 i = 1; i < Points.Num(); ++i)
 	    {
 	        // Square root of distance for centripetal parameterization
-	        const float Length = FMath::Sqrt(static_cast<float>((Points[i] - Points[i-1]).Size()));
+	    	const float Length = static_cast<float>(Math::CentripetalDistance(Points[i], Points[i-1]));
 	        Lengths.Add(Length);
 	        TotalLength += Length;
 	    }
@@ -506,7 +506,7 @@ public:
 	    // For closed loop, add the closing segment
 	    if (this->IsClosedLoop() && Points.Num() > 0)
 	    {
-	        const float ClosureLength = FMath::Sqrt(static_cast<float>((Points[0] - Points.Last()).Size()));
+	        const float ClosureLength = static_cast<float>(Math::CentripetalDistance(Points[0], Points.Last()));
 	        Lengths.Add(ClosureLength);
 	        TotalLength += ClosureLength;
 	    }
diff --git a/Engine/Source/Runtime/GeometryCore/Public/Splines/CustomAttributeExample.h b/Engine/Source/Runtime/GeometryCore/Public/Splines/CustomAttributeExample.h
--- a/Engine/Source/Runtime/GeometryCore/Public/Splines/CustomAttributeExample.h
+++ b/Engine/Source/Runtime/GeometryCore/Public/Splines/CustomAttributeExample.h
@@ -17,7 +17,14 @@ namespace Spline
 
 /**
  * Simple custom attribute type for splines
- * Demonstrates the minimum requirements for a custom attribute type
+ * Demonstrates the requirements for a custom attribute type
+ * 
+ * Custom types need to provide certain operations for use with splines:
+ * 1. Basic arithmetic: operator-(const T&), operator+(const T&), operator*(float)
+ * 2. Optional helper methods that improve performance: Size(), SizeSquared(), Dot(), etc.
+ * 
+ * If your type doesn't provide some of these methods, spline operations will still work
+ * using fallback implementations, but may be less efficient.
  */
 struct FFoo
 {
@@ -158,7 +165,18 @@ public:
     }
 };
 
-// Example function demonstrating the FFoo type with the MultiSpline
+/**
+ * Example function demonstrating the FFoo type with the MultiSpline
+ * 
+ * FFoo implements all the necessary methods for use with splines:
+ * - Basic arithmetic operators
+ * - Size() and SizeSquared() for distance calculations
+ * - Dot() for dot product calculations
+ * - GetSafeNormal() for normalization
+ * 
+ * These methods are automatically detected by the UE::Geometry::Spline::Math
+ * functions, allowing FFoo to work seamlessly with all spline operations.
+ */
 void CreatePolyBezierWithFooAttributes()
 {
 	// Create a PolyBezier spline with a single segment
@@ -169,7 +187,7 @@ void CreatePolyBezierWithFooAttributes()
 
     // Create a MultiSpline that will host both the main spline and attribute channels
     TMultiSpline<FPolyBezierSpline3d> Spline;
-    Spline.Spline = BaseSpline; // Set the main spline
+    Spline.GetSpline() = BaseSpline; // Set the main spline
     
     // Create attribute channel using type registry by name
     Spline.CreateAttributeChannel<FFoo>("FooValues", TEXT("BSpline3"));
@@ -184,7 +202,7 @@ void CreatePolyBezierWithFooAttributes()
     }
 
     // Add another segment to the main spline
-    Spline.Spline.AppendBezierSegment(
+    Spline.GetSpline().AppendBezierSegment(
 		FVector(150, 50, 0),   // First control point
 		FVector(150, -50, 0),  // Second control point 
 		FVector(200, 0, 0)     // End point
@@ -206,7 +224,7 @@ void CreatePolyBezierWithFooAttributes()
 	// Evaluate at various points
     for (float t = 0.0f; t <= 2.0f; t += 0.1f)
 	{
-        FVector Position = Spline.Spline.Evaluate(t);
+        FVector Position = Spline.GetSpline().Evaluate(t);
 		FFoo Foo = Spline.EvaluateAttribute<FFoo>("FooValues", t);
     
 		UE_LOG(LogTemp, Display, TEXT("At t=%.1f: Position=(%.1f, %.1f, %.1f), Foo=%.2f"),
diff --git a/Engine/Source/Runtime/GeometryCore/Public/Splines/MultiSpline.h b/Engine/Source/Runtime/GeometryCore/Public/Splines/MultiSpline.h
--- a/Engine/Source/Runtime/GeometryCore/Public/Splines/MultiSpline.h
+++ b/Engine/Source/Runtime/GeometryCore/Public/Splines/MultiSpline.h
@@ -6,6 +6,7 @@
 #include "SplineInterfaces.h"
 #include "BoxTypes.h"
 #include "SplineTypeRegistry.h"
+#include "Spline.h"
 
 namespace UE
 {
@@ -19,7 +20,7 @@ namespace Spline
  * Provides attribute functionality while allowing direct access to implementation-specific methods
  */
 template <typename SPLINETYPE>
-class TMultiSpline
+class TMultiSpline : public TSplineWrapper<SPLINETYPE>
 {
 public:
     enum class EMappingRangeSpace: uint8
@@ -110,10 +111,10 @@ private:
     };
 
 public:
-    typedef SPLINETYPE SplineType;
-    typedef typename SplineType::ValueType ValueType;
-
-    SplineType Spline;
+    using SplineType = typename TSplineWrapper<SPLINETYPE>::SplineType;
+	
+    SplineType& GetSpline() { return TSplineWrapper<SPLINETYPE>::InternalSpline; }
+    const SplineType& GetSpline() const { return TSplineWrapper<SPLINETYPE>::InternalSpline; }
     
     /**
      * Gets an attribute channel of the specified type and tries to cast it to the expected implementation type
@@ -174,7 +175,7 @@ public:
     {
         if (const FChildSpline* Child = Children.Find(Name))
         {
-            return Child->MapParameterToChildSpace(Parameter, Spline.GetParameterSpace());
+            return Child->MapParameterToChildSpace(Parameter, GetSpline().GetParameterSpace());
         }
 
         return 0.f;
@@ -190,7 +191,7 @@ public:
     {
         if (const FChildSpline* Child = Children.Find(Name))
         {
-            return Child->MapParameterFromChildSpace(Parameter, Spline.GetParameterSpace());
+            return Child->MapParameterFromChildSpace(Parameter, GetSpline().GetParameterSpace());
         }
         
         return 0.f;
@@ -203,7 +204,7 @@ public:
      * @return True if channel was created successfully
      */
     template <typename AttrType>
-    bool CreateAttributeChannel(const FName& Name, const FName& ImplName = FName(TEXT("BSpline3")))
+    bool CreateAttributeChannel(const FName& Name, const FName& ImplName)
     {
         // Check if channel already exists
         if (Children.Contains(Name))
@@ -247,6 +248,35 @@ public:
         return true;
     }
 
+    /**
+     * Clears all values from a specific attribute channel
+     * @param Name - Name of the attribute channel to clear
+     * @return true if the channel was found and cleared, false otherwise
+     */
+    bool ClearAttributeChannel(const FName& Name)
+    {
+        if (const FChildSpline* Channel = Children.Find(Name))
+        {
+            Channel->Spline->Clear();
+            return true;
+        }
+        return false;
+    }
+
+    /**
+     * Checks for an attribute channel called Name
+     * @param Name - Name of the attribute channel to find
+     * @return true if an attribute channel called Name exists.
+     */
+    bool HasAttributeChannel(const FName& Name) const
+    {
+        if (const FChildSpline* Channel = Children.Find(Name))
+        {
+            return true;
+        }
+        return false;
+    }
+    
     /**
      * Gets a value from an attribute channel at the specified parameter
      * @param Name Name of the attribute channel
diff --git a/Engine/Source/Runtime/GeometryCore/Public/Splines/ParameterizedTypes.h b/Engine/Source/Runtime/GeometryCore/Public/Splines/ParameterizedTypes.h
--- a/Engine/Source/Runtime/GeometryCore/Public/Splines/ParameterizedTypes.h
+++ b/Engine/Source/Runtime/GeometryCore/Public/Splines/ParameterizedTypes.h
@@ -95,7 +95,75 @@ struct THasArchiveOperator<T,
 	std::void_t<decltype(std::declval<FArchive&>() << std::declval<T&>())>
 > : std::true_type {};
 
+/**
+ * Trait detectors for various value type operations
+ * These traits are used to determine if a type supports certain operations
+ */
+
+// Detector for Size() method
+template<typename T, typename = void>
+struct THasSizeMethod : std::false_type {};
 
+template<typename T>
+struct THasSizeMethod<T, std::void_t<decltype(std::declval<T>().Size())>>
+    : std::true_type {};
+
+// Detector for SizeSquared() method
+template<typename T, typename = void>
+struct THasSizeSquaredMethod : std::false_type {};
+
+template<typename T>
+struct THasSizeSquaredMethod<T, std::void_t<decltype(std::declval<T>().SizeSquared())>>
+    : std::true_type {};
+
+// Detector for Dot() method
+template<typename T, typename = void>
+struct THasDotMethod : std::false_type {};
+
+template<typename T>
+struct THasDotMethod<T, std::void_t<decltype(std::declval<T>().Dot(std::declval<T>()))>>
+    : std::true_type {};
+
+// Detector for GetSafeNormal() method
+template<typename T, typename = void>
+struct THasGetSafeNormalMethod : std::false_type {};
+
+template<typename T>
+struct THasGetSafeNormalMethod<T, std::void_t<decltype(std::declval<T>().GetSafeNormal())>>
+    : std::true_type {};
+
+// Detector for operator- (subtraction)
+template<typename T, typename = void>
+struct THasSubtractionOperator : std::false_type {};
+
+template<typename T>
+struct THasSubtractionOperator<T, std::void_t<decltype(std::declval<T>() - std::declval<T>())>>
+    : std::true_type {};
+
+// Detector for operator+ (addition)
+template<typename T, typename = void>
+struct THasAdditionOperator : std::false_type {};
+
+template<typename T>
+struct THasAdditionOperator<T, std::void_t<decltype(std::declval<T>() + std::declval<T>())>>
+    : std::true_type {};
+
+// Detector for operator* (scalar multiplication)
+template<typename T, typename = void>
+struct THasScalarMultiplicationOperator : std::false_type {};
+
+template<typename T>
+struct THasScalarMultiplicationOperator<T, std::void_t<decltype(std::declval<T>() * std::declval<float>())>>
+    : std::true_type {};
+
+// Detector for ZeroVector static member (for proper Zero value)
+template<typename T, typename = void>
+struct THasZeroVectorMember : std::false_type {};
+
+template<typename T>
+struct THasZeroVectorMember<T, std::void_t<decltype(T::ZeroVector)>>
+    : std::true_type {};
+	
 
 } // end namespace UE::Geometry::Spline
 } // end namespace UE::Geometry
diff --git a/Engine/Source/Runtime/GeometryCore/Public/Splines/PolyBezierSpline.h b/Engine/Source/Runtime/GeometryCore/Public/Splines/PolyBezierSpline.h
--- a/Engine/Source/Runtime/GeometryCore/Public/Splines/PolyBezierSpline.h
+++ b/Engine/Source/Runtime/GeometryCore/Public/Splines/PolyBezierSpline.h
@@ -123,8 +123,8 @@ public:
            ValueType P3 = Points[1];
            
            // Get normalized direction vectors from center to points
-           ValueType Dir0 = (P0 - Center).GetSafeNormal();
-           ValueType Dir3 = (P3 - Center).GetSafeNormal();
+           ValueType Dir0 = Math::GetSafeNormal(P0 - Center);
+           ValueType Dir3 = Math::GetSafeNormal(P3 - Center);
            
            // Calculate perpendicular vectors (normalized)
            ValueType Perp0 = ValueType(-Dir0.Y, Dir0.X, 0);
@@ -148,8 +148,8 @@ public:
                P3 = Points[i+1];
                
                // Calculate exact perpendicular vectors for this segment
-               Dir0 = (P0 - Center).GetSafeNormal();
-               Dir3 = (P3 - Center).GetSafeNormal();
+               Dir0 = Math::GetSafeNormal(P0 - Center);
+               Dir3 = Math::GetSafeNormal(P3 - Center);
                Perp0 = ValueType(-Dir0.Y, Dir0.X, 0);
                Perp3 = ValueType(-Dir3.Y, Dir3.X, 0);
                
@@ -294,9 +294,9 @@ public:
 
     	// Setup coefficients in standard Bezier polynomial form relative to test point
     	Coeffs[0] = P0 - Point;                    // constant term
-    	Coeffs[1] = 3 * (P1 - P0);                 // linear coefficient
-    	Coeffs[2] = 3 * (P2 - 2*P1 + P0);          // quadratic coefficient  
-    	Coeffs[3] = P3 - 3*P2 + 3*P1 - P0;         // cubic coefficient
+    	Coeffs[1] = (P1 - P0) * 3;                 // linear coefficient
+    	Coeffs[2] = (P2 - P1*2 + P0) * 3;          // quadratic coefficient  
+    	Coeffs[3] = P3 - P2*3 + P1*3 - P0;         // cubic coefficient
 
     	float LocalT = Math::FindNearestPoint_Cubic(Coeffs, 0.0f, 1.0f, OutSquaredDistance);
     	// Get the actual parameter bounds for this segment from the knot vector
@@ -441,7 +441,7 @@ public:
 	            // For subsequent segments, check if P0 matches the last point of the previous segment
 	            ValueType LastPoint = Base::GetValue(Base::NumKeys() - 1);
 	            
-	            if ((P0 - LastPoint).SizeSquared() < UE_KINDA_SMALL_NUMBER)
+	            if (Math::SizeSquared(P0 - LastPoint) < UE_KINDA_SMALL_NUMBER)
 	            {
 	                // Points are the same, use AppendBezierSegment which only needs P1, P2, P3
 	                this->AppendBezierSegment(P1, P2, P3, ParameterizationPolicy);
@@ -464,7 +464,7 @@ public:
 	        ValueType LastPoint = Base::GetValue(Base::NumKeys() - 1);
 	        
 	        // If not already a closed loop, add a segment to close it
-	        if ((LastPoint - FirstPoint).SizeSquared() > UE_KINDA_SMALL_NUMBER && NumSegments > 0)
+	        if (Math::SizeSquared(LastPoint - FirstPoint) > UE_KINDA_SMALL_NUMBER && NumSegments > 0)
 	        {
 	            // Add a segment that connects back to the start
 	            // We need to compute appropriate control points
@@ -475,9 +475,9 @@ public:
 	            ValueType InTangent = (Points[1] - Points[0]);
 	            
 	            // Scale tangents to match segment length
-	            float SegmentLength = static_cast<float>((FirstPoint - LastPoint).Size());
-	            OutTangent = OutTangent.GetSafeNormal() * (SegmentLength * 0.33f);
-	            InTangent = InTangent.GetSafeNormal() * (SegmentLength * 0.33f);
+	            float SegmentLength = static_cast<float>(Math::Distance(FirstPoint, LastPoint));
+	            OutTangent = Math::GetSafeNormal(OutTangent) * (SegmentLength * 0.33f);
+	            InTangent = Math::GetSafeNormal(InTangent) * (SegmentLength * 0.33f);
 	            
 	            // Create control points
 	            ValueType P1 = LastPoint + OutTangent;
@@ -557,9 +557,8 @@ public:
 	        const ValueType& LastExistingPoint = this->Base::GetValue(this->Base::NumKeys() - 1);
 	        
 	        // Check if the connection point is reasonably close
-	        constexpr double ConnectionTolerance = 1e-4;
-	        const double DistanceSquared = (LastExistingPoint - FirstNewPoint).SizeSquared();
-	        if (DistanceSquared > ConnectionTolerance)
+	        const double DistanceSquared = Math::SizeSquared(LastExistingPoint, FirstNewPoint);
+	        if (DistanceSquared > UE_KINDA_SMALL_NUMBER)
 	        {
 	            UE_LOG(LogSpline, Warning, TEXT("Connection point mismatch when appending segments. Distance: %f"),
 	                FMath::Sqrt(static_cast<float>(DistanceSquared)));
@@ -573,9 +572,8 @@ public:
 	        const ValueType& FirstExistingPoint = this->Base::GetValue(0);
 	        
 	        // Check if the connection point is reasonably close
-	        constexpr double ConnectionTolerance = 1e-4;
-	        const double DistanceSquared = (LastNewPoint - FirstExistingPoint).SizeSquared();
-	        if (DistanceSquared > ConnectionTolerance)
+	        const double DistanceSquared = Math::SizeSquared(LastNewPoint, FirstExistingPoint);
+	        if (DistanceSquared > UE_KINDA_SMALL_NUMBER)
 	        {
 	            UE_LOG(LogSpline, Warning, TEXT("Connection point mismatch when prepending segments. Distance: %f"),
 	                FMath::Sqrt(static_cast<float>(DistanceSquared)));
@@ -1161,9 +1159,9 @@ private:
     	switch (ParameterizationPolicy)
     	{
     	case EParameterizationPolicy::ChordLength:
-    		return static_cast<float>((P3 - P0).Size());
+    		return static_cast<float>(Math::Distance(P0, P3));
     	case EParameterizationPolicy::Centripetal:
-    		return FMath::Sqrt(static_cast<float>((P3 - P0).Size()));
+    		return static_cast<float>(Math::CentripetalDistance(P0, P3));
     	case EParameterizationPolicy::Uniform:
     	default:
 			return 1.0f;
@@ -1399,7 +1397,7 @@ public:
 				ValueType LastPoint = this->GetValue((NumSegments - 1) * 4 + 3);
             
 				// If the last point isn't already equal to the first point
-				if (!FirstPoint.Equals(LastPoint, UE_SMALL_NUMBER))
+				if (!Math::Equals(FirstPoint, LastPoint, UE_SMALL_NUMBER))	
 				{
 					// Calculate control points for a smooth transition
 					ValueType Direction = FirstPoint - LastPoint;
@@ -1645,22 +1643,22 @@ private:
             
 			if (EndIndex < Points.Num())
 			{
-				float Length;
+				double Length;
                 
 				// Use appropriate distance calculation based on knot mode
 				if (Mode == EParameterizationPolicy::Centripetal)
 				{
 					// Square root of distance for centripetal mode
-					Length = FMath::Sqrt(static_cast<float>((Points[EndIndex] - Points[StartIndex]).Size()));
+					Length = Math::CentripetalDistance(Points[StartIndex], Points[EndIndex]);
 				}
 				else // Chord length mode
 				{
 					// Actual distance for chord length mode
-					Length = static_cast<float>((Points[EndIndex] - Points[StartIndex]).Size());
+					Length = Math::Distance(Points[EndIndex], Points[StartIndex]);
 				}
                 
-				SegmentLengths.Add(Length);
-				TotalLength += Length;
+				SegmentLengths.Add(static_cast<float>(Length));
+				TotalLength += static_cast<float>(Length);
 			}
 		}
         
@@ -1688,10 +1686,8 @@ private:
 				AccumLength += SegmentLengths[i];
                 
 				// Use normalized accumulated length as the knot value
-				const float t = AccumLength;/*TotalLength > UE_SMALL_NUMBER ? 
-					(AccumLength / TotalLength) * static_cast<float>(NumSegments) : // Scale to match segment count
-					static_cast<float>(i + 1);										// Fallback to uniform if total length ~ 0*/
-                
+				const float t = AccumLength;
+				
 				// Interior knots at segment boundaries with multiplicity 3
 				for (int32 j = 0; j < 3; ++j)
 				{
@@ -1701,7 +1697,6 @@ private:
 		}
 
 		// End knots with multiplicity 4
-		// const float EndKnot = static_cast<float>(NumSegments);
 		const float EndKnot = AccumLength + (SegmentLengths.Num() > 0 ? SegmentLengths.Last() : 0.0f);
 		while (CurrentKnot < NumKnots)
 		{
diff --git a/Engine/Source/Runtime/GeometryCore/Public/Splines/Spline.h b/Engine/Source/Runtime/GeometryCore/Public/Splines/Spline.h
--- a/Engine/Source/Runtime/GeometryCore/Public/Splines/Spline.h
+++ b/Engine/Source/Runtime/GeometryCore/Public/Splines/Spline.h
@@ -22,7 +22,7 @@ class TSplineWrapper : public TSplineInterface<typename SPLINETYPE::ValueType>
 {
 public:
     typedef SPLINETYPE SplineType;
-	typedef typename SPLINETYPE::ValueType ValueType;  // Explicitly define ValueType
+	typedef typename SplineType::ValueType ValueType;  // Explicitly define ValueType
 
     virtual ~TSplineWrapper() override = default;
 
@@ -51,9 +51,9 @@ public:
         { return InternalSpline.FindNearest(Point, OutSquaredDistance); }
     virtual TUniquePtr<ISplineInterface> Clone() const override
     {
-    	TUniquePtr<TSplineWrapper<SPLINETYPE>> ClonedWrapper = MakeUnique<TSplineWrapper<SPLINETYPE>>();
+    	TUniquePtr<TSplineWrapper<SplineType>> ClonedWrapper = MakeUnique<TSplineWrapper<SplineType>>();
     
-    	// Clone the internal spline (assuming SPLINETYPE has proper copy semantics)
+    	// Clone the internal spline (assuming SplineType has proper copy semantics)
     	ClonedWrapper->InternalSpline = InternalSpline;
     
     	// Copy infinity modes and any other wrapper-specific properties
@@ -68,7 +68,7 @@ public:
     	return Ar;	
     }
 protected:
-    SPLINETYPE InternalSpline;
+    SplineType InternalSpline;
 };
 
 /**
diff --git a/Engine/Source/Runtime/GeometryCore/Public/Splines/SplineMath.h b/Engine/Source/Runtime/GeometryCore/Public/Splines/SplineMath.h
--- a/Engine/Source/Runtime/GeometryCore/Public/Splines/SplineMath.h
+++ b/Engine/Source/Runtime/GeometryCore/Public/Splines/SplineMath.h
@@ -2,6 +2,7 @@
 
 #pragma once
 #include "CoreMinimal.h"
+#include "ParameterizedTypes.h"
 namespace UE
 {
 namespace Geometry
@@ -11,6 +12,104 @@ namespace Spline
 namespace Math
 {
 
+
+/**
+ * Size/magnitude of a value
+ */
+template<typename T>
+double Size(const T& Value)
+{
+    if constexpr (THasSizeMethod<T>::value)
+        return static_cast<double>(Value.Size());
+    else if constexpr (std::is_arithmetic_v<T>)
+        return FMath::Abs(static_cast<double>(Value));
+    else
+        return 1.0; // Default for other types
+}
+
+/**
+ * Squared size (optimization to avoid sqrt)
+ */
+template<typename T>
+double SizeSquared(const T& Value)
+{
+    if constexpr (THasSizeSquaredMethod<T>::value)
+        return static_cast<double>(Value.SizeSquared());
+    else if constexpr (std::is_arithmetic_v<T>)
+        return static_cast<double>(Value * Value);
+    else
+        return Size(Value) * Size(Value);
+}
+
+/**
+ * Distance between values
+ */
+template<typename T>
+double Distance(const T& A, const T& B)
+{
+    if constexpr (THasSubtractionOperator<T>::value)
+        return Size(B - A);
+    else
+        return 1.0; // Default fallback
+}
+	
+/**
+ * Centripetal distance calculation (for parameterization)
+ */
+template<typename T>
+double CentripetalDistance(const T& A, const T& B)
+{
+	return FMath::Sqrt(Distance(A, B));
+}
+	
+/**
+ * Dot product between values
+ */
+template<typename T>
+double Dot(const T& A, const T& B)
+{
+    if constexpr (THasDotMethod<T>::value)
+        return static_cast<double>(A.Dot(B));
+    else if constexpr (std::is_arithmetic_v<T>)
+        return static_cast<double>(A * B);
+    else
+        return 0.0; // Fallback
+}
+
+/**
+ * Get normalized version of a value
+ */
+template<typename T>
+T GetSafeNormal(const T& Value)
+{
+    if constexpr (THasGetSafeNormalMethod<T>::value)
+        return Value.GetSafeNormal();
+    else if constexpr (std::is_arithmetic_v<T>)
+    {
+        const double AbsValue = FMath::Abs(static_cast<double>(Value));
+        return AbsValue < UE_KINDA_SMALL_NUMBER ? T(0) : (Value > 0 ? T(1) : T(-1));
+    }
+    else
+        return Value; // Fallback
+}
+
+/**
+ * Check if two values are equal within tolerance
+ */
+template<typename T>
+bool Equals(const T& A, const T& B, float Tolerance = UE_KINDA_SMALL_NUMBER)
+{
+    if constexpr (THasEqualsMethod<T>::value)
+        return A.Equals(B, Tolerance);
+    else if constexpr (std::is_arithmetic_v<T>)
+        return FMath::Abs(static_cast<float>(A - B)) <= Tolerance;
+    else if constexpr (THasSubtractionOperator<T>::value)
+        return Distance(A, B) <= Tolerance;
+    else
+        return false; // Fallback
+}
+
+	
 /** Helper for validating types and inputs */
 struct FSplineValidation
 {
@@ -82,7 +181,7 @@ struct FSplineValidation
 		// Check for degenerate/repeated points
 		for (int32 i = 1; i < Window.Num(); ++i)
 		{
-			if ((*Window[i] - *Window[i-1]).SizeSquared() < Tolerance)
+			if (DistanceSquared(*Window[i], *Window[i-1]) < Tolerance)
 			{
 				return false;
 			}
@@ -540,10 +639,10 @@ float FindNearestPoint_Cubic(
 	const ValueType& P0 = Coeffs[0]; // Already relative to target
 	const ValueType At1 = P0 + Coeffs[1] + Coeffs[2] + Coeffs[3]; // Curve at t=1
 
-	float BestDistSq = static_cast<float>(P0.SizeSquared());
+	float BestDistSq = static_cast<float>(Math::SizeSquared(P0));
 	float BestParam = StartParam;
 
-	const float EndDistSq = static_cast<float>(At1.SizeSquared());
+	const float EndDistSq = static_cast<float>(Math::SizeSquared(At1));
 	if (EndDistSq < BestDistSq)
 	{
 		BestDistSq = EndDistSq;
@@ -556,12 +655,12 @@ float FindNearestPoint_Cubic(
 	const ValueType& P3 = Coeffs[3];
 
 	TStaticArray<double, 6> RootCoeffs;
-	RootCoeffs[5] = 3.0 * P3.Dot(P3);
-	RootCoeffs[4] = 5.0 * P3.Dot(P2);
-	RootCoeffs[3] = 4.0 * P3.Dot(P1) + 2.0 * P2.Dot(P2);
-	RootCoeffs[2] = 3.0 * P2.Dot(P1) + 3.0 * P3.Dot(P0);
-	RootCoeffs[1] = P1.Dot(P1) + 2.0 * P2.Dot(P0);
-	RootCoeffs[0] = P1.Dot(P0);
+	RootCoeffs[5] = 3.0 * Dot(P3, P3);
+	RootCoeffs[4] = 5.0 * Dot(P3, P2);
+	RootCoeffs[3] = 4.0 * Dot(P3, P1) + 2.0 * Dot(P2, P2);
+	RootCoeffs[2] = 3.0 * Dot(P2, P1) + 3.0 * Dot(P3, P0);
+	RootCoeffs[1] = Dot(P1, P1) + 2.0 * Dot(P2, P0);
+	RootCoeffs[0] = Dot(P1, P0);
 
 	// Find roots in [0,1]  
 	UE::Math::TPolynomialRootSolver<double, 5> RootSolver(RootCoeffs, 0, 1);
@@ -577,7 +676,7 @@ float FindNearestPoint_Cubic(
 								P2 * (t * t) +
 								P3 * (t * t * t);
 
-		const float DistSq = static_cast<float>(Point.SizeSquared());
+		const float DistSq = static_cast<float>(Math::SizeSquared(Point));
 		if (DistSq < BestDistSq)
 		{
 			BestDistSq = DistSq;
diff --git a/Engine/Source/Runtime/GeometryCore/Public/Splines/TangentBezierSpline.h b/Engine/Source/Runtime/GeometryCore/Public/Splines/TangentBezierSpline.h
--- a/Engine/Source/Runtime/GeometryCore/Public/Splines/TangentBezierSpline.h
+++ b/Engine/Source/Runtime/GeometryCore/Public/Splines/TangentBezierSpline.h
@@ -40,43 +40,44 @@ enum class ETangentMode : uint8
 /**
  * Structure representing a control point on a tangent-based spline curve
  */
-struct FTangentBezierControlPoint 
+template<typename ValueType>
+struct TTangentBezierControlPoint 
 {
 	/** Position of the control point */
-	FVector Position;
+	ValueType Position;
 
 	/** Incoming tangent vector */
-	FVector TangentIn;
+	ValueType TangentIn;
 
 	/** Outgoing tangent vector */
-	FVector TangentOut;
+	ValueType TangentOut;
 
 	/** Tangent computation mode */
 	ETangentMode TangentMode;
 
 	/** Default constructor creates a point at origin with zero tangents */
-	FTangentBezierControlPoint ()
-		: Position(FVector::ZeroVector)
-		, TangentIn(FVector::ZeroVector)
-		, TangentOut(FVector::ZeroVector)
+	TTangentBezierControlPoint ()
+		: Position(ValueType())
+		, TangentIn(ValueType())
+		, TangentOut(ValueType())
 		, TangentMode(ETangentMode::Auto)
 	{
 	}
 
 	/** Constructor with position only - defaults to auto tangents */
-	explicit FTangentBezierControlPoint (const FVector& InPosition)
+	explicit TTangentBezierControlPoint (const ValueType& InPosition)
 		: Position(InPosition)
-		, TangentIn(FVector::ZeroVector)
-		, TangentOut(FVector::ZeroVector)
+		, TangentIn(ValueType())
+		, TangentOut(ValueType())
 		, TangentMode(ETangentMode::Auto)
 	{
 	}
 
 	/** Full constructor */
-	FTangentBezierControlPoint (
-		const FVector& InPosition,
-		const FVector& InTangentIn,
-		const FVector& InTangentOut,
+	TTangentBezierControlPoint (
+		const ValueType& InPosition,
+		const ValueType& InTangentIn,
+		const ValueType& InTangentOut,
 		ETangentMode InTangentMode)
 		: Position(InPosition)
 		, TangentIn(InTangentIn)
@@ -101,13 +102,13 @@ struct FTangentBezierControlPoint
 		
 	}
     
-	friend FArchive& operator<<(FArchive& Ar, FTangentBezierControlPoint& Point)
+	friend FArchive& operator<<(FArchive& Ar, TTangentBezierControlPoint& Point)
 	{
 		Point.Serialize(Ar);
 		return Ar;
 	}
 
-	bool operator==(const FTangentBezierControlPoint& Other) const
+	bool operator==(const TTangentBezierControlPoint& Other) const
 	{
 		return Position == Other.Position
 			&& TangentIn == Other.TangentIn
@@ -129,6 +130,7 @@ class TTangentBezierSpline :
 public:
 	typedef typename TSplineWrapper<TPolyBezierSpline<VALUETYPE>>::ValueType ValueType;
 	using TSplineWrapper<TPolyBezierSpline<VALUETYPE>>::InternalSpline;
+	using FTangentBezierControlPoint = TTangentBezierControlPoint<ValueType>;
 
 	// Generate compile-time type ID for TangentBezier
 	DECLARE_SPLINE_TYPE_ID(
@@ -140,7 +142,7 @@ public:
 	virtual ~TTangentBezierSpline() override = default;
 	
 	/** Default constructor with at least one segment */
-	TTangentBezierSpline(const FVector& StartPoint, const FVector& EndPoint)
+	TTangentBezierSpline(const ValueType& StartPoint, const ValueType& EndPoint)
 		: Tension(0.0f)
 		, TangentModes({ETangentMode::Auto, ETangentMode::Auto})
 	{
@@ -151,10 +153,10 @@ public:
 	 * Constructor with full control over tangents
 	 */
 	TTangentBezierSpline(
-		const FVector& StartPoint, 
-		const FVector& EndPoint,
-		const FVector& StartTangent,
-		const FVector& EndTangent,
+		const ValueType& StartPoint, 
+		const ValueType& EndPoint,
+		const ValueType& StartTangent,
+		const ValueType& EndTangent,
 		bool bAutoTangents = false)
 		: Tension(0.0f)
 		, TangentModes({bAutoTangents ? ETangentMode::Auto : ETangentMode::User,
@@ -225,6 +227,11 @@ public:
         return Ar;
     }
 
+	bool operator==(const TTangentBezierSpline<ValueType>& Other) const
+	{
+		return InternalSpline == Other.InternalSpline;
+	}
+
 	// Static shape generators
 
     /**
@@ -234,8 +241,8 @@ public:
      * @return New TangentBezierSpline initialized as a line
      */
     static TTangentBezierSpline CreateLine(
-        const FVector& StartPoint, 
-        const FVector& EndPoint)
+        const ValueType& StartPoint, 
+        const ValueType& EndPoint)
     {
 		// Create empty spline - we'll manually build it
 		TTangentBezierSpline Result;
@@ -257,7 +264,7 @@ public:
      * @return New TangentBezierSpline initialized as an arc
      */
     static TTangentBezierSpline CreateCircleArc(
-        const FVector& Center,
+        const ValueType& Center,
         float Radius,
         float StartAngle,
         float EndAngle,
@@ -268,7 +275,7 @@ public:
         
         // Clear the default initialization
         Result.InternalSpline = FPolyBezierSpline3d::CreateCircleArc(
-            FVector3d(Center), Radius, StartAngle, EndAngle, NumSegments);
+            Center, Radius, StartAngle, EndAngle, NumSegments);
         
 		// Set up tangent status for all points
 		Result.TangentModes.Init(ETangentMode::Auto, Result.GetNumPoints());
@@ -288,14 +295,14 @@ public:
      * @return New TangentBezierSpline initialized as a circle
      */
     static TTangentBezierSpline CreateCircle(
-        const FVector& Center,
+        const ValueType& Center,
         float Radius,
         int32 NumSegments = 4)
     {
         // Create the circle
         TTangentBezierSpline Result;
         Result.InternalSpline = FPolyBezierSpline3d::CreateCircle(
-            FVector3d(Center), Radius, NumSegments);
+            ValueType(Center), Radius, NumSegments);
         
         // Set up tangent status for all points
 		Result.TangentModes.Init(ETangentMode::Auto, Result.GetNumPoints());
@@ -315,7 +322,7 @@ public:
      * @return New TangentBezierSpline initialized as an ellipse
      */
     static TTangentBezierSpline CreateEllipse(
-        const FVector& Center,
+        const ValueType& Center,
         float RadiusX,
         float RadiusY,
         int32 NumSegments = 4)
@@ -323,7 +330,7 @@ public:
         // Create the ellipse
         TTangentBezierSpline Result;
         Result.InternalSpline = FPolyBezierSpline3d::CreateEllipse(
-            FVector3d(Center), RadiusX, RadiusY, NumSegments);
+            Center, RadiusX, RadiusY, NumSegments);
         
 		// Set up tangent status for all points
 		Result.TangentModes.Init(ETangentMode::Auto, Result.GetNumPoints());
@@ -369,16 +376,16 @@ public:
 	        if (NumPoints >= 2)
 	        {
 	            // Get first and last points
-	            const FVector FirstPos = GetPosition(0);
-	            const FVector LastPos = GetPosition(NumPoints - 1);
+	            const ValueType FirstPos = GetValue(0);
+	            const ValueType LastPos = GetValue(NumPoints - 1);
 	            
 	            // Compute tangents for the closure segment
-	            FVector LastOutTangent, FirstInTangent;
+	            ValueType LastOutTangent, FirstInTangent;
 	            
 	            if (IsAutoTangent(NumPoints - 1) && IsAutoTangent(0))
 	            {
 	                // Auto tangents for closure
-	                FVector ClosingTangent = (1.0f - Tension) * (FirstPos - LastPos);
+	                ValueType ClosingTangent = (FirstPos - LastPos) * (1.0f - Tension);
 	                LastOutTangent = ClosingTangent;
 	                FirstInTangent = -ClosingTangent;
 	            }
@@ -390,10 +397,10 @@ public:
 	            }
 	            
 	            // Calculate control points for the closing segment
-	            FVector P0 = LastPos;
-	            FVector P3 = FirstPos;
-	            FVector P1 = P0 + (LastOutTangent / 3.0f);
-	            FVector P2 = P3 - (FirstInTangent / 3.0f);
+	            ValueType P0 = LastPos;
+	            ValueType P3 = FirstPos;
+	            ValueType P1 = P0 + (LastOutTangent / 3.0f);
+	            ValueType P2 = P3 - (FirstInTangent / 3.0f);
 	            
 	            // Add the closing segment to the internal spline
 	            InternalSpline.AppendBezierSegment(P1, P2, P3);
@@ -440,7 +447,7 @@ public:
 	    }
 	    
 	    // Convert to cubic Bezier control points for the internal spline
-	    TArray<FVector> BezierPoints;
+	    TArray<ValueType> BezierPoints;
 	    
 	    // For each segment (between consecutive points), we need 4 control points
 	    // Reserve space based on number of segments
@@ -452,11 +459,11 @@ public:
 	        const FTangentBezierControlPoint& StartPoint = Points[i];
 	        const FTangentBezierControlPoint& EndPoint = Points[i + 1];
 	        
-	        FVector P0 = StartPoint.Position;
-	        FVector P3 = EndPoint.Position;
+	        ValueType P0 = StartPoint.Position;
+	        ValueType P3 = EndPoint.Position;
 	        
 	        // Calculate tangents based on the points' tangent modes
-	        FVector OutTangent, InTangent;
+	        ValueType OutTangent, InTangent;
 	        
 	        // Process start point's outgoing tangent
 	        if (StartPoint.TangentMode == ETangentMode::Auto || 
@@ -487,12 +494,12 @@ public:
 	            // Apply clamping for AutoClamped
 	            if (StartPoint.TangentMode == ETangentMode::AutoClamped)
 	            {
-	                const float SegmentLength = static_cast<float>((P3 - P0).Size());
+	                const float SegmentLength = Math::Distance(P3, P0);
 	                const float MaxTangentSize = SegmentLength * 0.33f;
 	                
-	                if (OutTangent.Size() > MaxTangentSize && !FMath::IsNearlyZero(OutTangent.Size()))
+	                if (Math::Size(OutTangent) > MaxTangentSize && !FMath::IsNearlyZero(Math::Size(OutTangent)))
 	                {
-	                    OutTangent = OutTangent.GetSafeNormal() * MaxTangentSize;
+	                    OutTangent = Math::GetSafeNormal(OutTangent) * MaxTangentSize;
 	                }
 	            }
 	        }
@@ -504,7 +511,7 @@ public:
 	        else if (StartPoint.TangentMode == ETangentMode::Constant)
 	        {
 	            // Zero tangent for constant interpolation
-	            OutTangent = FVector::ZeroVector;
+	            OutTangent = ValueType();
 	        }
 	        else
 	        {
@@ -541,12 +548,12 @@ public:
 	            // Apply clamping for AutoClamped
 	            if (EndPoint.TangentMode == ETangentMode::AutoClamped)
 	            {
-	                const float SegmentLength = static_cast<float>((P3 - P0).Size());
+	                const float SegmentLength = Math::Distance(P3, P0);
 	                const float MaxTangentSize = SegmentLength * 0.33f;
 	                
-	                if (InTangent.Size() > MaxTangentSize && !FMath::IsNearlyZero(InTangent.Size()))
+	                if (Math::Size(InTangent) > MaxTangentSize && !FMath::IsNearlyZero(Math::Size(InTangent)))
 	                {
-	                    InTangent = InTangent.GetSafeNormal() * MaxTangentSize;
+	                    InTangent = Math::GetSafeNormal(InTangent) * MaxTangentSize;
 	                }
 	            }
 	        }
@@ -558,7 +565,7 @@ public:
 	        else if (EndPoint.TangentMode == ETangentMode::Constant)
 	        {
 	            // Zero tangent for constant interpolation
-	            InTangent = FVector::ZeroVector;
+	            InTangent = ValueType();
 	        }
 	        else
 	        {
@@ -567,8 +574,8 @@ public:
 	        }
 	        
 	        // Calculate bezier control points
-	        FVector P1 = P0 + OutTangent / 3.0f;
-	        FVector P2 = P3 - InTangent / 3.0f;
+	        ValueType P1 = P0 + OutTangent / 3.0f;
+	        ValueType P2 = P3 - InTangent / 3.0f;
 	        
 	        // Add all 4 control points for this segment
 	        BezierPoints.Add(P0);
@@ -630,11 +637,11 @@ public:
                 if (i == 0)
                 {
                     const int32 ExistingNumPoints = GetNumPoints();
-                    const FVector LastExistingPoint = GetPosition(ExistingNumPoints - 1);
+                    const ValueType LastExistingPoint = GetValue(ExistingNumPoints - 1);
                     
                     // Check if connection points are reasonably close
                     constexpr double ConnectionTolerance = 1e-4;
-                    const double DistanceSquared = (LastExistingPoint - Points[0].Position).SizeSquared();
+                    const double DistanceSquared = Math::SizeSquared(LastExistingPoint - Points[0].Position);
                     
                     if (DistanceSquared <= ConnectionTolerance)
                     {
@@ -660,11 +667,11 @@ public:
                 // Skip the last point if it's the same as our first existing point
                 if (i == NumPoints - 1)
                 {
-                    const FVector FirstExistingPoint = GetPosition(0);
+                    const ValueType FirstExistingPoint = GetValue(0);
                     
                     // Check if connection points are reasonably close
                     constexpr double ConnectionTolerance = 1e-4;
-                    const double DistanceSquared = (FirstExistingPoint - Points[i].Position).SizeSquared();
+                    const double DistanceSquared = Math::SizeSquared(FirstExistingPoint - Points[i].Position);
                     
                     if (DistanceSquared <= ConnectionTolerance)
                     {
@@ -822,10 +829,10 @@ public:
     	}
 		        
 		// Calculate control points for the cubic Bezier segment
-		FVector P0 = ControlPoint.Position;
+		ValueType P0 = ControlPoint.Position;
         
 		// Handle tangent computation based on mode
-		FVector OutTangent;
+		ValueType OutTangent;
         
 		if (ControlPoint.TangentMode == ETangentMode::User || 
 			ControlPoint.TangentMode == ETangentMode::Broken)
@@ -836,11 +843,11 @@ public:
 		else
 		{
 			// For auto modes, use a reasonable default that will be updated later
-			OutTangent = GetPosition(0) - ControlPoint.Position;
+			OutTangent = GetValue(0) - ControlPoint.Position;
 		}
         
-		FVector P1 = P0 + (OutTangent / 3.0f);
-		FVector P2 = GetPosition(0) - (GetTangentIn(0) / 3.0f);
+		ValueType P1 = P0 + (OutTangent / 3.0f);
+		ValueType P2 = GetValue(0) - (GetTangentIn(0) / 3.0f);
         
 		// Let InternalSpline handle the prepending efficiently
 		InternalSpline.PrependBezierSegment(P0, P1, P2);
@@ -870,10 +877,10 @@ public:
     	
     	
 	    // Normal case - add segment to existing spline
-        FVector P0 = GetPosition(NewIndex - 1);
+        ValueType P0 = GetValue(NewIndex - 1);
 	    
 		        
-		FVector P3 = ControlPoint.Position;
+		ValueType P3 = ControlPoint.Position;
     
 		// Update the previous point's tangents if needed
 		if (TangentModes[NewIndex - 1] == ETangentMode::Auto || 
@@ -883,13 +890,13 @@ public:
 		}
     
 		// Get its outgoing tangent (may have been updated above)
-		FVector OutTangentPrev = GetTangentOut(NewIndex - 1);
+		ValueType OutTangentPrev = GetTangentOut(NewIndex - 1);
     
 		// Calculate control points for the new segment
-		FVector P1 = P0 + (OutTangentPrev / 3.0f);
+		ValueType P1 = P0 + (OutTangentPrev / 3.0f);
     
 		// For the incoming tangent of the new point
-		FVector P2;
+		ValueType P2;
     
 		if (ControlPoint.TangentMode == ETangentMode::User || 
 			ControlPoint.TangentMode == ETangentMode::Broken)
@@ -901,7 +908,7 @@ public:
 		{
 			// For other modes, we'll set up a reasonable default
 			// The UpdatePointTangents call below will update it properly
-			FVector DefaultTangent = (1.0f - Tension) * (P3 - P0);
+			ValueType DefaultTangent = (1.0f - Tension) * (P3 - P0);
 			P2 = P3 - (DefaultTangent / 3.0f);
 		}
     
@@ -923,7 +930,7 @@ public:
             }
         }
     }
-
+	
 	/**
 	 * Gets the control point at the specified index
 	 * @param Index Index of the point to retrieve
@@ -938,7 +945,7 @@ public:
 		}
     
 		FTangentBezierControlPoint Result;
-		Result.Position = GetPosition(Index);
+		Result.Position = GetValue(Index);
 		Result.TangentIn = GetTangentIn(Index);
 		Result.TangentOut = GetTangentOut(Index);
     
@@ -1052,7 +1059,7 @@ public:
             return;
         }
 
-        const FVector OldPosition = GetPosition(Index);
+        const ValueType OldPosition = GetValue(Index);
         bool bPositionChanged = (OldPosition != ControlPoint.Position);
         
         SetPointTangentMode(Index, ControlPoint.TangentMode);
@@ -1114,53 +1121,61 @@ public:
         }
     }
 
+	void SetValue(int32 Index, const ValueType& NewValue)
+	{
+		FTangentBezierControlPoint NewPoint = GetControlPoint(Index);
+		NewPoint.Position = NewValue;
+		ModifyPoint(Index, NewPoint);
+	}
+	
     /** Gets position at specified parent space parameter */
-    FVector GetPosition(float Parameter) const
+    ValueType GetValue(float Parameter) const
     {
-        return FVector(InternalSpline.Evaluate(Parameter));
+        return InternalSpline.Evaluate(Parameter);
     }
     
-    /** Gets position of the specified point */
-    FVector GetPosition(int32 Index) const
+    /** Gets the value of the specified point */
+    ValueType GetValue(int32 Index) const
     {
     	const int32 NumKeys = InternalSpline.NumKeys();
 
     	if (NumKeys == 0 || Index < 0 || Index >= GetNumPoints())
-    		return FVector::ZeroVector;
+    		return ValueType();
     
     	if (NumKeys == 1)
     	{
     		// Special case when we only have one point
     		if (Index == 0)
-    			return FVector(InternalSpline.GetValue(0));
-    		return FVector::ZeroVector;
+    			return InternalSpline.GetValue(0);
+    		
+    		return ValueType();
     	}
 
     	// Regular case - convert from point index to bezier control point index
     	if (Index == 0)
     	{
     		// First point is P0 of first segment
-    		return FVector(InternalSpline.GetValue(0));
+    		return InternalSpline.GetValue(0);
     	}
     	else
     	{
     		// Other points are P3 of previous segments
     		const int32 ControlPointIndex = (Index - 1) * 4 + 3;
     		if (ControlPointIndex < NumKeys)
-    			return FVector(InternalSpline.GetValue(ControlPointIndex));
+    			return InternalSpline.GetValue(ControlPointIndex);
     	}
 
-    	return FVector::ZeroVector;
+    	return ValueType();
     }
 
     /** Gets tangent at parameter in spline space */
-    FVector GetTangent(float Parameter) const
+    ValueType GetTangent(float Parameter) const
     {
-        return FVector(InternalSpline.template EvaluateDerivative<1>(Parameter));
+        return InternalSpline.template EvaluateDerivative<1>(Parameter);
     }
     
     /** Gets incoming tangent for the specified point */
-    FVector GetTangentIn(int32 Index) const
+    ValueType GetTangentIn(int32 Index) const
     {
         if (Index == 0)
         {
@@ -1169,46 +1184,46 @@ public:
             {
                 // For closed loop, look at last segment
                 const int32 LastSegment = InternalSpline.GetNumSegments() - 1;
-                FVector P3 = FVector(InternalSpline.GetValue(LastSegment * 4 + 3));
-                FVector P2 = FVector(InternalSpline.GetValue(LastSegment * 4 + 2));
+                ValueType P3 = InternalSpline.GetValue(LastSegment * 4 + 3);
+                ValueType P2 = InternalSpline.GetValue(LastSegment * 4 + 2);
                 return 3.0f * (P3 - P2);
             }
-            return FVector::ZeroVector;
+            return ValueType();
         }
         
         // Normal case - get from previous segment's control points
         const int32 BezierIndex = (Index - 1) * 4;
         if (BezierIndex + 3 < InternalSpline.NumKeys())
         {
-            FVector P3 = FVector(InternalSpline.GetValue(BezierIndex + 3));
-            FVector P2 = FVector(InternalSpline.GetValue(BezierIndex + 2));
+            ValueType P3 = InternalSpline.GetValue(BezierIndex + 3);
+            ValueType P2 = InternalSpline.GetValue(BezierIndex + 2);
             return 3.0f * (P3 - P2);
         }
-        return FVector::ZeroVector;
+        return ValueType();
     }
     
     /** Gets outgoing tangent for the specified point */
-    FVector GetTangentOut(int32 Index) const
+    ValueType GetTangentOut(int32 Index) const
     {
         const int32 NumPoints = GetNumPoints();
         if (Index == NumPoints - 1 && !IsClosedLoop())
 		{
-			return FVector::ZeroVector;
+			return ValueType();
 		}
         
         // Normal case - get from next segment's control points
         const int32 BezierIndex = Index * 4;
         if (BezierIndex + 1 < InternalSpline.NumKeys())
         {
-            FVector P0 = FVector(InternalSpline.GetValue(BezierIndex));
-            FVector P1 = FVector(InternalSpline.GetValue(BezierIndex + 1));
+            ValueType P0 = InternalSpline.GetValue(BezierIndex);
+            ValueType P1 = InternalSpline.GetValue(BezierIndex + 1);
             return 3.0f * (P1 - P0);
         }
-        return FVector::ZeroVector;
+        return ValueType();
     }
 
     /** Sets tangent in for the specified point */
-    void SetTangentIn(int32 Index, const FVector& NewTangent)
+    void SetTangentIn(int32 Index, const ValueType& NewTangent)
     {
 		if (Index < 0 || Index >= GetNumPoints())
 		{
@@ -1221,8 +1236,8 @@ public:
             if (InternalSpline.IsClosedLoop() && InternalSpline.GetNumSegments() > 0)
             {
                 const int32 LastSegment = InternalSpline.GetNumSegments() - 1;
-                FVector P3 = FVector(InternalSpline.GetValue(LastSegment * 4 + 3));
-                FVector P2 = P3 - (NewTangent / 3.0f);
+                ValueType P3 = InternalSpline.GetValue(LastSegment * 4 + 3);
+                ValueType P2 = P3 - (NewTangent / 3.0f);
                 UpdateBezierControlPoint(LastSegment * 4 + 2, P2);
             }
             return;
@@ -1232,14 +1247,14 @@ public:
         const int32 BezierIndex = (Index - 1) * 4 + 2;
         if (BezierIndex < InternalSpline.NumKeys())
         {
-            FVector P3 = FVector(InternalSpline.GetValue(BezierIndex + 1));
-            FVector P2 = P3 - (NewTangent / 3.0f);
+            ValueType P3 = InternalSpline.GetValue(BezierIndex + 1);
+            ValueType P2 = P3 - (NewTangent / 3.0f);
             UpdateBezierControlPoint(BezierIndex, P2);
         }
     }
     
     /** Sets tangent out for the specified point */
-    void SetTangentOut(int32 Index, const FVector& NewTangent)
+    void SetTangentOut(int32 Index, const ValueType& NewTangent)
     {
         const int32 NumPoints = GetNumPoints();
         if (Index == NumPoints - 1 && !IsClosedLoop())
@@ -1251,8 +1266,8 @@ public:
         const int32 BezierIndex = Index * 4 + 1;
         if (BezierIndex < InternalSpline.NumKeys())
         {
-            FVector P0 = FVector(InternalSpline.GetValue(BezierIndex - 1));
-            FVector P1 = P0 + (NewTangent / 3.0f);
+            ValueType P0 = InternalSpline.GetValue(BezierIndex - 1);
+            ValueType P1 = P0 + (NewTangent / 3.0f);
             UpdateBezierControlPoint(BezierIndex, P1);
         }
     }
@@ -1325,122 +1340,7 @@ public:
 	{
 		return InternalSpline.FindSegmentIndex(Parameter, OutLocalParam);
 	}
-    template <typename AttrType>
-    bool CreateAttributeChannel(const FName& Name)
-    {
-        return false;//InternalSpline.CreateAttributeChannel<AttrType>(Name);
-    }
-    
-    template <typename AttrType>
-    AttrType EvaluateAttribute(FName Name, float Parameter) const
-    {
-        return AttrType();//InternalSpline.EvaluateAttribute<AttrType>(Name, Parameter);
-    }
-
-    template <typename AttrType>
-    int32 AddAttributeValue(FName Name, const AttrType& Value, float Parameter)
-    {
-        return -1;//InternalSpline.AddAttributeValue<AttrType>(Name, Value, Parameter);
-    }
-
-	/**
-	 * Checks if the spline has an attribute channel with the given name
-	 * @param Name - Name of the attribute channel 
-	 * @return true if the channel exists, false otherwise
-	 */
-	bool HasAttributeChannel(const FName& Name) const
-	{
-		return InternalSpline.HasAttributeChannel(Name);
-	}
-
-	/**
-	 * Retrieves the number of values in a specific attribute channel
-	 * @param Name - Name of the attribute channel
-	 * @return Number of values or zero if the channel does not exist
-	 */
-	int32 NumAttributeValues(const FName& Name) const
-	{
-		return 0;//InternalSpline.NumAttributeValues(Name);
-	}
-
-	/**
-	 * Retrieves the parameter value of a specific attribute at the given index
-	 * @param Name - Name of the attribute channel
-	 * @param Index - Index of the value
-	 * @return Parameter value
-	 */
-	float GetAttributeParameter(const FName& Name, int32 Index) const
-	{
-		return 0;//InternalSpline.GetAttributeParameter(Name, Index);
-	}
 
-	/**
-	 * Retrieves all attribute channels in the spline
-	 * @return Map of attribute channels by name
-	 */
-	// const TMap<FName, TUniquePtr<IAttributeChannelBase>>& GetAttributeChannels() const
-	// {
-	// 	return TMap<FName, TUniquePtr<IAttributeChannelBase>>();//InternalSpline.GetAttributeChannels();
-	// }
-
-	/**
-	 * Updates the parameter of an attribute by index, possibly changing its index to preserve the ordering of the attribute storage
-	 * @param Name - Name of the attribute channel
-	 * @param Index - Index of the attribute to modify
-	 * @param ParentSpaceParameter - Parameter to set in parent spline space
-	 * @return New index of the attribute originally at the given index after any reordering
-	 */
-	int32 SetAttributeParameter(const FName& Name, int32 Index, float ParentSpaceParameter)
-	{
-		return -1;//InternalSpline.SetAttributeParameter(Name, Index, ParentSpaceParameter);
-	}
-	
-	/**
-	 * Retrieves the value of a specific attribute at the given index
-	 * @tparam AttrType - Custom attribute type
-	 * @param Name - Name of the attribute channel
-	 * @param Index - Index of the value
-	 * @return Attribute value
-	 */
-	template <typename AttrType>
-	AttrType GetAttributeValue(const FName& Name, int32 Index) const
-	{
-		return AttrType();//InternalSpline.GetAttributeValue<AttrType>(Name, Index);
-	}
-
-	/**
-	 * Removes a value from an attribute channel
-	 * @param Name - Name of the attribute channel
-	 * @param Index - Index of the value to remove
-	 * @return true if the value was removed, false otherwise
-	 */
-	bool RemoveAttributeValue(FName Name, int32 Index)
-	{
-		return false;//InternalSpline.RemoveAttributeValue(Name, Index);
-	}
-
-	/**
-	 * Sets the value of a specific attribute at the given index
-	 * @tparam AttrType - Custom attribute type
-	 * @param Name - Name of the attribute channel
-	 * @param Index - Index of the value
-	 * @param Value - New value to set
-	 */
-	template <typename AttrType>
-	void SetAttributeValue(const FName& Name, int32 Index, const AttrType& Value)
-	{
-		//InternalSpline.SetAttributeValue<AttrType>(Name, Index, Value);
-	}
-
-	/**
-	 * Clears all values from a specific attribute channel
-	 * @param Name - Name of the attribute channel to clear
-	 * @return true if the channel was found and cleared, false otherwise
-	 */
-	bool ClearAttributeChannel(const FName& Name)
-	{
-		return false;//InternalSpline.ClearAttributeChannel(Name);
-	}
     /**
 	 * Updates tangents for all control points based on their tangent modes
 	 */
@@ -1474,12 +1374,12 @@ public:
 		
 
 		// Get the position of the current point
-		const FVector& Current = GetPosition(Index);
+		const ValueType& Current = GetValue(Index);
 
 		// Get previous and next points (handling closed loop case)
-		const FVector& Prev = (Index > 0) ? GetPosition(Index - 1) : (IsClosedLoop() ? GetPosition(NumPoints - 1) : Current);
+		const ValueType& Prev = (Index > 0) ? GetValue(Index - 1) : (IsClosedLoop() ? GetValue(NumPoints - 1) : Current);
 
-		const FVector& Next = (Index < NumPoints - 1) ? GetPosition(Index + 1) : (IsClosedLoop() ? GetPosition(0) : Current);
+		const ValueType& Next = (Index < NumPoints - 1) ? GetValue(Index + 1) : (IsClosedLoop() ? GetValue(0) : Current);
 
 		// Get actual parameter values for points
 		float PrevParam = (Index > 0)
@@ -1501,8 +1401,8 @@ public:
 		}
 
 		// Calculate tangents based on mode
-		FVector InTangent = GetTangentIn(Index); // Keep existing if not changing
-		FVector OutTangent = GetTangentOut(Index);
+		ValueType InTangent = GetTangentIn(Index); // Keep existing if not changing
+		ValueType OutTangent = GetTangentOut(Index);
 		
 		switch (Mode)
 		{
@@ -1513,7 +1413,7 @@ public:
 				float PrevToNextParamDiff = FMath::Max<float>(UE_KINDA_SMALL_NUMBER, NextParam - PrevParam);
 
 				// Calculate basic tangent - matches legacy formula
-				FVector AutoTangent = (1.0f - Tension) * ((Next - Prev) / PrevToNextParamDiff);
+				ValueType AutoTangent = (1.0f - Tension) * ((Next - Prev) / PrevToNextParamDiff);
 
 				// For AutoClamped mode, apply clamping like the legacy implementation
 				if (Mode == ETangentMode::AutoClamped)
@@ -1521,25 +1421,25 @@ public:
 					float PrevToCurParamDiff = FMath::Max<float>(UE_KINDA_SMALL_NUMBER, CurrentParam - PrevParam);
 					float NextToCurParamDiff = FMath::Max<float>(UE_KINDA_SMALL_NUMBER, NextParam - CurrentParam);
 
-					FVector PrevToCurTangent = (Current - Prev) / PrevToCurParamDiff;
-					FVector NextToCurTangent = (Next - Current) / NextToCurParamDiff;
+					ValueType PrevToCurTangent = (Current - Prev) / PrevToCurParamDiff;
+					ValueType NextToCurTangent = (Next - Current) / NextToCurParamDiff;
 
-					float PrevToCurLength = static_cast<float>(PrevToCurTangent.Size());
-					float NextToCurLength = static_cast<float>(NextToCurTangent.Size());
+					float PrevToCurLength = static_cast<float>(Math::Size(PrevToCurTangent));
+					float NextToCurLength = static_cast<float>(Math::Size(NextToCurTangent));
 
 					// Clamp tangent magnitude
 					float MaxScale = FMath::Min(PrevToCurLength, NextToCurLength);
 
-					if (AutoTangent.Size() > MaxScale && !FMath::IsNearlyZero(AutoTangent.Size()))
+					if (Math::Size(AutoTangent) > MaxScale && !FMath::IsNearlyZero(Math::Size(AutoTangent)))
 					{
-						AutoTangent = AutoTangent.GetSafeNormal() * MaxScale;
+						AutoTangent = Math::GetSafeNormal(AutoTangent) * MaxScale;
 					}
 				}
 
 				// Handle stationary endpoints
 				if (bStationaryEndpoint && (Index == 0 || Index == NumPoints - 1) && !IsClosedLoop())
 				{
-					AutoTangent = FVector::ZeroVector;
+					AutoTangent = ValueType();
 				}
 
 				InTangent = AutoTangent;
@@ -1556,7 +1456,7 @@ public:
 				}
 				else
 				{
-					InTangent = FVector::ZeroVector;
+					InTangent = ValueType();
 				}
 
 				// Next segment tangent (outgoing)
@@ -1566,7 +1466,7 @@ public:
 				}
 				else
 				{
-					OutTangent = FVector::ZeroVector;
+					OutTangent = ValueType();
 				}
 			}
 			break;
@@ -1574,8 +1474,8 @@ public:
 		case ETangentMode::Constant:
 			{
 				// Zero tangents for constant interpolation
-				InTangent = FVector::ZeroVector;
-				OutTangent = FVector::ZeroVector;
+				InTangent = ValueType();
+				OutTangent = ValueType();
 			}
 			break;
 
@@ -1583,7 +1483,7 @@ public:
 		default:
 			// Default to Auto mode if unknown
 			{
-				FVector Tangent = (1.0f - Tension) * (Next - Prev);
+				ValueType Tangent = (1.0f - Tension) * (Next - Prev);
 				InTangent = Tangent;
 				OutTangent = Tangent;
 			}
@@ -1596,7 +1496,7 @@ public:
 
 	}
 
-	virtual float FindNearest(const FVector& Point, float& OutSquaredDistance) const override
+	virtual float FindNearest(const ValueType& Point, float& OutSquaredDistance) const override
 	{
         // Delegate to internal spline
         return InternalSpline.FindNearest(Point, OutSquaredDistance);
@@ -1634,7 +1534,13 @@ public:
 		}
 		return TangentModes[Index];
 	}
-	TArray<ETangentMode> GetTangentModes() const
+
+	void SetTangentModes(const TArray<ETangentMode>& InTangentModes)
+	{
+		TangentModes = InTangentModes;
+	}
+	
+	const TArray<ETangentMode>& GetTangentModes() const
 	{
 		return TangentModes;
 	}
@@ -1650,12 +1556,12 @@ public:
 		return bStationaryEndpoint;
 	}
 	
-	virtual float FindNearestOnSegment(const ValueType& Point, int32 SegmentIndex, float& OutSquaredDistance) const
+	float FindNearestOnSegment(const ValueType& Point, int32 SegmentIndex, float& OutSquaredDistance) const
 		{ return InternalSpline.FindNearestOnSegment(Point, SegmentIndex, OutSquaredDistance); }
 
 protected:
 	/** Updates a Bezier control point in the internal spline */
-	void UpdateBezierControlPoint(int32 Index, const FVector& NewValue)
+	void UpdateBezierControlPoint(int32 Index, const ValueType& NewValue)
 	{
 		if (Index < InternalSpline.NumKeys())
 		{
@@ -1679,7 +1585,7 @@ private:
     	if (InternalSpline.NumKeys() == 0)
     	{
     		// First point - just store it directly
-    		InternalSpline.AddValue(FVector(ControlPoint.Position));
+    		InternalSpline.AddValue(ControlPoint.Position);
     		// update tangent mode
             if (TangentModes.IsEmpty())
             {
@@ -1697,15 +1603,15 @@ private:
     	{
     		// We have exactly one point stored - this is our second point
     		// Get the first point, then clear to start building segments
-    		FVector P0 = FVector(InternalSpline.GetValue(0));
-    		FVector P3 = ControlPoint.Position;
+    		ValueType P0 = InternalSpline.GetValue(0);
+    		ValueType P3 = ControlPoint.Position;
     		// Calculate control points
-    		FVector P1, P2;
+    		ValueType P1, P2;
     		if (ControlPoint.TangentMode == ETangentMode::Auto ||
 				ControlPoint.TangentMode == ETangentMode::AutoClamped)
     		{
     			// For first segment with auto tangents, use simple interpolation
-    			FVector AutoTangent = (1.0f - Tension) * (P3 - P0);
+    			ValueType AutoTangent = (1.0f - Tension) * (P3 - P0);
     			P1 = P0 + (AutoTangent / 3.0f);
     			P2 = P3 - (AutoTangent / 3.0f);
     		}
