[NetPhysics] DebugDraw improvements for networked physics
- Don't scrap debug draw calls that happen just before a resimulation
- Implement debug draw to show which physics replication is used, color-coded.
- Implement debug draw to show what triggers a resimulation color-coded. (Position, Rotation, Velocity, AngularVelocity)

diff --git a/Engine/Source/Runtime/Engine/Private/PhysicsEngine/PhysicsReplication.cpp b/Engine/Source/Runtime/Engine/Private/PhysicsEngine/PhysicsReplication.cpp
--- a/Engine/Source/Runtime/Engine/Private/PhysicsEngine/PhysicsReplication.cpp
+++ b/Engine/Source/Runtime/Engine/Private/PhysicsEngine/PhysicsReplication.cpp
@@ -78,6 +78,12 @@ namespace CharacterMovementCVars
 	static FAutoConsoleVariableRef CVarApplyAsyncSleepState(TEXT("p.ApplyAsyncSleepState"), ApplyAsyncSleepState, TEXT(""));
 }
 
+namespace RenderInterpolationCVars
+{
+	bool bRenderInterpDebugDrawResimTrigger = false;
+	static FAutoConsoleVariableRef CVarRenderInterpDebugDrawResimTrigger(TEXT("p.RenderInterp.DebugDraw.ResimTrigger"), bRenderInterpDebugDrawResimTrigger, TEXT("Draw debug lines for physics render interpolation, also needs p.Chaos.DebugDraw.Enabled set"));
+}
+
 namespace PhysicsReplicationCVars
 {
 	int32 SkipSkeletalRepOptimization = 1;
@@ -90,6 +96,12 @@ namespace PhysicsReplicationCVars
 	int32 EnableDefaultReplication = 0;
 	static FAutoConsoleVariableRef CVarEnableDefaultReplication(TEXT("np2.EnableDefaultReplication"), EnableDefaultReplication, TEXT("Enable default replication in the networked physics prediction flow."));
 
+	int32 DebugDrawShowRepMode = 0;
+	static FAutoConsoleVariableRef CVarPhysicsNetDebugDrawShowRepMode(TEXT("p.Net.DebugDraw.ShowRepMode"), DebugDrawShowRepMode, TEXT("Debug draw to show which physics replication mode is being used and where target states are being received. Green = Legacy Default, Light Blue = New Default, Yellow = Predictive Interpolation, Red = Resimulation, White = Something's Wrong. NOTE, Requires CVar p.Chaos.DebugDraw.Enabled 1"));
+
+	float DebugDrawLifeTime = 3.0f;
+	static FAutoConsoleVariableRef CVarPhysicsNetDebugDrawLifeTime(TEXT("p.Net.DebugDraw.LifeTime"), DebugDrawLifeTime, TEXT("Networked physics debug draw lifetime"));
+
 	namespace DefaultReplicationCVars
 	{
 		bool bHardsnapLegacyInPT = false;
@@ -826,7 +838,7 @@ bool FPhysicsReplication::ApplyRigidBodyState(float DeltaSeconds, FBodyInstance*
 			if (UWorld* OwningWorld = GetOwningWorld())
 			{
 				FColor Color = FColor::White;
-				DrawDebugDirectionalArrow(OwningWorld, CurrentState.Position, TargetPos, 5.0f, Color, true, CharacterMovementCVars::NetCorrectionLifetime, 0, 1.5f);
+				DrawDebugDirectionalArrow(OwningWorld, CurrentState.Position, TargetPos, 5.0f, Color, false, CharacterMovementCVars::NetCorrectionLifetime, 0, 1.5f);
 				DrawDebugFloatHistory(*OwningWorld, PhysicsTarget.ErrorHistory, CurrentState.Position + FVector(0.0f, 0.0f, 100.0f), FVector2D(100.0f, 50.0f), FColor::White, false, 0, -1);
 			}
 		}
@@ -952,6 +964,8 @@ void FPhysicsReplicationAsync::OnPreSimulate_Internal()
 			UpdateRewindDataTarget(Input);
 			UpdateAsyncTarget(Input, RigidsSolver);
 			
+			DebugDrawReplicationMode(Input);
+
 			// Deprecated, legacy BodyInstance flow for Default Replication
 			if (Input.Proxy != nullptr)
 			{
@@ -1148,7 +1162,7 @@ void FPhysicsReplicationAsync::UpdateAsyncTarget(const FPhysicsRepAsyncInputData
 			if (PhysicsReplicationCVars::PredictiveInterpolationCVars::bDrawDebugTargets)
 			{
 				const FVector Offset = FVector(0.0f, 0.0f, PhysicsReplicationCVars::PredictiveInterpolationCVars::DrawDebugZOffset);
-				Chaos::FDebugDrawQueue::GetInstance().DrawDebugBox(Input.TargetState.Position + Offset, FVector(15.0f, 15.0f, 15.0f), Input.TargetState.Quaternion, FColor::MakeRandomSeededColor(Input.ServerFrame), false, CharacterMovementCVars::NetCorrectionLifetime, 0, 1.0f);
+				Chaos::FDebugDrawQueue::GetInstance().DrawDebugBox(Input.TargetState.Position + Offset, FVector(15.0f, 15.0f, 15.0f), Input.TargetState.Quaternion, FColor::MakeRandomSeededColor(Input.ServerFrame), false, PhysicsReplicationCVars::DebugDrawLifeTime, 0, 1.0f);
 			}
 #endif
 
@@ -1743,7 +1757,7 @@ bool FPhysicsReplicationAsync::PredictiveInterpolation(Chaos::FPBDRigidParticleH
 		const FVector Offset = FVector(0.0f, 0.0f, PhysicsReplicationCVars::PredictiveInterpolationCVars::DrawDebugZOffset);
 		const FVector Pos = Target.TargetState.Position + Offset;
 		const int32 SizeMultiplier = FMath::Clamp(Target.TickCount, -4, 30);
-		Chaos::FDebugDrawQueue::GetInstance().DrawDebugBox(Pos, FVector(5.0f + SizeMultiplier * 0.75f, 5.0f + SizeMultiplier * 0.75f, 5.0f + SizeMultiplier * 0.75f), Target.TargetState.Quaternion, FColor::MakeRandomSeededColor(Target.ServerFrame), false, CharacterMovementCVars::NetCorrectionLifetime, 0, 1.0f);
+		Chaos::FDebugDrawQueue::GetInstance().DrawDebugBox(Pos, FVector(5.0f + SizeMultiplier * 0.75f, 5.0f + SizeMultiplier * 0.75f, 5.0f + SizeMultiplier * 0.75f), Target.TargetState.Quaternion, FColor::MakeRandomSeededColor(Target.ServerFrame), false, PhysicsReplicationCVars::DebugDrawLifeTime, 0, 1.0f);
 	}
 #endif
 
@@ -1953,7 +1967,7 @@ bool FPhysicsReplicationAsync::PredictiveInterpolation(Chaos::FPBDRigidParticleH
 				const FVector Offset = FVector(0.0f, 0.0f, PhysicsReplicationCVars::PredictiveInterpolationCVars::DrawDebugZOffset);
 				const FVector Pos = KinTargetPos + Offset;
 				const int32 SizeMultiplier = FMath::Clamp(Target.TickCount, -4, 30);
-				Chaos::FDebugDrawQueue::GetInstance().DrawDebugSphere(Pos, 3.0f + SizeMultiplier * 0.75f, 8, FColor::MakeRandomSeededColor(Target.ServerFrame), false, CharacterMovementCVars::NetCorrectionLifetime, 0, 1.0f);
+				Chaos::FDebugDrawQueue::GetInstance().DrawDebugSphere(Pos, 3.0f + SizeMultiplier * 0.75f, 8, FColor::MakeRandomSeededColor(Target.ServerFrame), false, PhysicsReplicationCVars::DebugDrawLifeTime, 0, 1.0f);
 			}
 #endif
 		}
@@ -2185,11 +2199,24 @@ bool FPhysicsReplicationAsync::ResimulationReplication(Chaos::FPBDRigidParticleH
 	const bool bCompareW = Chaos::FPhysicsSolverBase::GetResimulationErrorAngularVelocityThresholdEnabled() || SettingsCurrent.ResimulationSettings.bOverrideResimulationErrorAngularVelocityThreshold;
 	bool bShouldTriggerResim = false;
 
+#if !(UE_BUILD_SHIPPING || UE_BUILD_TEST)
+	// Debugging
+	FColor DebugColor = FColor::Black;
+	bool bResimV = false;
+	bool bResimW = false;
+#endif
+
 	// Check for positional discrepancy in Distance between client and server
 	if (bCompareX)
 	{
 		const float ResimPositionErrorThreshold = SettingsCurrent.ResimulationSettings.GetResimulationErrorPositionThreshold(Chaos::FPhysicsSolverBase::GetResimulationErrorPositionThreshold());
 		bShouldTriggerResim = Chaos::FRewindData::CheckVectorThreshold(Target.TargetState.Position, PastState.GetX(), ResimPositionErrorThreshold);
+#if !(UE_BUILD_SHIPPING || UE_BUILD_TEST)
+		if (bShouldTriggerResim)
+		{
+			DebugColor = FColor::Orange;
+		}
+#endif
 	}
 
 	// Check for linear velocity discrepancy in Distance / s between client and server
@@ -2197,13 +2224,25 @@ bool FPhysicsReplicationAsync::ResimulationReplication(Chaos::FPBDRigidParticleH
 	{
 		const float ResimLinVelocityErrorThreshold = SettingsCurrent.ResimulationSettings.GetResimulationErrorLinearVelocityThreshold(Chaos::FPhysicsSolverBase::GetResimulationErrorLinearVelocityThreshold());
 		bShouldTriggerResim = Chaos::FRewindData::CheckVectorThreshold(Target.TargetState.LinVel, PastState.GetV(), ResimLinVelocityErrorThreshold);
+#if !(UE_BUILD_SHIPPING || UE_BUILD_TEST)
+		if (bShouldTriggerResim)
+		{
+			bResimV = true;
+		}
+#endif
 	}
 
 	// Check for angular velocity discrepancy in Degrees / s between client and server
 	if (!bShouldTriggerResim && bCompareW)
 	{
 		const float ResimAngVelocityErrorThreshold = SettingsCurrent.ResimulationSettings.GetResimulationErrorAngularVelocityThreshold(Chaos::FPhysicsSolverBase::GetResimulationErrorAngularVelocityThreshold());
-		bShouldTriggerResim = Chaos::FRewindData::CheckVectorThreshold(FMath::DegreesToRadians(Target.TargetState.AngVel), PastState.GetW(), ResimAngVelocityErrorThreshold);
+		bShouldTriggerResim = Chaos::FRewindData::CheckVectorThreshold(Target.TargetState.AngVel, FMath::RadiansToDegrees(PastState.GetW()), ResimAngVelocityErrorThreshold);
+#if !(UE_BUILD_SHIPPING || UE_BUILD_TEST)
+		 if (bShouldTriggerResim)
+		 {
+			 bResimW = true;
+		 }
+#endif
 	}
 
 	// Check for rotational discrepancy in Degrees between client and server
@@ -2211,6 +2250,12 @@ bool FPhysicsReplicationAsync::ResimulationReplication(Chaos::FPBDRigidParticleH
 	{
 		const float ResimRotationErrorThreshold = SettingsCurrent.ResimulationSettings.GetResimulationErrorRotationThreshold(Chaos::FPhysicsSolverBase::GetResimulationErrorRotationThreshold());
 		bShouldTriggerResim = Chaos::FRewindData::CheckQuaternionThreshold(Target.TargetState.Quaternion, PastState.GetR(), ResimRotationErrorThreshold);
+#if !(UE_BUILD_SHIPPING || UE_BUILD_TEST)
+		if (bShouldTriggerResim)
+		{
+			DebugColor = FColor::Magenta;
+		}
+#endif
 	}
 
 #if !(UE_BUILD_SHIPPING || UE_BUILD_TEST)
@@ -2225,16 +2270,49 @@ bool FPhysicsReplicationAsync::ResimulationReplication(Chaos::FPBDRigidParticleH
 		UE_LOG(LogTemp, Log, TEXT("Particle Target Omega = %s | Current Omega= %s"), *Target.TargetState.AngVel.ToString(), *PastState.GetW().ToString());
 	}
 
-	if (PhysicsReplicationCVars::ResimulationCVars::bDrawDebug)
-	{ 
-		static constexpr float BoxSize = 5.0f;
-		const float ColorLerp = bShouldTriggerResim ? 1.0f : 0.0f;
-		const FColor DebugColor = FLinearColor::LerpUsingHSV(FLinearColor::Green, FLinearColor::Red, ColorLerp).ToFColor(false);
+	if (PhysicsReplicationCVars::ResimulationCVars::bDrawDebug || RenderInterpolationCVars::bRenderInterpDebugDrawResimTrigger)
+	{
+		if (bShouldTriggerResim)
+		{
+			FVector Box = RenderInterpolationCVars::bRenderInterpDebugDrawResimTrigger ? FVector(6, 3, 2) : FVector(40, 20, 10);
+			const float DrawThickness = RenderInterpolationCVars::bRenderInterpDebugDrawResimTrigger ? 0.5f : 1.5f;
+
+			if (RenderInterpolationCVars::bRenderInterpDebugDrawResimTrigger) // Resim debug draw extension for render interpolation 
+			{
+				Chaos::FDebugDrawQueue::GetInstance().DrawDebugBox(PastState.GetX(), Box, PastState.GetR(), FColor::White, false, PhysicsReplicationCVars::DebugDrawLifeTime, 0, DrawThickness);
+				Chaos::FDebugDrawQueue::GetInstance().DrawDebugBox(Target.TargetState.Position, Box, Target.TargetState.Quaternion, DebugColor, false, PhysicsReplicationCVars::DebugDrawLifeTime, 0, DrawThickness);
 
-		Chaos::FDebugDrawQueue::GetInstance().DrawDebugBox(Target.TargetState.Position, FVector(BoxSize, BoxSize, BoxSize), Target.TargetState.Quaternion, FColor::Orange, true, CharacterMovementCVars::NetCorrectionLifetime, 0, 1.0f);
-		Chaos::FDebugDrawQueue::GetInstance().DrawDebugBox(PastState.GetX(), FVector(6, 6, 6), PastState.GetR(), DebugColor, true, CharacterMovementCVars::NetCorrectionLifetime, 0, 1.0f);
+				Chaos::FDebugDrawQueue::GetInstance().DrawDebugDirectionalArrow(Handle->GetX(), PastState.GetX(), 5.0f, FColor::White, false, PhysicsReplicationCVars::DebugDrawLifeTime, 0, DrawThickness);
+				Chaos::FDebugDrawQueue::GetInstance().DrawDebugDirectionalArrow(PastState.GetX(), Target.TargetState.Position, 5.0f, FColor::Black, false, PhysicsReplicationCVars::DebugDrawLifeTime, 0, DrawThickness);
 
-		Chaos::FDebugDrawQueue::GetInstance().DrawDebugDirectionalArrow(PastState.GetX(), Target.TargetState.Position, 5.0f, FColor::MakeRandomSeededColor(LocalFrame), true, CharacterMovementCVars::NetCorrectionLifetime, 0, 0.5f);
+				if (bResimV)
+				{
+					const FVector DiffV = Target.TargetState.LinVel - PastState.GetV();
+					Chaos::FDebugDrawQueue::GetInstance().DrawDebugDirectionalArrow(Target.TargetState.Position, Target.TargetState.Position + DiffV, 5.0f, FColor::Orange, false, PhysicsReplicationCVars::DebugDrawLifeTime, 0, DrawThickness);
+				}
+				if (bResimW)
+				{
+					const FVector DiffW = Target.TargetState.AngVel - FMath::RadiansToDegrees(PastState.GetW());
+					Chaos::FDebugDrawQueue::GetInstance().DrawDebugDirectionalArrow(Target.TargetState.Position + DiffW, Target.TargetState.Position, 5.0f, FColor::Magenta, false, PhysicsReplicationCVars::DebugDrawLifeTime, 0, DrawThickness);
+				}
+			}
+			else // Resim trigger debug draw
+			{
+				Chaos::FDebugDrawQueue::GetInstance().DrawDebugBox(Handle->GetX(), Box, PastState.GetR(), FColor::White, false, PhysicsReplicationCVars::DebugDrawLifeTime, 0, DrawThickness);
+				Chaos::FDebugDrawQueue::GetInstance().DrawDebugBox(Handle->GetX() + (Target.TargetState.Position - PastState.GetX()), Box, Target.TargetState.Quaternion, DebugColor, false, PhysicsReplicationCVars::DebugDrawLifeTime, 0, DrawThickness);
+
+				if (bResimV)
+				{
+					const FVector DiffV = Target.TargetState.LinVel - PastState.GetV();
+					Chaos::FDebugDrawQueue::GetInstance().DrawDebugDirectionalArrow(Handle->GetX(), Handle->GetX() + DiffV, 5.0f, FColor::Orange, false, PhysicsReplicationCVars::DebugDrawLifeTime, 0, DrawThickness);
+				}
+				if (bResimW)
+				{
+					const FVector DiffW = Target.TargetState.AngVel - FMath::RadiansToDegrees(PastState.GetW());
+					Chaos::FDebugDrawQueue::GetInstance().DrawDebugDirectionalArrow(Handle->GetX() + DiffW, Handle->GetX(), 5.0f, FColor::Magenta, false, PhysicsReplicationCVars::DebugDrawLifeTime, 0, DrawThickness);
+				}
+			}
+		}
 	}
 #endif
 
@@ -2290,7 +2368,7 @@ bool FPhysicsReplicationAsync::ResimulationReplication(Chaos::FPBDRigidParticleH
 #if !(UE_BUILD_SHIPPING || UE_BUILD_TEST)
 			if (PhysicsReplicationCVars::ResimulationCVars::bDrawDebug)
 			{
-				Chaos::FDebugDrawQueue::GetInstance().DrawDebugDirectionalArrow(Handle->GetX(), CorrectedX, 5.0f, FColor::MakeRandomSeededColor(LocalFrame), true, CharacterMovementCVars::NetCorrectionLifetime, 0, 0.5f);
+				Chaos::FDebugDrawQueue::GetInstance().DrawDebugDirectionalArrow(Handle->GetX(), CorrectedX, 5.0f, FColor::MakeRandomSeededColor(LocalFrame), false, PhysicsReplicationCVars::DebugDrawLifeTime, 0, 0.5f);
 			}
 #endif
 			// Apply correction to position and rotation
@@ -2322,6 +2400,66 @@ bool FPhysicsReplicationAsync::ResimulationReplication(Chaos::FPBDRigidParticleH
 	return bClearTarget;
 }
 
+void FPhysicsReplicationAsync::DebugDrawReplicationMode(const FPhysicsRepAsyncInputData& Input)
+{
+#if !(UE_BUILD_SHIPPING || UE_BUILD_TEST)
+
+	if (!PhysicsReplicationCVars::DebugDrawShowRepMode)
+	{
+		return;
+	}
+
+	if (Input.PhysicsObject == nullptr && Input.Proxy == nullptr)
+	{
+		return;
+	}
+
+	FColor DebugColor = FColor::White;
+	FVector BoxExtent = FVector(10.0f, 10.0f, 10.0f);
+	FQuat Rotation = FQuat::Identity;
+
+	if (Input.PhysicsObject)
+	{
+		if (FReplicatedPhysicsTargetAsync* Target = ObjectToTarget.Find(Input.PhysicsObject))
+		{
+			Chaos::FReadPhysicsObjectInterface_Internal Interface = Chaos::FPhysicsObjectInternalInterface::GetRead();
+			if (Chaos::FGeometryParticleHandle* Handle = Interface.GetParticle(Input.PhysicsObject))
+			{
+				BoxExtent = Handle->LocalBounds().Extents() * 0.5f;
+				Rotation = Handle->GetR();
+			}
+
+			const EPhysicsReplicationMode RepMode = Target->IsWaiting() ? Target->RepModeOverride : Target->RepMode;
+			switch (RepMode)
+			{
+			case EPhysicsReplicationMode::PredictiveInterpolation:
+				DebugColor = FColor::Yellow;
+				break;
+			case EPhysicsReplicationMode::Resimulation:
+				DebugColor = FColor::Red;
+				break;
+			case EPhysicsReplicationMode::Default:
+			default:
+				DebugColor = FColor::Cyan;
+				break;
+			}
+		}
+	}
+	else if (Input.Proxy != nullptr)
+	{
+		// Legacy Default physics replication
+
+		Chaos::FSingleParticlePhysicsProxy* Proxy = Input.Proxy;
+		Chaos::FRigidBodyHandle_Internal* Handle = Proxy->GetPhysicsThreadAPI();
+
+		Rotation = Handle->GetR();
+		DebugColor = FColor::Green;
+	}
+
+	Chaos::FDebugDrawQueue::GetInstance().DrawDebugBox(Input.TargetState.Position, BoxExtent, Rotation, DebugColor, false, PhysicsReplicationCVars::DebugDrawLifeTime, 0, 1.0f);
+#endif
+}
+
 FName FPhysicsReplicationAsync::GetFNameForStatId() const
 {
 	const static FLazyName StaticName("FPhysicsReplicationAsyncCallback");
diff --git a/Engine/Source/Runtime/Engine/Public/PhysicsReplication.h b/Engine/Source/Runtime/Engine/Public/PhysicsReplication.h
--- a/Engine/Source/Runtime/Engine/Public/PhysicsReplication.h
+++ b/Engine/Source/Runtime/Engine/Public/PhysicsReplication.h
@@ -232,7 +232,8 @@ private:
 	bool UsePhysicsReplicationLOD();
 	void CheckTargetResimValidity(FReplicatedPhysicsTargetAsync& Target);
 	void ApplyPhysicsReplicationLOD(Chaos::FConstPhysicsObjectHandle PhysicsObjectHandle, FReplicatedPhysicsTargetAsync& Target, const uint32 LODFlags);
-	
+	void DebugDrawReplicationMode(const FPhysicsRepAsyncInputData& Input);
+
 	/** Static function to extrapolate a target for N ticks using X DeltaSeconds */
 	static void ExtrapolateTarget(FReplicatedPhysicsTargetAsync& Target, const int32 ExtrapolateFrames, const float DeltaSeconds);
 	/** Static function to extrapolate a target for N Seconds */
diff --git a/Engine/Source/Runtime/Experimental/Chaos/Private/PBDRigidsSolver.cpp b/Engine/Source/Runtime/Experimental/Chaos/Private/PBDRigidsSolver.cpp
--- a/Engine/Source/Runtime/Experimental/Chaos/Private/PBDRigidsSolver.cpp
+++ b/Engine/Source/Runtime/Experimental/Chaos/Private/PBDRigidsSolver.cpp
@@ -1353,7 +1353,10 @@ namespace Chaos
 		const FReal StartSimTime = GetSolverTime();
 
 #if CHAOS_DEBUG_DRAW
-		ChaosDD::Private::FChaosDDScopeTimelineContext DDContext(CDDFrameTimeline, StartSimTime, GetLastDt());
+		if (!GetEvolution()->IsResimming())
+		{
+			ChaosDD::Private::FChaosDDScopeTimelineContext DDContext(CDDFrameTimeline, StartSimTime, GetLastDt());
+		}
 #endif
 
 #if !(UE_BUILD_SHIPPING || UE_BUILD_TEST)
