- Material Expression Color Ramp add to new translator
- FIX: ColorRamp Cooking Error

diff --git a/Engine/Source/Runtime/Engine/Private/Materials/MaterialExpressions.cpp b/Engine/Source/Runtime/Engine/Private/Materials/MaterialExpressions.cpp
--- a/Engine/Source/Runtime/Engine/Private/Materials/MaterialExpressions.cpp
+++ b/Engine/Source/Runtime/Engine/Private/Materials/MaterialExpressions.cpp
@@ -4743,8 +4743,10 @@ void UMaterialExpressionLinearInterpolate::GetCaption(TArray<FString>& OutCaptio
 UMaterialExpressionColorRamp::UMaterialExpressionColorRamp(const FObjectInitializer& ObjectInitializer) 
 	: Super(ObjectInitializer)
 {
-	// Initialize ColorCurve RGB value showing a white->black gradient
+#if WITH_EDITOR
 	ColorCurve = ObjectInitializer.CreateDefaultSubobject<UCurveLinearColor>(this, TEXT("ColorCurve"));
+
+	// Initialize ColorCurve RGB value showing a white->black gradient
 	for (int i = 0; i < 3; i++)
 	{
 		if (FRichCurve* Curve = &ColorCurve->FloatCurves[i])
@@ -4754,6 +4756,7 @@ UMaterialExpressionColorRamp::UMaterialExpressionColorRamp(const FObjectInitiali
 			Curve->AddKey(1.f, 0.f);
 		}
 	}
+#endif // WITH_EDITOR
 }
 
 #if WITH_EDITOR
@@ -4783,29 +4786,34 @@ int32 UMaterialExpressionColorRamp::Compile(FMaterialCompiler* Compiler, int32 O
 	int32 InputCode = Compiler->ComponentMask(Input.Compile(Compiler), true, false, false, false);
 
 	// Helper lambda to evaluate a curve using vectorized multiply–add operations
-	auto EvaluateCurve = [Compiler](const FRichCurve& Curve, int32 InValue) -> int32
+	auto EvaluateCurve = [Compiler, InputCode](const FRichCurve& Curve) -> int32
 		{
 			const int32 NumKeys = Curve.Keys.Num();
-			if (NumKeys == 0) { return Compiler->Constant(0.0f); }
+			switch (NumKeys)
+			{
+				case 0:
+					return Compiler->Constant(0.0f);
 
-			if (NumKeys == 1) { return Compiler->Constant(Curve.Keys[0].Value); }
+				case 1:
+					return Compiler->Constant(Curve.Keys[0].Value);
 
-			if (NumKeys == 2)
-			{
-				const float StartTime = Curve.Keys[0].Time;
-				const float EndTime = Curve.Keys[1].Time;
-				const float StartValue = Curve.Keys[0].Value;
-				const float EndValue = Curve.Keys[1].Value;
+				case 2:
+				{
+					const float StartTime = Curve.Keys[0].Time;
+					const float EndTime = Curve.Keys[1].Time;
+					const float StartValue = Curve.Keys[0].Value;
+					const float EndValue = Curve.Keys[1].Value;
 
-				int32 TimeDelta = Compiler->Constant(EndTime - StartTime);
-				int32 TimeDiff = Compiler->Sub(InValue, Compiler->Constant(StartTime));
-				int32 Fraction = Compiler->Div(TimeDiff, TimeDelta);
+					int32 TimeDelta = Compiler->Constant(EndTime - StartTime);
+					int32 TimeDiff = Compiler->Sub(InputCode, Compiler->Constant(StartTime));
+					int32 Fraction = Compiler->Div(TimeDiff, TimeDelta);
 
-				return Compiler->Lerp(Compiler->Constant(StartValue), Compiler->Constant(EndValue), Fraction);
+					return Compiler->Lerp(Compiler->Constant(StartValue), Compiler->Constant(EndValue), Fraction);
+				}
 			}
 
 			// Turn input into a vector4
-			int32 InValueVec = Compiler->AppendVector(Compiler->AppendVector(InValue, InValue), Compiler->AppendVector(InValue, InValue));
+			int32 InValueVec = Compiler->AppendVector(Compiler->AppendVector(InputCode, InputCode), Compiler->AppendVector(InputCode, InputCode));
 
 			int32 Result = Compiler->Constant(Curve.Keys[0].Value);
 			int32 i = 0;
@@ -4861,7 +4869,7 @@ int32 UMaterialExpressionColorRamp::Compile(FMaterialCompiler* Compiler, int32 O
 				int32 TimeDelta = Compiler->Constant(EndTime - StartTime);
 				int32 ValueDelta = Compiler->Constant(EndValue - StartValue);
 
-				int32 TimeDiff = Compiler->Sub(InValue, Compiler->Constant(StartTime));
+				int32 TimeDiff = Compiler->Sub(InputCode, Compiler->Constant(StartTime));
 				int32 Fraction = Compiler->Div(TimeDiff, TimeDelta);
 				int32 SatFraction = Compiler->Saturate(Fraction);
 
@@ -4871,10 +4879,10 @@ int32 UMaterialExpressionColorRamp::Compile(FMaterialCompiler* Compiler, int32 O
 			return Result;
 		};
 
-	int32 Red = EvaluateCurve(ColorCurve->FloatCurves[0], InputCode);
-	int32 Green = EvaluateCurve(ColorCurve->FloatCurves[1], InputCode);
-	int32 Blue = EvaluateCurve(ColorCurve->FloatCurves[2], InputCode);
-	int32 Alpha = EvaluateCurve(ColorCurve->FloatCurves[3], InputCode);
+	int32 Red = EvaluateCurve(ColorCurve->FloatCurves[0]);
+	int32 Green = EvaluateCurve(ColorCurve->FloatCurves[1]);
+	int32 Blue = EvaluateCurve(ColorCurve->FloatCurves[2]);
+	int32 Alpha = EvaluateCurve(ColorCurve->FloatCurves[3]);
 	int32 FinalVector = Compiler->AppendVector(Compiler->AppendVector(Red, Green), Compiler->AppendVector(Blue, Alpha));
 	return FinalVector;
 }
diff --git a/Engine/Source/Runtime/Engine/Private/Materials/MaterialExpressionsIR.cpp b/Engine/Source/Runtime/Engine/Private/Materials/MaterialExpressionsIR.cpp
--- a/Engine/Source/Runtime/Engine/Private/Materials/MaterialExpressionsIR.cpp
+++ b/Engine/Source/Runtime/Engine/Private/Materials/MaterialExpressionsIR.cpp
@@ -39,6 +39,7 @@
 #include "Materials/MaterialExpressionClearCoatNormalCustomOutput.h"
 #include "Materials/MaterialExpressionCloudLayer.h"
 #include "Materials/MaterialExpressionCollectionParameter.h"
+#include "Materials/MaterialExpressionColorRamp.h"
 #include "Materials/MaterialExpressionComponentMask.h"
 #include "Materials/MaterialExpressionConstant.h"
 #include "Materials/MaterialExpressionConstant2Vector.h"
@@ -1014,6 +1015,7 @@ void UMaterialExpressionClamp::Build(MIR::FEmitter& Em)
 	Em.Output(0, OutputValue);
 }
 
+
 void BuildTernaryArithmeticOperator(MIR::FEmitter& Em, MIR::EOperator Op, FExpressionInput* A, float ConstA, FExpressionInput* B, float ConstB, FExpressionInput* C, float ConstC)
 {
 	FValueRef ValueA = Em.InputDefaultFloat(A, ConstA);
@@ -1022,6 +1024,136 @@ void BuildTernaryArithmeticOperator(MIR::FEmitter& Em, MIR::EOperator Op, FExpre
 	Em.Output(0, Em.Operator(Op, ValueA, ValueB, ValueC));
 }
 
+void UMaterialExpressionColorRamp::Build(MIR::FEmitter& Em)
+{	
+	// Check that the ColorCurve is set
+	if (!ColorCurve)
+	{
+		Em.Errorf(TEXT("Missing ColorCurve"));
+		return;
+	}
+
+	FValueRef InputValue = Em.CastToFloat(Em.InputDefaultFloat(&Input, ConstInput), 1);
+
+	// If the input is constant, evaluate at compile time.
+	if (const MIR::FConstant* Constant = MIR::As<MIR::FConstant>(InputValue))
+	{
+		FLinearColor ColorValue = ColorCurve->GetLinearColorValue(Constant->Float);
+		Em.Output(0, Em.ConstantFloat4(ColorValue));
+		return;
+	}
+
+	// Helper lambda to evaluate a curve
+	auto EvaluateCurve = [&Em, InputValue](const FRichCurve& Curve) -> FValueRef
+		{
+			const int32 NumKeys = Curve.Keys.Num();
+
+			switch (NumKeys)
+			{
+				case 0:
+					return Em.ConstantFloat(0.0f);
+
+				case 1:
+					return Em.ConstantFloat(Curve.Keys[0].Value);
+
+				case 2:
+				{
+					float StartTime = Curve.Keys[0].Time;
+					float EndTime = Curve.Keys[1].Time;
+					float StartValue = Curve.Keys[0].Value;
+					float EndValue = Curve.Keys[1].Value;
+
+					FValueRef TimeDelta = Em.ConstantFloat(EndTime - StartTime);
+					FValueRef TimeDiff = Em.Subtract(InputValue, Em.ConstantFloat(StartTime));
+					FValueRef Fraction = Em.Divide(TimeDiff, TimeDelta);
+
+					return Em.Lerp(Em.ConstantFloat(StartValue), Em.ConstantFloat(EndValue), Fraction);
+				}
+			}
+
+			FValueRef InValueVec = Em.Vector4(InputValue, InputValue, InputValue, InputValue);
+
+			FValueRef Result = Em.ConstantFloat(Curve.Keys[0].Value);
+			int32 i = 0;
+
+			// Use vector operations for segments of 4
+			for (; i < NumKeys - 4; i += 4)
+			{
+				FVector4f StartTimeVector(
+					Curve.Keys[i].Time,
+					Curve.Keys[i + 1].Time,
+					Curve.Keys[i + 2].Time,
+					Curve.Keys[i + 3].Time
+				);
+				FValueRef StartTimeVec = Em.ConstantFloat4(StartTimeVector);
+
+				FVector4f EndTimeVector(
+					Curve.Keys[i + 1].Time,
+					Curve.Keys[i + 2].Time,
+					Curve.Keys[i + 3].Time,
+					Curve.Keys[i + 4].Time
+				);
+				FValueRef EndTimeVec = Em.ConstantFloat4(EndTimeVector);
+
+				FVector4f StartValueVector(
+					Curve.Keys[i].Value,
+					Curve.Keys[i + 1].Value,
+					Curve.Keys[i + 2].Value,
+					Curve.Keys[i + 3].Value
+				);
+				FValueRef StartValueVec = Em.ConstantFloat4(StartValueVector);
+
+				FVector4f EndValueVector(
+					Curve.Keys[i + 1].Value,
+					Curve.Keys[i + 2].Value,
+					Curve.Keys[i + 3].Value,
+					Curve.Keys[i + 4].Value
+				);
+				FValueRef EndValueVec = Em.ConstantFloat4(EndValueVector);
+
+				FValueRef TimeDeltaVec = Em.Subtract(EndTimeVec, StartTimeVec);
+				FValueRef ValueDeltaVec = Em.Subtract(EndValueVec, StartValueVec);
+
+				FValueRef TimeDiffVec = Em.Subtract(InValueVec, StartTimeVec);
+				FValueRef FractionVec = Em.Divide(TimeDiffVec, TimeDeltaVec);
+				FValueRef SatFractionVec = Em.Saturate(FractionVec);
+				FValueRef ContributionVec = Em.Multiply(ValueDeltaVec, SatFractionVec);
+
+				FVector4f Ones(1.0f, 1.0f, 1.0f, 1.0f);
+				FValueRef OnesVec = Em.ConstantFloat4(Ones);
+				FValueRef ContributionSum = Em.Dot(ContributionVec, OnesVec);
+
+				Result = Em.Add(Result, ContributionSum);
+			}
+			
+			// Use scalar operations for the remaining keys
+			for (; i < NumKeys - 1; i++)
+			{
+				float StartTime = Curve.Keys[i].Time;
+				float EndTime = Curve.Keys[i + 1].Time;
+				float StartValue = Curve.Keys[i].Value;
+				float EndValue = Curve.Keys[i + 1].Value;
+
+				FValueRef TimeDelta = Em.ConstantFloat(EndTime - StartTime);
+				FValueRef ValueDelta = Em.ConstantFloat(EndValue - StartValue);
+				FValueRef TimeDiff = Em.Subtract(InputValue, Em.ConstantFloat(StartTime));
+				FValueRef Fraction = Em.Divide(TimeDiff, TimeDelta);
+				FValueRef SatFraction = Em.Saturate(Fraction);
+				FValueRef Contribution = Em.Multiply(ValueDelta, SatFraction);
+				Result = Em.Add(Result, Contribution);
+			}
+			return Result;
+		};
+
+	FValueRef Red = EvaluateCurve(ColorCurve->FloatCurves[0]);
+	FValueRef Green = EvaluateCurve(ColorCurve->FloatCurves[1]);
+	FValueRef Blue = EvaluateCurve(ColorCurve->FloatCurves[2]);
+	FValueRef Alpha = EvaluateCurve(ColorCurve->FloatCurves[3]);
+
+	FValueRef FinalVector = Em.Vector4(Red, Green, Blue, Alpha);
+	Em.Output(0, FinalVector);
+}
+
 void UMaterialExpressionLinearInterpolate::Build(MIR::FEmitter& Em)
 {
 	BuildTernaryArithmeticOperator(Em, MIR::TO_Lerp, &A, ConstA, &B, ConstB, &Alpha, ConstAlpha);
diff --git a/Engine/Source/Runtime/Engine/Public/Materials/MaterialExpressionColorRamp.h b/Engine/Source/Runtime/Engine/Public/Materials/MaterialExpressionColorRamp.h
--- a/Engine/Source/Runtime/Engine/Public/Materials/MaterialExpressionColorRamp.h
+++ b/Engine/Source/Runtime/Engine/Public/Materials/MaterialExpressionColorRamp.h
@@ -10,7 +10,7 @@
 #include "Curves/CurveLinearColor.h"
 #include "MaterialExpressionColorRamp.generated.h"
 
-UCLASS(Abstract, MinimalAPI)
+UCLASS(MinimalAPI)
 class UMaterialExpressionColorRamp : public UMaterialExpression
 {
 	GENERATED_BODY()
@@ -30,9 +30,10 @@ public:
 
 	//~ Begin UMaterialExpression Interface
 #if WITH_EDITOR
-	virtual void HandleCurvePropertyChanged(UCurveBase* Curve, EPropertyChangeType::Type ChangeType);
+	virtual void Build(MIR::FEmitter& Emitter) override;
 	virtual int32 Compile(class FMaterialCompiler* Compiler, int32 OutputIndex) override;
 	virtual void GetCaption(TArray<FString>& OutCaptions) const override;
+	virtual void HandleCurvePropertyChanged(UCurveBase* Curve, EPropertyChangeType::Type ChangeType);
 #endif // WITH_EDITOR
 	//~ End UMaterialExpression Interface
 };
\ No newline at end of file
