Fixing PostLoad methods in animation to be correct about handling cooked vs uncooked data.

#rnx

diff --git a/Engine/Plugins/Animation/PoseSearch/Source/Runtime/Private/PoseSearchAssetSampler.cpp b/Engine/Plugins/Animation/PoseSearch/Source/Runtime/Private/PoseSearchAssetSampler.cpp
--- a/Engine/Plugins/Animation/PoseSearch/Source/Runtime/Private/PoseSearchAssetSampler.cpp
+++ b/Engine/Plugins/Animation/PoseSearch/Source/Runtime/Private/PoseSearchAssetSampler.cpp
@@ -8,6 +8,7 @@
 #include "Interfaces/ITargetPlatformManagerModule.h"
 #include "PoseSearch/PoseSearchAnimNotifies.h"
 #include "PoseSearch/PoseSearchDefines.h"
+#include "UObject/Package.h"
 
 namespace UE::PoseSearch
 {
@@ -257,7 +258,10 @@ void FAnimationAssetSampler::Init(const UAnimationAsset* InAnimationAsset, const
 	PlatformHash = FIoHash::Zero;
 	if (const UAnimSequence* AnimSequence = Cast<UAnimSequence>(InAnimationAsset))
 	{		
-		PlatformHash = AnimSequence->GetDerivedDataKeyHash(TargetPlatform);
+		if (!AnimationAssetPtr->GetPackage()->HasAnyPackageFlags(PKG_Cooked))
+		{
+			PlatformHash = AnimSequence->GetDerivedDataKeyHash(TargetPlatform);
+		}
 	}
 #endif
 
@@ -421,7 +425,7 @@ void FAnimationAssetSampler::ExtractPose(const FAnimExtractContext& ExtractionCt
 #if WITH_EDITOR
 		if (Cast<UAnimSequence>(SequenceBase))
 		{
-			check(PlatformHash != FIoHash::Zero);
+			check(SequenceBase->GetPackage()->HasAnyPackageFlags(PKG_Cooked) || PlatformHash != FIoHash::Zero);
 			SequenceContext.TargetPlatformHash = PlatformHash;
 		}
 #endif // WITH_EDITOR
diff --git a/Engine/Plugins/Animation/PoseSearch/Source/Runtime/Private/PoseSearchDerivedData.cpp b/Engine/Plugins/Animation/PoseSearch/Source/Runtime/Private/PoseSearchDerivedData.cpp
--- a/Engine/Plugins/Animation/PoseSearch/Source/Runtime/Private/PoseSearchDerivedData.cpp
+++ b/Engine/Plugins/Animation/PoseSearch/Source/Runtime/Private/PoseSearchDerivedData.cpp
@@ -2693,7 +2693,7 @@ EAsyncBuildIndexResult FAsyncPoseSearchDatabasesManagement::RequestAsyncBuildInd
 		return EAsyncBuildIndexResult::Failed;
 	}
 
-	if (Database->GetPackage()->bIsCookedForEditor)
+	if (Database->GetPackage()->HasAnyPackageFlags(PKG_Cooked))
 	{
 		// Don't cache for cooked packages
 		return EAsyncBuildIndexResult::Success;
diff --git a/Engine/Plugins/Chooser/Source/ProxyTable/Private/ProxyTable.cpp b/Engine/Plugins/Chooser/Source/ProxyTable/Private/ProxyTable.cpp
--- a/Engine/Plugins/Chooser/Source/ProxyTable/Private/ProxyTable.cpp
+++ b/Engine/Plugins/Chooser/Source/ProxyTable/Private/ProxyTable.cpp
@@ -4,6 +4,7 @@
 #include "AssetRegistry/AssetRegistryModule.h"
 #include "Logging/LogMacros.h"
 #include "Misc/StringBuilder.h"
+#include "UObject/Package.h"
 
 DEFINE_LOG_CATEGORY_STATIC(LogProxyTable,Log,All);
 
@@ -162,7 +163,10 @@ void UProxyTable::PostLoad()
 {
 	Super::PostLoad();
 #if WITH_EDITORONLY_DATA
-	BuildRuntimeData();
+	if (!GetPackage()->HasAnyPackageFlags(PKG_Cooked))
+	{
+		BuildRuntimeData();
+	}
 #endif
 
 	// compilation for property accesses
diff --git a/Engine/Source/Runtime/Engine/Private/Animation/AnimSequence.cpp b/Engine/Source/Runtime/Engine/Private/Animation/AnimSequence.cpp
--- a/Engine/Source/Runtime/Engine/Private/Animation/AnimSequence.cpp
+++ b/Engine/Source/Runtime/Engine/Private/Animation/AnimSequence.cpp
@@ -108,15 +108,6 @@ namespace UE::Anim::Private
 	{
 		return GStripAdditiveRefPose != 0;
 	}
-
-	bool IsAdditiveRefPoseStripped()
-	{
-#if WITH_EDITOR
-		return false;
-#else
-		return ShouldStripAdditiveRefPose();
-#endif
-	}
 } // namespace UE::Anim::Private
 
 #if WITH_EDITOR
@@ -2571,14 +2562,22 @@ bool UAnimSequence::IsValidAdditive() const
 { 
 	if (AdditiveAnimType != AAT_None)
 	{
+		bool bIsAdditiveRefPoseStripped = UE::Anim::Private::ShouldStripAdditiveRefPose();
+#if WITH_EDITOR
+		const bool bIsCookedData = GetPackage()->HasAnyPackageFlags(PKG_Cooked);
+		if (!bIsCookedData)
+		{
+			bIsAdditiveRefPoseStripped = false;
+		}
+#endif
 		switch (RefPoseType)
 		{
 		case ABPT_RefPose:
 			return true;
 		case ABPT_AnimScaled:
-			return UE::Anim::Private::IsAdditiveRefPoseStripped() || (RefPoseSeq != nullptr);
+			return bIsAdditiveRefPoseStripped || (RefPoseSeq != nullptr);
 		case ABPT_AnimFrame:
-			return (UE::Anim::Private::IsAdditiveRefPoseStripped() || RefPoseSeq != nullptr) && (RefFrameIndex >= 0);
+			return (bIsAdditiveRefPoseStripped || RefPoseSeq != nullptr) && (RefFrameIndex >= 0);
 		case ABPT_LocalAnimFrame:
 			return (RefFrameIndex >= 0);
 		default:
diff --git a/Engine/Source/Runtime/Engine/Private/Animation/AnimSequenceBase.cpp b/Engine/Source/Runtime/Engine/Private/Animation/AnimSequenceBase.cpp
--- a/Engine/Source/Runtime/Engine/Private/Animation/AnimSequenceBase.cpp
+++ b/Engine/Source/Runtime/Engine/Private/Animation/AnimSequenceBase.cpp
@@ -299,8 +299,11 @@ void UAnimSequenceBase::PostDuplicate(EDuplicateMode::Type DuplicateMode)
 	Super::PostDuplicate(DuplicateMode);
 
 #if WITH_EDITOR
-	checkf(DataModelInterface.GetObject()->GetOuter() == this, TEXT("Animation Data Model interface has incorrect outer, expected %s - found %s"), *this->GetName(), *DataModelInterface.GetObject()->GetOuter()->GetName());
-	BindToModelModificationEvent();
+	if (!GetPackage()->HasAnyPackageFlags(PKG_Cooked))
+	{
+		checkf(DataModelInterface.GetObject()->GetOuter() == this, TEXT("Animation Data Model interface has incorrect outer, expected %s - found %s"), *this->GetName(), *DataModelInterface.GetObject()->GetOuter()->GetName());
+		BindToModelModificationEvent();
+	}
 #endif // WITH_EDITOR
 }
 
