Resample for SoundWaves on Cook

[RN] Auto SRC SoundWaves that don't meet supported codec rates
#tests AudioUnitTests / Standardbuild Cook

diff --git a/Engine/Source/Runtime/AudioMixer/Private/AudioMixerBuffer.cpp b/Engine/Source/Runtime/AudioMixer/Private/AudioMixerBuffer.cpp
--- a/Engine/Source/Runtime/AudioMixer/Private/AudioMixerBuffer.cpp
+++ b/Engine/Source/Runtime/AudioMixer/Private/AudioMixerBuffer.cpp
@@ -292,7 +292,11 @@ namespace Audio
 		if (Buffer->DecompressionState && Buffer->DecompressionState->StreamCompressedInfo(InWave, &QualityInfo))
 		{
 			// Refresh the wave data
-			InWave->SetSampleRate(QualityInfo.SampleRate);
+			const int64 NumFrames = QualityInfo.NumChannels && QualityInfo.SampleDataSize > 0 ?
+				QualityInfo.SampleDataSize / sizeof(int16) / QualityInfo.NumChannels : 0;
+			
+			InWave->SetNumFrames(IntCastChecked<uint32>(NumFrames));
+			InWave->SetSampleRate(QualityInfo.SampleRate, /* bFromDecoders */ true);
 			InWave->NumChannels = QualityInfo.NumChannels;
 			if (QualityInfo.SampleDataSize != 0)
 			{
diff --git a/Engine/Source/Runtime/AudioMixer/Private/AudioMixerBuffer.h b/Engine/Source/Runtime/AudioMixer/Private/AudioMixerBuffer.h
--- a/Engine/Source/Runtime/AudioMixer/Private/AudioMixerBuffer.h
+++ b/Engine/Source/Runtime/AudioMixer/Private/AudioMixerBuffer.h
@@ -62,6 +62,7 @@ namespace Audio
 		int32 GetNumChannels() const { return NumChannels; }
 		uint32 GetNumFrames() const { return NumFrames; }
 		void InitSampleRate(const float InSampleRate) { SampleRate = InSampleRate; }
+		void InitNumFrames(const uint32 InNumFrames) { NumFrames = InNumFrames; }
 
 	private:
 
diff --git a/Engine/Source/Runtime/AudioMixer/Private/AudioMixerSource.cpp b/Engine/Source/Runtime/AudioMixer/Private/AudioMixerSource.cpp
--- a/Engine/Source/Runtime/AudioMixer/Private/AudioMixerSource.cpp
+++ b/Engine/Source/Runtime/AudioMixer/Private/AudioMixerSource.cpp
@@ -533,8 +533,9 @@ namespace Audio
 			// Grab the source's source data override plugin settings
 			InitParams.SourceDataOverridePluginSettings = UseSourceDataOverridePlugin() ? WaveInstance->SourceDataOverridePluginSettings : nullptr;
 
-			// Update the buffer sample rate to the wave instance sample rate in case it was serialized incorrectly
+			// Update the buffer sample rate to the wave instance sample rate, as it could have changed during decoder parse.
 			MixerBuffer->InitSampleRate(WaveData->GetSampleRateForCurrentPlatform());
+			MixerBuffer->InitNumFrames(WaveData->GetNumFrames());
 
 			// Retrieve the raw pcm buffer data and the precached buffers before initializing so we can avoid having USoundWave ptrs in audio renderer thread
 			EBufferType::Type BufferType = MixerBuffer->GetType();
diff --git a/Engine/Source/Runtime/Engine/Classes/Sound/SoundWave.h b/Engine/Source/Runtime/Engine/Classes/Sound/SoundWave.h
--- a/Engine/Source/Runtime/Engine/Classes/Sound/SoundWave.h
+++ b/Engine/Source/Runtime/Engine/Classes/Sound/SoundWave.h
@@ -1211,15 +1211,17 @@ public:
 	/**
 	* Overwrite sample rate. Used for procedural soundwaves, as well as sound waves that are resampled on compress/decompress.
 	*/
-	void SetSampleRate(uint32 InSampleRate)
-	{
-		SampleRate = InSampleRate;
-#if !WITH_EDITOR
-		// Ensure that we invalidate our cached sample rate if the FProperty sample rate is changed.
-		bCachedSampleRateFromPlatformSettings = false;
-		bSampleRateManuallyReset = true;
-#endif //WITH_EDITOR
-	}
+	ENGINE_API void SetSampleRate(uint32 InSampleRate, const bool bFromDecoders=false);
+
+	/**
+	* Overwrite number of frames. Used for procedural soundwaves, as well as sound waves that are resampled on compress/decompress.
+	*/
+	ENGINE_API void SetNumFrames(uint32 InNumFrames);
+
+	/**
+	 * Return the number of frames. (can change after decoders read the headers).
+	 */
+	ENGINE_API uint32 GetNumFrames() const;
 
 	/**
 	 *	@param		Format		Format to check
@@ -1634,6 +1636,8 @@ public:
 	const uint8* GetResourceData() const { return ResourceData.GetView().GetData(); }
 	uint32 GetNumChannels() { return NumChannels; }
 
+	ENGINE_API void SetSampleRate(const int32 InSampleRate) { SampleRate = InSampleRate; }
+	ENGINE_API void SetNumFrames(const int32 InNumFrames) { NumFrames = InNumFrames; }
 
 private:
 	ENGINE_API void DiscardZerothChunkData();
diff --git a/Engine/Source/Runtime/Engine/Private/AudioDecompress.cpp b/Engine/Source/Runtime/Engine/Private/AudioDecompress.cpp
--- a/Engine/Source/Runtime/Engine/Private/AudioDecompress.cpp
+++ b/Engine/Source/Runtime/Engine/Private/AudioDecompress.cpp
@@ -830,7 +830,7 @@ void FAsyncAudioDecompressWorker::DoWork()
 			}
 #endif
 
-			Wave->SetSampleRate(QualityInfo.SampleRate);
+			Wave->SetSampleRate(QualityInfo.SampleRate, /* bFromDecoders*/ true);
 			Wave->NumChannels = QualityInfo.NumChannels;
 			if (QualityInfo.Duration > 0.0f)
 			{ 
diff --git a/Engine/Source/Runtime/Engine/Private/AudioDerivedData.cpp b/Engine/Source/Runtime/Engine/Private/AudioDerivedData.cpp
--- a/Engine/Source/Runtime/Engine/Private/AudioDerivedData.cpp
+++ b/Engine/Source/Runtime/Engine/Private/AudioDerivedData.cpp
@@ -177,8 +177,8 @@ Derived data key generation.
 
 // If you want to bump this version, generate a new guid using
 // VS->Tools->Create GUID and paste it here. https://www.guidgen.com works too.
-#define AUDIO_DERIVEDDATA_VER				TEXT("E0B90A86877043398B7D2C0963526C72")				// Change this if you want to bump all audio formats at once
-#define STREAMEDAUDIO_DERIVEDDATA_VER		TEXT("D092D8E4F0024390B5D31F4EE3C4A2C5")				// This depends on the above key, but will regenerate all streaming chunk data derived from the compressed audio
+#define AUDIO_DERIVEDDATA_VER				TEXT("64e45415311549acbee28698853aa3a4")				// Change this if you want to bump all audio formats at once
+#define STREAMEDAUDIO_DERIVEDDATA_VER		TEXT("adfedd669bf247aaa06e5f0d25b9f4ce")				// This depends on the above key, but will regenerate all streaming chunk data derived from the compressed audio
 
 #if WITH_EDITORONLY_DATA
 
@@ -475,21 +475,22 @@ class FStreamedAudioCacheDerivedDataWorker : public FNonAbandonableTask
 	bool bIsStreaming = false;
 	/** Initial chunk size */
 	int32 ZerothChunkSizeSoundWaveOverride = 0;
-	/* Cached platform specific sample rate */
-	uint32 SampleRate = 0;
-	/* Cached platform specific channel count */
-	int32 NumChannels = 0;
-	/* Cached platform specific number of AudioFrames */
-	uint64 NumOfAudioFrames = 0;
-	/* Size in Audio Frames of First Audio Chunk */
-	uint32 SizeOfFirstAudioChunkInFrames = 0;
+	/* Size in Seconds of First Audio Chunk */
+	float SizeOfFirstAudioChunkInSeconds = 0.f;
 
-	bool GetCompressedData(TArray<uint8>& OutData)
+	bool GetCompressedData(FAudioCookOutputs& OutData)
 	{
 		if (CompressedAudioHandle)
 		{
 			GetDerivedDataCacheRef().WaitAsynchronousCompletion(CompressedAudioHandle);
-			bool bResult = GetDerivedDataCacheRef().GetAsynchronousResults(CompressedAudioHandle, OutData);
+			TArray<uint8> ByteStream;
+			const bool bResult = GetDerivedDataCacheRef().GetAsynchronousResults(CompressedAudioHandle, ByteStream);
+			if (bResult)
+			{
+				FMemoryReader Reader(ByteStream);
+				ensure(OutData.Serialize(Reader));
+				ensure(OutData.EncodedData.Num() > 0);
+			}
 			CompressedAudioHandle = 0;
 			return bResult;
 		}
@@ -563,7 +564,7 @@ class FStreamedAudioCacheDerivedDataWorker : public FNonAbandonableTask
 		OutChunkOffsets.Add(ChunkTable.FindTime(0));
 	}
 
-	bool SplitUsingSeekTable(const IAudioFormat::FSeekTable& InTable, const TArray<uint8>& InSrcBuffer,
+	bool SplitUsingSeekTable(const IAudioFormat::FSeekTable& InTable, const FAudioCookOutputs& InCookOuputs,
 	                         TArray<TArray<uint8>>& OutBuffers,
 	                         const int32 InFirstChunkMaxSize,const int32 InMaxChunkSize,
 	                         TArray<uint32>& OutChunkOffsets,
@@ -575,6 +576,7 @@ class FStreamedAudioCacheDerivedDataWorker : public FNonAbandonableTask
 		// Chunk 1 - Audio (optionally inlined, containing n secs of audio).
 		// Chunk N - Audio - Sized using Max Chunk size.
 		
+		const TArray<uint8>& InSrcBuffer = InCookOuputs.EncodedData;
 		const TArray<uint32>& Offsets = InTable.Offsets;
 		const TArray<uint32>& Times = InTable.Times;
 
@@ -657,7 +659,7 @@ class FStreamedAudioCacheDerivedDataWorker : public FNonAbandonableTask
 		
 				UE_LOG(LogAudio, Verbose, TEXT("Adding Chunk %d because (%s): Blocks=%d (%u bytes), AudioFrames=%u (%2.2f seconds), SeekTableEntries=%d (%d bytes), ChunkSize=%u bytes, PercentFull=%2.2f, Remaining=%u bytes"),
 					OutBuffers.Num(), bOverSize ? TEXT("Enough Bytes") : TEXT("Enough Samples"), BlockCount, AudioSize, CurrentAudioFrameCount,
-					(float)CurrentAudioFrameCount / SampleRate, BlockCount, CurrentTableSize, Chunk.Num(), ((float)Chunk.Num() / Budget) * 100.f, Budget-Chunk.Num()
+					(float)CurrentAudioFrameCount / InCookOuputs.SampleRate, BlockCount, CurrentTableSize, Chunk.Num(), ((float)Chunk.Num() / Budget) * 100.f, Budget-Chunk.Num()
 					);
 
 				OutBuffers.Add(Chunk);
@@ -699,7 +701,7 @@ class FStreamedAudioCacheDerivedDataWorker : public FNonAbandonableTask
 
 			UE_LOG(LogAudio, Verbose, TEXT("Adding FINAL Chunk %d because (%s): Blocks=%d (%u bytes), AudioFrames=%u (%2.2f seconds), SeekTableEntries=%d (%d bytes), ChunkSize=%u bytes, PercentFull=%2.2f, Remaining=%u bytes"),
 				OutBuffers.Num(), TEXT("EOF"), BlockCount, AudioSize, CurrentAudioFrameCount,
-				(float)CurrentAudioFrameCount / SampleRate, BlockCount, CurrentTableSize, Chunk.Num(), ((float)Chunk.Num() / Budget) * 100.f, Budget-Chunk.Num()
+				(float)CurrentAudioFrameCount / InCookOuputs.SampleRate, BlockCount, CurrentTableSize, Chunk.Num(), ((float)Chunk.Num() / Budget) * 100.f, Budget-Chunk.Num()
 				);
 
 			OutBuffers.Add(Chunk);
@@ -708,7 +710,7 @@ class FStreamedAudioCacheDerivedDataWorker : public FNonAbandonableTask
 		return true;
 	}
 
-	bool SplitDataForStreaming(const IAudioFormat* InFormat, const TArray<uint8>& InSrcBuffer, TArray<TArray<uint8>>& OutBuffers, const int32 InFirstChunkMaxSize, const int32 InMaxChunkSize, TArray<uint32>& OutChunkOffsets,
+	bool SplitDataForStreaming(const IAudioFormat* InFormat, const FAudioCookOutputs& InCookOutputs, TArray<TArray<uint8>>& OutBuffers, const int32 InFirstChunkMaxSize, const int32 InMaxChunkSize, TArray<uint32>& OutChunkOffsets,
 		const uint32 InFirstChunkMaxLengthAudioFrames, const uint32 InMaxChunkLengthAudioFrames) const
 	{
 		// Split using the seek-table if we have one.
@@ -716,10 +718,10 @@ class FStreamedAudioCacheDerivedDataWorker : public FNonAbandonableTask
 		{
 			// Because the extract call modifies the source buffer inline (to remove the embedded seek table) we must keep a copy as if
 			// the call was to fail to split the calling code must contend with the original.
-			TArray<uint8> InSrcBufferCopy = InSrcBuffer;
+			FAudioCookOutputs InSrcBufferCopy = InCookOutputs;
 		
 			IAudioFormat::FSeekTable Table;			
-			if (!ensureMsgf(InFormat->ExtractSeekTableForStreaming(InSrcBufferCopy, Table), TEXT("SoundWave: '%s' requires a Seektable, but doesn't contain one."), *SoundWaveFullName))
+			if (!ensureMsgf(InFormat->ExtractSeekTableForStreaming(InSrcBufferCopy.EncodedData, Table), TEXT("SoundWave: '%s' requires a Seektable, but doesn't contain one."), *SoundWaveFullName))
 			{
 				return false;
 			}
@@ -736,7 +738,7 @@ class FStreamedAudioCacheDerivedDataWorker : public FNonAbandonableTask
 		}
 
 		// otherwise... ask the format to split.
-		return InFormat->SplitDataForStreaming(InSrcBuffer, OutBuffers, InFirstChunkMaxSize, InMaxChunkSize);
+		return InFormat->SplitDataForStreaming(InCookOutputs.EncodedData, OutBuffers, InFirstChunkMaxSize, InMaxChunkSize);
 	}
 
 	/** Build the streamed audio. This function is safe to call from any thread. */
@@ -762,8 +764,8 @@ class FStreamedAudioCacheDerivedDataWorker : public FNonAbandonableTask
 		{
 			DerivedData->AudioFormat = AudioFormatName;
 
-			TArray<uint8> CompressedBuffer;
-			if (GetCompressedData(CompressedBuffer))
+			FAudioCookOutputs CookOutputs;
+			if (GetCompressedData(CookOutputs))
 			{
 				TArray<TArray<uint8>> ChunkBuffers;
 
@@ -773,7 +775,7 @@ class FStreamedAudioCacheDerivedDataWorker : public FNonAbandonableTask
 				// By default, the zeroth chunk's max size is the same as the other chunks.
 				int32 ZerothChunkSize = MaxChunkSizeForCurrentWave;
 
-				const int32 MinimumChunkSize = AudioFormat->GetMinimumSizeForInitialChunk(AudioFormatName, CompressedBuffer);
+				const int32 MinimumChunkSize = AudioFormat->GetMinimumSizeForInitialChunk(AudioFormatName, CookOutputs.EncodedData);
 				const bool bForceLegacyStreamChunking = bIsStreaming && CompressionOverrides && CompressionOverrides->StreamCachingSettings.bForceLegacyStreamChunking;
 
 				// If the zeroth size for this sound wave was overridden, use that:
@@ -784,7 +786,7 @@ class FStreamedAudioCacheDerivedDataWorker : public FNonAbandonableTask
 				else
 				{
 					// Ensure that the minimum chunk size is nonzero if our compressed buffer is not empty.
-					checkf(CompressedBuffer.Num() == 0 || MinimumChunkSize != 0, TEXT("To use Load On Demand, please override GetMinimumSizeForInitialChunk"));
+					checkf(CookOutputs.EncodedData.Num() == 0 || MinimumChunkSize != 0, TEXT("To use Load On Demand, please override GetMinimumSizeForInitialChunk"));
 
 					//
 					if (bForceLegacyStreamChunking)
@@ -827,16 +829,16 @@ class FStreamedAudioCacheDerivedDataWorker : public FNonAbandonableTask
 				uint32 MaxLengthOfFirstAudioChunkInFrames = MAX_uint32;
 				if (bInlineFirstAudioChunk)
 				{
-					MaxLengthOfFirstAudioChunkInFrames = SizeOfFirstAudioChunkInFrames;
+					MaxLengthOfFirstAudioChunkInFrames = CookOutputs.SampleRate * SizeOfFirstAudioChunkInSeconds;
 				}
 
 				TArray<uint32> ChunkOffsets;
-				if (SplitDataForStreaming(AudioFormat, CompressedBuffer, ChunkBuffers, ZerothChunkSize,
+				if (SplitDataForStreaming(AudioFormat, CookOutputs, ChunkBuffers, ZerothChunkSize,
 				                          MaxChunkSizeForCurrentWave, ChunkOffsets, MaxLengthOfFirstAudioChunkInFrames,
 				                          MAX_uint32))
 				{
 					UE_LOG(LogAudio, Display, TEXT("Chunk stats for (%s: Duration=%2.2f secs, Channels=%d), Settings(FirstChunkSize=%u frames, InlineFirst=%s, MaxChunkSize=%d), Chunks=%d, Chunk0=%d bytes, Chunk1=%d bytes, ChunkN=%d)"),
-						*SoundWaveFullName,(float)NumOfAudioFrames / SampleRate, NumChannels,
+						*SoundWaveFullName,(float)CookOutputs.NumFrames / CookOutputs.SampleRate, CookOutputs.NumChannels,
 						MaxLengthOfFirstAudioChunkInFrames, ToCStr(LexToString(bInlineFirstAudioChunk)), MaxChunkSizeForCurrentWave,
 						ChunkBuffers.Num(),
 						ChunkBuffers.IsValidIndex(0) ? ChunkBuffers[0].Num() : 0,
@@ -902,7 +904,7 @@ class FStreamedAudioCacheDerivedDataWorker : public FNonAbandonableTask
 						// NOTE we should only do this if we've been given a size greater than 0 to put there
 						// otherwise the chunk will use the normal size boundary.
 						const bool bIsFirstAudioChunk = NewChunk->SeekOffsetInAudioFrames == 0;
-						NewChunk->bInlineChunk = bIsFirstAudioChunk && bInlineFirstAudioChunk && SizeOfFirstAudioChunkInFrames > 0;
+						NewChunk->bInlineChunk = bIsFirstAudioChunk && bInlineFirstAudioChunk && SizeOfFirstAudioChunkInSeconds > 0;
 						
 						if (NewChunk->BulkData.IsLocked())
 						{
@@ -929,7 +931,7 @@ class FStreamedAudioCacheDerivedDataWorker : public FNonAbandonableTask
 					// Could not split so copy compressed data into a single chunk
 					FStreamedAudioChunk* NewChunk = new FStreamedAudioChunk();
 					DerivedData->Chunks.Add(NewChunk);
-					NewChunk->DataSize = CompressedBuffer.Num();
+					NewChunk->DataSize = CookOutputs.EncodedData.Num();
 					NewChunk->AudioDataSize = NewChunk->DataSize;
 
 					if (NewChunk->BulkData.IsLocked())
@@ -938,8 +940,8 @@ class FStreamedAudioCacheDerivedDataWorker : public FNonAbandonableTask
 					}
 
 					NewChunk->BulkData.Lock(LOCK_READ_WRITE);
-					void* NewChunkData = NewChunk->BulkData.Realloc(CompressedBuffer.Num());
-					FMemory::Memcpy(NewChunkData, CompressedBuffer.GetData(), CompressedBuffer.Num());
+					void* NewChunkData = NewChunk->BulkData.Realloc(CookOutputs.EncodedData.Num());
+					FMemory::Memcpy(NewChunkData, CookOutputs.EncodedData.GetData(), CookOutputs.EncodedData.Num());
 					NewChunk->BulkData.Unlock();
 				}
 
@@ -993,23 +995,12 @@ public:
 		, bIsProcedural(InSoundWave->IsA<USoundWaveProcedural>())
 		, bIsStreaming(InSoundWave->bStreaming)
 		, ZerothChunkSizeSoundWaveOverride(InSoundWave->InitialChunkSize_DEPRECATED)
+		, SizeOfFirstAudioChunkInSeconds(InSoundWave->GetSizeOfFirstAudioChunkInSeconds(InTargetPlatform))
 	{
 		// Gather all USoundWave object inputs to avoid race-conditions that could result when touching the UObject from another thread.
 		GetStreamedAudioDerivedDataKeySuffix(*InSoundWave, AudioFormatName, CompressionOverrides, InTargetPlatform, KeySuffix);
 		const FName PlatformSpecificFormat = InSoundWave->GetPlatformSpecificFormat(InAudioFormatName, CompressionOverrides);
 
-		// Cache num channels (only used for logging currently).
-		NumChannels = InSoundWave->NumChannels;
-
-		// Cache the sample-rate (making sure to use the compression overrides version for this platform).
-		const float OverrideSampleRate = InSoundWave->GetSampleRateForCompressionOverrides(InCompressionOverrides);
-		SampleRate = OverrideSampleRate > 0 ? static_cast<uint32>(OverrideSampleRate)
-			: static_cast<uint32>(InSoundWave->GetSampleRateForCurrentPlatform());
-
-		// Determine amount to inline in Audio frames. (careful to round for determinism reasons).
-		const float SizeOfFirstAudioChunkInSeconds = InSoundWave->GetSizeOfFirstAudioChunkInSeconds(InTargetPlatform);
-		SizeOfFirstAudioChunkInFrames = FMath::CeilToInt(static_cast<float>(SampleRate * SizeOfFirstAudioChunkInSeconds));
-
  		// Fetch compressed data directly from the DDC to ensure thread-safety. Will be async if the compressor is thread-safe.
 		FDerivedAudioDataCompressor* DeriveAudioData = new FDerivedAudioDataCompressor(InSoundWave, InAudioFormatName, PlatformSpecificFormat, CompressionOverrides, InTargetPlatform);
 		CompressedAudioHandle = GetDerivedDataCacheRef().GetAsynchronous(DeriveAudioData);
@@ -1080,8 +1071,8 @@ public:
 		}
 
 		// Cleanup the async DDC query if needed
-		TArray<uint8> Dummy;
-		GetCompressedData(Dummy);
+		FAudioCookOutputs DummyOutputs;
+		GetCompressedData(DummyOutputs);
 
 		return bLoadedFromDDC;
 	}
@@ -1602,7 +1593,7 @@ static void ResampleWaveData(Audio::FAlignedFloatBuffer& WaveData, int32 NumChan
 	}
 
 	double TimeDelta = FPlatformTime::Seconds() - StartTime;
-	UE_LOG(LogAudioDerivedData, Display, TEXT("Resampling file from %f to %f took %f seconds."), SourceSampleRate, DestinationSampleRate, TimeDelta);
+	UE_LOG(LogAudioDerivedData, Display, TEXT("Resampling file from %d to %d took %2.2f seconds."), (int32)SourceSampleRate, (int32)DestinationSampleRate, TimeDelta);
 }
 
 struct FAudioCookInputs
@@ -1682,7 +1673,6 @@ struct FAudioCookInputs
 		}
 	}
 };
-
 #if WITH_EDITORONLY_DATA
 
 static int32 CalculateModifiedCompressionQuality(int32 InQuality, float InQualityModifier)
@@ -1690,11 +1680,88 @@ static int32 CalculateModifiedCompressionQuality(int32 InQuality, float InQualit
 	float ModifiedCompressionQuality = (float)InQuality * InQualityModifier;
 	return FMath::Clamp<int32>(FMath::FloorToInt(ModifiedCompressionQuality), 1, 100);
 }
+static TOptional<int32> GetResampleRate(const IAudioFormat* InCompressor, const int32 InSampleRateOverride,
+	const FName InFormatName, const int32 InCurrentSampleRate)
+{
+	using namespace AudioDerivedDataPrivate;
+	if (!ensure(InCompressor))
+	{
+		return {};
+	}
+
+	// Start with the rate of the wave.
+	int32 ResampleRate = InCurrentSampleRate;
+	
+	// Inputs specify an override?
+	if (InSampleRateOverride > 0 && InSampleRateOverride != InCurrentSampleRate)
+	{
+		ResampleRate = InSampleRateOverride;
+
+		static bool bLogOnce = false;
+		if (!InCompressor->IsSampleRateSupported(ResampleRate) && !bLogOnce)
+		{
+			UE_LOG(LogAudio, Warning,  TEXT("SampleRate Override '%d' not supported by AudioFormat '%s'"), InSampleRateOverride, *InFormatName.ToString());
+			bLogOnce = true;
+		}
+	}
+
+	// Do we support this current rate/override?
+	if (!InCompressor->IsSampleRateSupported(ResampleRate))
+	{
+		// Ask for the rates supported by this format.
+		TArray<int32> SupportedRates(InCompressor->GetSupportedSampleRates());
+
+		// If we've failed the "IsSampleRateSupported" check above, we should have some rate we support defined here.
+		checkf(SupportedRates.Num() > 0, TEXT("AudioFormat=%s, returned no supported rates, but rejected our current rate=%d"),
+			*InFormatName.ToString(), InCurrentSampleRate);
+		
+		// Just in case, sort the item, as we binary search them.
+		SupportedRates.Sort();
+
+		const int32 MaxSupportedRate = SupportedRates.Last();
+		const int32 MinSupportedRate = SupportedRates[0];
+
+		// Special case for just one entry.
+		if (SupportedRates.Num() == 1)
+		{
+			ResampleRate = SupportedRates.Last();
+		}		
+		// < Min
+		else if (ResampleRate < MinSupportedRate)
+		{
+			ResampleRate = MinSupportedRate;
+		}
+		// > Max.
+		else if (ResampleRate > MaxSupportedRate)
+		{
+			ResampleRate = MaxSupportedRate;	
+		}
+		// ... Find first greatest.
+		else
+		{
+			const int32 Index = Algo::LowerBound(SupportedRates, ResampleRate);
+			check(SupportedRates.IsValidIndex(Index));
+			ResampleRate = SupportedRates[Index];
+		}
+
+		// We should have found something now.
+		ensure(InCompressor->IsSampleRateSupported(ResampleRate));
+	}
+
+	// If we're different return the new rate.
+	if (ResampleRate != InCurrentSampleRate && ResampleRate > 0)
+	{
+		return ResampleRate;
+	}
+
+	// Don't resample if the rate is the same at the current.
+	return {};
+}
 
 /**
  * Cook a simple mono or stereo wave
  */
-static void CookSimpleWave(const FAudioCookInputs& Inputs, TArray<uint8>& OutputBuffer)
+static void CookSimpleWave(const FAudioCookInputs& Inputs, FAudioCookOutputs& Outputs)
 {
 	// Warning: Existing released assets should maintain bitwise exact encoded audio
 	// in order to minimize patch sizes. Changing anything in this function can 
@@ -1704,7 +1771,7 @@ static void CookSimpleWave(const FAudioCookInputs& Inputs, TArray<uint8>& Output
 
 	FWaveModInfo WaveInfo;
 	TArray<uint8> Input;
-	check(!OutputBuffer.Num());
+	check(!Outputs.EncodedData.Num());
 
 	bool bWasLocked = false;
 
@@ -1741,7 +1808,8 @@ static void CookSimpleWave(const FAudioCookInputs& Inputs, TArray<uint8>& Output
 	int32 NumBytes = Input.Num();
 	int32 NumSamples = NumBytes / sizeof(int16);
 
-	const bool bNeedsResample = (Inputs.SampleRateOverride > 0 && Inputs.SampleRateOverride != (float)WaveSampleRate);
+	const TOptional<int32> ResampleRate = GetResampleRate(Inputs.Compressor, Inputs.SampleRateOverride, Inputs.BaseFormat, WaveSampleRate);
+	const bool bNeedsResample = ResampleRate.IsSet(); 
 	const bool bNeedsToApplyWaveTransformation = (Inputs.WaveTransformations.Num() > 0);
 
 	// Only convert PCM wave data to float if needed. The conversion alters the sample
@@ -1782,9 +1850,9 @@ static void CookSimpleWave(const FAudioCookInputs& Inputs, TArray<uint8>& Output
 		// Resample if necessary
 		if (bNeedsResample)
 		{
-			ResampleWaveData(InputFloatBuffer, NumChannels, WaveSampleRate, Inputs.SampleRateOverride);
+			ResampleWaveData(InputFloatBuffer, NumChannels, WaveSampleRate, *ResampleRate);
 			
-			WaveSampleRate = Inputs.SampleRateOverride;
+			WaveSampleRate = *ResampleRate;
 			NumSamples = InputFloatBuffer.Num();
 		}
 
@@ -1828,16 +1896,21 @@ static void CookSimpleWave(const FAudioCookInputs& Inputs, TArray<uint8>& Output
 	}
 
 	// Cook the data.
-	if (!Inputs.Compressor->Cook(Inputs.BaseFormat, Input, QualityInfo, OutputBuffer))
+	if (!Inputs.Compressor->Cook(Inputs.BaseFormat, Input, QualityInfo, Outputs.EncodedData))
 	{
 		UE_LOG(LogAudioDerivedData, Warning, TEXT("Cooking sound failed: %s"), *Inputs.SoundFullName);
 	}
+
+	// Record the results.
+	Outputs.NumChannels = NumChannels;
+	Outputs.SampleRate = WaveSampleRate;
+	Outputs.NumFrames = NumSamples / NumChannels;
 }
 
 /**
  * Cook a multistream (normally 5.1) wave
  */
-static void CookSurroundWave(const FAudioCookInputs& Inputs,  TArray<uint8>& OutputBuffer)
+static void CookSurroundWave(const FAudioCookInputs& Inputs, FAudioCookOutputs& Outputs)
 {
 	// Warning: Existing released assets should maintain bitwise exact encoded audio
 	// in order to minimize patch sizes. Changing anything in this function can 
@@ -1845,7 +1918,7 @@ static void CookSurroundWave(const FAudioCookInputs& Inputs,  TArray<uint8>& Out
 
 	TRACE_CPUPROFILER_EVENT_SCOPE(CookSurroundWave);
 
-	check(!OutputBuffer.Num());
+	check(!Outputs.EncodedData.Num());
 
 	int32					i;
 	size_t					SampleDataSize = 0;
@@ -1974,7 +2047,9 @@ static void CookSurroundWave(const FAudioCookInputs& Inputs,  TArray<uint8>& Out
 	int32 NumFrames = SampleDataSize / sizeof(int16);
 
 	// bNeedsResample could change if a transformation changes the sample rate
-	bool bNeedsResample = Inputs.SampleRateOverride > 0 && Inputs.SampleRateOverride != (float)WaveSampleRate;
+
+	const TOptional<int32> ResampleRate = GetResampleRate(Inputs.Compressor, Inputs.SampleRateOverride, Inputs.BaseFormat, WaveSampleRate);
+	const bool bNeedsResample = ResampleRate.IsSet(); 
 	
 	const bool bContainsTransformations = Inputs.WaveTransformations.Num() > 0;
 	const bool bNeedsDeinterleave = bNeedsResample || bContainsTransformations;
@@ -2011,20 +2086,16 @@ static void CookSurroundWave(const FAudioCookInputs& Inputs,  TArray<uint8>& Out
 			}
 
 			UE_CLOG(WaveSampleRate != TransformationInfo.SampleRate, LogAudioDerivedData, Warning, TEXT("Wave transformations which alter the sample rate are not supported. Cooked audio for %s may be incorrect"), *Inputs.SoundFullName);
-			WaveSampleRate = TransformationInfo.SampleRate;
-			bNeedsResample = Inputs.SampleRateOverride > 0 && Inputs.SampleRateOverride != (float)WaveSampleRate;
-			
 			UE_CLOG(ChannelCount != TransformationInfo.NumChannels, LogAudioDerivedData, Error, TEXT("Wave transformations which alter number of channels are not supported. Cooked audio for %s may be incorrect"), *Inputs.SoundFullName);
-			ChannelCount = TransformationInfo.NumChannels;
 
 			NumFrames = InterleavedFloatBuffer.Num() / ChannelCount;
 		}
 
 		if (bNeedsResample)
 		{
-			ResampleWaveData(InterleavedFloatBuffer, ChannelCount, WaveSampleRate, Inputs.SampleRateOverride);
+			ResampleWaveData(InterleavedFloatBuffer, ChannelCount, WaveSampleRate, *ResampleRate);
 			
-			WaveSampleRate = Inputs.SampleRateOverride;
+			WaveSampleRate = *ResampleRate; 
 			NumFrames = InterleavedFloatBuffer.Num() / ChannelCount;
 		}
 
@@ -2076,14 +2147,56 @@ static void CookSurroundWave(const FAudioCookInputs& Inputs,  TArray<uint8>& Out
 	}
 
 	//@todo tighten up the checking for empty results here
-	if (!Inputs.Compressor->CookSurround(Inputs.BaseFormat, SourceBuffers, QualityInfo, OutputBuffer))
+	if (!Inputs.Compressor->CookSurround(Inputs.BaseFormat, SourceBuffers, QualityInfo,  Outputs.EncodedData))
 	{
 		UE_LOG(LogAudioDerivedData, Warning, TEXT("Cooking surround sound failed: %s"), *Inputs.SoundFullName);
 	}
+ 
+	// Record the results.
+	Outputs.NumChannels = ChannelCount;
+	Outputs.SampleRate = WaveSampleRate;
+	Outputs.NumFrames = NumFrames;
 }
 
 #endif // WITH_EDITORONLY_DATA
 
+int32 FAudioCookOutputs::GetVersion()
+{
+	static int32 Ver = 2;
+	return Ver;
+}
+
+bool FAudioCookOutputs::Serialize(FArchive& Ar)
+{
+	if (Ar.IsSaving())
+	{
+		Id = GetId(); // Always set to known good ID
+		Version = GetVersion();
+	}
+
+	Ar << Id;
+
+	// Doses this look sane?
+	if (Ar.IsLoading() && Id != GetId())
+	{
+		return false; 
+	}
+
+	// Reject different versions. As we are DDC this indicates an error.
+	Ar << Version;
+	if (Ar.IsLoading() && Version != GetVersion())
+	{
+		return false;
+	}
+
+	Ar << NumChannels;
+	Ar << SampleRate;
+	Ar << NumFrames;
+	Ar << EncodedData;
+
+	return true;
+}
+
 FDerivedAudioDataCompressor::FDerivedAudioDataCompressor(USoundWave* InSoundNode, FName InBaseFormat, FName InHashedFormat, const FPlatformAudioCookOverrides* InCompressionOverrides, const ITargetPlatform* InTargetPlatform)
 	: CookInputs(MakeUnique<FAudioCookInputs>(InSoundNode, InBaseFormat, InHashedFormat, InCompressionOverrides, InTargetPlatform))
 {
@@ -2110,7 +2223,8 @@ FString FDerivedAudioDataCompressor::GetPluginSpecificCacheKeySuffix() const
 	SoundWaveHash = CookInputs->SoundWaveHash;
 #endif //WITH_EDITORONLY_DATA
 	
-	return FString::Printf(TEXT("%s_%04X_%s%s"), *FormatHash, FormatVersion, *SoundWaveHash, *CookInputs->CompressedDataGuid.ToString());
+	return FString::Printf(TEXT("%s_%04X_%s%s%d"), *FormatHash, FormatVersion, *SoundWaveHash, *CookInputs->CompressedDataGuid.ToString(),
+		FAudioCookOutputs::GetVersion());
 }
 
 bool FDerivedAudioDataCompressor::IsBuildThreadsafe() const
@@ -2140,21 +2254,26 @@ bool FDerivedAudioDataCompressor::Build(TArray<uint8>& OutData)
 	{
 		return false;
 	}
-	
+
+	FAudioCookOutputs CookOutputs;
 	if (!CookInputs->ChannelSizes.Num())
 	{
 		check(!CookInputs->ChannelOffsets.Num());
-		CookSimpleWave(*CookInputs, OutData);
+		CookSimpleWave(*CookInputs, CookOutputs);
 	}
 	else
 	{
 		check(CookInputs->ChannelOffsets.Num() == SPEAKER_Count);
 		check(CookInputs->ChannelSizes.Num() == SPEAKER_Count);
-		CookSurroundWave(*CookInputs, OutData);
+		CookSurroundWave(*CookInputs, CookOutputs);
 	}
 
+	// Serialize to bitstream.
+	FMemoryWriter Writer(OutData);
+	ensure(CookOutputs.Serialize(Writer));
+		
 	const uint64 BeforeSize = CookInputs->BulkData.GetPayloadSize();
-	const uint64 AfterSize = OutData.Num();
+	const uint64 AfterSize = CookOutputs.EncodedData.Num();
 	const float Percent = BeforeSize > 0 ? ((float)AfterSize / BeforeSize) * 100.f : 0.0f;
 
 	// Log message about the completed results
diff --git a/Engine/Source/Runtime/Engine/Private/AudioFormatSettings.cpp b/Engine/Source/Runtime/Engine/Private/AudioFormatSettings.cpp
--- a/Engine/Source/Runtime/Engine/Private/AudioFormatSettings.cpp
+++ b/Engine/Source/Runtime/Engine/Private/AudioFormatSettings.cpp
@@ -130,14 +130,7 @@ namespace Audio
 				return false;		
 			}
 			
-			// Sample rate ok? 
-			if (!Format->IsSampleRateSupported(InWave.SampleRate))
-			{
-				UE_LOG(LogAudio, Verbose, TEXT("Wave '%s', format '%s' doesn't support sample-rate: '%d'"),
-					*InWave.Name.ToString(), *InWave.FormatName.ToString(), InWave.SampleRate);
-				return false;	
-			}
-
+		
 			// Channel count ok?
 			if (!Format->IsChannelCountSupported(InWave.NumChannels))
 			{
diff --git a/Engine/Source/Runtime/Engine/Private/SoundWave.cpp b/Engine/Source/Runtime/Engine/Private/SoundWave.cpp
--- a/Engine/Source/Runtime/Engine/Private/SoundWave.cpp
+++ b/Engine/Source/Runtime/Engine/Private/SoundWave.cpp
@@ -37,7 +37,9 @@
 #include "Templates/UnrealTemplate.h"
 #include "UObject/ObjectSaveContext.h"
 #include "ISoundWaveCloudStreaming.h"
+#include "SoundWaveCompiler.h"
 #include "SoundFileIO/SoundFileIO.h"
+#include "Serialization/MemoryReader.h"
 
 #if WITH_EDITORONLY_DATA
 #include "uewav.h"
@@ -1147,6 +1149,40 @@ void USoundWave::GetAssetRegistryTags(FAssetRegistryTagsContext Context) const
 #endif
 }
 
+
+void USoundWave::SetSampleRate(uint32 InSampleRate, const bool bFromDecoders)
+{
+	if (SoundWaveDataPtr)
+	{
+		SoundWaveDataPtr->SetSampleRate(IntCastChecked<int32>(InSampleRate));
+	}
+	SampleRate = InSampleRate;;
+
+	if (bFromDecoders)
+	{
+		// Ensure that we invalidate our cached sample rate if the FProperty sample rate is changed.
+		bSampleRateManuallyReset = true;
+		bCachedSampleRateFromPlatformSettings = false;
+	}
+}
+
+void USoundWave::SetNumFrames(uint32 InNumFrames)
+{
+	if (SoundWaveDataPtr)
+	{
+		SoundWaveDataPtr->SetNumFrames(InNumFrames);	
+	}
+}
+
+uint32 USoundWave::GetNumFrames() const
+{
+	if (SoundWaveDataPtr)
+	{
+		return SoundWaveDataPtr->GetNumFrames();
+	}
+	return 0;
+}
+
 void USoundWave::Serialize( FArchive& Ar )
 {
 	LLM_SCOPE(ELLMTag::AudioSoundWaves);
@@ -1926,8 +1962,13 @@ FByteBulkData* USoundWave::GetCompressedData(FName Format, const FPlatformAudioC
 			if (bGetSuccessful)
 			{
 				COOK_STAT(Timer.AddHitOrMiss(bDataWasBuilt ? FCookStats::CallStats::EHitOrMiss::Miss : FCookStats::CallStats::EHitOrMiss::Hit, OutData.Num()));
+
+				FAudioCookOutputs CookOutput;
+				FMemoryReader Reader(OutData);
+				ensure(CookOutput.Serialize(Reader));
+				
 				Result->Lock(LOCK_READ_WRITE);
-				FMemory::Memcpy(Result->Realloc(OutData.Num()), OutData.GetData(), OutData.Num());
+				FMemory::Memcpy(Result->Realloc(CookOutput.EncodedData.Num()), CookOutput.EncodedData.GetData(), CookOutput.EncodedData.Num());
 				Result->Unlock();
 			}
 		}
@@ -4121,6 +4162,11 @@ float USoundWave::GetSampleRateForCurrentPlatform() const
 	// GIsEditor is false in standalone but WITH_EDITORONLY_DATA is true
 	if (GIsEditor)
 	{
+		if (bSampleRateManuallyReset)
+		{
+			return SampleRate;
+		}
+		
 		float SampleRateOverride = FPlatformCompressionUtilities::GetTargetSampleRateForPlatform(SampleRateQuality);
 		if (SampleRateOverride > 0)
 		{
diff --git a/Engine/Source/Runtime/Engine/Public/AudioDerivedData.h b/Engine/Source/Runtime/Engine/Public/AudioDerivedData.h
--- a/Engine/Source/Runtime/Engine/Public/AudioDerivedData.h
+++ b/Engine/Source/Runtime/Engine/Public/AudioDerivedData.h
@@ -5,11 +5,78 @@
 #include "CoreMinimal.h"
 #include "DerivedDataPluginInterface.h"
 
+#include "AudioDerivedData.generated.h"
+
 class IAudioFormat;
 class USoundWave;
 struct FPlatformAudioCookOverrides;
 struct FAudioCookInputs;
 
+
+/**
+ * Struct that's serialized into the DDC record for "Audio" DDC task.
+ * Filled in by the encoders, additional relevant state is passed back here.
+ * Like decisions make during encoding, like resampling, transformations etc.
+ */
+USTRUCT()
+struct FAudioCookOutputs
+{
+	GENERATED_BODY()
+
+	/**
+	 * Get the expected ID of this struct. uint32 so we can peek at memory.
+	 * @return uint32 Id.
+	 */
+	static uint32 GetId()
+	{
+		// Magic value is "ACO"
+		return ((uint32)'O' << 16 | (uint32)'C' << 8 | (uint32)'A');
+	}
+
+	/**
+	 * Gets the current Version expected. Any other versions would be an error.
+	 * @return int32 Version.
+	 */
+	static int32 GetVersion();
+
+	/**
+	 * Serializer (read/write)
+	 * @param Ar Reader/Writer Archive.
+	 * @return True if successfully read/written, false otherwise.
+	 */
+	bool Serialize(FArchive& Ar);
+
+	UPROPERTY()
+	uint32 Id = GetId();
+	
+	UPROPERTY()
+	int32 Version = GetVersion();
+
+	/**
+	 * Final channel count that was encoded by the IAudioFormat
+	 */
+	UPROPERTY()
+	int32 NumChannels = 0;
+
+	/**
+	 * Final Sample rate that was encoded by IAudioFormat 
+	 */
+	UPROPERTY()
+	int32 SampleRate = 0;
+
+	/**
+	 * Number of frames in the encoded data.
+	 */
+	UPROPERTY()
+	uint32 NumFrames = 0;
+
+	/**
+	 * The Binary output of the IAudioFormat. 
+	 */
+	UPROPERTY()
+	TArray<uint8> EncodedData;
+};
+
 class FDerivedAudioDataCompressor : public FDerivedDataPluginInterface
 {
 private:
