Fix Linux CIS Issue

diff --git a/Engine/Source/Runtime/Engine/Classes/Sound/SoundWave.h b/Engine/Source/Runtime/Engine/Classes/Sound/SoundWave.h
--- a/Engine/Source/Runtime/Engine/Classes/Sound/SoundWave.h
+++ b/Engine/Source/Runtime/Engine/Classes/Sound/SoundWave.h
@@ -1636,7 +1636,7 @@ public:
 	const uint8* GetResourceData() const { return ResourceData.GetView().GetData(); }
 	uint32 GetNumChannels() { return NumChannels; }
 
-	ENGINE_API void SetSampleRate(const int32 InSampleRate) { SampleRate = InSampleRate; }
+	ENGINE_API void SetSampleRate(const int32 InSampleRate) { SampleRate = static_cast<float>(InSampleRate); }
 	ENGINE_API void SetNumFrames(const int32 InNumFrames) { NumFrames = InNumFrames; }
 
 private:
