[shaders] fix job status reset on a retry to reset the status set on the root pipeline job in addition to the stage jobs (the former is the one that's actually updated in the pipeline case)

diff --git a/Engine/Source/Runtime/Engine/Private/ShaderCompiler/ShaderCompiler.cpp b/Engine/Source/Runtime/Engine/Private/ShaderCompiler/ShaderCompiler.cpp
--- a/Engine/Source/Runtime/Engine/Private/ShaderCompiler/ShaderCompiler.cpp
+++ b/Engine/Source/Runtime/Engine/Private/ShaderCompiler/ShaderCompiler.cpp
@@ -2536,6 +2536,7 @@ bool FShaderCompilingManager::HandlePotentialRetry(TMap<int32, FShaderMapFinaliz
 			
 			// Reset DDC query request owner
 			CurrentJob->RequestOwner.Reset();
+			CurrentJob->JobStatusPtr->Reset();
 			CurrentJob->bInputHashSet = false;
 			CurrentJob->bFinalized = false;
 			// Need to force reissued jobs to skip the cache queries otherwise jobs with warnings will just be cache hits and not actually recompile
