Fixed dllstorage. Removed dllstorage on types and moved to variables/methods plus removed dllstorage on inlined methods

#[robomerge] 5.6

diff --git a/Engine/Plugins/Experimental/GameplayGraph/Source/GameplayGraph/Public/Graph/GraphDefaultSerialization.h b/Engine/Plugins/Experimental/GameplayGraph/Source/GameplayGraph/Public/Graph/GraphDefaultSerialization.h
--- a/Engine/Plugins/Experimental/GameplayGraph/Source/GameplayGraph/Public/Graph/GraphDefaultSerialization.h
+++ b/Engine/Plugins/Experimental/GameplayGraph/Source/GameplayGraph/Public/Graph/GraphDefaultSerialization.h
@@ -24,7 +24,7 @@ struct FSerializedEdgeData
 	friend GAMEPLAYGRAPH_API bool operator==(const FSerializedEdgeData& Lhs, const FSerializedEdgeData& Rhs);
 	friend GAMEPLAYGRAPH_API bool operator!=(const FSerializedEdgeData& Lhs, const FSerializedEdgeData& Rhs);
 
-	friend uint32 GAMEPLAYGRAPH_API GetTypeHash(const FSerializedEdgeData& EdgeData)
+	friend uint32 GetTypeHash(const FSerializedEdgeData& EdgeData)
 	{
 		return HashCombine(GetTypeHash(EdgeData.Node1), GetTypeHash(EdgeData.Node2));
 	}
diff --git a/Engine/Source/Runtime/Engine/Classes/Curves/IndexedCurve.h b/Engine/Source/Runtime/Engine/Classes/Curves/IndexedCurve.h
--- a/Engine/Source/Runtime/Engine/Classes/Curves/IndexedCurve.h
+++ b/Engine/Source/Runtime/Engine/Classes/Curves/IndexedCurve.h
@@ -38,16 +38,16 @@ public:
 	}
 
 	/** Get number of keys in curve. */
-	ENGINE_API virtual int32 GetNumKeys() const PURE_VIRTUAL(FIndexedCurve::GetNumKeys, return 0;);
+	virtual int32 GetNumKeys() const PURE_VIRTUAL(FIndexedCurve::GetNumKeys, return 0;);
 
 	/** Move a key to a new time. This may change the index of the key, so the new key index is returned. */
-	ENGINE_API virtual void SetKeyTime(FKeyHandle KeyHandle, float NewTime) PURE_VIRTUAL(FIndexedCurve::SetKeyTime, );
+	virtual void SetKeyTime(FKeyHandle KeyHandle, float NewTime) PURE_VIRTUAL(FIndexedCurve::SetKeyTime, );
 
 	/** Get the time for the Key with the specified index. */
-	ENGINE_API virtual float GetKeyTime(FKeyHandle KeyHandle) const PURE_VIRTUAL(FIndexedCurve::GetKeyTime, return 0.f;);
+	virtual float GetKeyTime(FKeyHandle KeyHandle) const PURE_VIRTUAL(FIndexedCurve::GetKeyTime, return 0.f;);
 	
 	/** Allocates a duplicate of the curve */
-	ENGINE_API virtual FIndexedCurve* Duplicate() const PURE_VIRTUAL(FIndexedCurve::Duplicate, return nullptr;);
+	virtual FIndexedCurve* Duplicate() const PURE_VIRTUAL(FIndexedCurve::Duplicate, return nullptr;);
 
 	/** Checks to see if the key handle is valid for this curve. */
 	ENGINE_API bool IsKeyHandleValid(FKeyHandle KeyHandle) const;
diff --git a/Engine/Source/Runtime/Engine/Classes/Curves/RealCurve.h b/Engine/Source/Runtime/Engine/Classes/Curves/RealCurve.h
--- a/Engine/Source/Runtime/Engine/Classes/Curves/RealCurve.h
+++ b/Engine/Source/Runtime/Engine/Classes/Curves/RealCurve.h
@@ -139,7 +139,7 @@ public:
 	  * @param	bUnwindRotation		When true, the value will be treated like a rotation value in degrees, and will automatically be unwound to prevent flipping 360 degrees from the previous key 
 	  * @param  KeyHandle			Optionally can specify what handle this new key should have, otherwise, it'll make a new one
 	  */
-	ENGINE_API virtual FKeyHandle AddKey(float InTime, float InValue, const bool bUnwindRotation = false, FKeyHandle KeyHandle = FKeyHandle()) PURE_VIRTUAL(FRealCurve::AddKey, return FKeyHandle::Invalid(););
+	virtual FKeyHandle AddKey(float InTime, float InValue, const bool bUnwindRotation = false, FKeyHandle KeyHandle = FKeyHandle()) PURE_VIRTUAL(FRealCurve::AddKey, return FKeyHandle::Invalid(););
 
 	/**
 	 *  Remove the specified key from the curve.
@@ -147,10 +147,10 @@ public:
 	 * @param KeyHandle The handle of the key to remove.
 	 * @see AddKey, SetKeys
 	 */
-	ENGINE_API virtual void DeleteKey(FKeyHandle KeyHandle) PURE_VIRTUAL(FRealCurve::DeleteKey,);
+	virtual void DeleteKey(FKeyHandle KeyHandle) PURE_VIRTUAL(FRealCurve::DeleteKey,);
 
 	/** Finds the key at InTime, and updates its value. If it can't find the key within the KeyTimeTolerance, it adds one at that time */
-	ENGINE_API virtual FKeyHandle UpdateOrAddKey(float InTime, float InValue, const bool bUnwindRotation = false, float KeyTimeTolerance = UE_KINDA_SMALL_NUMBER) PURE_VIRTUAL(FRealCurve::UpdateOrAddKey, return FKeyHandle::Invalid(););
+	virtual FKeyHandle UpdateOrAddKey(float InTime, float InValue, const bool bUnwindRotation = false, float KeyTimeTolerance = UE_KINDA_SMALL_NUMBER) PURE_VIRTUAL(FRealCurve::UpdateOrAddKey, return FKeyHandle::Invalid(););
 
 	/** Finds a key a the specified time */
 	ENGINE_API FKeyHandle FindKey(float KeyTime, float KeyTimeTolerance = UE_KINDA_SMALL_NUMBER) const;
@@ -159,15 +159,15 @@ public:
 	ENGINE_API bool KeyExistsAtTime(float KeyTime, float KeyTimeTolerance = UE_KINDA_SMALL_NUMBER) const;
 
 	/** Set the value of the specified key */
-	ENGINE_API virtual void SetKeyValue(FKeyHandle KeyHandle, float NewValue, bool bAutoSetTangents = true) PURE_VIRTUAL(FRealCurve::SetKeyValue,);
+	virtual void SetKeyValue(FKeyHandle KeyHandle, float NewValue, bool bAutoSetTangents = true) PURE_VIRTUAL(FRealCurve::SetKeyValue,);
 
 	/** Returns the value of the specified key */
-	ENGINE_API virtual float GetKeyValue(FKeyHandle KeyHandle) const PURE_VIRTUAL(FRealCurve::GetKeyValue, return 0.f;);
+	virtual float GetKeyValue(FKeyHandle KeyHandle) const PURE_VIRTUAL(FRealCurve::GetKeyValue, return 0.f;);
 
 	/** Returns a <Time, Value> pair for the specified key */
-	ENGINE_API virtual TPair<float, float> GetKeyTimeValuePair(FKeyHandle KeyHandle) const PURE_VIRTUAL(FRealCurve::GetKeyTimeValuePair, return (TPair<float,float>(0.f,0.f)););
+	virtual TPair<float, float> GetKeyTimeValuePair(FKeyHandle KeyHandle) const PURE_VIRTUAL(FRealCurve::GetKeyTimeValuePair, return (TPair<float,float>(0.f,0.f)););
 
-	ENGINE_API virtual void SetKeyInterpMode(FKeyHandle KeyHandle, ERichCurveInterpMode NewInterpMode) PURE_VIRTUAL(FRealCurve::SetKeyInterpMode,);
+	virtual void SetKeyInterpMode(FKeyHandle KeyHandle, ERichCurveInterpMode NewInterpMode) PURE_VIRTUAL(FRealCurve::SetKeyInterpMode,);
 
 	/** Set the default value of the curve */
 	void SetDefaultValue(float InDefaultValue) { DefaultValue = InDefaultValue; }
@@ -178,37 +178,37 @@ public:
 	/** Removes the default value for this curve. */
 	void ClearDefaultValue() { DefaultValue = MAX_flt; }
 
-	ENGINE_API virtual ERichCurveInterpMode GetKeyInterpMode(FKeyHandle KeyHandle) const PURE_VIRTUAL(FRealCurve::GetTimeRange, return RCIM_None; );
+	virtual ERichCurveInterpMode GetKeyInterpMode(FKeyHandle KeyHandle) const PURE_VIRTUAL(FRealCurve::GetTimeRange, return RCIM_None; );
 
 	/** Get range of input time values. Outside this region curve continues constantly the start/end values. */
-	ENGINE_API virtual void GetTimeRange(float& MinTime, float& MaxTime) const PURE_VIRTUAL(FRealCurve::GetTimeRange, );
+	virtual void GetTimeRange(float& MinTime, float& MaxTime) const PURE_VIRTUAL(FRealCurve::GetTimeRange, );
 
 	/** Get range of output values. */
-	ENGINE_API virtual void GetValueRange(float& MinValue, float& MaxValue) const PURE_VIRTUAL(FRealCurve::GetValueRange, );
+	virtual void GetValueRange(float& MinValue, float& MaxValue) const PURE_VIRTUAL(FRealCurve::GetValueRange, );
 
 	/** Clear all keys. */
-	ENGINE_API virtual void Reset() PURE_VIRTUAL(FRealCurve::Reset, );
+	virtual void Reset() PURE_VIRTUAL(FRealCurve::Reset, );
 
 	/** Remap InTime based on pre and post infinity extrapolation values */
-	ENGINE_API virtual void RemapTimeValue(float& InTime, float& CycleValueOffset) const PURE_VIRTUAL(FRealCurve::RemapTimeValue, );
+	virtual void RemapTimeValue(float& InTime, float& CycleValueOffset) const PURE_VIRTUAL(FRealCurve::RemapTimeValue, );
 
 	/** Evaluate this curve at the specified time */
-	ENGINE_API virtual float Eval(float InTime, float InDefaultValue = 0.0f) const PURE_VIRTUAL(FRealCurve::Eval, return 0.f;);
+	virtual float Eval(float InTime, float InDefaultValue = 0.0f) const PURE_VIRTUAL(FRealCurve::Eval, return 0.f;);
 
 	/** Resize curve length to the [MinTimeRange, MaxTimeRange] */
-	ENGINE_API virtual void ReadjustTimeRange(float NewMinTimeRange, float NewMaxTimeRange, bool bInsert/* whether insert or remove*/, float OldStartTime, float OldEndTime) PURE_VIRTUAL(FRealCurve::ReadjustTimeRange, );
+	virtual void ReadjustTimeRange(float NewMinTimeRange, float NewMaxTimeRange, bool bInsert/* whether insert or remove*/, float OldStartTime, float OldEndTime) PURE_VIRTUAL(FRealCurve::ReadjustTimeRange, );
 
 	/** Bake curve given the sample rate */
-	ENGINE_API virtual void BakeCurve(float SampleRate) PURE_VIRTUAL(FRealCurve::BakeCurve, );
-	ENGINE_API virtual void BakeCurve(float SampleRate, float FirstKeyTime, float LastKeyTime) PURE_VIRTUAL(FRealCurve::BakeCurve, );
+	virtual void BakeCurve(float SampleRate) PURE_VIRTUAL(FRealCurve::BakeCurve, );
+	virtual void BakeCurve(float SampleRate, float FirstKeyTime, float LastKeyTime) PURE_VIRTUAL(FRealCurve::BakeCurve, );
 
 	/** Remove redundant keys, comparing against Tolerance (and optional use of sample-rate for additional testing) */
-	ENGINE_API virtual void RemoveRedundantKeys(float Tolerance, FFrameRate SampleRate = FFrameRate(0,0)) PURE_VIRTUAL(FRealCurve::RemoveRedundantKeys, );
-	ENGINE_API virtual void RemoveRedundantKeys(float Tolerance, float FirstKeyTime, float LastKeyTime, FFrameRate SampleRate = FFrameRate(0,0)) PURE_VIRTUAL(FRealCurve::RemoveRedundantKeys, );
+	virtual void RemoveRedundantKeys(float Tolerance, FFrameRate SampleRate = FFrameRate(0,0)) PURE_VIRTUAL(FRealCurve::RemoveRedundantKeys, );
+	virtual void RemoveRedundantKeys(float Tolerance, float FirstKeyTime, float LastKeyTime, FFrameRate SampleRate = FFrameRate(0,0)) PURE_VIRTUAL(FRealCurve::RemoveRedundantKeys, );
 
 protected:
 	static ENGINE_API void CycleTime(float MinTime, float MaxTime, float& InTime, int& CycleCount);
-	ENGINE_API virtual int32 GetKeyIndex(float KeyTime, float KeyTimeTolerance) const PURE_VIRTUAL(FRealCurve::GetKeyIndex, return INDEX_NONE;);
+	virtual int32 GetKeyIndex(float KeyTime, float KeyTimeTolerance) const PURE_VIRTUAL(FRealCurve::GetKeyIndex, return INDEX_NONE;);
 
 public:
 
