FastGeo: Avoid broadcasting World's AddLevelToWorldExtensionEvent and RemoveLevelFromWorldExtensionEvent when not necessary (when the streaming level is not is MakingVisible/MakingInvisible).

diff --git a/Engine/Source/Runtime/Engine/Private/World.cpp b/Engine/Source/Runtime/Engine/Private/World.cpp
--- a/Engine/Source/Runtime/Engine/Private/World.cpp
+++ b/Engine/Source/Runtime/Engine/Private/World.cpp
@@ -3780,6 +3780,7 @@ void UWorld::AddToWorld(ULevel* Level, const FTransform& LevelTransform, bool bC
 		bPerformedLastStep = true;
 	}
 
+	if (IsLevelMakingVisible(Level))
 	{
 		TRACE_CPUPROFILER_EVENT_SCOPE(AddToWorld_Extension);
 		bool bWaitForCompletion = !bConsiderTimeLimit;
@@ -4028,6 +4029,7 @@ void UWorld::RemoveFromWorld(ULevel* Level, bool bAllowIncrementalRemoval, const
 			check(bFinishRemovingLevel);
 		}
 
+		if (Level->bIsBeingRemoved)
 		{
 			TRACE_CPUPROFILER_EVENT_SCOPE(RemoveFromWorld_Extension);
 			check(bWaitForCompletion == !bAllowIncrementalRemoval);
