Timecode: Timecode provider and FApp::GetTimecode now don't strip subframes.

The problem with subframe truncation is that live data being evaluated by timecode will also evaluate with truncated resolution and may look choppy with respect to the rest of the animation.

#rnx

diff --git a/Engine/Source/Runtime/Core/Private/Misc/App.cpp b/Engine/Source/Runtime/Core/Private/Misc/App.cpp
--- a/Engine/Source/Runtime/Core/Private/Misc/App.cpp
+++ b/Engine/Source/Runtime/Core/Private/Misc/App.cpp
@@ -317,7 +317,7 @@ FTimecode FApp::GetTimecode()
 {
 	if (CurrentFrameTime.IsSet())
 	{
-		return FTimecode::FromFrameNumber(CurrentFrameTime->Time.GetFrame(), CurrentFrameTime->Rate);
+		return CurrentFrameTime->ToTimecode();
 	}
 	return FTimecode();
 }
diff --git a/Engine/Source/Runtime/Engine/Private/Engine/TimecodeProvider.cpp b/Engine/Source/Runtime/Engine/Private/Engine/TimecodeProvider.cpp
--- a/Engine/Source/Runtime/Engine/Private/Engine/TimecodeProvider.cpp
+++ b/Engine/Source/Runtime/Engine/Private/Engine/TimecodeProvider.cpp
@@ -15,14 +15,12 @@ FQualifiedFrameTime UTimecodeProvider::GetDelayedQualifiedFrameTime() const
 
 FTimecode UTimecodeProvider::GetTimecode() const
 {
-	const FQualifiedFrameTime NewFrameTime = GetQualifiedFrameTime();
-	return FTimecode::FromFrameNumber(NewFrameTime.Time.GetFrame(), NewFrameTime.Rate);
+	return GetQualifiedFrameTime().ToTimecode();
 }
 
 
 FTimecode UTimecodeProvider::GetDelayedTimecode() const
 {
-	const FQualifiedFrameTime NewFrameTime = GetDelayedQualifiedFrameTime();
-	return FTimecode::FromFrameNumber(NewFrameTime.Time.GetFrame(), NewFrameTime.Rate);
+	return GetDelayedQualifiedFrameTime().ToTimecode();
 }
 
