VirtualTextureAdapter now uses PooledRenderTarget instead of RHI Texture as target.
Makes transition logic simpler.
#rnx

diff --git a/Engine/Source/Runtime/Engine/Private/VT/VirtualTextureAdapter.cpp b/Engine/Source/Runtime/Engine/Private/VT/VirtualTextureAdapter.cpp
--- a/Engine/Source/Runtime/Engine/Private/VT/VirtualTextureAdapter.cpp
+++ b/Engine/Source/Runtime/Engine/Private/VT/VirtualTextureAdapter.cpp
@@ -59,6 +59,7 @@ namespace VirtualTextureAdapter
 	/** Final copy to output RDG parameters. */
 	BEGIN_SHADER_PARAMETER_STRUCT(FCopyToOutputParameters, )
 		RDG_TEXTURE_ACCESS(Input, ERHIAccess::CopySrc)
+		RDG_TEXTURE_ACCESS(Output, ERHIAccess::CopyDest)
 	END_SHADER_PARAMETER_STRUCT()
 
 
@@ -91,7 +92,6 @@ namespace VirtualTextureAdapter
 	static void RenderTile(
 		FRDGBuilder& GraphBuilder, 
 		FShaderResourceViewRHIRef& InSourceSRV,
-		FRHITexture* InDestTexture, 
 		IPooledRenderTarget* InDestRenderTarget, 
 		EPixelFormat InSourceFormat,
 		EPixelFormat InIntermediateFormat,
@@ -119,7 +119,7 @@ namespace VirtualTextureAdapter
 		const uint32 FinalTexelCountY = InDestRect.Max.Y - InDestRect.Min.Y;
 
 		const bool bSourceSrgb = InSourceSRV->GetDesc().Texture.SRV.GetViewInfo(InSourceSRV->GetTexture()).bSRGB;
-		const bool bDestSrgb = EnumHasAnyFlags(InDestTexture->GetDesc().Flags, ETextureCreateFlags::SRGB);
+		const bool bDestSrgb = EnumHasAnyFlags(InDestRenderTarget->GetRHI()->GetDesc().Flags, ETextureCreateFlags::SRGB);
 		const ETextureCreateFlags TextureCreateFlagsSrgb = bSourceSrgb ? ETextureCreateFlags::SRGB : ETextureCreateFlags::None;
 		const bool bFinalPassConvertToSrgb = !bSourceSrgb && bDestSrgb;
 
@@ -271,11 +271,14 @@ namespace VirtualTextureAdapter
 		// Final copy if we didn't already directly write to the physical texture.
 		if (!bIsDone)
 		{
+			FRDGTextureRef LastOutput = CurrentOutput;
+			CurrentOutput = GraphBuilder.RegisterExternalTexture(InDestRenderTarget, ERDGTextureFlags::None);
+
 			FRHICopyTextureInfo CopyInfo;
 			CopyInfo.DestPosition = FIntVector(InDestRect.Min.X, InDestRect.Min.Y, 0);
 			CopyInfo.Size = FIntVector(FinalTexelCountX, FinalTexelCountY, 0);
 
-			const bool bUseSourceTextureA = CurrentOutput == nullptr;
+			const bool bUseSourceTextureA = LastOutput == nullptr;
 			if (bUseSourceTextureA)
 			{
 				// Copying directly from the source.
@@ -295,21 +298,19 @@ namespace VirtualTextureAdapter
 			}
 
 			FCopyToOutputParameters* Parameters = GraphBuilder.AllocParameters<FCopyToOutputParameters>();
-			Parameters->Input = CurrentOutput;
+			Parameters->Input = LastOutput;
+			Parameters->Output = CurrentOutput;
 
 			GraphBuilder.AddPass(
 				RDG_EVENT_NAME("VirtualTextureAdapterCopyToOutput"),
 				Parameters,
 				ERDGPassFlags::Copy | ERDGPassFlags::NeverCull,
-				[bUseSourceTextureA, InputTextureA = InSourceSRV, InputTextureB = CurrentOutput, OutputTexture = InDestTexture, CopyInfo](FRDGAsyncTask, FRHICommandList& RHICmdList)
-				{
-					FRHITexture* InputTexture = bUseSourceTextureA ? InputTextureA->GetTexture() : InputTextureB->GetRHI();
-
-					RHICmdList.Transition(FRHITransitionInfo(OutputTexture, ERHIAccess::SRVMask, ERHIAccess::CopyDest));
-					RHICmdList.CopyTexture(InputTexture, OutputTexture, CopyInfo);
-					RHICmdList.Transition(FRHITransitionInfo(OutputTexture, ERHIAccess::CopyDest, ERHIAccess::SRVMask));
-				});
-		}
+				[bUseSourceTextureA, InputTextureA = InSourceSRV, InputTextureB = LastOutput, OutputTexture = CurrentOutput, CopyInfo](FRDGAsyncTask, FRHICommandList& RHICmdList)
+			{
+				FRHITexture* InputTexture = bUseSourceTextureA ? InputTextureA->GetTexture() : InputTextureB->GetRHI();
+				RHICmdList.CopyTexture(InputTexture, OutputTexture->GetRHI(), CopyInfo);
+			});
+ 		}
 	}
 }
 
@@ -373,7 +374,6 @@ protected:
 	{
 		for (FTileEntry const& Tile : TilesToRender)
 		{
-			FRHITexture* DestTexture = Tile.Target.TextureRHI != nullptr ? Tile.Target.TextureRHI->GetTexture2D() : nullptr;
 			IPooledRenderTarget* DestRenderTarget = Tile.Target.PooledRenderTarget;
 
 			const float X = (float)FMath::ReverseMortonCode2_64(Tile.vAddress);
@@ -394,7 +394,7 @@ protected:
 			const int32 SRVIndex = bHasSRVForLevel ? Tile.vLevel : 0;
 			const int32 vLevel = bHasSRVForLevel ? 0 : Tile.vLevel;
 
-			VirtualTextureAdapter::RenderTile(GraphBuilder, SourceSRVs[SRVIndex], DestTexture, DestRenderTarget, SourceFormat, IntermediateFormat, DestFormat, vLevel, UVRange, DestRect);
+			VirtualTextureAdapter::RenderTile(GraphBuilder, SourceSRVs[SRVIndex], DestRenderTarget, SourceFormat, IntermediateFormat, DestFormat, vLevel, UVRange, DestRect);
 		}
 
 		TilesToRender.Reset();
