Virtual texture upload cache now uses RDG to manage upload.
Default upload path is changed to use UpdateTexture2D instead of managing staging textures.
This avoids need for immediate context Lock/Unlock and testing shows it is performant even on D3D11.
Added r.VT.UploadUseLegacyPath read only console variable to return to old behavior.
Change also removes the complexity from the physical pool texture transition logic.
#rnx

diff --git a/Engine/Source/Runtime/Engine/Private/Texture2D.cpp b/Engine/Source/Runtime/Engine/Private/Texture2D.cpp
--- a/Engine/Source/Runtime/Engine/Private/Texture2D.cpp
+++ b/Engine/Source/Runtime/Engine/Private/Texture2D.cpp
@@ -1634,162 +1634,109 @@ void FVirtualTexture2DResource::InitRHI(FRHICommandListBase& RHICmdList)
 #if WITH_EDITOR
 void FVirtualTexture2DResource::InitializeEditorResources(IVirtualTexture* InVirtualTexture)
 {
-	// Create a texture resource from the lowest resolution VT page data
-	// this will then be used during asset tumbnails/hitproxies/...
+	// Create a texture resource for use when we try to access the VT as a regular texture.
+	// This may happen for debugging or when we are in a preview mode that doesn't support VT.
 	if (GIsEditor && !IsRunningCommandlet())
 	{
-		struct FPageToProduce
-		{
-			uint64 Handle;
-			uint32 TileX;
-			uint32 TileY;
-		};
-
 		// Choose a mip level for the thumbnail texture to ensure proper size
+		const uint32 TargetSize = FMath::Clamp(128u, GetTileSize(), 1024u);
 		const uint32 MaxMipLevel = VTData->GetNumMips() - 1u;
-		const uint32 MaxTextureSize = FMath::Min<uint32>(GetMax2DTextureDimension(), 1024u);
+		
 		uint32 MipLevel = 0;
 		uint32 MipWidth = GetSizeX();
 		uint32 MipHeight = GetSizeY();
-		while (((MipWidth > 128u && MipHeight > 128u) || MipWidth > MaxTextureSize || MipHeight > MaxTextureSize) && MipLevel < MaxMipLevel)
+		while (MipWidth > TargetSize && MipHeight > TargetSize && MipLevel < MaxMipLevel)
 		{
 			++MipLevel;
 			MipWidth = FMath::DivideAndRoundUp(MipWidth, 2u);
 			MipHeight = FMath::DivideAndRoundUp(MipHeight, 2u);
 		}
 
-		FRHICommandListImmediate& RHICommandList = FRHICommandListExecutor::GetImmediateCommandList();
-
+		FString Name = TextureName.ToString();
 		EPixelFormat PixelFormat = VTData->LayerTypes[0];
 
-		FString Name = TextureName.ToString();
+		FTextureRHIRef Texture2DRHI = nullptr;
 
-		FTextureRHIRef Texture2DRHI;
+		// We only handle the single layer case.
+		bool bUseFallbackTexture = GetNumLayers() != 1;
+		// If the VT texture is very skewed on one dimension so that other dimension is very small then we may have less than VT block size.
+		bUseFallbackTexture |= IsBlockCompressedFormat(PixelFormat) && ((MipWidth % 4 != 0) || (MipHeight % 4 != 0));
 
-		if (IsBlockCompressedFormat(PixelFormat) && ((MipWidth % 4 != 0) || (MipHeight % 4 != 0)))
+		if (bUseFallbackTexture)
 		{
-			// This means that VT texture is very skewed on one dimension, thus making other dimension
-			// very small - sometimes less than VT block size (128) so potentially it can be non-multiple of
-			// compressed texture block size. In such case create 1x1 texture thumbnail from fallback color
-
-			uint32 TexWidth = 1;
-			uint32 TexHeight = 1;
-			PixelFormat = EPixelFormat::PF_B8G8R8A8;
-
-			// create thumbnail from fallback color
-
-			FRHITextureCreateDesc StagingDesc = FRHITextureCreateDesc::Create2D(*Name, TexWidth, TexHeight, PixelFormat);
-			StagingDesc.SetFlags(TexCreateFlags);
-			if (!IsRunningRHIInSeparateThread())
-			{
-				StagingDesc.AddFlags(ETextureCreateFlags::CPUWritable);
-			}
-
-			FTextureRHIRef StagingTexture2DRHI = RHICreateTexture(StagingDesc);
-
-			uint32 Stride;
-			void* TexMemory = RHICommandList.LockTexture2D(StagingTexture2DRHI, 0, RLM_WriteOnly, Stride, false, false);
-			*static_cast<FColor*>(TexMemory) = VTData->LayerFallbackColors[0].ToFColor(true);
-			RHICommandList.UnlockTexture2D(StagingTexture2DRHI, 0u, false, false);
+			// Fallback to simple 1x1 texture of the fallback color.
+			FRHITextureCreateDesc TextureDesc = FRHITextureCreateDesc::Create2D(*Name, 1, 1, EPixelFormat::PF_B8G8R8A8);
+			TextureDesc.SetFlags(TexCreateFlags);
+			
+			const FColor Color = VTData->LayerFallbackColors[0].ToFColor(true);
+			FResourceBulkDataArrayView BulkDataInterface(&Color, sizeof(FColor));
+			TextureDesc.SetBulkData(&BulkDataInterface);
 
-			Texture2DRHI = MoveTemp(StagingTexture2DRHI);
+			Texture2DRHI = RHICreateTexture(TextureDesc);
 		}
 		else
 		{
+			FRHICommandListImmediate& RHICommandList = FRHICommandListExecutor::GetImmediateCommandList();
+			FRDGBuilder GraphBuilder(RHICommandList);
+
+			// We want to strip borders when compositing tiles since we're just laying out tiles in a regular texture.
+			// But if we have block compressed formats with border not divisible by 4 then doing this will lead to an unaligned copy. We keep the small unwanted borders in that case.
 			const bool bCopyUnwantedBordersForAlignment = (VTData->TileBorderSize % 4 != 0) && IsBlockCompressedFormat(PixelFormat);
 			const uint32 MipScaleFactor = (1u << MipLevel);
 			const uint32 MipWidthInTiles = FMath::DivideAndRoundUp(GetNumTilesX(), MipScaleFactor);
 			const uint32 MipHeightInTiles = FMath::DivideAndRoundUp(GetNumTilesY(), MipScaleFactor);
 			const uint32 TileSizeInPixels = bCopyUnwantedBordersForAlignment ? GetTileSize() + 2 * GetBorderSize() : GetTileSize();
+			const uint32 LayerMask = 1u;
+
+			// Create the target texture to composite tiles to.
+			FRDGTextureRef Texture = GraphBuilder.CreateTexture(FRDGTextureDesc::Create2D(
+				FIntPoint(MipWidthInTiles * TileSizeInPixels, MipHeightInTiles * TileSizeInPixels),
+				PixelFormat, FClearValueBinding::Black,	TexCreateFlags), *Name);
 
-			const uint32 LayerMask = 1u; // FVirtualTexture2DResource should only have a single layer
+			TRefCountPtr<IPooledRenderTarget> PooledRenderTarget = GraphBuilder.ConvertToExternalTexture(Texture);
+			Texture2DRHI = PooledRenderTarget->GetRHI();
+
+			// Produce all the tiles.
+			TArray<IVirtualTextureFinalizer*> Finalizers;
 
-			TArray<FPageToProduce> PagesToProduce;
-			PagesToProduce.Reserve(MipWidthInTiles * MipHeightInTiles);
 			for (uint32 TileY = 0u; TileY < MipHeightInTiles; ++TileY)
 			{
 				for (uint32 TileX = 0u; TileX < MipWidthInTiles; ++TileX)
 				{
 					const uint32 vAddress = FMath::MortonCode2(TileX) | (FMath::MortonCode2(TileY) << 1);
-					const FVTRequestPageResult RequestResult = InVirtualTexture->RequestPageData(FRHICommandListExecutor::GetImmediateCommandList(), ProducerHandle, LayerMask, MipLevel, vAddress, EVTRequestPagePriority::High);
-
-					// High priority request should never be Saturated
-					// It's possible for status to be Invalid, if requesting data from a mip level that doesn't exist for the given producer (when using sparse UDIMs)
-					// Technically could try to handle this, by check LocalMipBias, grabbing lower resolution tile, and resizing...but that would make this code much more complex for very little gain
-					ensure(RequestResult.Status != EVTRequestPageStatus::Saturated);
+					const FVTRequestPageResult RequestResult = InVirtualTexture->RequestPageData(RHICommandList, ProducerHandle, LayerMask, MipLevel, vAddress, EVTRequestPagePriority::High);
 
+					// High priority request should never be Saturated.
+					// It's possible for status to be Invalid, if requesting data from a mip level that doesn't exist for the given producer (when using sparse UDIMs).
+					// Technically we could try to handle this, by checking LocalMipBias, grabbing lower resolution tile, and resizing...but that would make this code much more complex for very little gain.
 					if (VTRequestPageStatus_HasData(RequestResult.Status))
 					{
-						PagesToProduce.Add({ RequestResult.Handle, TileX, TileY });
+						FVTProduceTargetLayer TargetLayer;
+						TargetLayer.PooledRenderTarget = PooledRenderTarget;
+						TargetLayer.pPageLocation = FIntVector(TileX, TileY, 0);
+
+						IVirtualTextureFinalizer* Finalizer = InVirtualTexture->ProducePageData(
+							RHICommandList,
+							GMaxRHIFeatureLevel,
+							bCopyUnwantedBordersForAlignment ? EVTProducePageFlags::None : EVTProducePageFlags::SkipPageBorders,
+							ProducerHandle, LayerMask, MipLevel, vAddress,
+							RequestResult.Handle,
+							&TargetLayer);
+
+						Finalizers.AddUnique(Finalizer);
 					}
 				}
 			}
 
-			FRHITextureCreateDesc Desc = FRHITextureCreateDesc::Create2D(*Name, MipWidthInTiles * TileSizeInPixels, MipHeightInTiles * TileSizeInPixels, PixelFormat);
-			Desc.AddFlags(TexCreateFlags);
-
-			Texture2DRHI = RHICreateTexture(Desc);
-
-			// We want to strip borders when compositing tiles since we're just laying out tiles in a regular texture.
-			// But if we have block compressed formats with border not divisible by 4 then doing this will lead to an unaligned copy. We keep the small unwanted borders in that case.
-			const EVTProducePageFlags ProducePageFlags = bCopyUnwantedBordersForAlignment ? EVTProducePageFlags::None : EVTProducePageFlags::SkipPageBorders;
-
-			TArray<IVirtualTextureFinalizer*> Finalizers;
-			for (const FPageToProduce& Page : PagesToProduce)
+			for (IVirtualTextureFinalizer* Finalizer : Finalizers)
 			{
-				const uint32 vAddress = FMath::MortonCode2(Page.TileX) | (FMath::MortonCode2(Page.TileY) << 1);
-
-				FVTProduceTargetLayer TargetLayer;
-				TargetLayer.TextureRHI = Texture2DRHI;
-				TargetLayer.pPageLocation = FIntVector(Page.TileX, Page.TileY, 0);
-
-				IVirtualTextureFinalizer* Finalizer = InVirtualTexture->ProducePageData(RHICommandList,
-					GMaxRHIFeatureLevel,
-					ProducePageFlags,
-					ProducerHandle, LayerMask, MipLevel, vAddress,
-					Page.Handle,
-					&TargetLayer);
 				if (Finalizer)
-				{
-					Finalizers.AddUnique(Finalizer);
-				}
-			}
-
-			{
-				FRDGBuilder GraphBuilder(RHICommandList);
-				for (IVirtualTextureFinalizer* Finalizer : Finalizers)
 				{
 					Finalizer->Finalize(GraphBuilder);
 				}
-				GraphBuilder.Execute();
 			}
 
-			if (MipWidthInTiles * TileSizeInPixels != MipWidth || MipHeightInTiles * TileSizeInPixels != MipHeight)
-			{
-				// Logical dimensions of mip image may be smaller than tile size (in this case tile will contain mirrored/wrapped padding)
-				// In this case, copy the proper sub-image from the tiled texture we produced into a new texture of the correct size
-				check(MipWidth <= MipWidthInTiles * TileSizeInPixels);
-				check(MipHeight <= MipHeightInTiles * TileSizeInPixels);
-
-				const FRHITextureCreateDesc ResizedDesc =
-					FRHITextureCreateDesc::Create2D(*Name, MipWidth, MipHeight, PixelFormat)
-					.SetFlags(Desc.Flags)
-					.SetInitialState(ERHIAccess::CopyDest);
-
-				FTextureRHIRef ResizedTexture2DRHI = RHICreateTexture(ResizedDesc);
-
-				FRHICopyTextureInfo CopyInfo;
-				CopyInfo.Size = FIntVector(MipWidth, MipHeight, 1);
-
-				// Put the source texture in CopySrc mode. The destination texture is already in CopyDest mode because we created it that way.
-				RHICommandList.Transition({FRHITransitionInfo(Texture2DRHI, ERHIAccess::SRVMask, ERHIAccess::CopySrc)}, ERHIPipeline::All, ERHIPipeline::Graphics);
-
-				RHICommandList.CopyTexture(Texture2DRHI, ResizedTexture2DRHI, CopyInfo);
-				// Make the destination texture SRVMask again. We don't care about the source texture after this, so we won't bother transitioning it.
-				RHICommandList.Transition(FRHITransitionInfo(ResizedTexture2DRHI, ERHIAccess::CopyDest, ERHIAccess::SRVMask));
-
-				Texture2DRHI = MoveTemp(ResizedTexture2DRHI);
-			}
+			GraphBuilder.Execute();
 		}
 
 		TextureRHI = Texture2DRHI;
diff --git a/Engine/Source/Runtime/Engine/Private/VT/VirtualTextureChunkManager.cpp b/Engine/Source/Runtime/Engine/Private/VT/VirtualTextureChunkManager.cpp
--- a/Engine/Source/Runtime/Engine/Private/VT/VirtualTextureChunkManager.cpp
+++ b/Engine/Source/Runtime/Engine/Private/VT/VirtualTextureChunkManager.cpp
@@ -152,7 +152,7 @@ IVirtualTextureFinalizer* FVirtualTextureChunkStreamingManager::ProduceTile(FRHI
 		if (LayerMask & (1u << LayerIndex))
 		{
 			const FVTProduceTargetLayer& Target = TargetLayers[LayerIndex];
-			UploadCache.SubmitTile(RHICmdList, StageTileHandles[LayerIndex], Target.TextureRHI->GetTexture2D(), Target.pPageLocation.X, Target.pPageLocation.Y, SkipBorderSize);
+			UploadCache.SubmitTile(RHICmdList, StageTileHandles[LayerIndex], Target, SkipBorderSize);
 		}
 	}
 
diff --git a/Engine/Source/Runtime/Engine/Private/VT/VirtualTextureUploadCache.cpp b/Engine/Source/Runtime/Engine/Private/VT/VirtualTextureUploadCache.cpp
--- a/Engine/Source/Runtime/Engine/Private/VT/VirtualTextureUploadCache.cpp
+++ b/Engine/Source/Runtime/Engine/Private/VT/VirtualTextureUploadCache.cpp
@@ -4,20 +4,19 @@
 
 #include "RenderGraphBuilder.h"
 #include "RenderUtils.h"
-#include "VirtualTextureChunkManager.h"
 #include "Stats/StatsTrace.h"
-
-// Allow uploading CPU buffer directly to GPU texture
-// This is slow under D3D11
-// Should be pretty decent on D3D12X...UpdateTexture does make an extra copy of the data, but Lock/Unlock texture also buffers an extra copy of texture on this platform
-// Might also be worth enabling this path on PC D3D12, need to measure
-#if !defined(ALLOW_UPDATE_TEXTURE)
-	#define ALLOW_UPDATE_TEXTURE 0
-#endif
+#include "VirtualTextureChunkManager.h"
 
 DECLARE_MEMORY_STAT_POOL(TEXT("Total GPU Upload Memory"), STAT_TotalGPUUploadSize, STATGROUP_VirtualTextureMemory, FPlatformMemory::MCR_GPU);
 DECLARE_MEMORY_STAT(TEXT("Total CPU Upload Memory"), STAT_TotalCPUUploadSize, STATGROUP_VirtualTextureMemory);
 
+static TAutoConsoleVariable<int32> CVarVTUploadUseLegacyPath(
+	TEXT("r.VT.UploadUseLegacyPath"),
+	0,
+	TEXT("Use the legacy virtual texture upload path which locks staging textures."),
+	ECVF_ReadOnly
+);
+
 static TAutoConsoleVariable<int32> CVarVTUploadMemoryPageSize(
 	TEXT("r.VT.UploadMemoryPageSize"),
 	4,
@@ -40,7 +39,7 @@ static TAutoConsoleVariable<int32> CVarMaxUploadRequests(
 	ECVF_RenderThreadSafe);
 
 
-uint32 FVTUploadTileAllocator::Allocate(FRHICommandList& RHICmdList, EPixelFormat InFormat, uint32 InTileSize)
+uint32 FVTUploadTileAllocator::Allocate(FRHICommandList& RHICmdList, EVTUploadType InUploadType, EPixelFormat InFormat, uint32 InTileSize)
 {
 	// Find matching FormatBuffer.
 	const FPixelFormatInfo& FormatInfo = GPixelFormats[InFormat];
@@ -97,7 +96,7 @@ uint32 FVTUploadTileAllocator::Allocate(FRHICommandList& RHICmdList, EPixelForma
 	if (StagingBuffer.Memory == nullptr)
 	{
 		// Staging buffer needs underlying buffer allocating.
-		StagingBuffer.Init(RHICmdList, Desc.BlockBytes, Desc.MemorySize);
+		StagingBuffer.Init(RHICmdList, InUploadType, Desc.BlockBytes, Desc.MemorySize);
 		NumAllocatedBytes += StagingBuffer.TileSizeAligned * StagingBuffer.NumTiles;
 	}
 
@@ -166,7 +165,7 @@ FVTUploadTileAllocator::FStagingBuffer::~FStagingBuffer()
 	Release(nullptr);
 }
 
-void FVTUploadTileAllocator::FStagingBuffer::Init(FRHICommandList& RHICmdList, uint32 InBufferStrideBytes, uint32 InTileSizeBytes)
+void FVTUploadTileAllocator::FStagingBuffer::Init(FRHICommandList& RHICmdList, EVTUploadType InUploadType, uint32 InBufferStrideBytes, uint32 InTileSizeBytes)
 {
 	TileSize = InTileSizeBytes;
 	TileSizeAligned = Align(InTileSizeBytes, 128u);
@@ -182,7 +181,7 @@ void FVTUploadTileAllocator::FStagingBuffer::Init(FRHICommandList& RHICmdList, u
 		TileFreeList[Index] = NumTiles - Index - 1;
 	}
 
-	if (GRHISupportsDirectGPUMemoryLock && GRHISupportsUpdateFromBufferTexture)
+	if (InUploadType == EVTUploadType::PersistentBuffer)
 	{
 		// Allocate staging buffer directly in GPU memory.
 		const FRHIBufferCreateDesc CreateDesc =
@@ -192,8 +191,8 @@ void FVTUploadTileAllocator::FStagingBuffer::Init(FRHICommandList& RHICmdList, u
 
 		RHIBuffer = RHICmdList.CreateBuffer(CreateDesc);
 
-		// Here we bypass 'normal' RHI operations in order to get a persistent pointer to GPU memory, on supported platforms
-		// This should be encapsulated into a proper RHI method at some point
+		// Here we bypass 'normal' RHI operations in order to get a persistent pointer to GPU memory, on supported platforms.
+		// This should be encapsulated into a proper RHI method at some point.
 		Memory = RHICmdList.LockBuffer(RHIBuffer, 0u, BufferSize, RLM_WriteOnly_NoOverwrite);
 
 		INC_MEMORY_STAT_BY(STAT_TotalGPUUploadSize, BufferSize);
@@ -238,6 +237,22 @@ void FVTUploadTileAllocator::FStagingBuffer::Release(FRHICommandList* RHICmdList
 	TileFreeList.Reset();
 }
 
+FVirtualTextureUploadCache::FVirtualTextureUploadCache()
+{
+	if (CVarVTUploadUseLegacyPath.GetValueOnGameThread())
+	{
+		UploadType = EVTUploadType::StagingTexture;
+	}
+	else if (GRHISupportsDirectGPUMemoryLock && GRHISupportsUpdateFromBufferTexture)
+	{
+		UploadType = EVTUploadType::PersistentBuffer;
+	}
+	else
+	{
+		UploadType = EVTUploadType::StagingCopy;
+	}
+}
+
 int32 FVirtualTextureUploadCache::GetOrCreatePoolIndex(EPixelFormat InFormat, uint32 InTileSize)
 {
 	for (int32 i = 0; i < Pools.Num(); ++i)
@@ -259,13 +274,25 @@ int32 FVirtualTextureUploadCache::GetOrCreatePoolIndex(EPixelFormat InFormat, ui
 
 void FVirtualTextureUploadCache::Finalize(FRDGBuilder& GraphBuilder)
 {
-	FRHICommandListImmediate& RHICmdList = GraphBuilder.RHICmdList;
-	SCOPED_DRAW_EVENT(RHICmdList, FVirtualTextureUploadCache_Finalize);
-	SCOPE_CYCLE_COUNTER(STAT_VTP_FlushUpload)
+	switch (UploadType)
+	{
+	case EVTUploadType::StagingTexture:
+		FinalizeWithLegacyCopyTexture(GraphBuilder);
+		break;
+	case EVTUploadType::StagingCopy:
+	case EVTUploadType::PersistentBuffer:
+		FinalizeWithUpdateTexture(GraphBuilder);
+		break;
+	}
+}
 
-	// Multi-GPU support:
-	SCOPED_GPU_MASK(RHICmdList, FRHIGPUMask::All());
+/** Parameter struct for tracking target texture resource states. */
+BEGIN_SHADER_PARAMETER_STRUCT(FVirtualTextureUploadCacheParameters, )
+	RDG_TEXTURE_ACCESS_ARRAY(TextureArray)
+END_SHADER_PARAMETER_STRUCT()
 
+void FVirtualTextureUploadCache::FinalizeWithLegacyCopyTexture(FRDGBuilder& GraphBuilder)
+{
 	for (int PoolIndex = 0; PoolIndex < Pools.Num(); ++PoolIndex)
 	{
 		FPoolEntry& PoolEntry = Pools[PoolIndex];
@@ -275,12 +302,9 @@ void FVirtualTextureUploadCache::Finalize(FRDGBuilder& GraphBuilder)
 			continue;
 		}
 
-		const FPixelFormatInfo& FormatInfo = GPixelFormats[PoolEntry.Format];
+		// Create/Resize the pool staging buffer texture.
+		const EPixelFormat Format = PoolEntry.Format;
 		const uint32 TileSize = PoolEntry.TileSize;
-		const uint32 BlockBytes = FormatInfo.BlockBytes;
-		const uint32 TileWidthInBlocks = FMath::DivideAndRoundUp(TileSize, (uint32)FormatInfo.BlockSizeX);
-		const uint32 TileHeightInBlocks = FMath::DivideAndRoundUp(TileSize, (uint32)FormatInfo.BlockSizeY);
-
 		const uint32 TextureIndex = PoolEntry.BatchTextureIndex;
 		PoolEntry.BatchTextureIndex = (PoolEntry.BatchTextureIndex + 1u) % FPoolEntry::NUM_STAGING_TEXTURES;
 		FStagingTexture& StagingTexture = PoolEntry.StagingTexture[TextureIndex];
@@ -323,73 +347,175 @@ void FVirtualTextureUploadCache::Finalize(FRDGBuilder& GraphBuilder)
 			INC_MEMORY_STAT_BY(STAT_TotalGPUUploadSize, CalcTextureSize(TileSize * WidthInTiles, TileSize * HeightInTiles, PoolEntry.Format, 1u));
 		}
 
-		uint32 BatchStride = 0u;
-		void* BatchMemory = RHICmdList.LockTexture2D(StagingTexture.RHITexture, 0, RLM_WriteOnly, BatchStride, false, false);
+		// Make a copy of required data for RDG pass.
+		FTileEntry* TileEntries = GraphBuilder.AllocPODArray<FTileEntry>(BatchCount);
+		FVTUploadTileBufferExt* UploadBuffers = GraphBuilder.AllocPODArray<FVTUploadTileBufferExt>(BatchCount);
+
+		FVirtualTextureUploadCacheParameters* UploadParameters = GraphBuilder.AllocParameters<FVirtualTextureUploadCacheParameters>();
 
-		// Copy all tiles to the staging texture
 		for (uint32 Index = 0u; Index < BatchCount; ++Index)
 		{
-			const FTileEntry& Entry = PoolEntry.PendingSubmit[Index];
-			const uint32_t SrcTileX = Index % StagingTexture.WidthInTiles;
-			const uint32_t SrcTileY = Index / StagingTexture.WidthInTiles;
-
-			const FVTUploadTileBufferExt UploadBuffer = TileAllocator.GetBufferFromHandleExt(Entry.TileHandle);
+			FTileEntry& TileEntry = PoolEntry.PendingSubmit[Index];
 
-			uint8* BatchDst = (uint8*)BatchMemory + TileHeightInBlocks * SrcTileY * BatchStride + TileWidthInBlocks * SrcTileX * BlockBytes;
-			for (uint32 y = 0u; y < TileHeightInBlocks; ++y)
+			// Common case is that all target textures are the same.
+			// So compare with the previous entry as a quick test for a faster path.
+			if (Index > 1 && TileEntry.PooledRenderTarget == TileEntries[Index - 1].PooledRenderTarget)
 			{
-				FMemory::Memcpy(
-					BatchDst + y * BatchStride,
-					(uint8*)UploadBuffer.BufferMemory + UploadBuffer.BufferOffset + y * UploadBuffer.Stride,
-					TileWidthInBlocks * BlockBytes);
+				TileEntry.Texture = TileEntries[Index - 1].Texture;
+			}
+			else
+			{
+				TileEntry.Texture = GraphBuilder.RegisterExternalTexture(TileEntry.PooledRenderTarget, ERDGTextureFlags::None);
+				UploadParameters->TextureArray.AddUnique(FRDGTextureAccess(TileEntry.Texture, ERHIAccess::CopyDest));
 			}
 
-			// Can release upload buffer.
-			TileAllocator.Free(RHICmdList, Entry.TileHandle);
+			TileEntries[Index] = TileEntry;
+			UploadBuffers[Index] = TileAllocator.GetBufferFromHandleExt(TileEntry.TileHandle);
 		}
 
-		RHICmdList.UnlockTexture2D(StagingTexture.RHITexture, 0u, false, false);
-		RHICmdList.Transition(FRHITransitionInfo(StagingTexture.RHITexture, ERHIAccess::SRVMask, ERHIAccess::CopySrc));
+		// Can clear the pending queue now.
+		PoolEntry.PendingSubmit.Reset();
 
-		// Upload each tile from staging texture to physical texture
-		for (uint32 Index = 0u; Index < BatchCount; ++Index)
+		GraphBuilder.AddPass(
+			RDG_EVENT_NAME("VirtualTextureUploadLegacy"),
+			UploadParameters,
+			ERDGPassFlags::Copy,
+			[this, TileSize, BatchCount, TileEntries, UploadBuffers, StagingTexture, Format](FRHICommandListImmediate& RHICmdList)
 		{
-			const FTileEntry& Entry = PoolEntry.PendingSubmit[Index];
-			const uint32_t SrcTileX = Index % StagingTexture.WidthInTiles;
-			const uint32_t SrcTileY = Index / StagingTexture.WidthInTiles;
+			const FPixelFormatInfo& FormatInfo = GPixelFormats[Format];
+			const uint32 BlockBytes = FormatInfo.BlockBytes;
+			const uint32 TileWidthInBlocks = FMath::DivideAndRoundUp(TileSize, (uint32)FormatInfo.BlockSizeX);
+			const uint32 TileHeightInBlocks = FMath::DivideAndRoundUp(TileSize, (uint32)FormatInfo.BlockSizeY);
 
-			const uint32 SkipBorderSize = Entry.SubmitSkipBorderSize;
-			const uint32 SubmitTileSize = TileSize - SkipBorderSize * 2;
-			const FIntVector SourceBoxStart(SrcTileX * TileSize + SkipBorderSize, SrcTileY * TileSize + SkipBorderSize, 0);
-			const FIntVector DestinationBoxStart(Entry.SubmitDestX * SubmitTileSize, Entry.SubmitDestY * SubmitTileSize, 0);
+			FRHILockTextureArgs LockArgs = FRHILockTextureArgs::Lock2D(StagingTexture.RHITexture, 0, EResourceLockMode::RLM_WriteOnly, false, false);
+			FRHILockTextureResult LockResult = RHICmdList.LockTexture(LockArgs);
+			void* BatchMemory = LockResult.Data;
 
-			if (!UpdatedTextures.Contains(Entry.RHISubmitTexture))
+			// Copy all tiles to the staging texture
+			for (uint32 Index = 0u; Index < BatchCount; ++Index)
+			{
+				const FTileEntry& Entry = TileEntries[Index];
+				const FVTUploadTileBufferExt& UploadBuffer = UploadBuffers[Index];
+
+				const uint32_t SrcTileX = Index % StagingTexture.WidthInTiles;
+				const uint32_t SrcTileY = Index / StagingTexture.WidthInTiles;
+		
+				uint8* BatchDst = (uint8*)BatchMemory + TileHeightInBlocks * SrcTileY * LockResult.Stride + TileWidthInBlocks * SrcTileX * BlockBytes;
+				for (uint32 y = 0u; y < TileHeightInBlocks; ++y)
+				{
+					FMemory::Memcpy(
+						BatchDst + y * LockResult.Stride,
+						(uint8*)UploadBuffer.BufferMemory + UploadBuffer.BufferOffset + y * UploadBuffer.Stride,
+						TileWidthInBlocks * BlockBytes);
+				}
+			}
+		
+			RHICmdList.UnlockTexture(LockArgs);
+			RHICmdList.Transition(FRHITransitionInfo(StagingTexture.RHITexture, ERHIAccess::SRVMask, ERHIAccess::CopySrc));
+		
+			// Upload each tile from staging texture to physical texture
+			for (uint32 Index = 0u; Index < BatchCount; ++Index)
 			{
-				RHICmdList.Transition(FRHITransitionInfo(Entry.RHISubmitTexture, ERHIAccess::Unknown, ERHIAccess::CopyDest), ERHITransitionCreateFlags::NoFence);
-				UpdatedTextures.Add(Entry.RHISubmitTexture);
+				const FTileEntry& Entry = TileEntries[Index];
+				
+				const uint32_t SrcTileX = Index % StagingTexture.WidthInTiles;
+				const uint32_t SrcTileY = Index / StagingTexture.WidthInTiles;
+		
+				const uint32 SkipBorderSize = Entry.SubmitSkipBorderSize;
+				const uint32 SubmitTileSize = TileSize - SkipBorderSize * 2;
+				const FIntVector SourceBoxStart(SrcTileX * TileSize + SkipBorderSize, SrcTileY * TileSize + SkipBorderSize, 0);
+				const FIntVector DestinationBoxStart(Entry.SubmitDestX * SubmitTileSize, Entry.SubmitDestY * SubmitTileSize, 0);
+		
+				FRHICopyTextureInfo CopyInfo;
+				CopyInfo.Size = FIntVector(SubmitTileSize, SubmitTileSize, 1);
+				CopyInfo.SourcePosition = SourceBoxStart;
+				CopyInfo.DestPosition = DestinationBoxStart;
+				RHICmdList.CopyTexture(StagingTexture.RHITexture, Entry.Texture->GetRHI(), CopyInfo);
 			}
+		
+			RHICmdList.Transition(FRHITransitionInfo(StagingTexture.RHITexture, ERHIAccess::CopySrc, ERHIAccess::SRVMask));
+		});
+	}
+}
 
-			FRHICopyTextureInfo CopyInfo;
-			CopyInfo.Size = FIntVector(SubmitTileSize, SubmitTileSize, 1);
-			CopyInfo.SourcePosition = SourceBoxStart;
-			CopyInfo.DestPosition = DestinationBoxStart;
-			RHICmdList.CopyTexture(StagingTexture.RHITexture, Entry.RHISubmitTexture, CopyInfo);
+void FVirtualTextureUploadCache::FinalizeWithUpdateTexture(FRDGBuilder& GraphBuilder)
+{
+	check (UploadType == EVTUploadType::StagingCopy || UploadType == EVTUploadType::PersistentBuffer);
+
+	for (int PoolIndex = 0; PoolIndex < Pools.Num(); ++PoolIndex)
+	{
+		FPoolEntry& PoolEntry = Pools[PoolIndex];
+		const uint32 BatchCount = PoolEntry.PendingSubmit.Num();
+		if (BatchCount == 0)
+		{
+			continue;
 		}
 
-		RHICmdList.Transition(FRHITransitionInfo(StagingTexture.RHITexture, ERHIAccess::CopySrc, ERHIAccess::SRVMask));
+		// Make a copy of required data for RDG pass.
+		const int32 PoolTileSize = PoolEntry.TileSize;
+		FTileEntry* TileEntries = GraphBuilder.AllocPODArray<FTileEntry>(BatchCount);
+		FVTUploadTileBufferExt* UploadBuffers = GraphBuilder.AllocPODArray<FVTUploadTileBufferExt>(BatchCount);
+		
+		FVirtualTextureUploadCacheParameters* UploadParameters = GraphBuilder.AllocParameters<FVirtualTextureUploadCacheParameters>();
 
+		for (uint32 Index = 0u; Index < BatchCount; ++Index)
+		{
+			FTileEntry& TileEntry = PoolEntry.PendingSubmit[Index];
+
+			// Common case is that all target textures are the same.
+			// So compare with the previous entry as a quick test for a faster path.
+			if (Index > 1 && TileEntry.PooledRenderTarget == TileEntries[Index - 1].PooledRenderTarget)
+			{
+				TileEntry.Texture = TileEntries[Index - 1].Texture;
+			}
+			else
+			{
+				TileEntry.Texture = GraphBuilder.RegisterExternalTexture(TileEntry.PooledRenderTarget, ERDGTextureFlags::None);
+				UploadParameters->TextureArray.AddUnique(FRDGTextureAccess(TileEntry.Texture, ERHIAccess::CopyDest));
+			}
+			
+			TileEntries[Index] = TileEntry;
+			UploadBuffers[Index] = TileAllocator.GetBufferFromHandleExt(TileEntry.TileHandle);
+		}
+
+		// Can clear the pending queue now.
 		PoolEntry.PendingSubmit.Reset();
-	}
 
-	// Transition all updates textures back to SRV
-	TArray<FRHITransitionInfo, SceneRenderingAllocator> SRVTransitions;
-	SRVTransitions.Reserve(UpdatedTextures.Num());
-	for (int32 Index = 0; Index < UpdatedTextures.Num(); ++Index)
-	{
-		SRVTransitions.Add(FRHITransitionInfo(UpdatedTextures[Index], ERHIAccess::CopyDest, ERHIAccess::SRVMask));
+		// Submit to RDG.
+		GraphBuilder.AddPass(
+			RDG_EVENT_NAME("VirtualTextureUpload"),
+			UploadParameters,
+			ERDGPassFlags::Copy,
+			[PoolTileSize, BatchCount, TileEntries, UploadBuffers](FRDGAsyncTask,  FRHICommandList& RHICmdList)
+		{
+			SCOPED_GPU_MASK(RHICmdList, FRHIGPUMask::All());
+
+			SCOPED_DRAW_EVENT(RHICmdList, FVirtualTextureUploadCache_Finalize);
+			SCOPE_CYCLE_COUNTER(STAT_VTP_FlushUpload)
+
+			for (uint32 Index = 0u; Index < BatchCount; ++Index)
+			{
+				const FTileEntry& Entry = TileEntries[Index];
+				const uint32 SubmitTileSize = PoolTileSize - Entry.SubmitSkipBorderSize * 2;
+				const FUpdateTextureRegion2D UpdateRegion(
+					Entry.SubmitDestX * SubmitTileSize, Entry.SubmitDestY * SubmitTileSize, 
+					Entry.SubmitSkipBorderSize, Entry.SubmitSkipBorderSize, 
+					SubmitTileSize, SubmitTileSize);
+				const FVTUploadTileBufferExt& UploadBuffer = UploadBuffers[Index];
+
+				if (UploadBuffer.RHIBuffer != nullptr)
+				{
+					// This is the EVTUploadType::PersistentBuffer upload path.
+					RHICmdList.UpdateFromBufferTexture2D(Entry.Texture->GetRHI(), 0u, UpdateRegion, UploadBuffer.Stride, UploadBuffer.RHIBuffer, UploadBuffer.BufferOffset);
+				}
+				else
+				{
+					// This is the EVTUploadType::StagingCopy upload path.
+					RHICmdList.UpdateTexture2D(Entry.Texture->GetRHI(), 0u, UpdateRegion, UploadBuffer.Stride, (uint8*)UploadBuffer.BufferMemory + UploadBuffer.BufferOffset);
+				}
+			}
+		});
 	}
-	RHICmdList.Transition(SRVTransitions, ERHIPipeline::Graphics, ERHIPipeline::All);
-	UpdatedTextures.Reset();
 }
 
 void FVirtualTextureUploadCache::ReleaseRHI()
@@ -402,6 +528,7 @@ void FVirtualTextureUploadCache::ReleaseRHI()
 	{
 		CancelTile(RHICmdList, FVTUploadTileHandle(It.GetIndex()));
 	}
+	
 	// Release staging textures.
 	Pools.Empty();
 }
@@ -410,7 +537,7 @@ FVTUploadTileHandle FVirtualTextureUploadCache::PrepareTileForUpload(FRHICommand
 {
 	SCOPE_CYCLE_COUNTER(STAT_VTP_StageTile)
 
-	uint32 TileHandle = TileAllocator.Allocate(RHICmdList, InFormat, InTileSize);
+	uint32 TileHandle = TileAllocator.Allocate(RHICmdList, UploadType, InFormat, InTileSize);
 	OutBuffer = TileAllocator.GetBufferFromHandle(TileHandle);
 
 	const int32 PoolIndex = GetOrCreatePoolIndex(InFormat, InTileSize);
@@ -424,60 +551,25 @@ FVTUploadTileHandle FVirtualTextureUploadCache::PrepareTileForUpload(FRHICommand
 	return FVTUploadTileHandle(Index);
 }
 
-void FVirtualTextureUploadCache::SubmitTile(FRHICommandList& RHICmdList, const FVTUploadTileHandle& InHandle, FRHITexture* InDestTexture, int InDestX, int InDestY, int InSkipBorderSize)
+void FVirtualTextureUploadCache::SubmitTile(FRHICommandList& RHICmdList, const FVTUploadTileHandle& InHandle, const FVTProduceTargetLayer& Target, int InSkipBorderSize)
 {
-	checkSlow(IsInParallelRenderingThread());
-
+	// Get entry and remove from pending uploads.
 	check(PendingUpload.IsValidIndex(InHandle.Index));
-	FTileEntry& Entry = PendingUpload[InHandle.Index];
-	Entry.FrameSubmitted = GFrameNumberRenderThread;
-
-	FPoolEntry& PoolEntry = Pools[Entry.PoolIndex];
-	const uint32 TileSize = PoolEntry.TileSize - InSkipBorderSize * 2;
-
-	const FVTUploadTileBufferExt UploadBuffer = TileAllocator.GetBufferFromHandleExt(Entry.TileHandle);
-	if (UploadBuffer.RHIBuffer != nullptr)
-	{
-		if (!UpdatedTextures.Contains(InDestTexture))
-		{
-			RHICmdList.Transition(FRHITransitionInfo(InDestTexture, ERHIAccess::Unknown, ERHIAccess::CopyDest), ERHITransitionCreateFlags::NoFence);
-			UpdatedTextures.Add(InDestTexture);
-		}
-
-		check(GRHISupportsUpdateFromBufferTexture);
-		const FUpdateTextureRegion2D UpdateRegion(InDestX * TileSize, InDestY * TileSize, InSkipBorderSize, InSkipBorderSize, TileSize, TileSize);
-		RHICmdList.UpdateFromBufferTexture2D(InDestTexture, 0u, UpdateRegion, UploadBuffer.Stride, UploadBuffer.RHIBuffer, UploadBuffer.BufferOffset);
-
-		// Move to pending list, so we won't re-use this buffer until the GPU has finished the copy.
-		// We're using persist mapped buffer here, so this is the only synchronization method in place...without this delay we'd get corrupt textures.
-		PendingRelease.Emplace(Entry);
-	}
-	else if(ALLOW_UPDATE_TEXTURE)
-	{
-		if (!UpdatedTextures.Contains(InDestTexture))
-		{
-			RHICmdList.Transition(FRHITransitionInfo(InDestTexture, ERHIAccess::Unknown, ERHIAccess::CopyDest), ERHITransitionCreateFlags::NoFence);
-			UpdatedTextures.Add(InDestTexture);
-		}
+	FTileEntry Entry = PendingUpload[InHandle.Index];
+	PendingUpload.RemoveAt(InHandle.Index);
 
-		const FUpdateTextureRegion2D UpdateRegion(InDestX * TileSize, InDestY * TileSize, InSkipBorderSize, InSkipBorderSize, TileSize, TileSize);
-		RHICmdList.UpdateTexture2D(InDestTexture, 0u, UpdateRegion, UploadBuffer.Stride, (uint8*)UploadBuffer.BufferMemory + UploadBuffer.BufferOffset);
+	// Place on deferred release queue.
+	Entry.FrameSubmitted = GFrameNumberRenderThread;
+	PendingRelease.Emplace(Entry);
 
-		// UpdateTexture2D makes internal copy of data, no need to wait before re-using tile
-		TileAllocator.Free(RHICmdList, Entry.TileHandle);
-	}
-	else
-	{
-		// Move to list of batched updates for the current pool
-		Entry.RHISubmitTexture = InDestTexture;
-		Entry.SubmitDestX = InDestX;
-		Entry.SubmitDestY = InDestY;
-		Entry.SubmitSkipBorderSize = InSkipBorderSize;
-		PoolEntry.PendingSubmit.Emplace(Entry);
-	}
+	// Move to list of batched updates for the current pool
+	Entry.PooledRenderTarget = Target.PooledRenderTarget;
+	Entry.SubmitDestX = Target.pPageLocation.X;
+	Entry.SubmitDestY = Target.pPageLocation.Y;
+	Entry.SubmitSkipBorderSize = InSkipBorderSize;
 
-	// Remove from pending uploads.
-	PendingUpload.RemoveAt(InHandle.Index);
+	FPoolEntry& PoolEntry = Pools[Entry.PoolIndex];
+	PoolEntry.PendingSubmit.Emplace(Entry);
 }
 
 void FVirtualTextureUploadCache::CancelTile(FRHICommandList& RHICmdList, const FVTUploadTileHandle& InHandle)
@@ -492,12 +584,14 @@ void FVirtualTextureUploadCache::UpdateFreeList(FRHICommandList& RHICmdList, boo
 {
 	const uint32 CurrentFrame = GFrameNumberRenderThread;
 
-	// Iterate tiles pending release and free them if they are old enough.
+	// Iterate tiles pending release and free them.
+	// We need to defer release of persistent mapped buffers here because there is no other synchronization method for the memory in these.
+	// Without deferring release we would get corrupt textures.
 	for (TSparseArray<FTileEntry>::TIterator It(PendingRelease); It; ++It)
 	{
 		check(CurrentFrame >= It->FrameSubmitted);
 		const uint32 FramesSinceSubmitted = CurrentFrame - It->FrameSubmitted;
-		if (FramesSinceSubmitted < 2u)
+		if (UploadType == EVTUploadType::PersistentBuffer && FramesSinceSubmitted < 2u)
 		{
 			break;
 		}
diff --git a/Engine/Source/Runtime/Engine/Private/VT/VirtualTextureUploadCache.h b/Engine/Source/Runtime/Engine/Private/VT/VirtualTextureUploadCache.h
--- a/Engine/Source/Runtime/Engine/Private/VT/VirtualTextureUploadCache.h
+++ b/Engine/Source/Runtime/Engine/Private/VT/VirtualTextureUploadCache.h
@@ -3,12 +3,14 @@
 #pragma once
 
 #include "CoreMinimal.h"
+#include "PixelFormat.h"
 #include "RendererInterface.h"
 #include "RenderResource.h"
-#include "PixelFormat.h"
 
+class FRDGTexture;
 class FRHITexture;
 class FRHIBuffer;
+struct FVTProduceTargetLayer;
 
 /** Opaque handle for referencing an upload tile returned by FVirtualTextureUploadCache::PrepareTileForUpload(). */
 struct FVTUploadTileHandle
@@ -39,12 +41,23 @@ struct FVTUploadTileBuffer
  */
 struct FVTUploadTileBufferExt
 {
-	TRefCountPtr<FRHIBuffer> RHIBuffer;
+	FRHIBuffer* RHIBuffer;
 	void* BufferMemory = nullptr;
 	uint32 BufferOffset = 0u;
 	uint32 Stride = 0u;
 };
 
+/** Enumeration of staging buffer types used for upload to the VT physical texture. */
+enum class EVTUploadType
+{
+	/** Copy image data to a (batched) staging texture for upload. Requires immediate context Lock() but was the best path for D3D11 at some point. */
+	StagingTexture,
+	/** "Directly" upload image data. Internally the RHI will copy to staging memory. */
+	StagingCopy,
+	/** Stream image data directly into a persistent staging buffer for upload. This is the optimal path but is only available on some platforms. */
+	PersistentBuffer
+};
+
 /** 
  * Handles allocation of staging buffer memory. 
  */
@@ -52,7 +65,7 @@ class FVTUploadTileAllocator
 {
 public:
 	/** Allocate a tile. Sometimes does an allocation of the backing CPU/GPU block of memory. */
-	uint32 Allocate(FRHICommandList& RHICmdList, EPixelFormat InFormat, uint32 InTileSize);
+	uint32 Allocate(FRHICommandList& RHICmdList, EVTUploadType InUploadType, EPixelFormat InFormat, uint32 InTileSize);
 	/** Free a tile. Sometimes does a free of the backing CPU/GPU block of memory. */
 	void Free(FRHICommandList& RHICmdList, uint32 InHandle);
 
@@ -86,7 +99,7 @@ private:
 	{
 		~FStagingBuffer();
 
-		void Init(FRHICommandList& RHICmdList, uint32 InBufferStrideBytes, uint32 InTileSizeBytes);
+		void Init(FRHICommandList& RHICmdList, EVTUploadType InUploadType, uint32 InBufferStrideBytes, uint32 InTileSizeBytes);
 		void Release(FRHICommandList* RHICmdList);
 
 		/** GPU buffer if used on platform. */
@@ -132,6 +145,8 @@ private:
 class FVirtualTextureUploadCache : public IVirtualTextureFinalizer, public FRenderResource
 {
 public:
+	FVirtualTextureUploadCache();
+
 	//~ Begin FRenderResource Interface.
 	virtual void ReleaseRHI() override;
 	//~ End FRenderResource Interface.
@@ -142,7 +157,7 @@ public:
 	 * Mark streamed upload data ready for upload to to the physical virtual texture.
 	 * Depending on the platform the upload might happen here, or be deferred to the Finalize() call.
 	 */
-	void SubmitTile(FRHICommandList& RHICmdList, const FVTUploadTileHandle& InHandle, FRHITexture* InDestTexture, int InDestX, int InDestY, int InSkipBorderSize);
+	void SubmitTile(FRHICommandList& RHICmdList, const FVTUploadTileHandle& InHandle, const FVTProduceTargetLayer& Target, int InSkipBorderSize);
 	/** Cancel a tile that was already in flight. */
 	void CancelTile(FRHICommandList& RHICmdList, const FVTUploadTileHandle& InHandle);
 
@@ -157,18 +172,25 @@ public:
 	uint32 IsInMemoryBudget() const;
 
 private:
+	/** Get a pool for a format and tile size. */
 	int32 GetOrCreatePoolIndex(EPixelFormat InFormat, uint32 InTileSize);
 
+	/** Legacy Finalize path that creates and locks as staging texture. */
+	void FinalizeWithLegacyCopyTexture(FRDGBuilder& GraphBuilder);
+	/** Finalize path that uses UpdateTexture2D. */
+	void FinalizeWithUpdateTexture(FRDGBuilder& GraphBuilder);
+
 	/** Description of a single allocated tile. Carries mutable state as tile moves from uploading to submitting to pending delete. */
 	struct FTileEntry
 	{
 		uint32 PoolIndex = 0;
-		uint32 TileHandle = 0u;
-		uint32 FrameSubmitted = 0u;
-		FRHITexture* RHISubmitTexture = nullptr;
+		uint32 TileHandle = 0;
+		IPooledRenderTarget* PooledRenderTarget = nullptr;
+		FRDGTexture* Texture = nullptr;
 		int32 SubmitDestX = 0;
 		int32 SubmitDestY = 0;
 		int32 SubmitSkipBorderSize = 0;
+		uint32 FrameSubmitted = 0;
 	};
 
 	/** Staging texture used for tile upload. Only used on platforms that don't have faster upload methods. */
@@ -193,9 +215,9 @@ private:
 		TArray<FTileEntry> PendingSubmit;
 	};
 
+	EVTUploadType UploadType;
 	TArray<FPoolEntry> Pools;
 	FVTUploadTileAllocator TileAllocator;
 	TSparseArray<FTileEntry> PendingUpload;
 	TSparseArray<FTileEntry> PendingRelease;
-	TArray<FRHITexture*> UpdatedTextures;
 };
