WorldSubsystem: added event OnWorldEndPlay along with existing OnWorldBeginPlay.

diff --git a/Engine/Source/Runtime/Engine/Private/World.cpp b/Engine/Source/Runtime/Engine/Private/World.cpp
--- a/Engine/Source/Runtime/Engine/Private/World.cpp
+++ b/Engine/Source/Runtime/Engine/Private/World.cpp
@@ -5898,7 +5898,8 @@ void UWorld::BeginPlay()
 	FAssetCompilingManager::Get().ProcessAsyncTasks();
 #endif
 
-	SubsystemCollection.ForEachSubsystem([this](UWorldSubsystem* WorldSubsystem){
+	SubsystemCollection.ForEachSubsystem([this](UWorldSubsystem* WorldSubsystem)
+	{
 		WorldSubsystem->OnWorldBeginPlay(*this);
 	});
 
@@ -5939,6 +5940,11 @@ bool UWorld::EndPlay(EEndPlayReason::Type EndPlayReason)
 		ActorIt->RouteEndPlay(EndPlayReason);
 	}
 
+	SubsystemCollection.ForEachSubsystem([this](UWorldSubsystem* WorldSubsystem)
+	{
+		WorldSubsystem->OnWorldEndPlay(*this);
+	});
+
 	SetBegunPlay(false);
 
 	return true;
diff --git a/Engine/Source/Runtime/Engine/Public/Subsystems/WorldSubsystem.h b/Engine/Source/Runtime/Engine/Public/Subsystems/WorldSubsystem.h
--- a/Engine/Source/Runtime/Engine/Public/Subsystems/WorldSubsystem.h
+++ b/Engine/Source/Runtime/Engine/Public/Subsystems/WorldSubsystem.h
@@ -35,6 +35,9 @@ public:
 	/** Called when world is ready to start gameplay before the game mode transitions to the correct state and call BeginPlay on all actors */
 	virtual void OnWorldBeginPlay(UWorld& InWorld) {}
 
+	/** Called when world is ready to stop gameplay before the game mode transitions to the correct state and call EndPlay on all actors */
+	virtual void OnWorldEndPlay(UWorld& InWorld) {}
+
 	/** Called after world components (e.g. line batcher and all level components) have been updated */
 	virtual void OnWorldComponentsUpdated(UWorld& World) {}
 
