Remove direct read to AActor::bIgnoreInPIE and replace with AActor::ShouldDuplicateInPIE virtual function so it can be overloaded by derived actor classes.

diff --git a/Engine/Source/Runtime/Engine/Classes/GameFramework/Actor.h b/Engine/Source/Runtime/Engine/Classes/GameFramework/Actor.h
--- a/Engine/Source/Runtime/Engine/Classes/GameFramework/Actor.h
+++ b/Engine/Source/Runtime/Engine/Classes/GameFramework/Actor.h
@@ -300,6 +300,10 @@ public:
 	UFUNCTION()
 	ENGINE_API virtual void OnRep_ReplicateMovement();
 
+#if WITH_EDITOR
+	ENGINE_API virtual bool ShouldDuplicateInPIE() const { return !bIgnoreInPIE; }
+#endif
+
 private:
 	/**
 	 * If true, replicate movement/location related properties.
diff --git a/Engine/Source/Runtime/Engine/Private/Level.cpp b/Engine/Source/Runtime/Engine/Private/Level.cpp
--- a/Engine/Source/Runtime/Engine/Private/Level.cpp
+++ b/Engine/Source/Runtime/Engine/Private/Level.cpp
@@ -735,7 +735,7 @@ void ULevel::Serialize( FArchive& Ar )
 				}
 			}
 
-			if ((Ar.GetPortFlags() & PPF_DuplicateForPIE) && Actor->bIgnoreInPIE)
+			if ((Ar.GetPortFlags() & PPF_DuplicateForPIE) && !Actor->ShouldDuplicateInPIE())
 			{
 				return false;
 			}
@@ -5145,4 +5145,4 @@ void ULevel::OnMultipleLSAsPopupDismissed()
 
 #endif // WITH_EDITOR
 
-#undef LOCTEXT_NAMESPACE
+#undef LOCTEXT_NAMESPACE
\ No newline at end of file
