Remove EAsyncLoadingProgress::Canceled as it is not handled by the async loading code and make it harder to implement progress callbacks.

diff --git a/Engine/Source/Runtime/CoreUObject/Public/UObject/UObjectGlobals.h b/Engine/Source/Runtime/CoreUObject/Public/UObject/UObjectGlobals.h
--- a/Engine/Source/Runtime/CoreUObject/Public/UObject/UObjectGlobals.h
+++ b/Engine/Source/Runtime/CoreUObject/Public/UObject/UObjectGlobals.h
@@ -686,9 +686,7 @@ enum class EAsyncLoadingProgress : uint32
 	/** Package has finished its serialization phase. */
 	Serialized,
 	/** Package has finished all loading phase successfully */
-	FullyLoaded,
-	/** Async loading was canceled */
-	Canceled
+	FullyLoaded
 };
 
 /** The type that represents an async loading priority */
diff --git a/Engine/Source/Runtime/Engine/Private/Tests/Loading/UnattendedAsyncLoadingTests.cpp b/Engine/Source/Runtime/Engine/Private/Tests/Loading/UnattendedAsyncLoadingTests.cpp
--- a/Engine/Source/Runtime/Engine/Private/Tests/Loading/UnattendedAsyncLoadingTests.cpp
+++ b/Engine/Source/Runtime/Engine/Private/Tests/Loading/UnattendedAsyncLoadingTests.cpp
@@ -136,8 +136,6 @@ namespace UnattendedLoadTestImpl
 										// callback even if it comes from the ALT thread.
 										case EAsyncLoadingProgress::Failed:
 											[[fallthrough]];
-										case EAsyncLoadingProgress::Canceled:
-											[[fallthrough]];
 										case EAsyncLoadingProgress::FullyLoaded:
 											[[fallthrough]];
 										case EAsyncLoadingProgress::Serialized:
