Substrate - added tooltip to the simple clear coat node.

#rnx
#tests PC editor

diff --git a/Engine/Source/Runtime/Engine/Private/Materials/MaterialExpressionSubstrate.cpp b/Engine/Source/Runtime/Engine/Private/Materials/MaterialExpressionSubstrate.cpp
--- a/Engine/Source/Runtime/Engine/Private/Materials/MaterialExpressionSubstrate.cpp
+++ b/Engine/Source/Runtime/Engine/Private/Materials/MaterialExpressionSubstrate.cpp
@@ -1416,6 +1416,14 @@ void UMaterialExpressionSubstrateSimpleClearCoatBSDF::GetCaption(TArray<FString>
 	OutCaptions.Add(TEXT("Substrate Simple Clear Coat"));
 }
 
+void UMaterialExpressionSubstrateSimpleClearCoatBSDF::GetExpressionToolTip(TArray<FString>& OutToolTip)
+{
+	OutToolTip.Add(TEXT("Substrate Simple Clear Coat BSDF"));
+	OutToolTip.Add(TEXT("It corresponds to the legacy \"Clear Coat\" shading model (also available using the \"Substrate Shading Models\" node)."));
+	OutToolTip.Add(TEXT("It has the complexity \"Single\"."));
+	OutToolTip.Add(TEXT("Clear coat top layer controls are limited to coverage and roughness. It always has a fixed grey scale transmittance and F0=0.04."));
+}
+
 EMaterialValueType UMaterialExpressionSubstrateSimpleClearCoatBSDF::GetOutputValueType(int32 OutputIndex)
 {
 	return MCT_Substrate;
diff --git a/Engine/Source/Runtime/Engine/Public/Materials/MaterialExpressionSubstrate.h b/Engine/Source/Runtime/Engine/Public/Materials/MaterialExpressionSubstrate.h
--- a/Engine/Source/Runtime/Engine/Public/Materials/MaterialExpressionSubstrate.h
+++ b/Engine/Source/Runtime/Engine/Public/Materials/MaterialExpressionSubstrate.h
@@ -477,6 +477,7 @@ class UMaterialExpressionSubstrateSimpleClearCoatBSDF : public UMaterialExpressi
 	virtual void GatherSubstrateMaterialInfo(FSubstrateMaterialInfo& SubstrateMaterialInfo, int32 OutputIndex) override;
 	virtual FSubstrateOperator* SubstrateGenerateMaterialTopologyTree(class FMaterialCompiler* Compiler, class UMaterialExpression* Parent, int32 OutputIndex) override;
 	virtual FName GetInputName(int32 InputIndex) const override;
+	virtual void GetExpressionToolTip(TArray<FString>& OutToolTip) override;
 #endif
 	//~ End UMaterialExpression Interface
 };
