AnimSequence: Fix uninitialized pose transforms being generated in some circumstances

Occured when:
- Retargetting was disabled
- Skeleton remapping (compatible skeletons) was not able to correctly remap the bone

PR #12840: Fix a crash caused by UAnimSequence::GetBonePose() may not initialize some bone transforms in a pose

diff --git a/Engine/Source/Runtime/Engine/Private/Animation/AnimSequence.cpp b/Engine/Source/Runtime/Engine/Private/Animation/AnimSequence.cpp
--- a/Engine/Source/Runtime/Engine/Private/Animation/AnimSequence.cpp
+++ b/Engine/Source/Runtime/Engine/Private/Animation/AnimSequence.cpp
@@ -1744,17 +1744,33 @@ void UAnimSequence::GetBonePose(FAnimationPoseData& OutAnimationPoseData, const
 		if (bDisableRetargeting)
 		{
 			TArray<FTransform> const& AuthoredOnRefSkeleton = GetRetargetTransforms();
+			const TArray<FTransform>& SkeletonRefPose = RequiredBones.GetSkeletonAsset()->GetRefLocalPoses();
+
 			// Map from this sequence its Skeleton to target
 			const FSkeletonRemapping& SkeletonRemapping = UE::Anim::FSkeletonRemappingRegistry::Get().GetRemapping(MySkeleton, RequiredBones.GetSkeletonAsset());
 			for (FCompactPoseBoneIndex PoseBoneIndex : OutPose.ForEachBoneIndex())
 			{
-				int32 SkeletonBoneIndex = RequiredBones.GetSkeletonPoseIndexFromCompactPoseIndex(PoseBoneIndex).GetInt();
-				SkeletonBoneIndex = SkeletonRemapping.IsValid() ? SkeletonRemapping.GetSourceSkeletonBoneIndex(SkeletonBoneIndex) : SkeletonBoneIndex;
+				const int32 SkeletonBoneIndex = RequiredBones.GetSkeletonPoseIndexFromCompactPoseIndex(PoseBoneIndex).GetInt();
+				const int32 RemappedSkeletonBoneIndex = SkeletonRemapping.IsValid() ? SkeletonRemapping.GetSourceSkeletonBoneIndex(SkeletonBoneIndex) : SkeletonBoneIndex;
 
 				// Virtual bones are part of the retarget transform pose, so if the pose has not been updated (recently) there might be a mismatch
-				if (SkeletonBoneIndex != INDEX_NONE && AuthoredOnRefSkeleton.IsValidIndex(SkeletonBoneIndex))
+				if (RemappedSkeletonBoneIndex != INDEX_NONE && AuthoredOnRefSkeleton.IsValidIndex(RemappedSkeletonBoneIndex))
 				{
-					OutPose[PoseBoneIndex] = AuthoredOnRefSkeleton[SkeletonBoneIndex];
+					OutPose[PoseBoneIndex] = AuthoredOnRefSkeleton[RemappedSkeletonBoneIndex];
+				}
+				else
+				{
+					// Only do this if we have a mesh. otherwise we're not retargeting animations.
+					if (RequiredBones.GetSkeletalMeshAsset())
+					{
+						// Pose bone index should always exist in Skeleton
+						checkSlow(SkeletonBoneIndex != INDEX_NONE);
+						OutPose[PoseBoneIndex] = SkeletonRefPose[SkeletonBoneIndex];
+					}
+					else
+					{
+						OutPose[PoseBoneIndex] = RequiredBones.GetRefPoseTransform(PoseBoneIndex);
+					}
 				}
 			}
 		}
