SkeletalMeshComponent: Fix meshes getting ticked offscreen incorrectly when using montage-related visibilty options

PR #12182: Speculative fix for skeletal meshes always tick unless set to OnlyTickPoseWhenRendered

diff --git a/Engine/Source/Runtime/Engine/Private/Components/SkeletalMeshComponent.cpp b/Engine/Source/Runtime/Engine/Private/Components/SkeletalMeshComponent.cpp
--- a/Engine/Source/Runtime/Engine/Private/Components/SkeletalMeshComponent.cpp
+++ b/Engine/Source/Runtime/Engine/Private/Components/SkeletalMeshComponent.cpp
@@ -1621,7 +1621,8 @@ bool USkeletalMeshComponent::ShouldOnlyTickMontagesAndRefreshBones(const float D
 	// RefreshBoneTransforms will need an updated graph.
 	return VisibilityBasedAnimTickOption == EVisibilityBasedAnimTickOption::OnlyTickMontagesAndRefreshBonesWhenPlayingMontages
 		&& !bRecentlyRendered
-		&& (DeltaTime > 0.f);
+		&& (DeltaTime > 0.f)
+		&& (GetAnimInstance() && GetAnimInstance()->IsAnyMontagePlaying());
 }
 
 void USkeletalMeshComponent::TickAnimation(float DeltaTime, bool bNeedsValidRootMotion)
@@ -1804,11 +1805,16 @@ bool USkeletalMeshComponent::ShouldTickPose() const
 	}
 #endif 
 
+	// Allow tick pose if there is any montage playing and has corresponding visibility based tick option while not rendered
+	const bool bShouldTickBasedOnMontage = (VisibilityBasedAnimTickOption == EVisibilityBasedAnimTickOption::OnlyTickMontagesWhenNotRendered || VisibilityBasedAnimTickOption == EVisibilityBasedAnimTickOption::OnlyTickMontagesAndRefreshBonesWhenPlayingMontages)
+		&& !bRecentlyRendered
+		&& (GetAnimInstance() && GetAnimInstance()->IsAnyMontagePlaying());
+
 	// Autonomous Ticking is allowed to occur multiple times per frame, as we can receive and process multiple networking updates the same frame.
 	const bool bShouldTickBasedOnAutonomousCheck = bIsAutonomousTickPose || (!bOnlyAllowAutonomousTickPose && !bAlreadyTickedThisFrame);
 	// When playing networked Root Motion Montages, we want these to play on dedicated servers and remote clients for networking and position correction purposes.
 	// So we force pose updates in that case to keep root motion and position in sync.
-	const bool bShouldTickBasedOnVisibility = ((VisibilityBasedAnimTickOption < EVisibilityBasedAnimTickOption::OnlyTickPoseWhenRendered) || bRecentlyRendered || IsPlayingNetworkedRootMotionMontage());
+	const bool bShouldTickBasedOnVisibility = (bShouldTickBasedOnMontage || (VisibilityBasedAnimTickOption <= EVisibilityBasedAnimTickOption::AlwaysTickPose) || bRecentlyRendered || IsPlayingNetworkedRootMotionMontage());
 
 	return (bShouldTickBasedOnVisibility && bShouldTickBasedOnAutonomousCheck && IsRegistered() && (AnimScriptInstance || PostProcessAnimInstance) && !bPauseAnims && GetWorld()->AreActorsInitialized() && !bNoSkeletonUpdate);
 }
