PropertyAccess: Fix memory leak when calling functions via property access that return non-POD values

Ensure we call DestroyValue after we operate on the returned address

diff --git a/Engine/Source/Runtime/Engine/Private/PropertyAccess.cpp b/Engine/Source/Runtime/Engine/Private/PropertyAccess.cpp
--- a/Engine/Source/Runtime/Engine/Private/PropertyAccess.cpp
+++ b/Engine/Source/Runtime/Engine/Private/PropertyAccess.cpp
@@ -450,13 +450,16 @@ struct FPropertyAccessSystem
 	{
 		// Buffer for function return values
 		TArray<uint8, TNonRelocatableInlineAllocator<64>> ArrayBuffer;
-	
+ 
+		// Property that we initialize the buffer with
+		FProperty* BufferProperty = nullptr;
+ 
 		void* Address = InContainer;
-
+ 
 		for(int32 IndirectionIndex = InAccess.IndirectionStartIndex; Address != nullptr && IndirectionIndex < InAccess.IndirectionEndIndex; ++IndirectionIndex)
 		{
 			const FPropertyAccessIndirection& Indirection = InLibrary.Indirections[IndirectionIndex];
-
+ 
 			switch(Indirection.Type)
 			{
 			case EPropertyAccessIndirectionType::Offset:
@@ -517,16 +520,17 @@ struct FPropertyAccessSystem
 				if(Indirection.Function != nullptr)
 				{
 					UObject* CalleeObject = static_cast<UObject*>(Address);
-
+ 
 					// Allocate buffer + alignment slack for the return value
 					ArrayBuffer.SetNumUninitialized(Indirection.ReturnBufferSize + Indirection.ReturnBufferAlignment, EAllowShrinking::No);
 					Address = Align(ArrayBuffer.GetData(), Indirection.ReturnBufferAlignment);
-
+ 
 					// Init value
 					check(Indirection.Function->GetReturnProperty());
 					check(Indirection.Function->GetReturnProperty() == Indirection.Property.Get());
 					Indirection.Property.Get()->InitializeValue(Address);
-
+					BufferProperty = Indirection.Property.Get();
+ 
 					if(Indirection.Type == EPropertyAccessIndirectionType::NativeFunction)
 					{
 						CallNativeAccessor(CalleeObject, Indirection.Function, Address);
@@ -535,7 +539,7 @@ struct FPropertyAccessSystem
 					{
 						CalleeObject->ProcessEvent(Indirection.Function, Address);
 					}
-
+ 
 					// Function access may return an object, so we need to follow that ptr
 					switch(Indirection.ObjectType)
 					{
@@ -573,10 +577,15 @@ struct FPropertyAccessSystem
 				check(false);
 			}
 		}
-
+ 
 		if(Address != nullptr)
 		{
 			InAddressFunction(Address);
+ 
+			if (BufferProperty != nullptr)
+			{
+				BufferProperty->DestroyValue(Address);
+			}
 		}
 	}
 	
