Remove deprecated functions in Materials that are UE 5.4 or less.

- Remove implementation of HasMaterialAttributes since the base class version has been deprecated and removed.

#rnx

diff --git a/Engine/Source/Editor/UnrealEd/Classes/MaterialGraph/MaterialGraphSchema.h b/Engine/Source/Editor/UnrealEd/Classes/MaterialGraph/MaterialGraphSchema.h
--- a/Engine/Source/Editor/UnrealEd/Classes/MaterialGraph/MaterialGraphSchema.h
+++ b/Engine/Source/Editor/UnrealEd/Classes/MaterialGraph/MaterialGraphSchema.h
@@ -218,24 +218,6 @@ class UMaterialGraphSchema : public UEdGraphSchema
 	UNREALED_API static const FLinearColor InactivePinColor;
 	UNREALED_API static const FLinearColor AlphaPinColor;
 
-	/**
-	 *  Add all linked to nodes to this pin to selection
-	 *  
-	 *	@param	Graph			CurrentGraph
-	 *	@param	InGraphPin		Pin clicked on
-	 */
-	UE_DEPRECATED(4.26, "Selecting all connected nodes is now handled through common pin actions on graph schemas. See SGraphEditorImpl::SelectAllNodesInDirection.")
-	void SelectAllInputNodes(UEdGraph* Graph, UEdGraphPin* InGraphPin);
-
-	/**
-	 * Get menu for breaking links to specific nodes
-	 *
-	 * @param	Menu		Menu we are populating
-	 * @param	InGraphPin	Pin with links to break
-	 */
-	UE_DEPRECATED(4.26, "Populating the break link to menu is handled through common pin actions on graph schemas. This function is no longer necessary.")
-	void GetBreakLinkToSubMenuActions(class UToolMenu* Menu, class UEdGraphPin* InGraphPin) const;
-
 	/**
 	 * Connect a pin to one of the Material Function's outputs
 	 *
diff --git a/Engine/Source/Editor/UnrealEd/Private/MaterialGraphSchema.cpp b/Engine/Source/Editor/UnrealEd/Private/MaterialGraphSchema.cpp
--- a/Engine/Source/Editor/UnrealEd/Private/MaterialGraphSchema.cpp
+++ b/Engine/Source/Editor/UnrealEd/Private/MaterialGraphSchema.cpp
@@ -291,82 +291,6 @@ UMaterialGraphSchema::UMaterialGraphSchema(const FObjectInitializer& ObjectIniti
 {
 }
 
-void UMaterialGraphSchema::SelectAllInputNodes(UEdGraph* Graph, UEdGraphPin* InGraphPin)
-{
-	TArray<UEdGraphPin*> AllPins = InGraphPin->LinkedTo;
-
-	if (AllPins.Num() == 0)
-	{
-		return;
-	}
-
-	for (UEdGraphPin* Pin : AllPins)
-	{
-		UMaterialGraphNode* MaterialNode = Cast<UMaterialGraphNode>(Pin->GetOwningNode());
-		FMaterialEditorUtilities::AddToSelection(Graph, MaterialNode->MaterialExpression);
-
-		TArray<UEdGraphPin*> LinkedPins = Pin->GetOwningNode()->GetAllPins();
-		for (UEdGraphPin* InputPin : LinkedPins)
-		{
-			if (InputPin->Direction == EEdGraphPinDirection::EGPD_Output)
-			{
-				continue;
-			}
-			else
-			{
-				PRAGMA_DISABLE_DEPRECATION_WARNINGS
-				SelectAllInputNodes(Graph, InputPin);
-				PRAGMA_ENABLE_DEPRECATION_WARNINGS
-			}
-		}
-	}
-}
-
-void UMaterialGraphSchema::GetBreakLinkToSubMenuActions(class UToolMenu* Menu, UEdGraphPin* InGraphPin) const
-{
-	FToolMenuSection& Section = Menu->FindOrAddSection("MaterialGraphSchemaPinActions");
-
-	// Make sure we have a unique name for every entry in the list
-	TMap< FString, uint32 > LinkTitleCount;
-
-	// Add all the links we could break from
-	for(TArray<class UEdGraphPin*>::TConstIterator Links(InGraphPin->LinkedTo); Links; ++Links)
-	{
-		UEdGraphPin* Pin = *Links;
-		FString TitleString = Pin->GetOwningNode()->GetNodeTitle(ENodeTitleType::ListView).ToString();
-		FText Title = FText::FromString( TitleString );
-		if ( Pin->PinName != TEXT("") )
-		{
-			TitleString = FString::Printf(TEXT("%s (%s)"), *TitleString, *Pin->PinName.ToString());
-
-			// Add name of connection if possible
-			FFormatNamedArguments Args;
-			Args.Add( TEXT("NodeTitle"), Title );
-			Args.Add( TEXT("PinName"), Pin->GetDisplayName() );
-			Title = FText::Format( LOCTEXT("BreakDescPin", "{NodeTitle} ({PinName})"), Args );
-		}
-
-		uint32 &Count = LinkTitleCount.FindOrAdd(TitleString);
-
-		FText Description;
-		FFormatNamedArguments Args;
-		Args.Add( TEXT("NodeTitle"), Title );
-		Args.Add( TEXT("NumberOfNodes"), Count );
-
-		if ( Count == 0 )
-		{
-			Description = FText::Format( LOCTEXT("BreakDesc", "Break link to {NodeTitle}"), Args );
-		}
-		else
-		{
-			Description = FText::Format( LOCTEXT("BreakDescMulti", "Break link to {NodeTitle} ({NumberOfNodes})"), Args );
-		}
-		++Count;
-		Section.AddMenuEntry(NAME_None, Description, Description, FSlateIcon(), FUIAction(
-			FExecuteAction::CreateUObject((UMaterialGraphSchema*const)this, &UMaterialGraphSchema::BreakSinglePinLink, const_cast< UEdGraphPin* >(InGraphPin), *Links) ) );
-	}
-}
-
 void UMaterialGraphSchema::OnConnectToFunctionOutput(UEdGraphPin* InGraphPin, UEdGraphPin* InFuncPin)
 {
 	const FScopedTransaction Transaction( NSLOCTEXT("UnrealEd", "GraphEd_CreateConnection", "Create Pin Link") );
diff --git a/Engine/Source/Runtime/Engine/Private/Materials/Material.cpp b/Engine/Source/Runtime/Engine/Private/Materials/Material.cpp
--- a/Engine/Source/Runtime/Engine/Private/Materials/Material.cpp
+++ b/Engine/Source/Runtime/Engine/Private/Materials/Material.cpp
@@ -6038,13 +6038,6 @@ bool UMaterial::CanBeClusterRoot() const
 	return true;
 }
 
-void UMaterial::GetAssetRegistryTags(TArray<FAssetRegistryTag>& OutTags) const
-{
-	PRAGMA_DISABLE_DEPRECATION_WARNINGS;
-	Super::GetAssetRegistryTags(OutTags);
-	PRAGMA_ENABLE_DEPRECATION_WARNINGS;
-}
-
 void UMaterial::GetAssetRegistryTags(FAssetRegistryTagsContext Context) const
 {
 	Super::GetAssetRegistryTags(Context);
diff --git a/Engine/Source/Runtime/Engine/Private/Materials/MaterialExpressions.cpp b/Engine/Source/Runtime/Engine/Private/Materials/MaterialExpressions.cpp
--- a/Engine/Source/Runtime/Engine/Private/Materials/MaterialExpressions.cpp
+++ b/Engine/Source/Runtime/Engine/Private/Materials/MaterialExpressions.cpp
@@ -13367,13 +13367,6 @@ void UMaterialFunctionInterface::DeclareConstructClasses(TArray<FTopLevelAssetPa
 }
 #endif
 
-void UMaterialFunctionInterface::GetAssetRegistryTags(TArray<FAssetRegistryTag>& OutTags) const
-{
-	PRAGMA_DISABLE_DEPRECATION_WARNINGS;
-	Super::GetAssetRegistryTags(OutTags);
-	PRAGMA_ENABLE_DEPRECATION_WARNINGS;
-}
-
 void UMaterialFunctionInterface::GetAssetRegistryTags(FAssetRegistryTagsContext Context) const
 {
 	Super::GetAssetRegistryTags(Context);
diff --git a/Engine/Source/Runtime/Engine/Private/Materials/MaterialInterface.cpp b/Engine/Source/Runtime/Engine/Private/Materials/MaterialInterface.cpp
--- a/Engine/Source/Runtime/Engine/Private/Materials/MaterialInterface.cpp
+++ b/Engine/Source/Runtime/Engine/Private/Materials/MaterialInterface.cpp
@@ -779,13 +779,6 @@ void UMaterialInterface::PostEditChangeProperty(FPropertyChangedEvent& PropertyC
 	Super::PostEditChangeProperty(PropertyChangedEvent);
 }
 
-void UMaterialInterface::GetAssetRegistryTags(TArray<FAssetRegistryTag>& OutTags) const
-{
-	PRAGMA_DISABLE_DEPRECATION_WARNINGS;
-	Super::GetAssetRegistryTags(OutTags);
-	PRAGMA_ENABLE_DEPRECATION_WARNINGS;
-}
-
 void UMaterialInterface::GetAssetRegistryTags(FAssetRegistryTagsContext Context) const
 {
 	if (AssetImportData)
diff --git a/Engine/Source/Runtime/Engine/Private/Materials/MaterialRenderProxy.cpp b/Engine/Source/Runtime/Engine/Private/Materials/MaterialRenderProxy.cpp
--- a/Engine/Source/Runtime/Engine/Private/Materials/MaterialRenderProxy.cpp
+++ b/Engine/Source/Runtime/Engine/Private/Materials/MaterialRenderProxy.cpp
@@ -504,11 +504,6 @@ void FMaterialRenderProxy::EvaluateUniformExpressions(FRHICommandListBase& RHICm
 	++UniformExpressionCacheSerialNumber;
 }
 
-void FMaterialRenderProxy::EvaluateUniformExpressions(FUniformExpressionCache& OutUniformExpressionCache, const FMaterialRenderContext& Context, FUniformExpressionCacheAsyncUpdater* Updater) const
-{
-	EvaluateUniformExpressions(FRHICommandListImmediate::Get(), OutUniformExpressionCache, Context, Updater);
-}
-
 void FMaterialRenderProxy::CacheUniformExpressions(FRHICommandListBase& RHICmdList, bool bRecreateUniformBuffer)
 {
 	// Register the render proxy's as a render resource so it can receive notifications to free the uniform buffer.
@@ -538,11 +533,6 @@ void FMaterialRenderProxy::CacheUniformExpressions(FRHICommandListBase& RHICmdLi
 	}
 }
 
-void FMaterialRenderProxy::CacheUniformExpressions(bool bRecreateUniformBuffer)
-{
-	CacheUniformExpressions(FRHICommandListImmediate::Get(), bRecreateUniformBuffer);
-}
-
 void FMaterialRenderProxy::CacheUniformExpressions_GameThread(bool bRecreateUniformBuffer)
 {
 	if (FApp::CanEverRender())
diff --git a/Engine/Source/Runtime/Engine/Private/Materials/MaterialShared.cpp b/Engine/Source/Runtime/Engine/Private/Materials/MaterialShared.cpp
--- a/Engine/Source/Runtime/Engine/Private/Materials/MaterialShared.cpp
+++ b/Engine/Source/Runtime/Engine/Private/Materials/MaterialShared.cpp
@@ -1618,9 +1618,6 @@ bool FMaterialResource::IsVolumetricPrimitive() const { return Material->Materia
 bool FMaterialResource::IsSpecialEngineMaterial() const { return Material->bUsedAsSpecialEngineMaterial; }
 bool FMaterialResource::HasVertexPositionOffsetConnected() const { return GetCachedExpressionData().IsPropertyConnected(MP_WorldPositionOffset); }
 bool FMaterialResource::HasPixelDepthOffsetConnected() const { return Material->HasPixelDepthOffsetConnected(); }
-#if WITH_EDITOR
-bool FMaterialResource::HasMaterialAttributesConnected() const { return (Material->bUseMaterialAttributes && Material->GetEditorOnlyData()->MaterialAttributes.IsConnected()); }
-#endif
 EMaterialShadingRate FMaterialResource::GetShadingRate() const { return Material->ShadingRate; }
 bool FMaterialResource::IsVariableRateShadingAllowed() const {
 	const bool bDisableForCustomPrimitiveData = Material->HasCustomPrimitiveData() && CVarDisableVRSForCustomPrimitiveData.GetValueOnAnyThread();
@@ -2046,7 +2043,6 @@ bool FMaterialResource::HasSpecularConnected() const
 
 bool FMaterialResource::HasMetallicConnected() const
 {
-	// HasMaterialAttributesConnected() should be captured by HasMetallicConnected() now
 	return Material->HasMetallicConnected();
 }
 
@@ -2108,7 +2104,7 @@ bool FMaterialResource::HasMaterialPropertyConnected(EMaterialProperty In) const
 		case MP_Normal: 			return Material->HasNormalConnected();
 		case MP_Roughness: 			return Material->HasRoughnessConnected();
 		case MP_Specular: 			return Material->HasSpecularConnected();
-		case MP_Metallic: 			return Material->HasMetallicConnected(); // HasMetallicConnected() should be properly capturing HasMaterialAttributesConnected() as well
+		case MP_Metallic: 			return Material->HasMetallicConnected();
 		case MP_Anisotropy: 		return Material->HasAnisotropyConnected();
 		case MP_AmbientOcclusion: 	return Material->HasAmbientOcclusionConnected();
 		}
diff --git a/Engine/Source/Runtime/Engine/Public/MaterialCompiler.h b/Engine/Source/Runtime/Engine/Public/MaterialCompiler.h
--- a/Engine/Source/Runtime/Engine/Public/MaterialCompiler.h
+++ b/Engine/Source/Runtime/Engine/Public/MaterialCompiler.h
@@ -278,8 +278,6 @@ public:
 	virtual int32 LocalPosition(EPositionIncludedOffsets IncludedOffsets, ELocalPositionOrigin OriginType) = 0;
 	virtual int32 WorldPosition(EWorldPositionIncludedOffsets WorldPositionIncludedOffsets) = 0;
 	virtual int32 ObjectWorldPosition(EPositionOrigin OriginType) = 0;
-	UE_DEPRECATED(5.4, "Use ObjectWorldPosition(EPositionOrigin) instead")
-	int32 ObjectWorldPosition() { return ObjectWorldPosition(EPositionOrigin::Absolute); }
 	virtual int32 ObjectRadius() = 0;
 	virtual int32 ObjectBounds() = 0;
 	virtual int32 ObjectLocalBounds(int32 OutputIndex) = 0;
@@ -287,15 +285,11 @@ public:
 	virtual int32 PreSkinnedLocalBounds(int32 OutputIndex) = 0;
 	virtual int32 DistanceCullFade() = 0;
 	virtual int32 ActorWorldPosition(EPositionOrigin OriginType) = 0;
-	UE_DEPRECATED(5.4, "Use ActorWorldPosition(EPositionOrigin) instead")
-	int32 ActorWorldPosition() { return ActorWorldPosition(EPositionOrigin::Absolute); }
 	virtual int32 ParticleMacroUV() = 0;
 	virtual int32 ParticleSubUV(int32 TextureIndex, EMaterialSamplerType SamplerType, int32 MipValue0Index, int32 MipValue1Index, ETextureMipValueMode MipValueMode, bool bBlend) = 0;
 	virtual int32 ParticleSubUVProperty(int32 PropertyIndex) = 0;
 	virtual int32 ParticleColor() = 0;
 	virtual int32 ParticlePosition(EPositionOrigin OriginType) = 0;
-	UE_DEPRECATED(5.4, "Use ParticlePosition(EPositionOrigin) instead")
-	int32 ParticlePosition() { return ParticlePosition(EPositionOrigin::Absolute); }
 	virtual int32 ParticleRadius() = 0;
 	virtual int32 SphericalParticleOpacity(int32 Density) = 0;
 	virtual int32 ParticleRelativeTime() = 0;
@@ -338,8 +332,6 @@ public:
 	virtual int32 VirtualTextureUniform(int32 TextureIndex, int32 VectorIndex, UE::Shader::EValueType Type) = 0;
 	virtual int32 VirtualTextureUniform(FName ParameterName, int32 TextureIndex, int32 VectorIndex, UE::Shader::EValueType Type) = 0;
 	virtual int32 VirtualTextureWorldToUV(int32 WorldPositionIndex, int32 P0, int32 P1, int32 P2, EPositionOrigin PositionOrigin) = 0;
-	UE_DEPRECATED(5.4, "Use VirtualTextureWorldToUV(int32, int32, int32, int32, EPositionOrigin) instead")
-	int32 VirtualTextureWorldToUV(int32 WorldPositionIndex, int32 P0, int32 P1, int32 P2) { return VirtualTextureWorldToUV(WorldPositionIndex, P0, P1, P2, EPositionOrigin::Absolute); }
 	virtual int32 VirtualTextureUnpack(int32 CodeIndex0, int32 CodeIndex1, int32 CodeIndex2, int32 P0, EVirtualTextureUnpackType UnpackType) = 0;
 	virtual int32 VirtualTextureCustomData() = 0;
 
@@ -616,52 +608,17 @@ public:
 	virtual int32 Sobol(int32 Cell, int32 Index, int32 Seed) = 0;
 	virtual int32 TemporalSobol(int32 Index, int32 Seed) = 0;
 	virtual int32 Noise(int32 Position, EPositionOrigin PositionOrigin, float Scale, int32 Quality, uint8 NoiseFunction, bool bTurbulence, int32 Levels, float OutputMin, float OutputMax, float LevelScale, int32 FilterWidth, bool bTiling, uint32 RepeatSize) = 0;
-	UE_DEPRECATED(5.4, "Use Noise(int32, EPositionOrigin, float, int32, uint8, bool, int32, float, float, float, int32, bool, uint32) instead")
-	int32 Noise(int32 Position, float Scale, int32 Quality, uint8 NoiseFunction, bool bTurbulence, int32 Levels, float OutputMin, float OutputMax, float LevelScale, int32 FilterWidth, bool bTiling, uint32 RepeatSize)
-	{
-		return Noise(Position, EPositionOrigin::Absolute, Scale, Quality, NoiseFunction, bTurbulence, Levels, OutputMin, OutputMax, LevelScale, FilterWidth, bTiling, RepeatSize);
-	}
 	virtual int32 VectorNoise(int32 Position, EPositionOrigin PositionOrigin, int32 Quality, uint8 NoiseFunction, bool bTiling, uint32 RepeatSize) = 0;
-	UE_DEPRECATED(5.4, "Use VectorNoise(int32, EPositionOrigin, int32, int8, bool, uint32) instead")
-	int32 VectorNoise(int32 Position, int32 Quality, uint8 NoiseFunction, bool bTiling, uint32 RepeatSize)
-	{
-		return VectorNoise(Position, EPositionOrigin::Absolute, Quality, NoiseFunction, bTiling, RepeatSize);
-	}
 	virtual int32 ScalarBlueNoise() = 0;
 
 	virtual int32 BlackBody( int32 Temp ) = 0;
 	virtual int32 DistanceToNearestSurface(int32 PositionArg, EPositionOrigin PositionOrigin) = 0;
-	UE_DEPRECATED(5.4, "Use DistanceToNearestSurface(int32, EPositionOrigin) instead")
-	int32 DistanceToNearestSurface(int32 PositionArg)
-	{
-		return DistanceToNearestSurface(PositionArg, EPositionOrigin::Absolute);
-	}
 	virtual int32 DistanceFieldGradient(int32 PositionArg, EPositionOrigin PositionOrigin) = 0;
-	UE_DEPRECATED(5.4, "Use DistanceFieldGradient(int32, EPositionOrigin) instead")
-	int32 DistanceFieldGradient(int32 PositionArg)
-	{
-		return DistanceFieldGradient(PositionArg, EPositionOrigin::Absolute);
-	}
 	virtual int32 DistanceFieldApproxAO(int32 PositionArg, EPositionOrigin PositionOrigin, int32 NormalArg, int32 BaseDistanceArg, int32 RadiusArg, uint32 NumSteps, float StepScale) = 0;
-	UE_DEPRECATED(5.4, "Use DistanceFieldApproxAO(int32, EPositionOrigin, int32, int32, int32, uint32, float) instead")
-	int32 DistanceFieldApproxAO(int32 PositionArg, int32 NormalArg, int32 BaseDistanceArg, int32 RadiusArg, uint32 NumSteps, float StepScale)
-	{
-		return DistanceFieldApproxAO(PositionArg, EPositionOrigin::Absolute, NormalArg, BaseDistanceArg, RadiusArg, NumSteps, StepScale);
-	}
 	virtual int32 SamplePhysicsField(int32 PositionArg, EPositionOrigin PositionOrigin, const int32 OutputType, const int32 TargetIndex) = 0;
-	UE_DEPRECATED(5.4, "Use SamplePhysicsField(int32, EPositionOrigin, int32, int32) instead")
-	int32 SamplePhysicsField(int32 PositionArg, const int32 OutputType, const int32 TargetIndex)
-	{
-		return SamplePhysicsField(PositionArg, EPositionOrigin::Absolute, OutputType, TargetIndex);
-	}
 	virtual int32 DepthOfFieldFunction(int32 Depth, int32 FunctionValueIndex) = 0;
 	virtual int32 PostVolumeUserFlagTestFunction(int32 Input) = 0;
 	virtual int32 AtmosphericFogColor(int32 WorldPosition, EPositionOrigin PositionOrigin) = 0;
-	UE_DEPRECATED(5.4, "Use AtmosphericFogColor(int32, EPositionOrigin) instead")
-	int32 AtmosphericFogColor(int32 WorldPosition)
-	{
-		return AtmosphericFogColor(WorldPosition, EPositionOrigin::Absolute);
-	}
 	virtual int32 RotateScaleOffsetTexCoords(int32 TexCoordCodeIndex, int32 RotationScale, int32 Offset) = 0;
 	virtual int32 SpeedTree(int32 GeometryArg, int32 WindArg, int32 LODArg, float BillboardThreshold, bool bAccurateWindVelocities, bool bExtraBend, int32 ExtraBendArg) = 0;
 	virtual int32 EyeAdaptation() = 0;
@@ -670,23 +627,12 @@ public:
 	virtual int32 AtmosphericLightColor() = 0;
 
 	virtual int32 SkyAtmosphereLightIlluminance(int32 WorldPosition, EPositionOrigin PositionOrigin, int32 LightIndex) = 0;
-	UE_DEPRECATED(5.4, "Use SkyAtmosphereLightIlluminance(int32, EPositionOrigin, int32) instead")
-	int32 SkyAtmosphereLightIlluminance(int32 WorldPosition, int32 LightIndex)
-	{
-		return SkyAtmosphereLightIlluminance(WorldPosition, EPositionOrigin::Absolute, LightIndex);
-	}
 	virtual int32 SkyAtmosphereLightIlluminanceOnGround(int32 LightIndex) = 0;
 	virtual int32 SkyAtmosphereLightDirection(int32 LightIndex) = 0;
 	virtual int32 SkyAtmosphereLightDiskLuminance(int32 LightIndex, int32 OverrideAtmosphereLightDiscCosHalfApexAngle) = 0;
 	virtual int32 SkyAtmosphereViewLuminance(int32 WorldDirectionOverrideCodeChunk) = 0;
 	virtual int32 SkyAtmosphereAerialPerspective(int32 WorldPosition, EPositionOrigin PositionOrigin) = 0;
-	UE_DEPRECATED(5.4, "Use SkyAtmosphereAerialPerspective(int32, EPositionOrigin) instead")
-	int32 SkyAtmosphereAerialPerspective(int32 WorldPosition)
-	{
-		return SkyAtmosphereAerialPerspective(WorldPosition, EPositionOrigin::Absolute);
-	}
 	virtual int32 SkyAtmosphereDistantLightScatteredLuminance() = 0;
-
 	virtual int32 SkyLightEnvMapSample(int32 DirectionCodeChunk, int32 RoughnessCodeChunk) = 0;
 
 	virtual int32 GetCloudSampleAltitude() = 0;
diff --git a/Engine/Source/Runtime/Engine/Public/MaterialShared.h b/Engine/Source/Runtime/Engine/Public/MaterialShared.h
--- a/Engine/Source/Runtime/Engine/Public/MaterialShared.h
+++ b/Engine/Source/Runtime/Engine/Public/MaterialShared.h
@@ -2325,11 +2325,6 @@ public:
 
 	virtual bool IsPreview() const { return false; }
 
-#if WITH_EDITOR
-	UE_DEPRECATED(5.1, "Should instead check individual properties")
-	virtual bool HasMaterialAttributesConnected() const { return false; }
-#endif
-
 	ENGINE_API const FMaterialCachedExpressionData& GetCachedExpressionData() const;
 
 	/** Is the material required to be complete? Default materials, special engine materials, etc */
@@ -2442,22 +2437,10 @@ public:
 	ENGINE_API bool MaterialUsesDisplacement_RenderThread() const;
 	ENGINE_API bool MaterialUsesDisplacement_GameThread() const;
 
-	UE_DEPRECATED(5.0, "This function is deprecated. Use MaterialUsesWorldPositionOffset_GameThread() instead.")
-	inline bool UsesWorldPositionOffset_GameThread() const
-	{
-		return MaterialUsesWorldPositionOffset_GameThread();
-	}
-
 	/** Does the material use a pixel depth offset. */
 	ENGINE_API bool MaterialUsesPixelDepthOffset_RenderThread() const;
 	ENGINE_API bool MaterialUsesPixelDepthOffset_GameThread() const;
 
-	UE_DEPRECATED(5.0, "This function is deprecated. Use MaterialUsesPixelDepthOffset_RenderThread() instead.")
-	inline bool MaterialUsesPixelDepthOffset() const
-	{
-		return MaterialUsesPixelDepthOffset_RenderThread();
-	}
-
 	/** Does the material modify the mesh position. */
 	ENGINE_API bool MaterialModifiesMeshPosition_RenderThread() const;
 	ENGINE_API bool MaterialModifiesMeshPosition_GameThread() const;
@@ -2692,13 +2675,6 @@ protected:
 	virtual FString GetBaseMaterialPathName() const { return TEXT(""); }
 	virtual FString GetDebugName() const { return GetBaseMaterialPathName(); }
 
-	UE_DEPRECATED(4.26, "Parameter bInHasQualityLevelUsage is deprecated")
-	void SetQualityLevelProperties(EMaterialQualityLevel::Type InQualityLevel, bool bInHasQualityLevelUsage, ERHIFeatureLevel::Type InFeatureLevel)
-	{
-		QualityLevel = InQualityLevel;
-		FeatureLevel = InFeatureLevel;
-	}
-
 	void SetQualityLevelProperties(ERHIFeatureLevel::Type InFeatureLevel, EMaterialQualityLevel::Type InQualityLevel = EMaterialQualityLevel::Num)
 	{
 		QualityLevel = InQualityLevel;
@@ -2923,12 +2899,6 @@ public:
 		SetQualityLevelProperties(InFeatureLevel, InQualityLevel);
 	}
 
-	UE_DEPRECATED(4.26, "Parameter bInHasQualityLevelUsage is deprecated")
-	void SetMaterial(UMaterial* InMaterial, EMaterialQualityLevel::Type InQualityLevel, bool bInHasQualityLevelUsage, ERHIFeatureLevel::Type InFeatureLevel, UMaterialInstance* InInstance = nullptr)
-	{
-		SetMaterial(InMaterial, InInstance, InFeatureLevel, InQualityLevel);
-	}
-
 #if WITH_EDITOR
 	/** Returns the number of samplers used in this material, or -1 if the material does not have a valid shader map (compile error or still compiling). */
 	ENGINE_API int32 GetSamplerUsage() const;
@@ -3146,9 +3116,7 @@ protected:
 
 	ENGINE_API virtual bool HasVertexPositionOffsetConnected() const override;
 	ENGINE_API virtual bool HasPixelDepthOffsetConnected() const override;
-#if WITH_EDITOR
-	ENGINE_API virtual bool HasMaterialAttributesConnected() const override;
-#endif
+
 	/** Useful for debugging. */
 	ENGINE_API virtual FString GetBaseMaterialPathName() const override;
 	ENGINE_API virtual FString GetDebugName() const override;
diff --git a/Engine/Source/Runtime/Engine/Public/Materials/Material.h b/Engine/Source/Runtime/Engine/Public/Materials/Material.h
--- a/Engine/Source/Runtime/Engine/Public/Materials/Material.h
+++ b/Engine/Source/Runtime/Engine/Public/Materials/Material.h
@@ -1346,8 +1346,6 @@ public:
 	ENGINE_API static void AddReferencedObjects(UObject* InThis, FReferenceCollector& Collector);
 	ENGINE_API virtual bool CanBeClusterRoot() const override;
 	ENGINE_API virtual void GetAssetRegistryTags(FAssetRegistryTagsContext Context) const override;
-	UE_DEPRECATED(5.4, "Implement the version that takes FAssetRegistryTagsContext instead.")
-	ENGINE_API virtual void GetAssetRegistryTags(TArray<FAssetRegistryTag>& OutTags) const override;
 	//~ End UObject Interface
 
 	enum class EPostEditChangeEffectOnShaders
diff --git a/Engine/Source/Runtime/Engine/Public/Materials/MaterialFunctionInterface.h b/Engine/Source/Runtime/Engine/Public/Materials/MaterialFunctionInterface.h
--- a/Engine/Source/Runtime/Engine/Public/Materials/MaterialFunctionInterface.h
+++ b/Engine/Source/Runtime/Engine/Public/Materials/MaterialFunctionInterface.h
@@ -83,8 +83,7 @@ public:
 #endif
 
 	virtual void GetAssetRegistryTags(FAssetRegistryTagsContext Context) const override;
-	UE_DEPRECATED(5.4, "Implement the version that takes FAssetRegistryTagsContext instead.")
-	virtual void GetAssetRegistryTags(TArray<FAssetRegistryTag>& OutTags) const override;
+
 #if WITH_EDITOR
 	virtual void PostEditChangeProperty(FPropertyChangedEvent& PropertyChangedEvent) override;
 #endif
@@ -187,9 +186,6 @@ public:
 
 #if WITH_EDITORONLY_DATA
 	ENGINE_API TConstArrayView<TObjectPtr<UMaterialExpression>> GetExpressions() const;
-
-	UE_DEPRECATED(5.1, "Use GetExpressions()")
-	inline TConstArrayView<TObjectPtr<UMaterialExpression>> GetFunctionExpressions() const { return GetExpressions(); }
 #endif // WITH_EDITORONLY_DATA
 
 #if WITH_EDITOR
diff --git a/Engine/Source/Runtime/Engine/Public/Materials/MaterialInstance.h b/Engine/Source/Runtime/Engine/Public/Materials/MaterialInstance.h
--- a/Engine/Source/Runtime/Engine/Public/Materials/MaterialInstance.h
+++ b/Engine/Source/Runtime/Engine/Public/Materials/MaterialInstance.h
@@ -1122,9 +1122,6 @@ protected:
 	void GetTextureExpressionValues(const FMaterialResource* MaterialResource, TArray<UTexture*>& OutTextures, TArray<TArray<int32>>* OutIndices = nullptr, bool bIncludeTextureCollections = false) const;
 	void GetTextureCollectionExpressionValues(const FMaterialResource* MaterialResource, TArray<UTextureCollection*>& OutTextureCollections) const;
 
-	UE_DEPRECATED(4.26, "Calling UpdatePermutationAllocations is no longer necessary")
-	inline void UpdatePermutationAllocations(FMaterialResourceDeferredDeletionArray* ResourcesToFree = nullptr) {}
-
 #if WITH_EDITOR
 	/**
 	* Refresh parameter names using the stored reference to the expression object for the parameter.
diff --git a/Engine/Source/Runtime/Engine/Public/Materials/MaterialInterface.h b/Engine/Source/Runtime/Engine/Public/Materials/MaterialInterface.h
--- a/Engine/Source/Runtime/Engine/Public/Materials/MaterialInterface.h
+++ b/Engine/Source/Runtime/Engine/Public/Materials/MaterialInterface.h
@@ -453,8 +453,6 @@ public:
 #if WITH_EDITOR
 	ENGINE_API virtual void PostEditChangeProperty(FPropertyChangedEvent& PropertyChangedEvent) override;
 	ENGINE_API virtual void GetAssetRegistryTags(FAssetRegistryTagsContext Context) const override;
-	UE_DEPRECATED(5.4, "Implement the version that takes FAssetRegistryTagsContext instead.")
-	ENGINE_API virtual void GetAssetRegistryTags(TArray<FAssetRegistryTag>& OutTags) const override;
 #endif // WITH_EDITOR
 #if WITH_EDITOR
 	static ENGINE_API void AppendToClassSchema(FAppendToClassSchemaContext& Context);
@@ -585,16 +583,7 @@ public:
 	 * Handles updating dependent MI's and cached uniform expressions.
 	 */
 	virtual void OverrideNumericParameterDefault(EMaterialParameterType Type, const FHashedMaterialParameterInfo& ParameterInfo, const UE::Shader::FValue& Value, bool bOverride, ERHIFeatureLevel::Type FeatureLevel) PURE_VIRTUAL(UMaterialInterface::OverrideNumericParameterDefault, return;);
-	/**
-	 * DEPRECATED: Returns default value of the given parameter
-	 */
-	UE_DEPRECATED(4.19, "This function is deprecated. Use GetScalarParameterDefaultValue instead.")
-	virtual float GetScalarParameterDefault(const FHashedMaterialParameterInfo& ParameterInfo, ERHIFeatureLevel::Type FeatureLevel)
-	{
-		float Value;
-		GetScalarParameterDefaultValue(ParameterInfo, Value);
-		return Value;
-	};
+
 	/**
 	 * Checks if the material can be used with the given usage flag.  
 	 * If the flag isn't set in the editor, it will be set and the material will be recompiled with it.
diff --git a/Engine/Source/Runtime/Engine/Public/Materials/MaterialRenderProxy.h b/Engine/Source/Runtime/Engine/Public/Materials/MaterialRenderProxy.h
--- a/Engine/Source/Runtime/Engine/Public/Materials/MaterialRenderProxy.h
+++ b/Engine/Source/Runtime/Engine/Public/Materials/MaterialRenderProxy.h
@@ -116,17 +116,11 @@ public:
 	 */
 	ENGINE_API void EvaluateUniformExpressions(FRHICommandListBase& RHICmdList, FUniformExpressionCache& OutUniformExpressionCache, const FMaterialRenderContext& Context, FUniformExpressionCacheAsyncUpdater* Updater = nullptr) const;
 
-	UE_DEPRECATED(5.4, "EvaluateUniformExpressions requires an RHI command list.")
-	ENGINE_API void EvaluateUniformExpressions(FUniformExpressionCache& OutUniformExpressionCache, const FMaterialRenderContext& Context, FUniformExpressionCacheAsyncUpdater* Updater = nullptr) const;
-
 	/**
 	 * Caches uniform expressions for efficient runtime evaluation.
 	 */
 	ENGINE_API void CacheUniformExpressions(FRHICommandListBase& RHICmdList, bool bRecreateUniformBuffer);
 
-	UE_DEPRECATED(5.4, "CacheUniformExpressions now requires a command list.")
-	ENGINE_API void CacheUniformExpressions(bool bRecreateUniformBuffer);
-
 	/** Cancels an in-flight cache operation. */
 	ENGINE_API void CancelCacheUniformExpressions();
 
diff --git a/Engine/Source/Runtime/Renderer/Public/MeshMaterialShader.h b/Engine/Source/Runtime/Renderer/Public/MeshMaterialShader.h
--- a/Engine/Source/Runtime/Renderer/Public/MeshMaterialShader.h
+++ b/Engine/Source/Runtime/Renderer/Public/MeshMaterialShader.h
@@ -84,20 +84,6 @@ public:
 		const FMeshMaterialShaderElementData& ShaderElementData,
 		FMeshDrawSingleShaderBindings& ShaderBindings) const;
 
-	UE_DEPRECATED(5.4, "GetShaderBindings no longer has a DrawRenderState argument. Please update your shader type to use the new function signature to ensure it is set up correctly.")
-	void GetShaderBindings(
-		const FScene* Scene,
-		ERHIFeatureLevel::Type FeatureLevel,
-		const FPrimitiveSceneProxy* PrimitiveSceneProxy,
-		const FMaterialRenderProxy& MaterialRenderProxy,
-		const FMaterial& Material,
-		const FMeshPassProcessorRenderState& DrawRenderState,
-		const FMeshMaterialShaderElementData& ShaderElementData,
-		FMeshDrawSingleShaderBindings& ShaderBindings) const
-	{
-		checkNoEntry();
-	}
-
 	RENDERER_API void GetElementShaderBindings(
 		const FShaderMapPointerTable& PointerTable,
 		const FScene* Scene, 
