Fix pipelining of GGPUFrameTime
- Previously, RHIs wrote directly into the global GGPUFrameTime variable, and engine code read directly from it. Since there's no explicit synchronization between those threads, the high level can miss values, or get duplicates due to the timing of the reads and writes vs when frames tick on the game thread.
- Now, FRHIGPUFrameTimeHistory maintains a history of frame GPU values, which can be consumed in a loop by the higher level threads. This avoids duplicates / missing frames.
- RHIGetGPUFrameCycles is still implemented for legacy systems, and returns the latest value only.

Simplify Dynamic Resolution
- Refactored the dynamic res calculation to make use of the pipelined GPU frame timings.
- Removed "r.DynamicRes.DynamicFrameTime" and associated cvars / code paths. This was only enabled on Windows, which also never ran because neither D3D11 or D3D12 RHIs on PC marked dynamic res as supported.
- Allowing the GPU time to increase when the engine is CPU bound will result in increased input latency, and complicates frame pacing. The checks for whether the engine is CPU bound also relied on game, render and RHI thread timings, which are no longer directly representative of the CPU state of the engine, given the parallelism we have now.
- The "DynamicRenderScaling" budget scopes are still not pipelined correctly. This needs further work.

diff --git a/Engine/Config/Windows/BaseWindowsEngine.ini b/Engine/Config/Windows/BaseWindowsEngine.ini
--- a/Engine/Config/Windows/BaseWindowsEngine.ini
+++ b/Engine/Config/Windows/BaseWindowsEngine.ini
@@ -24,7 +24,6 @@ r.Nanite.Streaming.ReservedResources=1
 D3D12.Bindless.ResourceDescriptorHeapSize=32768
 D3D12.Bindless.SamplerDescriptorHeapSize=2048
 r.PSOPrecache.GlobalShaders=1
-r.DynamicRes.DynamicFrameTime=1
 r.VRS.EnableSoftware=1
 r.VRS.ContrastAdaptiveShading=1
 
diff --git a/Engine/Source/Runtime/Apple/MetalRHI/Private/MetalProfiler.cpp b/Engine/Source/Runtime/Apple/MetalRHI/Private/MetalProfiler.cpp
--- a/Engine/Source/Runtime/Apple/MetalRHI/Private/MetalProfiler.cpp
+++ b/Engine/Source/Runtime/Apple/MetalRHI/Private/MetalProfiler.cpp
@@ -141,9 +141,7 @@ void FMetalCommandBufferTimer::RecordFrame()
 	FrameEndGPUCycles = LastBufferTiming.EndTime * CyclesPerSecond;
 	
 	uint64 FrameGPUTimeCycles = uint64(CyclesPerSecond * RunningFrameTimeSeconds);
-PRAGMA_DISABLE_DEPRECATION_WARNINGS
-	FPlatformAtomics::AtomicStore_Relaxed((int32*)&GGPUFrameTime, int32(FrameGPUTimeCycles));
-PRAGMA_ENABLE_DEPRECATION_WARNINGS
+	GRHIGPUFrameTimeHistory.PushFrameCycles(1.0, RunningFrameTimeSeconds);
 	
 #if STATS
 	FPlatformAtomics::AtomicStore_Relaxed(&GMetalGPUWorkTime, FrameGPUTimeCycles);
diff --git a/Engine/Source/Runtime/D3D12RHI/Private/D3D12Submission.cpp b/Engine/Source/Runtime/D3D12RHI/Private/D3D12Submission.cpp
--- a/Engine/Source/Runtime/D3D12RHI/Private/D3D12Submission.cpp
+++ b/Engine/Source/Runtime/D3D12RHI/Private/D3D12Submission.cpp
@@ -1435,26 +1435,21 @@ void FD3D12DynamicRHI::ProcessTimestamps(FD3D12TimingArray const& TimingPerQueue
 	FD3D12Device* Device = GetAdapter().GetDevice(0);
 	double Frequency = Device->GetTimestampFrequency(ED3D12QueueType::Direct);
 
-	const double Scale32 = 1.0 / (Frequency * FPlatformTime::GetSecondsPerCycle());
 	const double Scale64 = 1.0 / (Frequency * FPlatformTime::GetSecondsPerCycle64());
 
 	// Update the global GPU frame time stats
 	SET_CYCLE_COUNTER(STAT_RHI_GPUTotalTime, FPlatformMath::TruncToInt(double(UnionBusyCycles) * Scale64));
-	 
+
 	double HardwareGPUTime = 0.0;
 	if (GetHardwareGPUFrameTime(HardwareGPUTime) && CVarGPUTimeFromTimestamps.GetValueOnAnyThread() == 0)
 	{
 		SET_CYCLE_COUNTER(STAT_RHI_GPUTotalTimeHW, HardwareGPUTime);
-	PRAGMA_DISABLE_DEPRECATION_WARNINGS
-		GGPUFrameTime = HardwareGPUTime;
-	PRAGMA_ENABLE_DEPRECATION_WARNINGS
+		GRHIGPUFrameTimeHistory.PushFrameCycles(1.0 / FPlatformTime::GetSecondsPerCycle64(), HardwareGPUTime);
 	}
 	else
 	{
 		SET_CYCLE_COUNTER(STAT_RHI_GPUTotalTimeHW, 0);
-	PRAGMA_DISABLE_DEPRECATION_WARNINGS
-		GGPUFrameTime = FPlatformMath::TruncToInt(double(UnionBusyCycles) * Scale32);
-	PRAGMA_ENABLE_DEPRECATION_WARNINGS
+		GRHIGPUFrameTimeHistory.PushFrameCycles(Frequency, UnionBusyCycles);
 	}
 
 	for (auto const& Current : TimingPerQueue)
diff --git a/Engine/Source/Runtime/Engine/Private/DynamicResolution.cpp b/Engine/Source/Runtime/Engine/Private/DynamicResolution.cpp
--- a/Engine/Source/Runtime/Engine/Private/DynamicResolution.cpp
+++ b/Engine/Source/Runtime/Engine/Private/DynamicResolution.cpp
@@ -43,29 +43,6 @@ static TAutoConsoleVariable<int32> CVarUseGameThreadCriticalPath(
 	TEXT("Whether to use game thread critical path time when determining whether game is CPU bound."),
 	ECVF_RenderThreadSafe | ECVF_Default);
 
-static TAutoConsoleVariable<int32> CVarDynamicFrameTimeEnable(
-	TEXT("r.DynamicRes.DynamicFrameTime"), 0,
-	TEXT("Whether the r.DynamicRes.FrameTimeBudget should automatically increases when frame rate is bound by CPU."),
-	ECVF_RenderThreadSafe | ECVF_Default);
-
-static TAutoConsoleVariable<float> CVarDynamicFrameTimeErrorMarginPercent(
-	TEXT("r.DynamicRes.DynamicFrameTime.ErrorMarginPercent"), 10.0f,
-	TEXT("How much headroom should be left between CPU and GPU."),
-	ECVF_RenderThreadSafe | ECVF_Default);
-
-static TAutoConsoleVariable<int32> CVarDynamicFrameTimeTrack(
-	TEXT("r.DynamicRes.DynamicFrameTime.Track"), 1,
-	TEXT("What to track to control the budget\n")
-	TEXT(" 0: Frametime (but can create feedback loop when GPU bound and VSync);")
-	TEXT(" 1: Threads time (default);"),
-	ECVF_RenderThreadSafe | ECVF_Default);
-
-static TAutoConsoleVariable<float> CVarDynamicFrameTimeRoundUpToVsync(
-	TEXT("r.DynamicRes.DynamicFrameTime.RoundUpToVSyncError"), 10.f,
-	TEXT("Error to use to round up the dynamic frame time to vsync boundaries (default=10%)."),
-	ECVF_RenderThreadSafe | ECVF_Default);
-
-
 static TAutoConsoleVariable<float> CVarTargetedGPUHeadRoomPercentage(
 	TEXT("r.DynamicRes.TargetedGPUHeadRoomPercentage"),
 	10.0f,
@@ -130,7 +107,12 @@ static TAutoConsoleVariable<int32> CVarUpperBoundQuantization(
 	TEXT("Only recommended for use with the transient allocator (on supported platforms) with a large transient texture cache (e.g RHI.TransientAllocator.TextureCacheSize=512)"),
 	ECVF_RenderThreadSafe | ECVF_Default);
 
-
+static TAutoConsoleVariable<bool> CVarUseCPUTimeLogic(
+	TEXT("r.DynamicRes.UseCPUTimeLogic"),
+	false,
+	TEXT("When true, enables legacy logic that checks whether the engine is game or render thread bound, and if so, allows the GPU to consider more frame history.\n")
+	TEXT("When false, dynamic resolution is driven only from GPU time, and multiple over-budget GPU frames will cause a sooner drop in resolution."),
+	ECVF_RenderThreadSafe | ECVF_Default);
 
 DynamicRenderScaling::FHeuristicSettings GetPrimaryDynamicResolutionSettings()
 {
@@ -159,9 +141,6 @@ static float TimeStampQueryResultToMilliSeconds(uint64 TimestampResult)
 FDynamicResolutionHeuristicProxy::FDynamicResolutionHeuristicProxy()
 {
 	check(IsInGameThread());
-
-	FrameCounter = 0;
-
 	ResetInternal();
 }
 
@@ -173,7 +152,6 @@ FDynamicResolutionHeuristicProxy::~FDynamicResolutionHeuristicProxy()
 void FDynamicResolutionHeuristicProxy::Reset_RenderThread()
 {
 	check(IsInRenderingThread());
-
 	ResetInternal();
 }
 
@@ -187,8 +165,6 @@ void FDynamicResolutionHeuristicProxy::ResetInternal()
 	NumberOfFramesSinceScreenPercentageChange = 0;
 	CurrentFrameResolutionFractions.SetAll(1.0f);
 	CurrentFrameMaxResolutionFractions.SetAll(1.0f);
-	DynamicFrameTimeBudgetMs = 0.0f;
-	DynamicFrameTimeVSyncFactor = 0;
 
 	TemporalUpscalerMinResolutionFraction = ISceneViewFamilyScreenPercentage::kMinResolutionFraction;
 	TemporalUpscalerMaxResolutionFraction = ISceneViewFamilyScreenPercentage::kMaxResolutionFraction;
@@ -197,65 +173,45 @@ void FDynamicResolutionHeuristicProxy::ResetInternal()
 	IgnoreFrameRemainingCount = 1;
 }
 
-uint64 FDynamicResolutionHeuristicProxy::CreateNewPreviousFrameTimings_RenderThread(
-	float FrameTimeMs, float GameThreadTimeMs, float RenderThreadTimeMs, float RHIThreadTimeMs)
+void FDynamicResolutionHeuristicProxy::CreateNewPreviousFrameTimings_RenderThread(float GameThreadTimeMs, float RenderThreadTimeMs, float TotalFrameGPUBusyTimeMs)
 {
 	check(IsInRenderingThread());
+	check(TotalFrameGPUBusyTimeMs >= 0.0f);
 
 	// Early return if want to ignore frames.
 	if (IgnoreFrameRemainingCount > 0)
 	{
 		IgnoreFrameRemainingCount--;
-		return FDynamicResolutionHeuristicProxy::kInvalidEntryId;
+		return;
 	}
 
 	ResizeHistoryIfNeeded();
 
 	// Update history state.
-	{
-		int32 NewHistoryEntryIndex = (PreviousFrameIndex + 1) % History.Num();
-		History[NewHistoryEntryIndex] = FrameHistoryEntry();
-		History[NewHistoryEntryIndex].ResolutionFractions = CurrentFrameResolutionFractions;
-		History[NewHistoryEntryIndex].FrameTimeMs = FrameTimeMs;
-		History[NewHistoryEntryIndex].GameThreadTimeMs = GameThreadTimeMs;
-		History[NewHistoryEntryIndex].RenderThreadTimeMs = RenderThreadTimeMs;
-		History[NewHistoryEntryIndex].RHIThreadTimeMs = RHIThreadTimeMs;
-		PreviousFrameIndex = NewHistoryEntryIndex;
-		HistorySize = FMath::Min(HistorySize + 1, History.Num());
-
-		for (TLinkedList<DynamicRenderScaling::FBudget*>::TIterator BudgetIt(DynamicRenderScaling::FBudget::GetGlobalList()); BudgetIt; BudgetIt.Next())
-		{
-			const DynamicRenderScaling::FBudget& Budget = **BudgetIt;
-			BudgetHistorySizes[Budget] = FMath::Min(BudgetHistorySizes[Budget] + 1, History.Num());
-		}
-	}
+	int32 NewHistoryEntryIndex = (PreviousFrameIndex + 1) % History.Num();
+	FrameHistoryEntry& Entry = History[NewHistoryEntryIndex];
+	PreviousFrameIndex = NewHistoryEntryIndex;
+
+	Entry = FrameHistoryEntry();
+	Entry.ResolutionFractions = CurrentFrameResolutionFractions;
+	Entry.GameThreadTimeMs = GameThreadTimeMs;
+	Entry.RenderThreadTimeMs = RenderThreadTimeMs;
+	Entry.TotalFrameGPUBusyTimeMs = TotalFrameGPUBusyTimeMs;
+		
+	HistorySize = FMath::Min(HistorySize + 1, History.Num());
 
-	return ++FrameCounter;
-}
+	const DynamicRenderScaling::TMap<uint64>& LatestTimings = DynamicRenderScaling::GetLatestTimings();
 
-void FDynamicResolutionHeuristicProxy::CommitPreviousFrameGPUTimings_RenderThread(
-	uint64 HistoryFrameId,
-	float TotalFrameGPUBusyTimeMs,
-	const DynamicRenderScaling::TMap<float>& BudgetTimingMs)
-{
-	check(TotalFrameGPUBusyTimeMs >= 0.0f);
+	Entry.BudgetTimingMs[GDynamicPrimaryResolutionFraction] = TotalFrameGPUBusyTimeMs;
 
-	// Conveniently early return if invalid history frame id.
-	if (HistoryFrameId == FDynamicResolutionHeuristicProxy::kInvalidEntryId ||
-		HistoryFrameId <= (FrameCounter - HistorySize))
+	for (const DynamicRenderScaling::FBudget* Budget : *DynamicRenderScaling::FBudget::GetGlobalList())
 	{
-		return;
+		Entry.BudgetTimingMs[*Budget] = TimeStampQueryResultToMilliSeconds(LatestTimings[*Budget]);
+		BudgetHistorySizes[*Budget] = FMath::Min(BudgetHistorySizes[*Budget] + 1, History.Num());
 	}
-
-	int32 EntryId = (History.Num() + PreviousFrameIndex - int32(FrameCounter - HistoryFrameId)) % History.Num();
-
-	// Make sure not overwriting.
-	check(History[EntryId].TotalFrameGPUBusyTimeMs == -1.0f);
-
-	History[EntryId].TotalFrameGPUBusyTimeMs = TotalFrameGPUBusyTimeMs;
-	History[EntryId].BudgetTimingMs = BudgetTimingMs;
 }
 
+
 void FDynamicResolutionHeuristicProxy::RefreshCurrentFrameResolutionFraction_RenderThread()
 {
 	// Global constants.
@@ -263,95 +219,7 @@ void FDynamicResolutionHeuristicProxy::RefreshCurrentFrameResolutionFraction_Ren
 	const int32 MaxConsecutiveOverBudgetGPUFrameCount = FMath::Max(CVarMaxConsecutiveOverBudgetGPUFrameCount.GetValueOnRenderThread(), 2);
 
 	const bool bCanChangeResolution = NumberOfFramesSinceScreenPercentageChange >= CVarFrameChangePeriod.GetValueOnRenderThread();
-
-	float MinGlobalFrameTime = 0.0f;
-	if (CVarDynamicFrameTimeEnable.GetValueOnRenderThread())
-	{
-		static const IConsoleVariable* VSyncCVar = IConsoleManager::Get().FindConsoleVariable(TEXT("r.VSync"));
-		check(VSyncCVar);
-
-		const int32 FrameTimeTrack = CVarDynamicFrameTimeTrack.GetValueOnRenderThread();
-
-		// Find out the frame time that should be used.
-		TArray<float, TInlineAllocator<16>> SortedFrameTime;
-		SortedFrameTime.Reserve(History.Num());
-		for (int32 BrowsingFrameId = 0; BrowsingFrameId < History.Num(); BrowsingFrameId++)
-		{
-			const FrameHistoryEntry& FrameEntry = History[(History.Num() + PreviousFrameIndex - BrowsingFrameId) % History.Num()];
-			float FrameTimeMs = 0.0f;
-			if (FrameTimeTrack == 0)
-			{
-				FrameTimeMs = FrameEntry.FrameTimeMs;
-			}
-			else
-			{
-				FrameTimeMs = FMath::Max(FrameTimeMs, FrameEntry.GameThreadTimeMs);
-				FrameTimeMs = FMath::Max(FrameTimeMs, FrameEntry.RenderThreadTimeMs);
-				FrameTimeMs = FMath::Max(FrameTimeMs, FrameEntry.RHIThreadTimeMs);
-			}
-
-
-			// Ignores invalid frame time yet. 
-			if (FrameTimeMs < 0)
-			{
-				continue;
-			}
-
-			SortedFrameTime.Add(FrameTimeMs);
-		}
-
-		if (SortedFrameTime.Num() > 0)
-		{
-			SortedFrameTime.Sort([](const float& A, const float& B) {
-				return  A < B;
-			});
-			float MedianFrameTime = SortedFrameTime[SortedFrameTime.Num() / 2];
-
-			// When CPU bound and vsync is enabled, the GPU may end up underused due to vsync wait, so we can crank up the available GPU time a bit more.
-			const float RoundUpToVSyncError = CVarDynamicFrameTimeRoundUpToVsync.GetValueOnRenderThread();
-			if (RoundUpToVSyncError > 0.0f && VSyncCVar->GetInt())
-			{
-				const float VSyncFrameTime = 1000.0f / float(FPlatformMisc::GetMaxRefreshRate());
-				float VSyncRateChangeMultiplier = 1.0f + RoundUpToVSyncError / 100.f;
-
-				int32 CurrentVSyncFactor = FMath::CeilToInt(MedianFrameTime / VSyncFrameTime);
-
-				float PreviousSyncedFrameTime = DynamicFrameTimeVSyncFactor * VSyncFrameTime;
-
-				if (DynamicFrameTimeVSyncFactor == 0)
-				{
-					DynamicFrameTimeVSyncFactor = CurrentVSyncFactor;
-				}
-				else if (CurrentVSyncFactor > DynamicFrameTimeVSyncFactor)
-				{
-					if (MedianFrameTime > PreviousSyncedFrameTime * VSyncRateChangeMultiplier)
-					{
-						DynamicFrameTimeVSyncFactor = CurrentVSyncFactor;
-					}
-				}
-				else if (CurrentVSyncFactor < DynamicFrameTimeVSyncFactor)
-				{
-					if (MedianFrameTime < PreviousSyncedFrameTime / VSyncRateChangeMultiplier)
-					{
-						DynamicFrameTimeVSyncFactor = CurrentVSyncFactor;
-					}
-				}
-
-				MedianFrameTime = DynamicFrameTimeVSyncFactor * VSyncFrameTime;
-			}
-			else
-			{
-				DynamicFrameTimeVSyncFactor = 0;
-			}
-
-			const float MaxWithFrameFraction = (CVarDynamicFrameTimeErrorMarginPercent.GetValueOnRenderThread() / 100.0f);
-
-			float TargetFrameTime = MedianFrameTime * (1.0f - MaxWithFrameFraction);
-
-			DynamicFrameTimeBudgetMs = FMath::Lerp(TargetFrameTime, DynamicFrameTimeBudgetMs, FrameWeightExponent);
-			MinGlobalFrameTime = DynamicFrameTimeBudgetMs;
-		}
-	}
+	const bool bUseCPUTimeLogic = CVarUseCPUTimeLogic.GetValueOnRenderThread();
 
 	// New ResolutionFraction to use for this frame.
 	DynamicRenderScaling::TMap<float> NewFrameResolutionFractions = CurrentFrameResolutionFractions;
@@ -366,24 +234,13 @@ void FDynamicResolutionHeuristicProxy::RefreshCurrentFrameResolutionFraction_Ren
 
 		const int32 BudgetHistorySize = FMath::Min(HistorySize, BudgetHistorySizes[Budget]);
 
-		bool bEarlyReturn = BudgetHistorySize == 0;
-
-		float BudgetBudgetMs = BudgetSettings.BudgetMs;
-		if (Budget == GDynamicPrimaryResolutionFraction)
-		{
-			BudgetBudgetMs = FMath::Max(BudgetBudgetMs, MinGlobalFrameTime);
-		}
-		else if (!BudgetSettings.IsEnabled())
-		{
-			bEarlyReturn = true;
-		}
-
 		// Early returns if not enough data to work with.
-		if (bEarlyReturn)
+		if (BudgetHistorySize == 0 || (Budget != GDynamicPrimaryResolutionFraction && !BudgetSettings.IsEnabled()))
 		{
 			continue;
 		}
 
+		float BudgetBudgetMs = BudgetSettings.BudgetMs;
 		const float BudgetTargetMs = BudgetSettings.GetTargetedMs(BudgetBudgetMs);
 
 		float NewFrameResolutionFraction = 0.0f;
@@ -402,14 +259,13 @@ void FDynamicResolutionHeuristicProxy::RefreshCurrentFrameResolutionFraction_Ren
 
 		// Number of frame browsed.
 		int32 FrameCount = 0;
-		int32 BrowsingFrameId = 0;
-		for (; BrowsingFrameId < BudgetHistorySize && bCarryOnBrowsingFrameHistory; BrowsingFrameId++)
+		for (int32 BrowsingFrameId = 0; BrowsingFrameId < BudgetHistorySize && bCarryOnBrowsingFrameHistory; BrowsingFrameId++)
 		{
-			const FrameHistoryEntry& NextFrameEntry = GetPreviousFrameEntry(BrowsingFrameId - 1);
-			const FrameHistoryEntry& FrameEntry = GetPreviousFrameEntry(BrowsingFrameId + 0);
-			const FrameHistoryEntry& PreviousFrameEntry = GetPreviousFrameEntry(BrowsingFrameId + 1);
+			const FrameHistoryEntry& FrameEntry = GetPreviousFrameEntry(BrowsingFrameId);
 
-			const float TotalFrameGPUBusyTimeMs = Budget == GDynamicPrimaryResolutionFraction ? FrameEntry.TotalFrameGPUBusyTimeMs : FrameEntry.BudgetTimingMs[Budget];
+			const float TotalFrameGPUBusyTimeMs = Budget == GDynamicPrimaryResolutionFraction
+				? FrameEntry.TotalFrameGPUBusyTimeMs
+				: FrameEntry.BudgetTimingMs[Budget];
 
 			// Ignores frames that does not have any GPU timing yet. 
 			if (TotalFrameGPUBusyTimeMs < 0)
@@ -428,7 +284,7 @@ void FDynamicResolutionHeuristicProxy::RefreshCurrentFrameResolutionFraction_Ren
 			const bool bIsRenderThreadBound = FrameEntry.RenderThreadTimeMs > GDynamicPrimaryResolutionFraction.GetSettings().GetTargetedMs(GDynamicPrimaryResolutionFraction.GetSettings().BudgetMs);
 
 			// Whether the frame is CPU bound.
-			const bool bIsCPUBound = bIsGameThreadBound || bIsRenderThreadBound;
+			const bool bIsCPUBound = bUseCPUTimeLogic && (bIsGameThreadBound || bIsRenderThreadBound);
 
 			// Whether GPU is over budget, when not CPU bound.
 			const bool bHasOverBudgetGPU = !bIsCPUBound && TotalFrameGPUBusyTimeMs > BudgetBudgetMs;
@@ -465,9 +321,7 @@ void FDynamicResolutionHeuristicProxy::RefreshCurrentFrameResolutionFraction_Ren
 				// TODO: If we have RHI guarantee of frame timing association, we could make an estimation of B
 				// At resolution change that happen every N frames, amortized over time and scaled down as the
 				// standard variation of the GPU timing over non resolution changing frames increases.
-				SuggestedResolutionFraction = (
-					BudgetSettings.EstimateResolutionFactor(BudgetTargetMs, TotalFrameGPUBusyTimeMs)
-					* FrameEntry.ResolutionFractions[Budget]);
+				SuggestedResolutionFraction = BudgetSettings.EstimateResolutionFactor(BudgetTargetMs, TotalFrameGPUBusyTimeMs) * FrameEntry.ResolutionFractions[Budget];
 			}
 
 			NewFrameResolutionFraction += SuggestedResolutionFraction * Weight;
@@ -713,7 +567,7 @@ public:
 		Heuristic.Reset_RenderThread();
 	}
 
-	void BeginFrame(FRHICommandList& RHICmdList, float PrevFrameTimeMs, float PrevGameThreadTimeMs)
+	void BeginFrame(float PrevGameThreadTimeMs)
 	{
 		check(IsInRenderingThread());
 		ensure(GRHISupportsFrameCyclesBubblesRemoval);
@@ -734,36 +588,18 @@ public:
 			DynamicRenderScaling::BeginFrame(bIsBudgetEnabled);
 		}
 
-		// Query render thread time Ms.
 		float PrevRenderThreadTimeMs = FPlatformTime::ToMilliseconds(GRenderThreadTime);
-		float PrevRHIThreadTimeMs    = FPlatformTime::ToMilliseconds(GRHIThreadTime);
-		float PrevFrameGPUTimeMs     = FPlatformTime::ToMilliseconds(RHIGetGPUFrameCycles());
 
-		uint64 HistoryEntryId = Heuristic.CreateNewPreviousFrameTimings_RenderThread(
-			PrevFrameTimeMs,
-			PrevGameThreadTimeMs,
-			PrevRenderThreadTimeMs,
-			PrevRHIThreadTimeMs
-		);
-
-		const DynamicRenderScaling::TMap<uint64>& LatestTimings = DynamicRenderScaling::GetLatestTimings();
-
-		DynamicRenderScaling::TMap<float> BudgetTimingMs;
-		for (TLinkedList<DynamicRenderScaling::FBudget*>::TIterator BudgetIt(DynamicRenderScaling::FBudget::GetGlobalList()); BudgetIt; BudgetIt.Next())
+		uint64 GPUFrameTimeCycles64;
+		while (GPUFrameTimeState.PopFrameCycles(GPUFrameTimeCycles64) != FRHIGPUFrameTimeHistory::EResult::Empty)
 		{
-			const DynamicRenderScaling::FBudget& Budget = **BudgetIt;
-			BudgetTimingMs[Budget] = TimeStampQueryResultToMilliSeconds(LatestTimings[Budget]);
+			float PrevFrameGPUTimeMs = FPlatformTime::ToMilliseconds64(GPUFrameTimeCycles64);
+			Heuristic.CreateNewPreviousFrameTimings_RenderThread(PrevGameThreadTimeMs, PrevRenderThreadTimeMs, PrevFrameGPUTimeMs);
+			Heuristic.RefreshCurrentFrameResolutionFraction_RenderThread();
 		}
-		BudgetTimingMs[GDynamicPrimaryResolutionFraction] = PrevFrameGPUTimeMs;
-
-		Heuristic.CommitPreviousFrameGPUTimings_RenderThread(HistoryEntryId,
-			/* TotalFrameGPUBusyTimeMs = */ PrevFrameGPUTimeMs,
-			/* BudgetTimingMs = */ BudgetTimingMs);
-
-		Heuristic.RefreshCurrentFrameResolutionFraction_RenderThread();
 	}
 
-	void ProcessEvent(FRHICommandList& RHICmdList, EDynamicResolutionStateEvent Event)
+	void ProcessEvent(EDynamicResolutionStateEvent Event)
 	{
 		check(IsInRenderingThread());
 
@@ -781,6 +617,8 @@ public:
 
 	// Heuristic's proxy.
 	FDynamicResolutionHeuristicProxy Heuristic;
+
+	FRHIGPUFrameTimeHistory::FState GPUFrameTimeState;
 };
 
 
@@ -858,13 +696,12 @@ public:
 		check(IsInGameThread());
 
 		// Deletes the proxy on the rendering thread to make sure we don't delete before a recommand using it has finished.
-		FDefaultDynamicResolutionStateProxy* P = Proxy;
 		ENQUEUE_RENDER_COMMAND(DeleteDynamicResolutionProxy)(
-			[P](class FRHICommandList&)
-			{
-				P->Finalize();
-				delete P;
-			});
+			[P = Proxy](class FRHICommandList&)
+		{
+			P->Finalize();
+			delete P;
+		});
 	}
 
 
@@ -884,12 +721,12 @@ public:
 	virtual void ResetHistory() override
 	{
 		check(IsInGameThread());
-		FDefaultDynamicResolutionStateProxy* P = Proxy;
+
 		ENQUEUE_RENDER_COMMAND(DynamicResolutionResetHistory)(
-			[P](class FRHICommandList&)
-			{
-				P->Reset();
-			});
+			[P = Proxy](class FRHICommandList&)
+		{
+			P->Reset();
+		});
 
 	}
 
@@ -946,8 +783,6 @@ public:
 
 		if (Event == EDynamicResolutionStateEvent::BeginFrame)
 		{
-			// Query game thread time in milliseconds.
-			float PrevFrameTimeMs = (FApp::GetCurrentTime() - FApp::GetLastTime()) * 1000.0f;
 			float PrevGameThreadTimeMs = FPlatformTime::ToMilliseconds(GGameThreadTime);
 
 			if (CVarUseGameThreadCriticalPath.GetValueOnAnyThread())
@@ -955,21 +790,19 @@ public:
 				PrevGameThreadTimeMs = FPlatformTime::ToMilliseconds(GGameThreadTimeCriticalPath);
 			}
 
-			FDefaultDynamicResolutionStateProxy* P = Proxy;
 			ENQUEUE_RENDER_COMMAND(DynamicResolutionBeginFrame)(
-				[PrevFrameTimeMs, PrevGameThreadTimeMs, P](class FRHICommandList& RHICmdList)
+				[P = Proxy, PrevGameThreadTimeMs](FRHICommandList&)
 			{
-				P->BeginFrame(RHICmdList, PrevFrameTimeMs, PrevGameThreadTimeMs);
+				P->BeginFrame(PrevGameThreadTimeMs);
 			});
 		}
 		else
 		{
 			// Forward event to render thread.
-			FDefaultDynamicResolutionStateProxy* P = Proxy;
 			ENQUEUE_RENDER_COMMAND(DynamicResolutionBeginFrame)(
-				[P, Event](class FRHICommandList& RHICmdList)
+				[P = Proxy, Event](FRHICommandList&)
 			{
-				P->ProcessEvent(RHICmdList, Event);
+				P->ProcessEvent(Event);
 			});
 
 			if (Event == EDynamicResolutionStateEvent::EndFrame)
diff --git a/Engine/Source/Runtime/Engine/Public/DynamicResolutionProxy.h b/Engine/Source/Runtime/Engine/Public/DynamicResolutionProxy.h
--- a/Engine/Source/Runtime/Engine/Public/DynamicResolutionProxy.h
+++ b/Engine/Source/Runtime/Engine/Public/DynamicResolutionProxy.h
@@ -11,26 +11,14 @@
 class FDynamicResolutionHeuristicProxy
 {
 public:
-	static constexpr uint64 kInvalidEntryId = ~uint64(0);
-
 	ENGINE_API FDynamicResolutionHeuristicProxy();
 	ENGINE_API ~FDynamicResolutionHeuristicProxy();
 
 	/** Resets the proxy. */
 	ENGINE_API void Reset_RenderThread();
 
-	/** Create a new previous frame and feeds its timings and returns history's unique id. */
-	ENGINE_API uint64 CreateNewPreviousFrameTimings_RenderThread(float FrameTime, float GameThreadTimeMs, float RenderThreadTimeMs, float RHIThreadTime);
-	FORCEINLINE uint64 CreateNewPreviousFrameTimings_RenderThread(float GameThreadTimeMs, float RenderThreadTimeMs)
-	{
-		return CreateNewPreviousFrameTimings_RenderThread(/* FrameTime = */ 0.0f, GameThreadTimeMs, RenderThreadTimeMs, /* RHIThreadTime = */ 0.0f);
-	}
-
-	/** Commit GPU busy times */
-	ENGINE_API void CommitPreviousFrameGPUTimings_RenderThread(
-		uint64 HistoryFrameId,
-		float TotalFrameGPUBusyTimeMs,
-		const DynamicRenderScaling::TMap<float>& BudgetTimingMs);
+	/** Create a new previous frame and feeds its timings. */
+	void CreateNewPreviousFrameTimings_RenderThread(float GameThreadTimeMs, float RenderThreadTimeMs, float TotalFrameGPUBusyTimeMs);
 
 	/** Refresh resolution fraction's from history. */
 	ENGINE_API void RefreshCurrentFrameResolutionFraction_RenderThread();
@@ -61,11 +49,8 @@ public:
 private:
 	struct FrameHistoryEntry
 	{
-		// Thread timings in milliseconds.
-		float FrameTimeMs = -1.0f;
 		float GameThreadTimeMs = -1.0f;
 		float RenderThreadTimeMs = -1.0f;
-		float RHIThreadTimeMs = -1.0f;
 
 		// Total GPU busy time for the entire frame in milliseconds.
 		float TotalFrameGPUBusyTimeMs = -1.0f;
@@ -76,14 +61,14 @@ private:
 		// The resolution fraction the frame was rendered with.
 		DynamicRenderScaling::TMap<float> ResolutionFractions;
 
-		inline FrameHistoryEntry()
+		FrameHistoryEntry()
 		{
 			ResolutionFractions.SetAll(1.0f);
 			BudgetTimingMs.SetAll(-1.0f);
 		}
 
 		// Returns whether GPU timings have landed.
-		inline bool HasGPUTimings() const
+		bool HasGPUTimings() const
 		{
 			return TotalFrameGPUBusyTimeMs >= 0.0f;
 		}
@@ -95,12 +80,6 @@ private:
 	int32 PreviousFrameIndex;
 	int32 HistorySize;
 
-	// Exponential average of the CPU frame time.
-	float DynamicFrameTimeBudgetMs;
-
-	// Multiple of the VSync used.
-	int32 DynamicFrameTimeVSyncFactor = 0.0f;
-
 	// Counts the number of frame since the last screen percentage change.
 	int32 NumberOfFramesSinceScreenPercentageChange;
 
@@ -112,14 +91,11 @@ private:
 	DynamicRenderScaling::TMap<float> CurrentFrameMaxResolutionFractions;
 	DynamicRenderScaling::TMap<int32> BudgetHistorySizes;
 
-	// Frame counter to allocate unique ID for CommitPreviousFrameGPUTimings_RenderThread().
-	uint64 FrameCounter;
-
 	// Minimum and maximum resolution fractions supported by the main view family's third-party temporal upscaler.
 	float TemporalUpscalerMinResolutionFraction;
 	float TemporalUpscalerMaxResolutionFraction;
 
-	inline const FrameHistoryEntry& GetPreviousFrameEntry(int32 BrowsingFrameId) const
+	const FrameHistoryEntry& GetPreviousFrameEntry(int32 BrowsingFrameId) const
 	{
 		if (BrowsingFrameId < 0 || BrowsingFrameId >= HistorySize)
 		{
@@ -139,4 +115,3 @@ private:
 
 	ENGINE_API void ResizeHistoryIfNeeded();
 };
-
diff --git a/Engine/Source/Runtime/OpenGLDrv/Private/OpenGLDrv.cpp b/Engine/Source/Runtime/OpenGLDrv/Private/OpenGLDrv.cpp
--- a/Engine/Source/Runtime/OpenGLDrv/Private/OpenGLDrv.cpp
+++ b/Engine/Source/Runtime/OpenGLDrv/Private/OpenGLDrv.cpp
@@ -61,11 +61,6 @@ void FOpenGLDynamicRHI::RHISetExternalGPUTime(uint64 InExternalGPUTime)
 #if RHI_NEW_GPU_PROFILER
 	Profiler.ExternalGPUTime = InExternalGPUTime;
 #else
-
-PRAGMA_DISABLE_DEPRECATION_WARNINGS
-	GGPUFrameTime = InExternalGPUTime;
-PRAGMA_ENABLE_DEPRECATION_WARNINGS
-
 	GPUProfilingData->ExternalGPUTime = InExternalGPUTime;
 #endif
 }
@@ -329,36 +324,37 @@ void FOpenGLGPUProfiler::EndFrame()
 	{
 		uint64 GPUTiming = FrameTiming.GetTiming();
 		uint64 GPUFreq = FrameTiming.GetTimingFrequency();
-	PRAGMA_DISABLE_DEPRECATION_WARNINGS
-		GGPUFrameTime = FMath::TruncToInt( double(GPUTiming) / double(GPUFreq) / FPlatformTime::GetSecondsPerCycle() );
-	PRAGMA_ENABLE_DEPRECATION_WARNINGS
+		GRHIGPUFrameTimeHistory.PushFrameCycles(GPUFreq, GPUTiming);
 	}
 	else if (FOpenGLDisjointTimeStampQuery::IsSupported())
 	{
-		static uint32 GLastGPUFrameTime = 0;
+		static uint64 GLastGPUTiming = 0;
+		static uint64 GLastGPUFreq = 0;
+
 		uint64 GPUTiming = 0;
 		uint64 GPUFreq = FOpenGLDisjointTimeStampQuery::GetTimingFrequency();
 		int OldestQueryIndex = (CurrentGPUFrameQueryIndex + 1) % MAX_GPUFRAMEQUERIES;
-		if ( DisjointGPUFrameTimeQuery[OldestQueryIndex].IsResultValid() && DisjointGPUFrameTimeQuery[OldestQueryIndex].GetResult(&GPUTiming) )
+		if (DisjointGPUFrameTimeQuery[OldestQueryIndex].IsResultValid() && DisjointGPUFrameTimeQuery[OldestQueryIndex].GetResult(&GPUTiming))
+		{
+			GRHIGPUFrameTimeHistory.PushFrameCycles(GPUFreq, GPUTiming);
+
+			GLastGPUTiming = GPUTiming;
+			GLastGPUFreq = GPUFreq;
+		}
+		else if (GLastGPUFreq > 0)
 		{
-		PRAGMA_DISABLE_DEPRECATION_WARNINGS
-			GGPUFrameTime = FMath::TruncToInt( double(GPUTiming) / double(GPUFreq) / FPlatformTime::GetSecondsPerCycle() );
-			GLastGPUFrameTime = GGPUFrameTime;
-		PRAGMA_ENABLE_DEPRECATION_WARNINGS
+			// Keep the timing of the last frame if the query was disjoint (e.g. GPU frequency changed and the result is undefined)
+			GRHIGPUFrameTimeHistory.PushFrameCycles(GLastGPUFreq, GLastGPUTiming);
 		}
 		else
 		{
-			// Keep the timing of the last frame if the query was disjoint (e.g. GPU frequency changed and the result is undefined)
-		PRAGMA_DISABLE_DEPRECATION_WARNINGS
-			GGPUFrameTime = GLastGPUFrameTime;
-		PRAGMA_ENABLE_DEPRECATION_WARNINGS
+			// No valid timing data
+			GRHIGPUFrameTimeHistory.PushFrameCycles(1, 0);
 		}
 	}
 	else
 	{
-	PRAGMA_DISABLE_DEPRECATION_WARNINGS
-		GGPUFrameTime = ExternalGPUTime;
-	PRAGMA_ENABLE_DEPRECATION_WARNINGS
+		GRHIGPUFrameTimeHistory.PushFrameCycles(1.0 / FPlatformTime::GetSecondsPerCycle64(), ExternalGPUTime);
 	}
 
 	// if we have a frame open, close it now.
diff --git a/Engine/Source/Runtime/RHI/Private/GPUProfiler.cpp b/Engine/Source/Runtime/RHI/Private/GPUProfiler.cpp
--- a/Engine/Source/Runtime/RHI/Private/GPUProfiler.cpp
+++ b/Engine/Source/Runtime/RHI/Private/GPUProfiler.cpp
@@ -2404,9 +2404,7 @@ namespace UE::RHI::GPUProfiler
 			if (MaxQueueBusyCycles.IsSet())
 			{
 				// Set the total GPU time stat according to the value directly provided by the platform RHI
-			PRAGMA_DISABLE_DEPRECATION_WARNINGS
-				GGPUFrameTime = FPlatformMath::TruncToInt(FPlatformTime::ToSeconds64(*MaxQueueBusyCycles) / FPlatformTime::GetSecondsPerCycle());
-			PRAGMA_ENABLE_DEPRECATION_WARNINGS
+				GRHIGPUFrameTimeHistory.PushFrameCycles(1.0 / FPlatformTime::GetSecondsPerCycle64(), *MaxQueueBusyCycles);
 			}
 			else
 			{
@@ -2418,10 +2416,8 @@ namespace UE::RHI::GPUProfiler
 				}
 				uint64 WholeFrameUnion = FTimestampStream::ComputeUnion(StreamPointers);
 
-				// Update the global GPU frame time stats - need to convert to Cycles32 rather than Cycles64.
-			PRAGMA_DISABLE_DEPRECATION_WARNINGS
-				GGPUFrameTime = FPlatformMath::TruncToInt(FPlatformTime::ToSeconds64(WholeFrameUnion) / FPlatformTime::GetSecondsPerCycle());
-			PRAGMA_ENABLE_DEPRECATION_WARNINGS
+				// Update the global GPU frame time stats
+				GRHIGPUFrameTimeHistory.PushFrameCycles(1.0 / FPlatformTime::GetSecondsPerCycle64(), WholeFrameUnion);
 			}
 
 			// @todo set global csv GPU time
@@ -2492,4 +2488,61 @@ namespace UE::RHI::GPUProfiler
 
 #endif // RHI_NEW_GPU_PROFILER
 
+RHI_API FRHIGPUFrameTimeHistory GRHIGPUFrameTimeHistory;
+
+FRHIGPUFrameTimeHistory::EResult FRHIGPUFrameTimeHistory::FState::PopFrameCycles(uint64& OutCycles64)
+{
+	return GRHIGPUFrameTimeHistory.PopFrameCycles(*this, OutCycles64);
+}
+
+FRHIGPUFrameTimeHistory::EResult FRHIGPUFrameTimeHistory::PopFrameCycles(FState& State, uint64& OutCycles64)
+{
+	FScopeLock Lock(&CS);
+
+	if (State.NextIndex == NextIndex)
+	{
+		OutCycles64 = 0;
+		return EResult::Empty;
+	}
+	else
+	{
+		uint64 MinHistoryIndex = NextIndex >= MaxLength ? NextIndex - MaxLength : 0;
+
+		if (State.NextIndex < MinHistoryIndex)
+		{
+			State.NextIndex = MinHistoryIndex;
+			OutCycles64 = History[State.NextIndex++ % MaxLength];
+			return EResult::Disjoint;
+		}
+		else
+		{
+			OutCycles64 = History[State.NextIndex++ % MaxLength];
+			return EResult::Ok;
+		}
+	}
+}
+
+void FRHIGPUFrameTimeHistory::PushFrameCycles(double GPUFrequency, uint64 GPUCycles)
+{
+	double Seconds = double(GPUCycles) / GPUFrequency;
+	double Cycles32 = Seconds / FPlatformTime::GetSecondsPerCycle();
+	double Cycles64 = Seconds / FPlatformTime::GetSecondsPerCycle64();
+
+	{
+		FScopeLock Lock(&CS);
+		History[NextIndex++ % MaxLength] = uint64(Cycles64);
+	}
+
+PRAGMA_DISABLE_DEPRECATION_WARNINGS
+	FPlatformAtomics::InterlockedExchange(reinterpret_cast<volatile int32*>(&GGPUFrameTime), int32(Cycles32));
+PRAGMA_ENABLE_DEPRECATION_WARNINGS
+}
+
+RHI_API uint32 RHIGetGPUFrameCycles(uint32 GPUIndex)
+{
+PRAGMA_DISABLE_DEPRECATION_WARNINGS
+	return (uint32)FPlatformAtomics::AtomicRead(reinterpret_cast<volatile int32*>(&GGPUFrameTime));
+PRAGMA_ENABLE_DEPRECATION_WARNINGS
+}
+
 #undef LOCTEXT_NAMESPACE
diff --git a/Engine/Source/Runtime/RHI/Public/DynamicRHI.h b/Engine/Source/Runtime/RHI/Public/DynamicRHI.h
--- a/Engine/Source/Runtime/RHI/Public/DynamicRHI.h
+++ b/Engine/Source/Runtime/RHI/Public/DynamicRHI.h
@@ -1410,12 +1410,7 @@ FORCEINLINE void RHIAdvanceFrameForGetViewportBackBuffer(FRHIViewport* Viewport)
 	return GDynamicRHI->RHIAdvanceFrameForGetViewportBackBuffer(Viewport);
 }
 
-FORCEINLINE uint32 RHIGetGPUFrameCycles(uint32 GPUIndex = 0)
-{
-PRAGMA_DISABLE_DEPRECATION_WARNINGS
-	return GGPUFrameTime;
-PRAGMA_ENABLE_DEPRECATION_WARNINGS
-}
+RHI_API uint32 RHIGetGPUFrameCycles(uint32 GPUIndex = 0);
 
 FORCEINLINE FViewportRHIRef RHICreateViewport(void* WindowHandle, uint32 SizeX, uint32 SizeY, bool bIsFullscreen, EPixelFormat PreferredPixelFormat)
 {
diff --git a/Engine/Source/Runtime/RHI/Public/GPUProfiler.h b/Engine/Source/Runtime/RHI/Public/GPUProfiler.h
--- a/Engine/Source/Runtime/RHI/Public/GPUProfiler.h
+++ b/Engine/Source/Runtime/RHI/Public/GPUProfiler.h
@@ -15,6 +15,7 @@
 
 #include "Containers/AnsiString.h"
 #include "Containers/SpscQueue.h"
+#include "Containers/StaticArray.h"
 
 #if RHI_NEW_GPU_PROFILER
 
@@ -910,3 +911,55 @@ struct FGPUProfiler
 };
 
 #endif
+
+//
+// Type used to pipe GPU frame timings from the end-of-pipe / RHI threads up to the game / render threads.
+// Stores a history of GPU frame timings, which can be retrieved by engine code via:
+//
+//       static FRHIGPUFrameTimeHistory::FState GPUFrameTimeState;
+//       uint64 GPUFrameTimeCycles64;
+//       while (GPUFrameTimeState.PopFrameCycles(GPUFrameTimeCycles64) != FRHIGPUFrameTimeHistory::EResult::Empty)
+//       {
+//           ...
+//       }
+//
+class FRHIGPUFrameTimeHistory
+{
+public:
+	enum class EResult
+	{
+		// The next frame timing has been retrieved
+		Ok,
+
+		// The next frame timing has been retrieved, but the client has also missed some frames.
+		Disjoint,
+
+		// No new frame timing data available.
+		Empty
+	};
+
+	class FState
+	{
+		friend FRHIGPUFrameTimeHistory;
+		uint64 NextIndex = 0;
+	public:
+		RHI_API EResult PopFrameCycles(uint64& OutCycles64);
+	};
+
+private:
+	// Total number of GPU frame timings to store
+	static constexpr uint32 MaxLength = 16;
+
+	uint64 NextIndex = 0;
+	TStaticArray<uint64, MaxLength> History { InPlace, 0 };
+
+	FCriticalSection CS;
+
+	EResult PopFrameCycles(FState& State, uint64& OutCycles64);
+
+public:
+	// Called by platform RHIs to submit new GPU timing data
+	RHI_API void PushFrameCycles(double GPUFrequency, uint64 GPUCycles);
+};
+
+extern RHI_API FRHIGPUFrameTimeHistory GRHIGPUFrameTimeHistory;
diff --git a/Engine/Source/Runtime/VulkanRHI/Private/Android/VulkanAndroidPlatform.cpp b/Engine/Source/Runtime/VulkanRHI/Private/Android/VulkanAndroidPlatform.cpp
--- a/Engine/Source/Runtime/VulkanRHI/Private/Android/VulkanAndroidPlatform.cpp
+++ b/Engine/Source/Runtime/VulkanRHI/Private/Android/VulkanAndroidPlatform.cpp
@@ -34,11 +34,7 @@ namespace AndroidVulkan
 {
 	void VKSwappyPostWaitCallback(void*, int64_t cpu_time_ns, int64_t gpu_time_ns)
 	{
-		const double GPUTimeInSeconds = (double)gpu_time_ns / 1000000000.0;
-
-	PRAGMA_DISABLE_DEPRECATION_WARNINGS
-		GGPUFrameTime = GPUTimeInSeconds / FPlatformTime::GetSecondsPerCycle64();
-	PRAGMA_ENABLE_DEPRECATION_WARNINGS
+		GRHIGPUFrameTimeHistory.PushFrameCycles(1'000'000'000.0, gpu_time_ns);
 	}
 
 	void SetSwappyPostWaitCallback()
diff --git a/Engine/Source/Runtime/VulkanRHI/Private/VulkanSubmission.cpp b/Engine/Source/Runtime/VulkanRHI/Private/VulkanSubmission.cpp
--- a/Engine/Source/Runtime/VulkanRHI/Private/VulkanSubmission.cpp
+++ b/Engine/Source/Runtime/VulkanRHI/Private/VulkanSubmission.cpp
@@ -662,11 +662,8 @@ void FVulkanDynamicRHI::CompletePayload(FVulkanPayload* Payload)
 		}
 		if (Payload->bEndFrame && (BusyCycles > 0))
 		{
-			const double SecondsPerCycle = FPlatformTime::GetSecondsPerCycle();
 			const double Frequency = double((uint64)((1000.0 * 1000.0 * 1000.0) / Device->GetLimits().timestampPeriod));
-		PRAGMA_DISABLE_DEPRECATION_WARNINGS
-			GGPUFrameTime = FMath::TruncToInt(double(BusyCycles) / Frequency / SecondsPerCycle);
-		PRAGMA_ENABLE_DEPRECATION_WARNINGS
+			GRHIGPUFrameTimeHistory.PushFrameCycles(Frequency, BusyCycles);
 			BusyCycles = 0;
 		}
 	}
diff --git a/Engine/Source/Runtime/Windows/D3D11RHI/Private/D3D11RHI.cpp b/Engine/Source/Runtime/Windows/D3D11RHI/Private/D3D11RHI.cpp
--- a/Engine/Source/Runtime/Windows/D3D11RHI/Private/D3D11RHI.cpp
+++ b/Engine/Source/Runtime/Windows/D3D11RHI/Private/D3D11RHI.cpp
@@ -407,15 +407,11 @@ void FD3DGPUProfiler::EndFrame()
 	{
 		uint64 GPUTiming = FrameTiming.GetTiming();
 		uint64 GPUFreq = FrameTiming.GetTimingFrequency();
-	PRAGMA_DISABLE_DEPRECATION_WARNINGS
-		GGPUFrameTime = FMath::TruncToInt( double(GPUTiming) / double(GPUFreq) / FPlatformTime::GetSecondsPerCycle() );
-	PRAGMA_ENABLE_DEPRECATION_WARNINGS
+		GRHIGPUFrameTimeHistory.PushFrameCycles(GPUFreq, GPUTiming);
 	}
 	else
 	{
-	PRAGMA_DISABLE_DEPRECATION_WARNINGS
-		GGPUFrameTime = 0;
-	PRAGMA_ENABLE_DEPRECATION_WARNINGS
+		GRHIGPUFrameTimeHistory.PushFrameCycles(1, 0);
 	}
 
 	// if we have a frame open, close it now.
