[MIR] Enable new translator to render basic Substrate materials.

- Adds SubstrateParameter as new IR value to encode MCT_Substrate.
- Adds ObjectKind_SubstrateData to MIR::EObjectKind to encode FSubstrateData.
- Replace default initialization of "FrontMaterial" with a setup via SubstrateConvertLegacyMaterialStatic(); adopted from the old translator.
- This also adds some hard coded HLSL expressions to bootstrap Substrate support with limited features.

#rnx

diff --git a/Engine/Source/Runtime/Engine/Private/Materials/MaterialExpressionsIR.cpp b/Engine/Source/Runtime/Engine/Private/Materials/MaterialExpressionsIR.cpp
--- a/Engine/Source/Runtime/Engine/Private/Materials/MaterialExpressionsIR.cpp
+++ b/Engine/Source/Runtime/Engine/Private/Materials/MaterialExpressionsIR.cpp
@@ -900,12 +900,6 @@ void UMaterialExpressionTextureSample::Build(MIR::FEmitter& Em)
 
 	UE_MIR_CHECKPOINT(Em);
 
-	if (TextureValue && GetValueMaterialSamplerType(TextureValue) != SAMPLERTYPE_Color)
-	{
-		Em.Error(TEXT("Input texture sampler type must be color."));
-		return;
-	}
-
 	BuildTextureSampleExpression(Em, this, TextureValue);
 }
 
diff --git a/Engine/Source/Runtime/Engine/Private/Materials/MaterialIR.cpp b/Engine/Source/Runtime/Engine/Private/Materials/MaterialIR.cpp
--- a/Engine/Source/Runtime/Engine/Private/Materials/MaterialIR.cpp
+++ b/Engine/Source/Runtime/Engine/Private/Materials/MaterialIR.cpp
@@ -42,6 +42,7 @@ const TCHAR* LexToString(EValueKind Kind)
 		case VK_ExternalInput: return TEXT("ExternalInput");
 		case VK_TextureObject: return TEXT("TextureObject");
 		case VK_UniformParameter: return TEXT("UniformParameter");
+		case VK_SubstrateParameter: return TEXT("SubstrateParameter");
 		case VK_Dimensional: return TEXT("Dimensional");
 		case VK_SetMaterialOutput: return TEXT("SetMaterialOutput");
 		case VK_Operator: return TEXT("Operator");
@@ -99,6 +100,7 @@ uint32 FValue::GetSizeInBytes() const
 		case VK_ExternalInput: return sizeof(FExternalInput);
 		case VK_TextureObject: return sizeof(FTextureObject);
 		case VK_UniformParameter: return sizeof(FUniformParameter);
+		case VK_SubstrateParameter: return sizeof(VK_SubstrateParameter);
 		case VK_Dimensional: return sizeof(FDimensional) + sizeof(FValue*) * static_cast<const FDimensional*>(this)->GetComponents().Num();
 		case VK_SetMaterialOutput: return sizeof(FSetMaterialOutput);
 		case VK_Operator: return sizeof(FOperator);
diff --git a/Engine/Source/Runtime/Engine/Private/Materials/MaterialIREmitter.cpp b/Engine/Source/Runtime/Engine/Private/Materials/MaterialIREmitter.cpp
--- a/Engine/Source/Runtime/Engine/Private/Materials/MaterialIREmitter.cpp
+++ b/Engine/Source/Runtime/Engine/Private/Materials/MaterialIREmitter.cpp
@@ -684,6 +684,12 @@ FValueRef FEmitter::CustomPrimitiveData(uint32 PrimitiveDataIndex)
 	return {};
 }
 
+FValueRef FEmitter::SubstrateParameter()
+{
+	MIR::FSubstrateParameter Prototype = MakePrototype<MIR::FSubstrateParameter>(FObjectType::GetSubstrateData());
+	return EmitPrototype(*this, Prototype);
+}
+
 /*------------------------------ Instruction emission ------------------------------*/
 
 FSetMaterialOutput* FEmitter::SetMaterialOutput(EMaterialProperty InProperty, FValue* Arg)
diff --git a/Engine/Source/Runtime/Engine/Private/Materials/MaterialIRInternal.cpp b/Engine/Source/Runtime/Engine/Private/Materials/MaterialIRInternal.cpp
--- a/Engine/Source/Runtime/Engine/Private/Materials/MaterialIRInternal.cpp
+++ b/Engine/Source/Runtime/Engine/Private/Materials/MaterialIRInternal.cpp
@@ -76,6 +76,8 @@ MIR::FValue* CreateMaterialAttributeDefaultValue(FEmitter& Emitter, const UMater
 		case MCT_UInt:
 		case MCT_UInt4: return Emitter.ConstantInt4({ (TInteger)DefaultValue.X, (TInteger)DefaultValue.Y, (TInteger)DefaultValue.Z, (TInteger)DefaultValue.W });
 
+		case MCT_Substrate: return Emitter.SubstrateParameter();
+
 		default: UE_MIR_UNREACHABLE();
 	}
 }
diff --git a/Engine/Source/Runtime/Engine/Private/Materials/MaterialIRModuleBuilder.cpp b/Engine/Source/Runtime/Engine/Private/Materials/MaterialIRModuleBuilder.cpp
--- a/Engine/Source/Runtime/Engine/Private/Materials/MaterialIRModuleBuilder.cpp
+++ b/Engine/Source/Runtime/Engine/Private/Materials/MaterialIRModuleBuilder.cpp
@@ -23,6 +23,7 @@
 #include "Materials/MaterialInsights.h"
 #include "Engine/Texture.h"
 #include "Misc/FileHelper.h"
+#include "RenderUtils.h"
 
 static TAutoConsoleVariable<int32> CVarMaterialIRDebugDumpLevel(
 	TEXT("r.Material.Translator.DebugDump"),
@@ -64,6 +65,18 @@ struct FAnalysisContext
 	}
 };
 
+static MIR::FValue* FindIRModuleOutputValue(const FMaterialIRModule& Module, MIR::EStage Stage, EMaterialProperty InProperty)
+{
+	for (const MIR::FSetMaterialOutput* Output : Module.GetOutputs(Stage))
+	{
+		if (Output->Property == InProperty)
+		{
+			return Output->Arg;
+		}
+	}
+	return nullptr;
+}
+
 struct FMaterialIRModuleBuilderImpl
 {
 	FMaterialIRModuleBuilder* Builder;
@@ -96,6 +109,12 @@ struct FMaterialIRModuleBuilderImpl
 				PrepareSingleMaterialAttribute((EMaterialProperty)Index);
 			}
 		}
+
+		if (Substrate::IsSubstrateEnabled())
+		{
+			// Bridge legacy material attributes to Substrate's "FrontMaterial" input last
+			PrepareSingleMaterialAttribute(MP_FrontMaterial);
+		}
 	}
 
 	void PrepareSingleMaterialAttribute(EMaterialProperty Property)
@@ -311,6 +330,18 @@ struct FMaterialIRModuleBuilderImpl
 				FMaterialInputDescription Input;
 				ensure(Builder->Material->GetExpressionInputDescription(Output->Property, Input));
 
+				// Connect output properties with substrate parameter if a FrontMaterial property is provided
+				if (Output->Property == MP_FrontMaterial && Output->Arg)
+				{
+					if (MIR::FSubstrateParameter* SubstrateParameter = Output->Arg->As<MIR::FSubstrateParameter>())
+					{
+						for (int32 PropertyIndex = 0; PropertyIndex < MP_MAX; ++PropertyIndex)
+						{
+							SubstrateParameter->PropertyArgs[PropertyIndex] = FindIRModuleOutputValue(*Module, MIR::EStage::Stage_Pixel, static_cast<EMaterialProperty>(PropertyIndex));
+						}
+					}
+				}
+
 				if (!Output->Arg)
 				{
 					MIR::FValue** ValuePtr = Context.OutputValues.Find(Input.Input->GetConnectedOutput());
diff --git a/Engine/Source/Runtime/Engine/Private/Materials/MaterialIRToHLSLTranslator.cpp b/Engine/Source/Runtime/Engine/Private/Materials/MaterialIRToHLSLTranslator.cpp
--- a/Engine/Source/Runtime/Engine/Private/Materials/MaterialIRToHLSLTranslator.cpp
+++ b/Engine/Source/Runtime/Engine/Private/Materials/MaterialIRToHLSLTranslator.cpp
@@ -271,8 +271,28 @@ struct FTranslator : FMaterialIRToHLSLTranslation
 
 		LowerBlock(Module->GetRootBlock(CurrentStage));
 
-		Printer << TEXT("PixelMaterialInputs.FrontMaterial = GetInitialisedSubstrateData()") << EndOfStatement;
-		Printer << TEXT("PixelMaterialInputs.Subsurface = 0") << EndOfStatement;
+		if (Stage == MIR::Stage_Pixel)
+		{
+			Printer << TEXT("PixelMaterialInputs.Subsurface = 0") << EndOfStatement;
+
+			if (Substrate::IsSubstrateEnabled())
+			{
+				Printer << TEXT("// Substrate normal and tangent") << NewLine;
+				Printer << TEXT("Parameters.SharedLocalBases.Normals[0] = Parameters.TangentToWorld[2]") << EndOfStatement;
+				Printer << TEXT("#if MATERIAL_TANGENTSPACENORMAL") << NewLine;
+				Printer << TEXT("Parameters.SharedLocalBases.Normals[0] *= Parameters.TwoSidedSign") << EndOfStatement;
+				Printer << TEXT("#endif") << NewLine;
+				Printer << TEXT("Parameters.SharedLocalBases.Tangents[0] = Parameters.TangentToWorld[0]") << EndOfStatement;
+				Printer << TEXT("#if MATERIAL_TANGENTSPACENORMAL") << NewLine;
+				Printer << TEXT("Parameters.SharedLocalBases.Tangents[0] *= Parameters.TwoSidedSign") << EndOfStatement;
+				Printer << TEXT("#endif") << NewLine;
+				Printer << TEXT("Parameters.SharedLocalBases.Count = 1") << EndOfStatement;
+			}
+			else
+			{
+				Printer << TEXT("PixelMaterialInputs.FrontMaterial = GetInitialisedSubstrateData()") << EndOfStatement;
+			}
+		}
 
 		EvaluateOtherMaterialAttributesHLSL[CurrentStage] = MoveTemp(Printer.Buffer);
 	}
@@ -356,6 +376,7 @@ struct FTranslator : FMaterialIRToHLSLTranslation
 			case MIR::VK_ExternalInput: LowerExternalInput(static_cast<const MIR::FExternalInput*>(InValue)); break;
 			case MIR::VK_TextureObject: LowerTextureObject(static_cast<const MIR::FTextureObject*>(InValue)); break;
 			case MIR::VK_UniformParameter: LowerUniformParameter(static_cast<const MIR::FUniformParameter*>(InValue)); break;
+			case MIR::VK_SubstrateParameter: LowerSubstrateParameter(static_cast<const MIR::FSubstrateParameter*>(InValue)); break;
 			default: UE_MIR_UNREACHABLE();
 		}
 
@@ -591,6 +612,110 @@ struct FTranslator : FMaterialIRToHLSLTranslation
 		}
 	}
 
+	struct FSubstrateLegacyArgument
+	{
+		const MIR::FType* ParameterType;
+		const TCHAR* ParameterName;
+		EMaterialProperty MaterialProperty;
+		const TCHAR* DefaultInlineHLSL;
+	};
+
+	static TConstArrayView<FSubstrateLegacyArgument> GetSubstrateLegacyConversionArguments()
+	{
+		static const FSubstrateLegacyArgument LegacyArguments[] =
+		{
+			{ MIR::FPrimitiveType::GetFloat3(), TEXT("BaseColor"), MP_BaseColor, nullptr },
+			{ MIR::FPrimitiveType::GetFloat(), TEXT("Specular"), MP_SpecularColor, TEXT("0.5f") },
+			{ MIR::FPrimitiveType::GetFloat(), TEXT("Metallic"), MP_Metallic, nullptr },
+			{ MIR::FPrimitiveType::GetFloat(), TEXT("Roughness"), MP_Roughness, nullptr },
+			{ MIR::FPrimitiveType::GetFloat(), TEXT("Anisotropy"), MP_Anisotropy, nullptr },
+			{ MIR::FPrimitiveType::GetFloat3(), TEXT("SubSurfaceColor"), MP_SubsurfaceColor, TEXT("(float3)1") },
+			{ MIR::FPrimitiveType::GetFloat(), TEXT("SubSurfaceProfileId"), MP_MAX, TEXT("0") },
+			{ MIR::FPrimitiveType::GetFloat(), TEXT("ClearCoat"), MP_MAX, TEXT("1.0f") },
+			{ MIR::FPrimitiveType::GetFloat(), TEXT("ClearCoatRoughness"), MP_MAX, TEXT("0.1f") },
+			{ MIR::FPrimitiveType::GetFloat3(), TEXT("Emissive"), MP_EmissiveColor, nullptr },
+			{ MIR::FPrimitiveType::GetFloat(), TEXT("Opacity"), MP_Opacity, nullptr },
+			{ MIR::FPrimitiveType::GetFloat3(), TEXT("ThinTranslucentTransmittanceColor"), MP_MAX, TEXT("(float3)0.5f") },
+			{ MIR::FPrimitiveType::GetFloat(), TEXT("ThinTranslucentSurfaceCoverage"), MP_MAX, TEXT("1.0f") },
+			{ MIR::FPrimitiveType::GetFloat3(), TEXT("WaterScatteringCoefficients"), MP_MAX, nullptr },
+			{ MIR::FPrimitiveType::GetFloat3(), TEXT("WaterAbsorptionCoefficients"), MP_MAX, nullptr },
+			{ MIR::FPrimitiveType::GetFloat(), TEXT("WaterPhaseG"), MP_MAX, nullptr },
+			{ MIR::FPrimitiveType::GetFloat3(), TEXT("ColorScaleBehindWater"), MP_MAX, TEXT("(float3)1.0f") },
+			{ MIR::FPrimitiveType::GetInt(), TEXT("ShadingModel"), MP_MAX, TEXT("1") },
+			{ MIR::FPrimitiveType::GetFloat3(), TEXT("RawNormal"), MP_Normal, TEXT("Parameters.TangentToWorld[2]") },
+			{ MIR::FPrimitiveType::GetFloat3(), TEXT("RawTangent"), MP_Tangent, TEXT("Parameters.TangentToWorld[2]") },
+			{ MIR::FPrimitiveType::GetFloat3(), TEXT("RawClearCoatNormal"), MP_Normal, TEXT("Parameters.TangentToWorld[2]") },
+			{ MIR::FPrimitiveType::GetFloat3(), TEXT("RawCustomTangent"), MP_Tangent, TEXT("Parameters.TangentToWorld[2]") },
+			{ MIR::FPrimitiveType::GetInt(), TEXT("SharedLocalBasisIndex"), MP_MAX, nullptr },
+			{ MIR::FPrimitiveType::GetInt(), TEXT("ClearCoatBottomNormal_SharedLocalBasisIndex"), MP_MAX, nullptr },
+			{ MIR::FPrimitiveType::GetInt(), TEXT("SharedLocalBasisTypes"), MP_MAX, TEXT("Parameters.SharedLocalBases.Types") },
+		};
+		return LegacyArguments;
+	}
+
+	static TConstArrayView<FSubstrateLegacyArgument> GetSubstrateUnlitArguments()
+	{
+		static const FSubstrateLegacyArgument LegacyArguments[] =
+		{
+			{ MIR::FPrimitiveType::GetFloat3(), TEXT("Emissive"), MP_EmissiveColor, nullptr },
+			{ MIR::FPrimitiveType::GetFloat(), TEXT("TransmittanceColor"), MP_MAX, TEXT("(float3)0.0f") },
+			{ MIR::FPrimitiveType::GetFloat3(), TEXT("RawNormal"), MP_Normal, TEXT("Parameters.TangentToWorld[2]") },
+		};
+		return LegacyArguments;
+	}
+
+	void LowerSubstrateParameter(const MIR::FSubstrateParameter* SubstrateParameter)
+	{
+		Printer << TEXT("Parameters.SubstrateTree.PromoteParameterBlendedBSDFToOperator") << BeginArgs;
+		Printer.Tabs++;
+		Printer << NewLine << (SubstrateParameter->bIsUnlit ? TEXT("GetSubstrateUnlitBSDF") : TEXT("SubstrateConvertLegacyMaterialStatic")) << BeginArgs;
+		Printer.Tabs++;
+		Printer << ListSeparator << NewLine << TEXT("/*PixelFootprint:*/ Parameters.SubstratePixelFootprint");
+
+		TConstArrayView<FSubstrateLegacyArgument> Arguments = SubstrateParameter->bIsUnlit ? GetSubstrateUnlitArguments() : GetSubstrateLegacyConversionArguments();
+
+		for (const FSubstrateLegacyArgument& Argument : Arguments)
+		{
+			Printer << ListSeparator << NewLine << TEXT("/*") << Argument.ParameterName << TEXT("*/ ");
+
+			if (Argument.MaterialProperty != MP_MAX && SubstrateParameter->PropertyArgs[Argument.MaterialProperty])
+			{
+				LowerValue(SubstrateParameter->PropertyArgs[Argument.MaterialProperty]);
+			}
+			else if (Argument.DefaultInlineHLSL)
+			{
+				Printer << Argument.DefaultInlineHLSL;
+			}
+			else if (const MIR::FPrimitiveType* PrimitiveType = Argument.ParameterType->AsPrimitive())
+			{
+				if (const MIR::FPrimitiveType* ScalarType = Argument.ParameterType->AsScalar())
+				{
+					Printer << TEXT("0.0f");
+				}
+				else
+				{
+					LowerType(Argument.ParameterType);
+					Printer << BeginArgs;
+					for (int32 ComponentIndex = 0; ComponentIndex < PrimitiveType->NumRows; ++ComponentIndex)
+					{
+						Printer << ListSeparator << TEXT("0.0f");
+					}
+					Printer << EndArgs;
+				}
+			}
+			else
+			{
+				UE_MIR_UNREACHABLE();
+			}
+		}
+
+		Printer.Tabs--;
+		Printer << NewLine << EndArgs << ListSeparator;
+		Printer << NewLine << TEXT("0, 0, 0, 1");
+		Printer.Tabs--;
+		Printer << NewLine << EndArgs;
+	}
+
 	void LowerPrimitiveUniformParameter(const MIR::FUniformParameter* UniformParameter)
 	{
 		const MIR::FPrimitiveType* PrimitiveType = UniformParameter->Type->AsPrimitive();
@@ -599,7 +724,7 @@ struct FTranslator : FMaterialIRToHLSLTranslation
 
 		if (PrimitiveType->ScalarKind == MIR::EScalarKind::ScalarKind_Int)
 		{
-			Printer << "asint(";
+			Printer << TEXT("asint(");
 		}
 
 		const FUniformExpressionSet& UniformExpressionSet = Module->GetCompilationOutput().UniformExpressionSet;
@@ -678,7 +803,35 @@ struct FTranslator : FMaterialIRToHLSLTranslation
 			Printer << TEXT("PixelMaterialInputs.") << PropertyName << TEXT(" = ");
 		}
 
-		LowerValue(Output->Arg);
+		if (Substrate::IsSubstrateEnabled())
+		{
+			// Lower to default value when Substrate is enabled
+			switch (Output->Property)
+			{
+				case MP_EmissiveColor:		Printer << TEXT("0"); break;
+				case MP_Opacity:			Printer << TEXT("1.0f"); break;
+				case MP_OpacityMask:		Printer << TEXT("1.0f"); break;
+				case MP_BaseColor:			Printer << TEXT("0"); break;
+				case MP_Metallic:			Printer << TEXT("0"); break;
+				case MP_Specular:			Printer << TEXT("0"); break;
+				case MP_Roughness:			Printer << TEXT("0"); break;
+				case MP_Anisotropy:			Printer << TEXT("0"); break;
+				case MP_Normal:				Printer << TEXT("float3(0.0f, 0.0f, 1.0f)"); break;
+				case MP_Tangent:			Printer << TEXT("0"); break;
+				case MP_SubsurfaceColor:	Printer << TEXT("0"); break;
+				case MP_AmbientOcclusion:	Printer << TEXT("1.0f"); break;
+				case MP_Refraction:			Printer << TEXT("0"); break;
+				case MP_PixelDepthOffset:	Printer << TEXT("0.0f"); break;
+				case MP_ShadingModel:		Printer << TEXT("0"); break;
+				case MP_SurfaceThickness:	Printer << TEXT("0.01f"); break;
+				case MP_Displacement:		Printer << TEXT("-1.0f"); break;
+				default: LowerValue(Output->Arg); break;
+			}
+		}
+		else
+		{
+			LowerValue(Output->Arg);
+		}
 	}
 
 	void LowerOperator(const MIR::FOperator* Operator)
@@ -1051,11 +1204,35 @@ struct FTranslator : FMaterialIRToHLSLTranslation
 		}
 		EvaluateMaterialDeclaration.Append(TEXT("\n}\n"));
 		Params.Add(TEXT("evaluate_material_attributes"), MoveTemp(EvaluateMaterialDeclaration));
+
+		FString UniformMaterialExpressions;
+		if (Substrate::IsSubstrateEnabled())
+		{
+			// Adde default Substrate functions
+			UniformMaterialExpressions.Append(
+				TEXT("// Substrate: HiddenMaterialAssetConversion\n")
+				TEXT("#if TEMPLATE_USES_SUBSTRATE\n")
+				TEXT("void FSubstratePixelHeader::PreUpdateAllBSDFWithBottomUpOperatorVisit(float3 V) {}\n")
+				TEXT("void FSubstratePixelHeader::UpdateAllBSDFsOperatorCoverageTransmittance(FSubstrateIntegrationSettings Settings, float3 V)\n")
+				TEXT("{\n")
+				TEXT("#if SUBSTRATE_COMPILER_SUPPORTS_STRUCT_FORWARD_DECLARATION\n")
+				TEXT("\tSubstrateTree.UpdateSingleBSDFOperatorCoverageTransmittance(this, 0, Settings, V);\n")
+				TEXT("#else\n")
+				TEXT("\tUpdateSingleBSDFOperatorCoverageTransmittance(SubstrateTree, this, 0, Settings, V);\n")
+				TEXT("#endif\n")
+				TEXT("}\n")
+				TEXT("void FSubstratePixelHeader::UpdateAllOperatorsCoverageTransmittance() {}\n")
+				TEXT("void FSubstratePixelHeader::UpdateAllBSDFWithBottomUpOperatorVisit() {}\n")
+				TEXT("#endif // TEMPLATE_USES_SUBSTRATE\n")
+				TEXT("\n")
+			);
+		}
+		Params.Add(TEXT("uniform_material_expressions"), MoveTemp(UniformMaterialExpressions));
 	}
 	
 	ENoOp LowerType(const MIR::FType* Type)
 	{
-		if (auto ArithmeticType = Type->AsPrimitive())
+		if (const MIR::FPrimitiveType* ArithmeticType = Type->AsPrimitive())
 		{
 			switch (ArithmeticType->ScalarKind)
 			{
@@ -1074,6 +1251,14 @@ struct FTranslator : FMaterialIRToHLSLTranslation
 				Printer << TEXT("x") << ArithmeticType->NumColumns;
 			}
 		}
+		else if (const MIR::FObjectType* ObjectType = Type->AsObject())
+		{
+			switch (ObjectType->ObjectKind)
+			{
+				case MIR::ObjectKind_SubstrateData: Printer << TEXT("FSubstrateData"); break;
+				default: UE_MIR_UNREACHABLE(); break;
+			}
+		}
 		else
 		{
 			UE_MIR_UNREACHABLE();
@@ -1203,9 +1388,17 @@ struct FTranslator : FMaterialIRToHLSLTranslation
 			}
 		}
 
-		OutEnvironment.SetDefine(TEXT("MATERIAL_IS_SUBSTRATE"), false);
+		const bool bIsSubstrate = Substrate::IsSubstrateEnabled();
+
+		OutEnvironment.SetDefine(TEXT("MATERIAL_IS_SUBSTRATE"), bIsSubstrate);
 		OutEnvironment.SetDefine(TEXT("DUAL_SOURCE_COLOR_BLENDING_ENABLED"), false);
 		OutEnvironment.SetDefine(TEXT("TEXTURE_SAMPLE_DEBUG"), false);
+
+		if (bIsSubstrate)
+		{
+			OutEnvironment.SetDefine(TEXT("SUBSTRATE_USE_FULLYSIMPLIFIED_MATERIAL"), false);
+			OutEnvironment.SetDefine(TEXT("SUBSTRATE_CLAMPED_CLOSURE_COUNT"), 1);
+		}
 	}
 };
 
diff --git a/Engine/Source/Runtime/Engine/Private/Materials/MaterialIRTypes.cpp b/Engine/Source/Runtime/Engine/Private/Materials/MaterialIRTypes.cpp
--- a/Engine/Source/Runtime/Engine/Private/Materials/MaterialIRTypes.cpp
+++ b/Engine/Source/Runtime/Engine/Private/Materials/MaterialIRTypes.cpp
@@ -13,6 +13,7 @@ const TCHAR* TypeKindToString(ETypeKind Kind)
 	{
 		case TypeKind_Void: return TEXT("void");
 		case TypeKind_Primitive: return TEXT("primitive");
+		case TypeKind_Object: return TEXT("object");
 		default: UE_MIR_UNREACHABLE();
 	}
 }
@@ -124,6 +125,7 @@ FStringView FType::GetSpelling() const
 		switch (ObjectType->ObjectKind)
 		{
 			case ObjectKind_Texture2D: return TEXT("Texture2D");
+			case ObjectKind_SubstrateData: return TEXT("SubstrateData");
 		}
 	}
 	UE_MIR_UNREACHABLE();
@@ -396,6 +398,12 @@ const FType* FObjectType::GetTexture2D()
 	return &Instance;
 }
 
+const FType* FObjectType::GetSubstrateData()
+{
+	static const FObjectType Instance{ { TypeKind_Object }, ObjectKind_SubstrateData };
+	return &Instance;
+}
+
 } // namespace MIR
 
 #endif // #if WITH_EDITOR
diff --git a/Engine/Source/Runtime/Engine/Public/Materials/MaterialIR.h b/Engine/Source/Runtime/Engine/Public/Materials/MaterialIR.h
--- a/Engine/Source/Runtime/Engine/Public/Materials/MaterialIR.h
+++ b/Engine/Source/Runtime/Engine/Public/Materials/MaterialIR.h
@@ -89,6 +89,7 @@ enum EValueKind
 	VK_ExternalInput,
 	VK_TextureObject,
 	VK_UniformParameter,
+	VK_SubstrateParameter,
 
 	// Instructions
 
@@ -424,6 +425,14 @@ struct FUniformParameter: TValue<VK_UniformParameter>
 	uint32 Analysis_UniformParameterIndex;
 };
 
+// Material IR representation of the FSubstrateData shader parameter to transition legacy materials into Substrate.
+struct FSubstrateParameter : TValue<VK_SubstrateParameter>
+{
+	FValue* PropertyArgs[MP_MAX];
+
+	bool bIsUnlit;
+};
+
 /*------------------------------------ Instructions ------------------------------------*/
 
 // A block of instructions.
diff --git a/Engine/Source/Runtime/Engine/Public/Materials/MaterialIREmitter.h b/Engine/Source/Runtime/Engine/Public/Materials/MaterialIREmitter.h
--- a/Engine/Source/Runtime/Engine/Public/Materials/MaterialIREmitter.h
+++ b/Engine/Source/Runtime/Engine/Public/Materials/MaterialIREmitter.h
@@ -304,6 +304,9 @@ public:
 
 	// Emits a custom primitive data value with given index.
 	FValueRef CustomPrimitiveData(uint32 PrimitiveDataIndex);
+
+	// Emits a parameter that encodes MCT_Substrate as IR value.
+	FValueRef SubstrateParameter();
 	
 	/*----------------------------- Instructions emission -----------------------------*/
 
diff --git a/Engine/Source/Runtime/Engine/Public/Materials/MaterialIRTypes.h b/Engine/Source/Runtime/Engine/Public/Materials/MaterialIRTypes.h
--- a/Engine/Source/Runtime/Engine/Public/Materials/MaterialIRTypes.h
+++ b/Engine/Source/Runtime/Engine/Public/Materials/MaterialIRTypes.h
@@ -22,6 +22,7 @@ enum ETypeKind
 enum EObjectKind 
 {
 	ObjectKind_Texture2D,
+	ObjectKind_SubstrateData,
 };
 
 const TCHAR* TypeKindToString(ETypeKind Kind);
@@ -180,6 +181,7 @@ struct FObjectType : FType
 	EObjectKind ObjectKind;
 
 	static const FType* GetTexture2D();
+	static const FType* GetSubstrateData();
 };
 
 } // namespace MIR
