Temporary set r.VT.UploadUseLegacyPath 1 while fixing bug in new path.
#rnx

diff --git a/Engine/Source/Runtime/Engine/Private/VT/VirtualTextureUploadCache.cpp b/Engine/Source/Runtime/Engine/Private/VT/VirtualTextureUploadCache.cpp
--- a/Engine/Source/Runtime/Engine/Private/VT/VirtualTextureUploadCache.cpp
+++ b/Engine/Source/Runtime/Engine/Private/VT/VirtualTextureUploadCache.cpp
@@ -12,7 +12,7 @@ DECLARE_MEMORY_STAT(TEXT("Total CPU Upload Memory"), STAT_TotalCPUUploadSize, ST
 
 static TAutoConsoleVariable<int32> CVarVTUploadUseLegacyPath(
 	TEXT("r.VT.UploadUseLegacyPath"),
-	0,
+	1,
 	TEXT("Use the legacy virtual texture upload path which locks staging textures."),
 	ECVF_ReadOnly
 );
