Fix use after free issue with stack string allocation.
#rnx

diff --git a/Engine/Source/Runtime/Engine/Private/Texture2D.cpp b/Engine/Source/Runtime/Engine/Private/Texture2D.cpp
--- a/Engine/Source/Runtime/Engine/Private/Texture2D.cpp
+++ b/Engine/Source/Runtime/Engine/Private/Texture2D.cpp
@@ -1691,7 +1691,7 @@ void FVirtualTexture2DResource::InitializeEditorResources(IVirtualTexture* InVir
 			// Create the target texture to composite tiles to.
 			FRDGTextureRef Texture = GraphBuilder.CreateTexture(FRDGTextureDesc::Create2D(
 				FIntPoint(MipWidthInTiles * TileSizeInPixels, MipHeightInTiles * TileSizeInPixels),
-				PixelFormat, FClearValueBinding::Black,	TexCreateFlags), *Name);
+				PixelFormat, FClearValueBinding::Black,	TexCreateFlags), TEXT("VTEditorResource"));
 
 			TRefCountPtr<IPooledRenderTarget> PooledRenderTarget = GraphBuilder.ConvertToExternalTexture(Texture);
 			Texture2DRHI = PooledRenderTarget->GetRHI();
