Fix use after free in virtual texture upload.
Reverted temp fix to use legacy upload path.
#rnx

diff --git a/Engine/Source/Runtime/Engine/Private/VT/VirtualTextureUploadCache.cpp b/Engine/Source/Runtime/Engine/Private/VT/VirtualTextureUploadCache.cpp
--- a/Engine/Source/Runtime/Engine/Private/VT/VirtualTextureUploadCache.cpp
+++ b/Engine/Source/Runtime/Engine/Private/VT/VirtualTextureUploadCache.cpp
@@ -12,7 +12,7 @@ DECLARE_MEMORY_STAT(TEXT("Total CPU Upload Memory"), STAT_TotalCPUUploadSize, ST
 
 static TAutoConsoleVariable<int32> CVarVTUploadUseLegacyPath(
 	TEXT("r.VT.UploadUseLegacyPath"),
-	1,
+	0,
 	TEXT("Use the legacy virtual texture upload path which locks staging textures."),
 	ECVF_ReadOnly
 );
@@ -582,16 +582,13 @@ void FVirtualTextureUploadCache::CancelTile(FRHICommandList& RHICmdList, const F
 
 void FVirtualTextureUploadCache::UpdateFreeList(FRHICommandList& RHICmdList, bool bForceFreeAll)
 {
+	// Iterate tiles pending release and free them once we can be sure they have been used.
 	const uint32 CurrentFrame = GFrameNumberRenderThread;
+	static uint32 ReleaseFrameDelay = 2u;
 
-	// Iterate tiles pending release and free them.
-	// We need to defer release of persistent mapped buffers here because there is no other synchronization method for the memory in these.
-	// Without deferring release we would get corrupt textures.
 	for (TSparseArray<FTileEntry>::TIterator It(PendingRelease); It; ++It)
 	{
-		check(CurrentFrame >= It->FrameSubmitted);
-		const uint32 FramesSinceSubmitted = CurrentFrame - It->FrameSubmitted;
-		if (UploadType == EVTUploadType::PersistentBuffer && FramesSinceSubmitted < 2u)
+		if (CurrentFrame < It->FrameSubmitted + ReleaseFrameDelay)
 		{
 			break;
 		}
