Fix leaks that can occur when searching Arrays for objects that allocate memory by default

diff --git a/Engine/Source/Runtime/Engine/Private/KismetArrayLibrary.cpp b/Engine/Source/Runtime/Engine/Private/KismetArrayLibrary.cpp
--- a/Engine/Source/Runtime/Engine/Private/KismetArrayLibrary.cpp
+++ b/Engine/Source/Runtime/Engine/Private/KismetArrayLibrary.cpp
@@ -474,7 +474,7 @@ void UKismetArrayLibrary::GenericArray_Get(void* TargetArray, const FArrayProper
 				*ArrayProp->GetOwnerVariant().GetPathName()),
 				ELogVerbosity::Warning,
 				GetOutOfBoundsWarning);
-			InnerProp->InitializeValue(Item);
+			InnerProp->ClearValue(Item);
 		}
 	}
 }
@@ -623,7 +623,7 @@ void UKismetArrayLibrary::GenericArray_Random(void* TargetArray, const FArrayPro
 			FFrame::KismetExecutionMessage(*FString::Printf(TEXT("Attempted to access random index from empty array!")),
 				ELogVerbosity::Warning,
 				RandomAccessToEmptyArrayWarning);
-			InnerProp->InitializeValue(OutItem);
+			InnerProp->ClearValue(OutItem);
 		}
 	}
 }
@@ -648,7 +648,7 @@ void UKismetArrayLibrary::GenericArray_RandomFromStream(void* TargetArray, const
 			FFrame::KismetExecutionMessage(*FString::Printf(TEXT("Attempted to access random index from empty array!")),
 				ELogVerbosity::Warning,
 				RandomAccessToEmptyArrayWarning);
-			InnerProp->InitializeValue(OutItem);
+			InnerProp->ClearValue(OutItem);
 		}
 	}
 }
