Add resource transitions for the texture in readback functions: ReadSurfaceFloatCommand() and ReadLinearColorPixels()
Without transitions we don't see rhi validation errors but we do see d3ddebug errors.
#rnx

diff --git a/Engine/Source/Runtime/Engine/Private/UnrealClient.cpp b/Engine/Source/Runtime/Engine/Private/UnrealClient.cpp
--- a/Engine/Source/Runtime/Engine/Private/UnrealClient.cpp
+++ b/Engine/Source/Runtime/Engine/Private/UnrealClient.cpp
@@ -109,7 +109,10 @@ bool FRenderTarget::ReadFloat16Pixels(TArray<FFloat16Color>& OutImageData, FRead
 	ENQUEUE_RENDER_COMMAND(ReadSurfaceFloatCommand)(
 		[RenderTarget_RT = this, SrcRect_RT = InSrcRect, OutData_RT = &OutImageData, Flags_RT = InFlags](FRHICommandListImmediate& RHICmdList)
 	{
-		RHICmdList.ReadSurfaceFloatData(RenderTarget_RT->GetShaderResourceTexture(), SrcRect_RT, *OutData_RT, Flags_RT);
+		FRHITexture* Texture = RenderTarget_RT->GetShaderResourceTexture();
+		RHICmdList.Transition(FRHITransitionInfo(Texture, ERHIAccess::SRVMask, ERHIAccess::CopySrc));
+		RHICmdList.ReadSurfaceFloatData(Texture, SrcRect_RT, *OutData_RT, Flags_RT);
+		RHICmdList.Transition(FRHITransitionInfo(Texture, ERHIAccess::CopySrc, ERHIAccess::SRVMask));
 	});
 	FlushRenderingCommands();
 
@@ -128,9 +131,12 @@ bool FRenderTarget::ReadLinearColorPixels(TArray<FLinearColor> &OutImageData, FR
 	// Read the render target surface data back.	
 	ENQUEUE_RENDER_COMMAND(ReadSurfaceCommand)(
 		[RenderTarget_RT = this, SrcRect_RT = InSrcRect, OutData_RT = &OutImageData, Flags_RT = InFlags](FRHICommandListImmediate& RHICmdList)
-		{
-			RHICmdList.ReadSurfaceData(RenderTarget_RT->GetShaderResourceTexture(), SrcRect_RT, *OutData_RT, Flags_RT);
-		});
+	{
+		FRHITexture* Texture = RenderTarget_RT->GetShaderResourceTexture();
+		RHICmdList.Transition(FRHITransitionInfo(Texture, ERHIAccess::SRVMask, ERHIAccess::CopySrc));
+		RHICmdList.ReadSurfaceData(Texture, SrcRect_RT, *OutData_RT, Flags_RT);
+		RHICmdList.Transition(FRHITransitionInfo(Texture, ERHIAccess::CopySrc, ERHIAccess::SRVMask));
+	});
 	FlushRenderingCommands();
 
 	return OutImageData.Num() > 0;
