- Remove temp tag from Cascade tooltip

#rnx

diff --git a/Engine/Source/Runtime/Engine/Classes/Particles/TypeData/ParticleModuleTypeDataGpu.h b/Engine/Source/Runtime/Engine/Classes/Particles/TypeData/ParticleModuleTypeDataGpu.h
--- a/Engine/Source/Runtime/Engine/Classes/Particles/TypeData/ParticleModuleTypeDataGpu.h
+++ b/Engine/Source/Runtime/Engine/Classes/Particles/TypeData/ParticleModuleTypeDataGpu.h
@@ -466,7 +466,7 @@ class UParticleModuleTypeDataGpu : public UParticleModuleTypeDataBase
 	UPROPERTY()
 	struct FGPUSpriteResourceData ResourceData;
 
-	/** TEMP: How much to stretch sprites based on camera motion blur. */
+	/** How much to stretch sprites based on camera motion blur. */
 	UPROPERTY(EditAnywhere, Category=ParticleModuleTypeDataGpu)
 	float CameraMotionBlurAmount;
 
