Undo 42025817 and 41982735, which are attempting to fix a leak but have changed behavior for UserDefinedStruct properties
#rnx

diff --git a/Engine/Source/Runtime/Engine/Private/BlueprintMapLibrary.cpp b/Engine/Source/Runtime/Engine/Private/BlueprintMapLibrary.cpp
--- a/Engine/Source/Runtime/Engine/Private/BlueprintMapLibrary.cpp
+++ b/Engine/Source/Runtime/Engine/Private/BlueprintMapLibrary.cpp
@@ -45,7 +45,7 @@ bool UBlueprintMapLibrary::GenericMap_Find(const void* TargetMap, const FMapProp
 			}
 			else
 			{
-				MapProperty->ValueProp->ClearValue(OutValuePtr);
+				MapProperty->ValueProp->InitializeValue(OutValuePtr);
 			}
 		}
 		return FoundValuePtr != nullptr;
diff --git a/Engine/Source/Runtime/Engine/Private/KismetArrayLibrary.cpp b/Engine/Source/Runtime/Engine/Private/KismetArrayLibrary.cpp
--- a/Engine/Source/Runtime/Engine/Private/KismetArrayLibrary.cpp
+++ b/Engine/Source/Runtime/Engine/Private/KismetArrayLibrary.cpp
@@ -474,7 +474,7 @@ void UKismetArrayLibrary::GenericArray_Get(void* TargetArray, const FArrayProper
 				*ArrayProp->GetOwnerVariant().GetPathName()),
 				ELogVerbosity::Warning,
 				GetOutOfBoundsWarning);
-			InnerProp->ClearValue(Item);
+			InnerProp->InitializeValue(Item);
 		}
 	}
 }
@@ -623,7 +623,7 @@ void UKismetArrayLibrary::GenericArray_Random(void* TargetArray, const FArrayPro
 			FFrame::KismetExecutionMessage(*FString::Printf(TEXT("Attempted to access random index from empty array!")),
 				ELogVerbosity::Warning,
 				RandomAccessToEmptyArrayWarning);
-			InnerProp->ClearValue(OutItem);
+			InnerProp->InitializeValue(OutItem);
 		}
 	}
 }
@@ -648,7 +648,7 @@ void UKismetArrayLibrary::GenericArray_RandomFromStream(void* TargetArray, const
 			FFrame::KismetExecutionMessage(*FString::Printf(TEXT("Attempted to access random index from empty array!")),
 				ELogVerbosity::Warning,
 				RandomAccessToEmptyArrayWarning);
-			InnerProp->ClearValue(OutItem);
+			InnerProp->InitializeValue(OutItem);
 		}
 	}
 }
