Anim Decompression: Remove usage of deprecated GetCompactPoseIndexFromSkeletonIndex

PR #13063: Replace deprecated GetCompactPoseIndexFromSkeletonIndex

diff --git a/Engine/Source/Runtime/Engine/Private/Animation/AnimationDecompression.cpp b/Engine/Source/Runtime/Engine/Private/Animation/AnimationDecompression.cpp
--- a/Engine/Source/Runtime/Engine/Private/Animation/AnimationDecompression.cpp
+++ b/Engine/Source/Runtime/Engine/Private/Animation/AnimationDecompression.cpp
@@ -95,7 +95,8 @@ void DecompressPose(FCompactPose& OutPose,
 
 			if (TargetSkeletonBoneIndex != INDEX_NONE)
 			{
-				const FCompactPoseBoneIndex BoneIndex = RequiredBones.GetCompactPoseIndexFromSkeletonIndex(TargetSkeletonBoneIndex);
+				const FSkeletonPoseBoneIndex SkeletonPoseBoneIndex(TargetSkeletonBoneIndex);
+				const FCompactPoseBoneIndex BoneIndex = RequiredBones.GetCompactPoseIndexFromSkeletonPoseIndex(SkeletonPoseBoneIndex);
 				//Nasty, we break our type safety, code in the lower levels should be adjusted for this
 				const int32 CompactPoseBoneIndex = BoneIndex.GetInt();
 				if (CompactPoseBoneIndex != INDEX_NONE)
