Creating an AnimStreamable asset immediately crashes the editor
#fix patch-up AnimStreamable to a non-crashing and playing back state

diff --git a/Engine/Source/Runtime/Engine/Private/Animation/AnimStreamable.cpp b/Engine/Source/Runtime/Engine/Private/Animation/AnimStreamable.cpp
--- a/Engine/Source/Runtime/Engine/Private/Animation/AnimStreamable.cpp
+++ b/Engine/Source/Runtime/Engine/Private/Animation/AnimStreamable.cpp
@@ -493,6 +493,9 @@ void UAnimStreamable::InitFrom(const UAnimSequence* InSourceSequence)
 	RetargetSource = InSourceSequence->RetargetSource;
 	NumberOfKeys = DataModelInterface->GetNumberOfKeys();
 	SamplingFrameRate = SourceSequence->GetDataModel()->GetFrameRate();
+	PRAGMA_DISABLE_DEPRECATION_WARNINGS
+	SequenceLength = SamplingFrameRate.AsSeconds(NumberOfKeys - 1);
+	PRAGMA_ENABLE_DEPRECATION_WARNINGS
 
 	Controller->NotifyPopulated();
 
@@ -732,6 +735,8 @@ void UAnimStreamable::RequestCompressedDataForChunk(const FString& ChunkDDCKey,
 				CompressibleData->RawFloatCurves = DataModelInterface->GetFloatCurves();
 			}
 
+			// Data is manually populated above, so skip AnimationCompression builds/fetch data steps
+			CompressibleData->bDataFetched = true;
 			AnimCompressor->SetCompressibleData(CompressibleData);
 
 			if (bSkipDDC)
