1. Add virtual mip allocator for the generic texture streaming path.
2. Move IsDataOfflineProcessed from UTextureAllMipDataProviderFactory into UTextureMipDataProviderFactory so that an ensure can be fired if there is a mismatch between how a mip provider is expecting data and how the texture is storing it.
3. Fixup engine's use cases of UTextureMipDataProviderFactory to ensure data is loaded and validated correctly.
4. Add ShouldAllowPlatformTiling function to UTextureMipDataProviderFactory, which is tested in UTexture::IsCookPlatformTilingDisabled, prevent platforms from offline processing texture data at the requires of a UTextureMipDataProviderFactory.

diff --git a/Engine/Plugins/Mutable/Source/CustomizableObject/Internal/MuCO/CustomizableObjectMipDataProvider.h b/Engine/Plugins/Mutable/Source/CustomizableObject/Internal/MuCO/CustomizableObjectMipDataProvider.h
--- a/Engine/Plugins/Mutable/Source/CustomizableObject/Internal/MuCO/CustomizableObjectMipDataProvider.h
+++ b/Engine/Plugins/Mutable/Source/CustomizableObject/Internal/MuCO/CustomizableObjectMipDataProvider.h
@@ -184,6 +184,8 @@ public:
 		return true;
 	}
 
+	virtual bool ShouldAllowPlatformTiling(const UTexture* Owner) const override;
+
 	// Todo: Simplify by replacing the reference to the Instance with some static parametrization or hash with enough info to reconstruct the texture
 	UPROPERTY(Transient)
 	TObjectPtr<UCustomizableObjectInstance> CustomizableObjectInstance = nullptr;
diff --git a/Engine/Plugins/Mutable/Source/CustomizableObject/Private/MuCO/CustomizableObjectInstance.cpp b/Engine/Plugins/Mutable/Source/CustomizableObject/Private/MuCO/CustomizableObjectInstance.cpp
--- a/Engine/Plugins/Mutable/Source/CustomizableObject/Private/MuCO/CustomizableObjectInstance.cpp
+++ b/Engine/Plugins/Mutable/Source/CustomizableObject/Private/MuCO/CustomizableObjectInstance.cpp
@@ -2553,9 +2553,12 @@ bool UCustomizableInstancePrivate::UpdateSkeletalMesh_PostBeginUpdate0(UCustomiz
 		ensure(RenderData && RenderData->LODRenderData.Num() > 0);
 		ensure(SkeletalMesh->GetLODNum() > 0);
 
-		for (FSkeletalMeshLODRenderData& LODResource : RenderData->LODRenderData)
+		if(RenderData)
 		{
-			UnrealConversionUtils::UpdateSkeletalMeshLODRenderDataBuffersSize(LODResource);
+			for (FSkeletalMeshLODRenderData& LODResource : RenderData->LODRenderData)
+			{
+				UnrealConversionUtils::UpdateSkeletalMeshLODRenderDataBuffersSize(LODResource);
+			}
 		}
 
 		if (OperationData->bUseMeshCache)
@@ -3327,6 +3330,14 @@ namespace
 	}
 }
 
+bool MutableTextureUsesOfflineProcessedData()
+{
+#if PLATFORM_DESKTOP || PLATFORM_ANDROID || PLATFORM_IOS
+	return true;
+#else
+	return false;
+#endif
+}
 
 void SetTexturePropertiesFromMutableImageProps(UTexture2D* Texture, const FMutableModelImageProperties& Props, bool bNeverStream)
 {
@@ -3339,14 +3350,10 @@ void SetTexturePropertiesFromMutableImageProps(UTexture2D* Texture, const FMutab
 	{
 		Texture->NeverStream = bNeverStream;
 	}
-
-#if !PLATFORM_ANDROID && !PLATFORM_IOS
-	Texture->bNotOfflineProcessed = true;
-#endif
-
 #else
 	Texture->NeverStream = bNeverStream;
 #endif
+	Texture->bNotOfflineProcessed = !MutableTextureUsesOfflineProcessedData();
 
 	Texture->SRGB = Props.SRGB;
 	Texture->Filter = Props.Filter;
diff --git a/Engine/Plugins/Mutable/Source/CustomizableObject/Private/MuCO/CustomizableObjectMipDataProvider.cpp b/Engine/Plugins/Mutable/Source/CustomizableObject/Private/MuCO/CustomizableObjectMipDataProvider.cpp
--- a/Engine/Plugins/Mutable/Source/CustomizableObject/Private/MuCO/CustomizableObjectMipDataProvider.cpp
+++ b/Engine/Plugins/Mutable/Source/CustomizableObject/Private/MuCO/CustomizableObjectMipDataProvider.cpp
@@ -18,6 +18,12 @@
 
 #define UE_MUTABLE_MIPDATA_PROVIDER_UPDATE_IMAGE_REGION		TEXT("Task_Mutable_UpdateImage")
 
+bool MutableTextureUsesOfflineProcessedData();
+bool UMutableTextureMipDataProviderFactory::ShouldAllowPlatformTiling(const UTexture* Owner) const
+{
+	return MutableTextureUsesOfflineProcessedData();
+}
+
 UMutableTextureMipDataProviderFactory::UMutableTextureMipDataProviderFactory(const FObjectInitializer& ObjectInitializer)
 	: Super(ObjectInitializer)
 {
diff --git a/Engine/Source/Runtime/Engine/Classes/Engine/TextureAllMipDataProviderFactory.h b/Engine/Source/Runtime/Engine/Classes/Engine/TextureAllMipDataProviderFactory.h
--- a/Engine/Source/Runtime/Engine/Classes/Engine/TextureAllMipDataProviderFactory.h
+++ b/Engine/Source/Runtime/Engine/Classes/Engine/TextureAllMipDataProviderFactory.h
@@ -52,9 +52,4 @@ public:
 	  */
 	ENGINE_API virtual FStreamableRenderResourceState GetResourcePostInitState(const UTexture* Owner, bool bAllowStreaming)
 		PURE_VIRTUAL(UTextureAllMipDataProviderFactory::GetResourcePostInitState, return FStreamableRenderResourceState(););
-
-	/**
-	  * Return whether the mip data provided by this provider is in "offline processed" format (i.e. pre-tiled specifically for the target platform)
-	  */
-	virtual bool IsDataOfflineProcessed(const UTexture* Owner) { return false; }
 };
diff --git a/Engine/Source/Runtime/Engine/Classes/Engine/TextureMipDataProviderFactory.h b/Engine/Source/Runtime/Engine/Classes/Engine/TextureMipDataProviderFactory.h
--- a/Engine/Source/Runtime/Engine/Classes/Engine/TextureMipDataProviderFactory.h
+++ b/Engine/Source/Runtime/Engine/Classes/Engine/TextureMipDataProviderFactory.h
@@ -39,4 +39,9 @@ public:
 	* even without loading from disk at all, so streaming can be enabled for their textures.
 	*/
 	virtual bool WillProvideMipDataWithoutDisk() const { return false; }
+
+	/**
+	  * Return whether the mip data provided by this provider is in "offline processed" format (i.e. pre-tiled specifically for the target platform)
+	  */
+	virtual bool ShouldAllowPlatformTiling(const UTexture* Owner) const { return false; }
 };
diff --git a/Engine/Source/Runtime/Engine/Private/Rendering/StreamableTextureResource.cpp b/Engine/Source/Runtime/Engine/Private/Rendering/StreamableTextureResource.cpp
--- a/Engine/Source/Runtime/Engine/Private/Rendering/StreamableTextureResource.cpp
+++ b/Engine/Source/Runtime/Engine/Private/Rendering/StreamableTextureResource.cpp
@@ -55,6 +55,28 @@ bool CanCreateWithPartiallyResidentMips(ETextureCreateFlags TexCreateFlags)
 #endif
 }
 
+static bool ShouldAllowPlatformTiling(UTexture* Texture)
+{
+	UTextureAllMipDataProviderFactory* AllMipProvider = nullptr;
+	const TArray<UAssetUserData*>* UserDataArray = Texture->GetAssetUserDataArray();
+
+	bool bTextureDataIsTiled = !Texture->bNotOfflineProcessed;
+	if (UTextureMipDataProviderFactory* MipProvider = Texture->GetAssetUserData<UTextureMipDataProviderFactory>())
+	{
+		// If an UTextureAllMipDataProviderFactory is provided, the original texture data is unused, and all
+		// texture data will come from the provider, and should always determine if the data is offline processed
+		// or not.
+		bool ProviderDataIsTiled = MipProvider->ShouldAllowPlatformTiling(Texture);
+		if (MipProvider->IsA<UTextureAllMipDataProviderFactory>())
+		{
+			return ProviderDataIsTiled;
+		}
+		ensureMsgf(ProviderDataIsTiled == bTextureDataIsTiled, TEXT("A mip provider must provide texture data in the same layout as the texture itself."));
+	}
+	
+	return bTextureDataIsTiled;
+}
+
 /** Scoped debug info that provides the texture name to memory allocation and crash callstacks. */
 class FStreamableTextureScopedDebugInfo : public FScopedDebugInfo
 {
@@ -122,23 +144,20 @@ FStreamableTextureResource::FStreamableTextureResource(UTexture* InOwner, const
 	SizeY = Mip0.SizeY;
 	SizeZ = Mip0.SizeZ;
 
-	UTextureAllMipDataProviderFactory* AllMipProvider = InOwner->GetAssetUserData<UTextureAllMipDataProviderFactory>();
-	const bool bNotOfflineProcessed = (AllMipProvider != nullptr) ? (!AllMipProvider->IsDataOfflineProcessed(InOwner)) : (InOwner->bNotOfflineProcessed);
-
 	MipFadeSetting = (LODGroup == TEXTUREGROUP_Lightmap || LODGroup == TEXTUREGROUP_Shadowmap) ? MipFade_Slow : MipFade_Normal;
-	CreationFlags = (InOwner->SRGB ? TexCreate_SRGB : TexCreate_None)  | (bNotOfflineProcessed ? TexCreate_None : TexCreate_OfflineProcessed) | TexCreate_ShaderResource | (InOwner->bNoTiling ? TexCreate_NoTiling : TexCreate_None);
+	CreationFlags = 
+		  (InOwner->SRGB ? TexCreate_SRGB : TexCreate_None)
+		| (ShouldAllowPlatformTiling(InOwner) ? TexCreate_OfflineProcessed : TexCreate_None)
+		| TexCreate_ShaderResource 
+		| (InOwner->bNoTiling ? TexCreate_NoTiling : TexCreate_None);
 
 	if (InPostInitState.MaxNumLODs > 1)
 	{
 		CreationFlags |= TexCreate_Streamable;
 	}
 
-	// Whether the virtual update path is enabled for this texture. This allows to map / unmap top mip memory in an out.
-	// Whether the texture will be created with TexCreate_Virtual depends on the requested mip count and "r.VirtualTextureReducedMemory"
-	// TODO: Remove the bHasNoMipProvider filter here, once virtual mip support is added to the Mip Provider streaming path
-	const bool bHasNoMipProvider = (InOwner->GetAssetUserData<UTextureMipDataProviderFactory>() == nullptr);
-	bUsePartiallyResidentMips = bHasNoMipProvider && bAllowPartiallyResidentMips && InPostInitState.bSupportsStreaming && CanCreateWithPartiallyResidentMips(CreationFlags);
-	 
+	bUsePartiallyResidentMips = bAllowPartiallyResidentMips && InPostInitState.bSupportsStreaming && CanCreateWithPartiallyResidentMips(CreationFlags);
+
 	STAT(LODGroupStatName = TextureGroupStatFNames[LODGroup]);
 	STAT(bIsNeverStream = InOwner->NeverStream);
 }
diff --git a/Engine/Source/Runtime/Engine/Private/Streaming/Texture2DMipAllocator_Virtual.cpp b/Engine/Source/Runtime/Engine/Private/Streaming/Texture2DMipAllocator_Virtual.cpp
new file mode 100644
--- /dev/null
+++ b/Engine/Source/Runtime/Engine/Private/Streaming/Texture2DMipAllocator_Virtual.cpp
@@ -0,0 +1,173 @@
+// Copyright Epic Games, Inc. All Rights Reserved.
+
+/*=============================================================================
+Texture2DStreamIn_Virtual.cpp: Load texture 2D mips using ITextureMipDataProvider
+=============================================================================*/
+
+#include "Texture2DMipAllocator_Virtual.h"
+#include "RenderUtils.h"
+#include "Rendering/StreamableTextureResource.h"
+#include "Streaming/TextureMipAllocator.h"
+#include "Containers/ResourceArray.h"
+
+extern TAutoConsoleVariable<int32> CVarFlushRHIThreadOnSTreamingTextureLocks;
+
+FTexture2DMipAllocator_Virtual::FTexture2DMipAllocator_Virtual(UTexture* Texture)
+	: FTextureMipAllocator(Texture, ETickState::AllocateMips, ETickThread::Render)
+{
+}
+
+FTexture2DMipAllocator_Virtual::~FTexture2DMipAllocator_Virtual()
+{
+	check(!LockedMipIndices.Num());
+}
+
+// ********************************************************
+// ********* FTextureMipAllocator implementation **********
+// ********************************************************
+
+bool FTexture2DMipAllocator_Virtual::AllocateMips(
+	const FTextureUpdateContext& Context, 
+	FTextureMipInfoArray& OutMipInfos, 
+	const FTextureUpdateSyncOptions& SyncOptions)
+{
+	check(PendingFirstLODIdx < CurrentFirstLODIdx);
+
+	FRHITexture* Texture2DRHI = Context.Resource ? Context.Resource->GetTexture2DRHI() : nullptr;
+	if (!Texture2DRHI)
+	{
+		return false;
+	}
+
+	// Step (1) : Create the texture on the renderthread .......
+	if (!IntermediateTextureRHI)
+	{
+		DoConvertToVirtualWithNewMips(Context);
+		// Run the next step, when IntermediateTextureRHI will be somewhat valid (after synchronization).
+		AdvanceTo(ETickState::AllocateMips, ETickThread::Render);
+		return true;
+	}
+	// Step (2) : Finalize the texture state through.... .
+	else
+	{
+		const bool bFlushRHIThread = CVarFlushRHIThreadOnSTreamingTextureLocks.GetValueOnAnyThread() > 0;
+
+		RHIVirtualTextureSetFirstMipInMemory(IntermediateTextureRHI, PendingFirstLODIdx);
+
+		OutMipInfos.AddDefaulted(CurrentFirstLODIdx);
+
+		for (int32 MipIndex = PendingFirstLODIdx; MipIndex < CurrentFirstLODIdx; ++MipIndex)
+		{
+			const FTexture2DMipMap& OwnerMip = *Context.MipsView[MipIndex];
+			FTextureMipInfo& MipInfo = OutMipInfos[MipIndex];
+
+			MipInfo.Format = Context.Resource->GetPixelFormat();
+			MipInfo.SizeX = OwnerMip.SizeX;
+			MipInfo.SizeY = OwnerMip.SizeY;
+#if WITH_EDITORONLY_DATA
+			MipInfo.DataSize = CalcTextureMipMapSize(MipInfo.SizeX, MipInfo.SizeY, MipInfo.Format, 0);
+#else // Hasn't really been used on console. To investigate!
+			MipInfo.DataSize = 0;
+#endif
+
+			MipInfo.DestData = RHILockTexture2D(IntermediateTextureRHI, MipIndex, RLM_WriteOnly, MipInfo.RowPitch, false, bFlushRHIThread);
+
+			// Add this mip in the locked list of mips so that it can safely be unlocked when needed.
+			LockedMipIndices.Add(MipIndex);
+		}
+
+		// New mips are ready to be unlocked by the FTextureMipDataProvider implementation.
+		AdvanceTo(ETickState::UploadMips, ETickThread::Render);
+		return true;
+	}
+}
+
+bool FTexture2DMipAllocator_Virtual::UploadMips(const FTextureUpdateContext& Context, const FTextureUpdateSyncOptions& SyncOptions)
+{
+	if (!IntermediateTextureRHI)
+	{
+		return false;
+	}
+
+	// Unlock the mips so that the texture can be updated.
+	UnlockNewMips();
+
+	if(IntermediateTextureRHI)
+	{
+		RHIVirtualTextureSetFirstMipVisible(IntermediateTextureRHI, PendingFirstLODIdx);
+	}
+
+	AdvanceTo(ETickState::Done, ETickThread::None);
+	return true;
+}
+
+void FTexture2DMipAllocator_Virtual::Cancel(const FTextureUpdateSyncOptions& SyncOptions)
+{
+	// Unlock any locked mips.
+	UnlockNewMips();
+
+	if (IntermediateTextureRHI)
+	{
+		RHIVirtualTextureSetFirstMipInMemory(IntermediateTextureRHI, CurrentFirstLODIdx);
+	}
+}
+
+FTextureMipAllocator::ETickThread FTexture2DMipAllocator_Virtual::GetCancelThread() const
+{
+	// Nothing to do.
+	return ETickThread::None;
+}
+
+// ****************************
+// ********* Helpers **********
+// ****************************
+
+void FTexture2DMipAllocator_Virtual::UnlockNewMips()
+{
+	// Unlock any locked mips.
+	if (IntermediateTextureRHI)
+	{
+		const bool bFlushRHIThread = CVarFlushRHIThreadOnSTreamingTextureLocks.GetValueOnAnyThread() > 0;
+		for (int32 MipIndex : LockedMipIndices)
+		{
+			RHIUnlockTexture2D(IntermediateTextureRHI, MipIndex, false, CVarFlushRHIThreadOnSTreamingTextureLocks.GetValueOnAnyThread() > 0 );
+		}
+		LockedMipIndices.Empty();
+	}
+}
+
+void FTexture2DMipAllocator_Virtual::DoConvertToVirtualWithNewMips(const FTextureUpdateContext& Context)
+{
+	if (Context.Resource)
+	{
+		// If the texture is not virtual, then make it virtual immediately.
+		if (!Context.Resource->GetTexture2DRHI() || (Context.Resource->GetTexture2DRHI()->GetFlags() & ETextureCreateFlags::Virtual) != ETextureCreateFlags::Virtual)
+		{
+			const FTexture2DMipMap& MipMap0 = *Context.MipsView[0];
+
+			FTexture2DResource* Resource2D = Context.Resource->GetTexture2DResource();
+
+			// Create a copy of the texture that is a virtual texture.
+			const FRHITextureCreateDesc Desc =
+				FRHITextureCreateDesc::Create2D(TEXT("FTexture2DUpdate"), MipMap0.SizeX, MipMap0.SizeY, Context.Resource->GetPixelFormat())
+				.SetNumMips(ResourceState.MaxNumLODs)
+				.SetFlags(Context.Resource->GetCreationFlags() | ETextureCreateFlags::Virtual)
+				.SetBulkData(Resource2D->ResourceMem);
+
+			IntermediateTextureRHI = RHICreateTexture(Desc);
+
+			RHIVirtualTextureSetFirstMipInMemory(IntermediateTextureRHI, CurrentFirstLODIdx);
+			RHIVirtualTextureSetFirstMipVisible(IntermediateTextureRHI, CurrentFirstLODIdx);
+
+			UE::RHI::CopySharedMips_AssumeSRVMaskState(
+				FRHICommandListExecutor::GetImmediateCommandList(),
+				Context.Resource->GetTexture2DRHI(),
+				IntermediateTextureRHI);
+		}
+		else
+		{
+			// Otherwise the current texture is already virtual and we can update it directly.
+			IntermediateTextureRHI = Context.Resource->GetTexture2DRHI();
+		}
+	}
+}
diff --git a/Engine/Source/Runtime/Engine/Private/Streaming/Texture2DMipAllocator_Virtual.h b/Engine/Source/Runtime/Engine/Private/Streaming/Texture2DMipAllocator_Virtual.h
new file mode 100644
--- /dev/null
+++ b/Engine/Source/Runtime/Engine/Private/Streaming/Texture2DMipAllocator_Virtual.h
@@ -0,0 +1,44 @@
+// Copyright Epic Games, Inc. All Rights Reserved.
+
+/*=============================================================================
+Texture2DMipAllocator_AsyncReallocate.h: Implementation of FTextureMipAllocator using RHIAsyncReallocateTexture2D
+=============================================================================*/
+
+#pragma once
+
+#include "CoreMinimal.h"
+#include "TextureMipAllocator.h"
+#include "Engine/Texture2D.h"
+#include "Rendering/Texture2DResource.h"
+
+class UTexture2D;
+
+/**
+* FTexture2DMipAllocator_Virtual is an allocator specialized in creating and managing virtual textures.
+*/
+class FTexture2DMipAllocator_Virtual : public FTextureMipAllocator
+{
+public:
+
+	FTexture2DMipAllocator_Virtual(UTexture* Texture);
+	~FTexture2DMipAllocator_Virtual();
+
+	// ********************************************************
+	// ********* FTextureMipAllocator implementation **********
+	// ********************************************************
+
+	bool AllocateMips(const FTextureUpdateContext& Context, FTextureMipInfoArray& OutMipInfos, const FTextureUpdateSyncOptions& SyncOptions) final override;
+	bool UploadMips(const FTextureUpdateContext& Context, const FTextureUpdateSyncOptions& SyncOptions) final override;
+	void Cancel(const FTextureUpdateSyncOptions& SyncOptions) final override;
+	ETickThread GetCancelThread() const final override;
+
+protected:
+
+	// Unlock the mips referenced in LockedMipIndices.
+	void UnlockNewMips();
+
+	void DoConvertToVirtualWithNewMips(const FTextureUpdateContext& Context);
+
+	// The list of mips that are currently locked.
+	TArray<int32, TInlineAllocator<MAX_TEXTURE_MIP_COUNT>> LockedMipIndices;
+};
diff --git a/Engine/Source/Runtime/Engine/Private/Streaming/TextureMipDataProvider.cpp b/Engine/Source/Runtime/Engine/Private/Streaming/TextureMipDataProvider.cpp
--- a/Engine/Source/Runtime/Engine/Private/Streaming/TextureMipDataProvider.cpp
+++ b/Engine/Source/Runtime/Engine/Private/Streaming/TextureMipDataProvider.cpp
@@ -33,3 +33,10 @@ FTextureMipDataProvider::FTextureMipDataProvider(const UTexture* Texture, ETickS
 	, NextTickThread(InTickThread) 
 {
 }
+
+bool FTextureMipDataProvider::PostprocessGPUMipData(const FTextureUpdateContext& Context, const FTextureRHIRef& TextureRHI, int32 FirstMip, int32 NumMips, const FTextureUpdateSyncOptions& SyncOptions)
+{
+PRAGMA_DISABLE_DEPRECATION_WARNINGS
+	return PostprocessGPUMipData(Context, TextureRHI, SyncOptions);
+PRAGMA_ENABLE_DEPRECATION_WARNINGS
+}
diff --git a/Engine/Source/Runtime/Engine/Private/Streaming/TextureStreamIn.cpp b/Engine/Source/Runtime/Engine/Private/Streaming/TextureStreamIn.cpp
--- a/Engine/Source/Runtime/Engine/Private/Streaming/TextureStreamIn.cpp
+++ b/Engine/Source/Runtime/Engine/Private/Streaming/TextureStreamIn.cpp
@@ -220,13 +220,16 @@ bool FTextureStreamIn::DoUploadNewMips(const FContext& Context)
 
 bool FTextureStreamIn::DoPostprocessGPUMipData(const FContext& Context)
 {
+	const int32 MipOffset = !!(MipAllocator->GetIntermediateTextureRHI()->GetFlags() & TexCreate_Virtual) ? 0 : PendingFirstLODIdx;
+	const int32 NumMips = CurrentFirstLODIdx - PendingFirstLODIdx;
+
 	for (TUniquePtr<FTextureMipDataProvider>& MipDataProvider : MipDataProviders)
 	{
 		check(MipDataProvider);
 		if (MipDataProvider->GetNextTickState() == FTextureMipDataProvider::ETickState::PostprocessGPUMipData
 			&& IsSameThread(MipDataProvider->GetNextTickThread(), Context.CurrentThread))
 		{
-			if (!MipDataProvider->PostprocessGPUMipData(Context, MipAllocator->GetIntermediateTextureRHI(), SyncOptions))
+			if (!MipDataProvider->PostprocessGPUMipData(Context, MipAllocator->GetIntermediateTextureRHI(), PendingFirstLODIdx - MipOffset, NumMips, SyncOptions))
 			{
 				return false;
 			}
diff --git a/Engine/Source/Runtime/Engine/Private/Texture.cpp b/Engine/Source/Runtime/Engine/Private/Texture.cpp
--- a/Engine/Source/Runtime/Engine/Private/Texture.cpp
+++ b/Engine/Source/Runtime/Engine/Private/Texture.cpp
@@ -1630,6 +1630,14 @@ TextureMipGenSettings UTexture::GetMipGenSettingsFromString(const TCHAR* InStr,
 
 bool UTexture::IsCookPlatformTilingDisabled(const ITargetPlatformSettings* TargetPlatformSettings) const
 {
+	if (const UTextureMipDataProviderFactory* MipProvider = const_cast<UTexture*>(this)->GetAssetUserData<UTextureMipDataProviderFactory>())
+	{
+		if (!MipProvider->ShouldAllowPlatformTiling(this))
+		{
+			return false;
+		}
+	}
+
 	if (CookPlatformTilingSettings.GetValue() == TextureCookPlatformTilingSettings::TCPTS_FromTextureGroup)
 	{
 		const UTextureLODSettings* TextureLODSettings = nullptr;
diff --git a/Engine/Source/Runtime/Engine/Private/Texture2D.cpp b/Engine/Source/Runtime/Engine/Private/Texture2D.cpp
--- a/Engine/Source/Runtime/Engine/Private/Texture2D.cpp
+++ b/Engine/Source/Runtime/Engine/Private/Texture2D.cpp
@@ -34,6 +34,7 @@
 #include "Streaming/TextureStreamIn.h"
 #include "Streaming/Texture2DMipAllocator_AsyncCreate.h"
 #include "Streaming/Texture2DMipAllocator_AsyncReallocate.h"
+#include "Streaming/Texture2DMipAllocator_Virtual.h"
 #include "Streaming/Texture2DMipDataProvider_DDC.h"
 #include "Streaming/Texture2DMipDataProvider_IO.h"
 #include "Engine/TextureMipDataProviderFactory.h"
@@ -1958,8 +1959,11 @@ bool UTexture2D::StreamIn(int32 NewMipCount, bool bHighPrio)
 				DefaultMipDataProvider = new FTexture2DMipDataProvider_IO(this, bHighPrio);
 			}
 
-			// FTexture2DMipAllocator_Virtual?
-			if (GRHISupportsAsyncTextureCreation)
+			if (Texture2DResource->bUsePartiallyResidentMips)
+			{
+				MipAllocator = new FTexture2DMipAllocator_Virtual(this);
+			}
+			else if (GRHISupportsAsyncTextureCreation)
 			{
 				MipAllocator = new FTexture2DMipAllocator_AsyncCreate(this);
 			}
diff --git a/Engine/Source/Runtime/Engine/Public/Rendering/Texture2DResource.h b/Engine/Source/Runtime/Engine/Public/Rendering/Texture2DResource.h
--- a/Engine/Source/Runtime/Engine/Public/Rendering/Texture2DResource.h
+++ b/Engine/Source/Runtime/Engine/Public/Rendering/Texture2DResource.h
@@ -64,6 +64,7 @@ private:
 	/** Texture streaming command classes that need to be friends in order to call Update/FinalizeMipCount.	*/
 	friend class UTexture2D;
 	friend class FTexture2DUpdate;
+	friend class FTexture2DMipAllocator_Virtual;
 
 	/** Resource memory allocated by the owner for serialize bulk mip data into								*/
 	FTexture2DResourceMem* ResourceMem;
diff --git a/Engine/Source/Runtime/Engine/Public/Streaming/TextureMipDataProvider.h b/Engine/Source/Runtime/Engine/Public/Streaming/TextureMipDataProvider.h
--- a/Engine/Source/Runtime/Engine/Public/Streaming/TextureMipDataProvider.h
+++ b/Engine/Source/Runtime/Engine/Public/Streaming/TextureMipDataProvider.h
@@ -152,10 +152,17 @@ public:
 
 	/*
 	* Once the texture has been uploaded to the GPU, the mip provider has the opportunity to perform any modifications necessary
-	* before the texture is finalized and applied to the texture resource
+	* before the texture is finalized and applied to the texture resource.
 	*/
+	UE_DEPRECATED(5.7, "Implement the version that takes FirstMip and NumMips.")
 	virtual bool PostprocessGPUMipData(const FTextureUpdateContext& Context, const FTextureRHIRef& NewTextureRHI, const FTextureUpdateSyncOptions& SyncOptions) { return true; }
 
+	/*
+	* Once the texture has been uploaded to the GPU, the mip provider has the opportunity to perform any modifications necessary
+	* before the texture is finalized and applied to the texture resource.
+	*/
+	ENGINE_API virtual bool PostprocessGPUMipData(const FTextureUpdateContext& Context, const FTextureRHIRef& TextureRHI, int32 FirstMip, int32 NumMips, const FTextureUpdateSyncOptions& SyncOptions);
+
 	/**
 	* Abort anything that could be stalling the update in PollMips(). Called from an async task.
 	* Expected behavior is to cancel any pending IO operations. 
diff --git a/Engine/Source/Runtime/Landscape/Classes/LandscapeProxy.h b/Engine/Source/Runtime/Landscape/Classes/LandscapeProxy.h
--- a/Engine/Source/Runtime/Landscape/Classes/LandscapeProxy.h
+++ b/Engine/Source/Runtime/Landscape/Classes/LandscapeProxy.h
@@ -1215,6 +1215,9 @@ PRAGMA_ENABLE_DEPRECATION_WARNINGS
 #endif // WITH_EDITOR
 
 	//~ Begin UObject Interface.
+#if WITH_EDITOR
+	virtual void BeginCacheForCookedPlatformData(const ITargetPlatform* TargetPlatform) override;
+#endif // WITH_EDITOR
 	virtual void PreSave(FObjectPreSaveContext ObjectSaveContext) override;
 	virtual void Serialize(FArchive& Ar) override;
 	static void AddReferencedObjects(UObject* InThis, FReferenceCollector& Collector);
@@ -1689,6 +1692,10 @@ private:
 
 	void PostLoadFixupLandscapeGuidsIfInstanced();
 
+#if WITH_EDITOR
+	void InstallOrUpdateTextureUserDatas(const ITargetPlatform* TargetPlatform);
+#endif // WITH_EDITOR
+
 #if WITH_EDITORONLY_DATA
 public:
 	static const TArray<ALandscapeProxy*>& GetLandscapeProxies() { return LandscapeProxies; }
diff --git a/Engine/Source/Runtime/Landscape/Private/Landscape.cpp b/Engine/Source/Runtime/Landscape/Private/Landscape.cpp
--- a/Engine/Source/Runtime/Landscape/Private/Landscape.cpp
+++ b/Engine/Source/Runtime/Landscape/Private/Landscape.cpp
@@ -2192,6 +2192,7 @@ ULandscapeHeightmapTextureEdgeFixup* ULandscapeComponent::InstallOrUpdateTexture
 			}
 			OverrideFactory->SetupEdgeFixup(EdgeFixup);
 		}
+
 	}
 	else
 	{
@@ -2204,6 +2205,11 @@ ULandscapeHeightmapTextureEdgeFixup* ULandscapeComponent::InstallOrUpdateTexture
 		}
 	}
 
+#if WITH_EDITOR
+	// The EdgeFixup will always require linear texture data and should not apply per platform offline processing
+	HeightmapTexture->bNotOfflineProcessed = (EdgeFixup != nullptr);
+#endif // WITH_EDITOR
+
 	// double check we've achieved the desired relationships..
 	if (EdgeFixup)
 	{
@@ -2220,6 +2226,42 @@ ULandscapeHeightmapTextureEdgeFixup* ULandscapeComponent::InstallOrUpdateTexture
 	return EdgeFixup;
 }
 
+#if WITH_EDITOR
+void ALandscapeProxy::InstallOrUpdateTextureUserDatas(const ITargetPlatform* TargetPlatform)
+{
+	FName IniPlatformName = *TargetPlatform->IniPlatformName();
+
+	int32 HeightmapCompressionMode = 0;
+	if (IConsoleVariable* PlatformCVar = CVarLandscapeHeightmapCompressionMode->GetPlatformValueVariable(IniPlatformName).Get())
+	{
+		PlatformCVar->GetValue(HeightmapCompressionMode);
+	}
+	else
+	{
+		CVarLandscapeHeightmapCompressionMode->GetValue(HeightmapCompressionMode);
+	}
+
+	int32 HeightmapCompressionMipThreshold = 32;
+	if (IConsoleVariable* PlatformCVar = CVarLandscapeHeightmapCompressionMipThreshold->GetPlatformValueVariable(IniPlatformName).Get())
+	{
+		PlatformCVar->GetValue(HeightmapCompressionMipThreshold);
+	}
+	else
+	{
+		CVarLandscapeHeightmapCompressionMipThreshold->GetValue(HeightmapCompressionMipThreshold);
+	}
+
+	const bool bShouldCompressHeightmap = (HeightmapCompressionMode > 0);
+	const bool bIsStreamingProxy = this->IsA<ALandscapeStreamingProxy>();
+	const bool bUseEdgeFixups = bIsStreamingProxy;
+	const bool bUpdateSnapshotNow = true;
+	for (ULandscapeComponent* LandscapeComponent : LandscapeComponents)
+	{
+		LandscapeComponent->InstallOrUpdateTextureUserDatas(bUseEdgeFixups, bShouldCompressHeightmap, bUpdateSnapshotNow, HeightmapCompressionMipThreshold);
+	}
+}
+#endif // WITH_EDITOR
+
 void ULandscapeComponent::OnRegister()
 {
 	Super::OnRegister();
@@ -3760,8 +3802,15 @@ void ALandscapeProxy::OnFeatureLevelChanged(ERHIFeatureLevel::Type NewFeatureLev
 
 	UpdateRenderingMethod();
 }
+
+void ALandscapeProxy::BeginCacheForCookedPlatformData(const ITargetPlatform* TargetPlatform)
+{
+	InstallOrUpdateTextureUserDatas(TargetPlatform);
+	Super::BeginCacheForCookedPlatformData(TargetPlatform);
+}
 #endif
 
+
 void ALandscapeProxy::PreSave(FObjectPreSaveContext ObjectSaveContext)
 {
 	Super::PreSave(ObjectSaveContext);
@@ -3865,36 +3914,7 @@ void ALandscapeProxy::PreSave(FObjectPreSaveContext ObjectSaveContext)
 
 	if (ObjectSaveContext.IsCooking())
 	{
-		FName IniPlatformName = *ObjectSaveContext.GetTargetPlatform()->IniPlatformName();
-
-		int32 HeightmapCompressionMode = 0;
-		if (IConsoleVariable* PlatformCVar = CVarLandscapeHeightmapCompressionMode->GetPlatformValueVariable(IniPlatformName).Get())
-		{
-			PlatformCVar->GetValue(HeightmapCompressionMode);
-		}
-		else
-		{
-			CVarLandscapeHeightmapCompressionMode->GetValue(HeightmapCompressionMode);
-		}
-
-		int32 HeightmapCompressionMipThreshold = 32;
-		if (IConsoleVariable* PlatformCVar = CVarLandscapeHeightmapCompressionMipThreshold->GetPlatformValueVariable(IniPlatformName).Get())
-		{
-			PlatformCVar->GetValue(HeightmapCompressionMipThreshold);
-		}
-		else
-		{
-			CVarLandscapeHeightmapCompressionMipThreshold->GetValue(HeightmapCompressionMipThreshold);
-		}
-
-		const bool bShouldCompressHeightmap = (HeightmapCompressionMode > 0);
-		const bool bIsStreamingProxy = this->IsA<ALandscapeStreamingProxy>();
-		const bool bUseEdgeFixups = bIsStreamingProxy;
-		const bool bUpdateSnapshotNow = true;
-		for (ULandscapeComponent* LandscapeComponent : LandscapeComponents)
-		{
-			LandscapeComponent->InstallOrUpdateTextureUserDatas(bUseEdgeFixups, bShouldCompressHeightmap, bUpdateSnapshotNow, HeightmapCompressionMipThreshold);
-		}
+		InstallOrUpdateTextureUserDatas(ObjectSaveContext.GetTargetPlatform());
 	}
 #endif // WITH_EDITOR
 }
