Fix incorrect double negative inverse logic when testing UTextureMipDataProviderFactory::ShouldAllowPlatformTiling

diff --git a/Engine/Source/Runtime/Engine/Private/Texture.cpp b/Engine/Source/Runtime/Engine/Private/Texture.cpp
--- a/Engine/Source/Runtime/Engine/Private/Texture.cpp
+++ b/Engine/Source/Runtime/Engine/Private/Texture.cpp
@@ -1634,7 +1634,7 @@ bool UTexture::IsCookPlatformTilingDisabled(const ITargetPlatformSettings* Targe
 	{
 		if (!MipProvider->ShouldAllowPlatformTiling(this))
 		{
-			return false;
+			return true;
 		}
 	}
 
