Fix material translation crash in TextureStreamingInfo when input coordinate expression is not valid.
#rnx

diff --git a/Engine/Source/Runtime/Engine/Private/Materials/HLSLMaterialTranslator.cpp b/Engine/Source/Runtime/Engine/Private/Materials/HLSLMaterialTranslator.cpp
--- a/Engine/Source/Runtime/Engine/Private/Materials/HLSLMaterialTranslator.cpp
+++ b/Engine/Source/Runtime/Engine/Private/Materials/HLSLMaterialTranslator.cpp
@@ -8167,6 +8167,11 @@ int32 FHLSLMaterialTranslator::TextureStreamingInfo(int32 TextureReferenceIndex,
 		return INDEX_NONE;
 	}
 
+	if (CoordinateIndex == INDEX_NONE)
+	{
+		return INDEX_NONE;
+	}
+
 	return AddCodeChunk(MCT_Float, TEXT("MaterialStoreTexCoordScale(Parameters, %s, %d)"), *CoerceParameter(CoordinateIndex, MCT_Float2), TextureReferenceIndex);
 }
 
