Move the ChunkMap Generation to the end of the cook.
ChunkMap generation being at the end of the cook allows us to provide all the packages which were actually cooked and the final state of the AssetRegistry to know all the actual package dependencies and references.
Some package dependencies and references might be convereted from editorOnly to Game dependencies during the cook.

diff --git a/Engine/Plugins/Editor/AssetManagerEditor/Source/AssetManagerEditor/Private/AssetManagerEditorModule.cpp b/Engine/Plugins/Editor/AssetManagerEditor/Source/AssetManagerEditor/Private/AssetManagerEditorModule.cpp
--- a/Engine/Plugins/Editor/AssetManagerEditor/Source/AssetManagerEditor/Private/AssetManagerEditorModule.cpp
+++ b/Engine/Plugins/Editor/AssetManagerEditor/Source/AssetManagerEditor/Private/AssetManagerEditorModule.cpp
@@ -2187,7 +2187,7 @@ void FAssetManagerEditorModule::SetCurrentRegistrySource(const FString& SourceNa
 
 void FAssetManagerEditorModule::RefreshRegistryData()
 {
-	UAssetManager::Get().UpdateManagementDatabase(true);
+	UAssetManager::Get().UpdateManagementDatabase(EUpdateManagementDatabaseFlags::BuildChunkMap | EUpdateManagementDatabaseFlags::ForceRefresh);
 
 	// Rescan registry sources, try to restore the current one
 	const FString OldSourceName = CurrentRegistrySource ? CurrentRegistrySource->SourceName : FString();
diff --git a/Engine/Source/Editor/UnrealEd/Private/Commandlets/AssetRegistryGenerator.cpp b/Engine/Source/Editor/UnrealEd/Private/Commandlets/AssetRegistryGenerator.cpp
--- a/Engine/Source/Editor/UnrealEd/Private/Commandlets/AssetRegistryGenerator.cpp
+++ b/Engine/Source/Editor/UnrealEd/Private/Commandlets/AssetRegistryGenerator.cpp
@@ -1750,6 +1750,11 @@ void FAssetRegistryGenerator::FinalizeChunkIDs(const TSet<FName>& InCookedPackag
 		}
 	});
 
+	if (bGenerateChunks)
+	{
+		// Build the chunk map with only the packages that have been cooked.
+		UAssetManager::Get().UpdateCachedChunkMapAfterCook(CookedPackages, StartupPackages);
+	}
 	// add all the packages to the unassigned package list
 	for (FName CookedPackage : CookedPackages)
 	{
diff --git a/Engine/Source/Runtime/Engine/Classes/Engine/AssetManager.h b/Engine/Source/Runtime/Engine/Classes/Engine/AssetManager.h
--- a/Engine/Source/Runtime/Engine/Classes/Engine/AssetManager.h
+++ b/Engine/Source/Runtime/Engine/Classes/Engine/AssetManager.h
@@ -67,6 +67,14 @@ struct FAssetManagerLoadParams
 	TOptional<FAssetManagerDownloadParams> DownloadParams;
 };
 
+enum class EUpdateManagementDatabaseFlags : uint32
+{
+	None = 0,							// No flags
+	BuildChunkMap = (1 << 0),			// if the CachedChunkMap should be built or not.
+	ForceRefresh = (1 << 1),			// When set the database will be forcibly refreshed.
+};
+ENUM_CLASS_FLAGS(EUpdateManagementDatabaseFlags);
+
 /** 
  * A singleton UObject that is responsible for loading and unloading PrimaryAssets, and maintaining game-specific asset references
  * Games should override this class and change the class reference
@@ -845,7 +853,6 @@ public:
 	 */
 	virtual FGuid GetChunkEncryptionKeyGuid(int32 InChunkId) const { return FGuid(); }
 
-
 	/** Get the encrypted chunk assignments for a give package. Only valid between PreSaveAssetRegistry and PostSaveAssetRegistry */
 	virtual TSet<int32> GetEncryptedChunkIDsForPackage(FName InPackageName) const { return TSet<int32>(); }
 
@@ -868,11 +875,21 @@ public:
 	ENGINE_API virtual void ReinitializeFromConfig();
 
 	/** Updates the asset management database if needed */
-	ENGINE_API virtual void UpdateManagementDatabase(bool bForceRefresh = false);
+	UE_DEPRECATED(5.6, "UpdateManagementDatabase as been updated to take InFlags instead.")
+	ENGINE_API virtual void UpdateManagementDatabase(bool bForceRefresh);
+	ENGINE_API virtual void UpdateManagementDatabase(EUpdateManagementDatabaseFlags InFlags = EUpdateManagementDatabaseFlags::BuildChunkMap);
 
 	/** Returns the chunk information map computed during UpdateManagementDatabase */
 	ENGINE_API const TMap<int32, FAssetManagerChunkInfo>& GetChunkManagementMap() const;
 
+	/**
+	 * Called post cook to build the chunk map only with the packages which have been cooked.
+	 * Provides the StartupPackages which will go into chunk 0.
+	 *
+	 * It's expected this function will set the value of CachedChunkMap which can be used for package chunk assignment.
+	 */
+	UE_INTERNAL ENGINE_API void UpdateCachedChunkMapAfterCook(const TSet<FName>& AllCookedPackages, const TSet<FName>& StartupPackages);
+
 	/** Handles applying Asset Labels and should be overridden to do any game-specific labelling */
 	ENGINE_API virtual void ApplyPrimaryAssetLabels();
 
@@ -1061,8 +1078,20 @@ protected:
 	 * @param PackagesToUpdateChunksFor - the packages needing chunk assignment as part of the management database update.
 	 * @return a new CachedChunkMap to be used.
 	 */
+	UE_DEPRECATED(5.6, "BuildChunkMap now takes 2 additional arguments")
 	ENGINE_API virtual TMap<int32, FAssetManagerChunkInfo> BuildChunkMap(const TSet<FName>& PackagesToUpdateChunksFor) const;
 
+	/**
+	 * Used by UpdateManagementDatabase and at the end of the cook to build the CachedChunkMap.
+	 * Should return a new cachedChunkMap.
+	 *
+	 * @param PackagesToUpdateChunksFor - the packages needing chunk assignment as part of the management database update.
+	 * @param StartupPackages - All the known startup packages. When called from UpdateManagementDatabase, this will be empty.
+	 * @param PackageThatWereCooked - All the known packages that were cooked. When called from UpdateManagementDatabase, this will be empty.
+	 * @return a new CachedChunkMap to be used.
+	 */
+	ENGINE_API virtual TMap<int32, FAssetManagerChunkInfo> BuildChunkMap(const TSet<FName>& PackagesToUpdateChunksFor, const TSet<FName>& StartupPackages, const TSet<FName>& PackageThatWereCooked) const;
+
 	/** Copy of the asset state before PIE was entered, return to that when PIE completes */
 	TMap<FPrimaryAssetId, TArray<FName>> PrimaryAssetStateBeforePIE;
 
diff --git a/Engine/Source/Runtime/Engine/Private/AssetManager.cpp b/Engine/Source/Runtime/Engine/Private/AssetManager.cpp
--- a/Engine/Source/Runtime/Engine/Private/AssetManager.cpp
+++ b/Engine/Source/Runtime/Engine/Private/AssetManager.cpp
@@ -3,6 +3,7 @@
 #include "Engine/AssetManager.h"
 
 #include "Algo/Compare.h"
+#include "Algo/Transform.h"
 #include "Algo/Unique.h"
 #include "AssetRegistry/ARFilter.h"
 #include "AssetRegistry/AssetBundleData.h"
@@ -1512,7 +1513,6 @@ void UAssetManager::SetPrimaryAssetRulesExplicitly(FPrimaryAssetId PrimaryAssetI
 
 		AssetRuleOverrides.Add(PrimaryAssetId, ExplicitRules);
 	}
-
 	bIsManagementDatabaseCurrent = false;
 }
 
@@ -3705,7 +3705,7 @@ void UAssetManager::PostInitialAssetScan()
 	if (bUpdateManagementDatabaseAfterScan)
 	{
 		bUpdateManagementDatabaseAfterScan = false;
-		UpdateManagementDatabase(true);
+		UpdateManagementDatabase(EUpdateManagementDatabaseFlags::BuildChunkMap | EUpdateManagementDatabaseFlags::ForceRefresh);
 	}
 #endif
 }
@@ -4047,6 +4047,16 @@ void UAssetManager::OnAssetRegistryFilesLoaded()
 }
 
 void UAssetManager::UpdateManagementDatabase(bool bForceRefresh)
+{
+	EUpdateManagementDatabaseFlags UpdateFlags = EUpdateManagementDatabaseFlags::BuildChunkMap;
+	if (bForceRefresh)
+	{
+		UpdateFlags |= EUpdateManagementDatabaseFlags::ForceRefresh;
+	}
+	UpdateManagementDatabase(UpdateFlags);
+}
+
+void UAssetManager::UpdateManagementDatabase(EUpdateManagementDatabaseFlags InFlags)
 {
 	if (!GIsEditor)
 	{
@@ -4063,6 +4073,7 @@ void UAssetManager::UpdateManagementDatabase(bool bForceRefresh)
 		return;
 	}
 
+	const bool bForceRefresh = EnumHasAnyFlags(InFlags, EUpdateManagementDatabaseFlags::ForceRefresh);
 	if (bIsManagementDatabaseCurrent && !bForceRefresh)
 	{
 		return;
@@ -4081,6 +4092,8 @@ void UAssetManager::UpdateManagementDatabase(bool bForceRefresh)
 	TMultiMap<FAssetIdentifier, FAssetIdentifier> NoReferenceManagementMap;
 
 	// List of packages that need to have their chunks updated
+	const UProjectPackagingSettings* ProjectPackagingSettings = GetDefault<UProjectPackagingSettings>();
+	const bool bWillBuildChunkMap = ProjectPackagingSettings && ProjectPackagingSettings->bGenerateChunks && EnumHasAnyFlags(InFlags, EUpdateManagementDatabaseFlags::BuildChunkMap);
 	TSet<FName> PackagesToUpdateChunksFor;
 	TArray<FName> AssetPackagesReferenced;
 
@@ -4111,7 +4124,10 @@ void UAssetManager::UpdateManagementDatabase(bool bForceRefresh)
 				else
 				{
 					AssetPackagesReferenced.Add(PackageName);
-					PackagesToUpdateChunksFor.Add(PackageName);
+					if (bWillBuildChunkMap)
+					{
+						PackagesToUpdateChunksFor.Add(PackageName);
+					}
 				}
 			}
 
@@ -4132,7 +4148,10 @@ void UAssetManager::UpdateManagementDatabase(bool bForceRefresh)
 						else
 						{
 							AssetPackagesReferenced.Add(PackageName);
-							PackagesToUpdateChunksFor.Add(PackageName);
+							if (bWillBuildChunkMap)
+							{
+								PackagesToUpdateChunksFor.Add(PackageName);
+							}
 						}
 					}
 				}
@@ -4160,10 +4179,10 @@ void UAssetManager::UpdateManagementDatabase(bool bForceRefresh)
 	const bool bAllowInPIE = true;
 	SlowTask.MakeDialog(bShowCancelButton, bAllowInPIE);
 
-	auto SetManagerPredicate = [this, &PackagesToUpdateChunksFor](UE::AssetRegistry::FShouldSetManagerContext& Context)
+	auto SetManagerPredicate = [this, bWillBuildChunkMap, &PackagesToUpdateChunksFor](UE::AssetRegistry::FShouldSetManagerContext& Context)
 		{
 			EAssetSetManagerResult::Type Result = this->ShouldSetManager(Context);
-			if (Result != EAssetSetManagerResult::DoNotSet && Context.Target.IsPackage())
+			if (bWillBuildChunkMap, Result != EAssetSetManagerResult::DoNotSet && Context.Target.IsPackage())
 			{
 				PackagesToUpdateChunksFor.Add(Context.Target.PackageName);
 			}
@@ -4256,13 +4275,11 @@ void UAssetManager::UpdateManagementDatabase(bool bForceRefresh)
 		SetManageContext.ManagerMap = &PrimaryAssetIdManagementMap;
 		AssetRegistry.SetManageReferences(SetManageContext);
 	}
-
-	UProjectPackagingSettings* ProjectPackagingSettings = GetMutableDefault<UProjectPackagingSettings>();
-	if (ProjectPackagingSettings && ProjectPackagingSettings->bGenerateChunks)
+	
+	if (bWillBuildChunkMap)
 	{
-		CachedChunkMap = BuildChunkMap(PackagesToUpdateChunksFor);
-	}
-
+		CachedChunkMap = BuildChunkMap(PackagesToUpdateChunksFor, {}, {});
+	}	
 	bIsManagementDatabaseCurrent = true;
 }
 
@@ -4746,6 +4763,11 @@ bool UAssetManager::GetPackageChunkIds(FName PackageName, const ITargetPlatform*
 	return GetPrimaryAssetSetChunkIds(Managers, TargetPlatform, ExistingChunkList, OutChunkList);
 }
 
+void UAssetManager::UpdateCachedChunkMapAfterCook(const TSet<FName>& AllCookedPackages, const TSet<FName>& StartupPackages)
+{
+	CachedChunkMap = BuildChunkMap(AllCookedPackages, StartupPackages, AllCookedPackages);
+}
+
 bool UAssetManager::GetPrimaryAssetSetChunkIds(const TSet<FPrimaryAssetId>& PrimaryAssetSet, const class ITargetPlatform* TargetPlatform, TArrayView<const int32> ExistingChunkList, TArray<int32>& OutChunkList) const
 {
 	bool bFoundAny = false;
@@ -4830,6 +4852,11 @@ void UAssetManager::EndPIE(bool bStartSimulate)
 }
 
 TMap<int32, FAssetManagerChunkInfo> UAssetManager::BuildChunkMap(const TSet<FName>& PackagesToUpdateChunksFor) const
+{
+	return BuildChunkMap(PackagesToUpdateChunksFor, {}, {});
+}
+
+TMap<int32, FAssetManagerChunkInfo> UAssetManager::BuildChunkMap(const TSet<FName>& PackagesToUpdateChunksFor, const TSet<FName>& StartupPackages, const TSet<FName>& PackageThatWereCooked) const
 {
 	TMap<int32, FAssetManagerChunkInfo> ReturnChunkMap;
 
