CIS fixup

diff --git a/Engine/Source/Runtime/Engine/Private/AssetManager.cpp b/Engine/Source/Runtime/Engine/Private/AssetManager.cpp
--- a/Engine/Source/Runtime/Engine/Private/AssetManager.cpp
+++ b/Engine/Source/Runtime/Engine/Private/AssetManager.cpp
@@ -4182,7 +4182,7 @@ void UAssetManager::UpdateManagementDatabase(EUpdateManagementDatabaseFlags InFl
 	auto SetManagerPredicate = [this, bWillBuildChunkMap, &PackagesToUpdateChunksFor](UE::AssetRegistry::FShouldSetManagerContext& Context)
 		{
 			EAssetSetManagerResult::Type Result = this->ShouldSetManager(Context);
-			if (bWillBuildChunkMap, Result != EAssetSetManagerResult::DoNotSet && Context.Target.IsPackage())
+			if (bWillBuildChunkMap && Result != EAssetSetManagerResult::DoNotSet && Context.Target.IsPackage())
 			{
 				PackagesToUpdateChunksFor.Add(Context.Target.PackageName);
 			}
