Add package name to PackageLoadCancelled Streamable error

diff --git a/Engine/Source/Runtime/Engine/Classes/Engine/Experimental/StreamableManagerError.h b/Engine/Source/Runtime/Engine/Classes/Engine/Experimental/StreamableManagerError.h
--- a/Engine/Source/Runtime/Engine/Classes/Engine/Experimental/StreamableManagerError.h
+++ b/Engine/Source/Runtime/Engine/Classes/Engine/Experimental/StreamableManagerError.h
@@ -9,7 +9,7 @@ DECLARE_ERROR_MODULE(StreamableManager, 0x10);
 
 
 DECLARE_ERROR_ONEPARAM(PackageLoadFailed, 1, StreamableManager, NSLOCTEXT("StreamableManager", "PackageLoadError", "Failed to load package {PackageName}"), FString, PackageName, TEXT("Unknown"));
-DECLARE_ERROR(PackageLoadCancelled, 2, StreamableManager, NSLOCTEXT("StreamableManager", "PackageLoadCancelled", "Async load canceled"));
+DECLARE_ERROR_ONEPARAM(PackageLoadCanceled, 2, StreamableManager, NSLOCTEXT("StreamableManager", "PackageLoadCancelled", "Async load canceled {PackageName}"), FString, PackageName, TEXT("Unknown"));
 DECLARE_ERROR(DownloadError, 3, StreamableManager, NSLOCTEXT("StreamableManager", "DownloadError", "Failed to download"));
 DECLARE_ERROR_ONEPARAM(PackageNameInvalid, 4, StreamableManager, NSLOCTEXT("StreamableManager", "PackageNameInvalid", "Found invalid package name {InvalidPackageName}"), FString, InvalidPackageName, TEXT("Unknown"));
 
