* Fixed up invalid dllstorage on types and methods
* Moved thread_local variable in static function to cpp file to make sure it is the same across dlls

diff --git a/Engine/Source/Runtime/CoreUObject/Internal/UObject/ScriptTimeLimiter.h b/Engine/Source/Runtime/CoreUObject/Internal/UObject/ScriptTimeLimiter.h
--- a/Engine/Source/Runtime/CoreUObject/Internal/UObject/ScriptTimeLimiter.h
+++ b/Engine/Source/Runtime/CoreUObject/Internal/UObject/ScriptTimeLimiter.h
@@ -10,14 +10,10 @@ namespace UE
 class FScriptTimeLimiter
 {
 public:
-	COREUOBJECT_API FScriptTimeLimiter() = default;
+	FScriptTimeLimiter() = default;
 
 	/** @return Reference to the thread-local singleton ScriptTimeLimiter */
-	COREUOBJECT_API static FScriptTimeLimiter& Get()
-	{
-		thread_local FScriptTimeLimiter TimeLimiterSingleton;
-		return TimeLimiterSingleton;
-	}
+	COREUOBJECT_API static FScriptTimeLimiter& Get();
 
 	/** 
 	 * The timer is started when the first (outermost) call to StartTimer occurs.
@@ -51,4 +47,4 @@ private:
 	double StartingTime = 0.0;
 };
 
-}
\ No newline at end of file
+}
diff --git a/Engine/Source/Runtime/CoreUObject/Private/UObject/ScriptTimeLimiter.cpp b/Engine/Source/Runtime/CoreUObject/Private/UObject/ScriptTimeLimiter.cpp
--- a/Engine/Source/Runtime/CoreUObject/Private/UObject/ScriptTimeLimiter.cpp
+++ b/Engine/Source/Runtime/CoreUObject/Private/UObject/ScriptTimeLimiter.cpp
@@ -5,6 +5,12 @@
 #include "HAL/PlatformTime.h"
 #include "VerseVM/VVMVerseHangDetection.h"
 
+UE::FScriptTimeLimiter& UE::FScriptTimeLimiter::Get()
+{
+	thread_local FScriptTimeLimiter TimeLimiterSingleton;
+	return TimeLimiterSingleton;
+}
+
 void UE::FScriptTimeLimiter::StartTimer()
 {
 	check(IsInGameThread());
diff --git a/Engine/Source/Runtime/Engine/Classes/Engine/AssetManager.h b/Engine/Source/Runtime/Engine/Classes/Engine/AssetManager.h
--- a/Engine/Source/Runtime/Engine/Classes/Engine/AssetManager.h
+++ b/Engine/Source/Runtime/Engine/Classes/Engine/AssetManager.h
@@ -1032,7 +1032,7 @@ protected:
 	/** Function used during creating Management references to decide when to recurse and set references */
 	ENGINE_API virtual EAssetSetManagerResult::Type ShouldSetManager(UE::AssetRegistry::FShouldSetManagerContext& Context) const;
 	UE_DEPRECATED(5.6, "Override the version that takes an FShouldSetManagerContext instead.")
-	ENGINE_API virtual EAssetSetManagerResult::Type ShouldSetManager(const FAssetIdentifier& Manager, const FAssetIdentifier& Source, const FAssetIdentifier& Target,
+	virtual EAssetSetManagerResult::Type ShouldSetManager(const FAssetIdentifier& Manager, const FAssetIdentifier& Source, const FAssetIdentifier& Target,
 		UE::AssetRegistry::EDependencyCategory Category, UE::AssetRegistry::EDependencyProperty Properties, EAssetSetManagerFlags::Type Flags) const
 	{
 		return EAssetSetManagerResult::DoNotSet;
diff --git a/Engine/Source/Runtime/Engine/Classes/GameFramework/Actor.h b/Engine/Source/Runtime/Engine/Classes/GameFramework/Actor.h
--- a/Engine/Source/Runtime/Engine/Classes/GameFramework/Actor.h
+++ b/Engine/Source/Runtime/Engine/Classes/GameFramework/Actor.h
@@ -301,7 +301,7 @@ public:
 	ENGINE_API virtual void OnRep_ReplicateMovement();
 
 #if WITH_EDITOR
-	ENGINE_API virtual bool ShouldDuplicateInPIE() const { return !bIgnoreInPIE; }
+	virtual bool ShouldDuplicateInPIE() const { return !bIgnoreInPIE; }
 #endif
 
 private:
diff --git a/Engine/Source/Runtime/Engine/Classes/Sound/QuartzSubscription.h b/Engine/Source/Runtime/Engine/Classes/Sound/QuartzSubscription.h
--- a/Engine/Source/Runtime/Engine/Classes/Sound/QuartzSubscription.h
+++ b/Engine/Source/Runtime/Engine/Classes/Sound/QuartzSubscription.h
@@ -99,7 +99,7 @@ protected:
 	UE_API void SetNotificationAnticipationAmountMusicalDuration(const EQuartzCommandQuantization Duration,  const double Multiplier);
 
 	UE_API void QuartzUnsubscribe();
-	UE_API virtual bool ShouldUnsubscribe() { return false; } 
+	virtual bool ShouldUnsubscribe() { return false; } 
 
 private:
 	struct FMetronomeDelegateGameThreadData
diff --git a/Engine/Source/Runtime/Engine/Classes/Sound/SoundWave.h b/Engine/Source/Runtime/Engine/Classes/Sound/SoundWave.h
--- a/Engine/Source/Runtime/Engine/Classes/Sound/SoundWave.h
+++ b/Engine/Source/Runtime/Engine/Classes/Sound/SoundWave.h
@@ -1636,8 +1636,8 @@ public:
 	const uint8* GetResourceData() const { return ResourceData.GetView().GetData(); }
 	uint32 GetNumChannels() { return NumChannels; }
 
-	ENGINE_API void SetSampleRate(const int32 InSampleRate) { SampleRate = static_cast<float>(InSampleRate); }
-	ENGINE_API void SetNumFrames(const int32 InNumFrames) { NumFrames = InNumFrames; }
+	void SetSampleRate(const int32 InSampleRate) { SampleRate = static_cast<float>(InSampleRate); }
+	void SetNumFrames(const int32 InNumFrames) { NumFrames = InNumFrames; }
 
 private:
 	ENGINE_API void DiscardZerothChunkData();
diff --git a/Engine/Source/Runtime/Engine/Public/Subtitles/SubtitlesAndClosedCaptionsDelegates.h b/Engine/Source/Runtime/Engine/Public/Subtitles/SubtitlesAndClosedCaptionsDelegates.h
--- a/Engine/Source/Runtime/Engine/Public/Subtitles/SubtitlesAndClosedCaptionsDelegates.h
+++ b/Engine/Source/Runtime/Engine/Public/Subtitles/SubtitlesAndClosedCaptionsDelegates.h
@@ -47,8 +47,8 @@ static constexpr float SubtitleMinDuration = 0.05f;
 /**
  *  Base class for subtitle data being attached to assets
  */
-UCLASS(BlueprintType)
-class ENGINE_API USubtitleAssetUserData : public UAssetUserData
+UCLASS(MinimalAPI, BlueprintType)
+class USubtitleAssetUserData : public UAssetUserData
 {
 	GENERATED_BODY()
 
diff --git a/Engine/Source/Runtime/Engine/Public/WorldPartition/WorldPartition.h b/Engine/Source/Runtime/Engine/Public/WorldPartition/WorldPartition.h
--- a/Engine/Source/Runtime/Engine/Public/WorldPartition/WorldPartition.h
+++ b/Engine/Source/Runtime/Engine/Public/WorldPartition/WorldPartition.h
@@ -233,10 +233,10 @@ public:
 	bool CanGenerateStreaming() const { return !StreamingPolicy; }
 
 	UE_DEPRECATED(5.3, "GenerateStreaming is deprecated, use GenerateStreaming with a param struct instead")
-	ENGINE_API bool GenerateStreaming(TArray<FString>* OutPackagesToGenerate = nullptr) { return false; }
+	bool GenerateStreaming(TArray<FString>* OutPackagesToGenerate = nullptr) { return false; }
 
 	UE_DEPRECATED(5.3, "GenerateContainerStreaming is deprecated, use GenerateContainerStreaming with a param struct instead")
-	ENGINE_API bool GenerateContainerStreaming(const UActorDescContainer* ActorDescContainer, TArray<FString>* OutPackagesToGenerate = nullptr) { return false; }
+	bool GenerateContainerStreaming(const UActorDescContainer* ActorDescContainer, TArray<FString>* OutPackagesToGenerate = nullptr) { return false; }
 
 	PRAGMA_DISABLE_DEPRECATION_WARNINGS
 	struct FGenerateStreamingParams
@@ -399,7 +399,7 @@ public:
 	PRAGMA_ENABLE_DEPRECATION_WARNINGS
 
 	UE_DEPRECATED(5.2, "CheckForErrors is deprecated, CheckForErrors with FCheckForErrorsParams should be used instead.")
-	static ENGINE_API void CheckForErrors(IStreamingGenerationErrorHandler* ErrorHandler, const UActorDescContainer* ActorDescContainer, bool bEnableStreaming, bool bIsChangelistValidation) {}
+	static void CheckForErrors(IStreamingGenerationErrorHandler* ErrorHandler, const UActorDescContainer* ActorDescContainer, bool bEnableStreaming, bool bIsChangelistValidation) {}
 
 	static ENGINE_API void CheckForErrors(const FCheckForErrorsParams& Params);
 
@@ -428,10 +428,10 @@ public:
 	FActorDescContainerInstancePreInitializeDelegate OnActorDescContainerInstancePreInitialize;
 
 	UE_DEPRECATED(5.4, "Use RegisterActorDescContainerInstance")
-	ENGINE_API UActorDescContainer* RegisterActorDescContainer(const FContainerRegistrationParams& InRegistrationParameters) { return nullptr; }
+	UActorDescContainer* RegisterActorDescContainer(const FContainerRegistrationParams& InRegistrationParameters) { return nullptr; }
 
 	UE_DEPRECATED(5.4, "Use UnregisterActorDescContainerInstance")
-	ENGINE_API bool UnregisterActorDescContainer(UActorDescContainer* Container) { return false; }
+	bool UnregisterActorDescContainer(UActorDescContainer* Container) { return false; }
 
 	ENGINE_API void UninitializeActorDescContainers();
 
@@ -488,7 +488,7 @@ public:
 	ENGINE_API bool CanStream() const;
 	ENGINE_API bool IsServer() const;
 	ENGINE_API bool IsServerStreamingEnabled() const;
-	ENGINE_API bool IsContentBundleEnabled() const { return !bDisableContentBundles; }
+	bool IsContentBundleEnabled() const { return !bDisableContentBundles; }
 	ENGINE_API bool IsServerStreamingOutEnabled() const;
 	ENGINE_API bool UseMakingVisibleTransactionRequests() const;
 	ENGINE_API bool UseMakingInvisibleTransactionRequests() const;
@@ -851,4 +851,4 @@ private:
 #if WITH_EDITOR
 	friend class FScopedCookingExternalStreamingObject;
 #endif
-};
\ No newline at end of file
+};
diff --git a/Engine/Source/Runtime/Experimental/Dataflow/Core/Public/Dataflow/DataflowNode.h b/Engine/Source/Runtime/Experimental/Dataflow/Core/Public/Dataflow/DataflowNode.h
--- a/Engine/Source/Runtime/Experimental/Dataflow/Core/Public/Dataflow/DataflowNode.h
+++ b/Engine/Source/Runtime/Experimental/Dataflow/Core/Public/Dataflow/DataflowNode.h
@@ -118,7 +118,7 @@ PRAGMA_ENABLE_DEPRECATION_WARNINGS
 	virtual bool IsExperimental() { return false; }
 
 	/** Override if the node could generate primitives .*/
-	DATAFLOWCORE_API virtual bool HasPrimitives() const
+	virtual bool HasPrimitives() const
 	{
 		return false;
 	}
