[Shaders] Minor bug fix and added log information about distributed shader compiler controller.

#rnx

diff --git a/Engine/Source/Runtime/Engine/Private/ShaderCompiler/ShaderCompiler.cpp b/Engine/Source/Runtime/Engine/Private/ShaderCompiler/ShaderCompiler.cpp
--- a/Engine/Source/Runtime/Engine/Private/ShaderCompiler/ShaderCompiler.cpp
+++ b/Engine/Source/Runtime/Engine/Private/ShaderCompiler/ShaderCompiler.cpp
@@ -1024,10 +1024,10 @@ FShaderCompileThreadRunnableBase* FShaderCompilingManager::LaunchShaderCompiling
 
 FShaderCompileThreadRunnableBase* FShaderCompilingManager::LaunchRemoteShaderCompilingThread(bool bDelayThreadExecution)
 {
-	// Check if there already is a local compile thread
-	if (FShaderCompileThreadRunnableBase* LocalCompileThread = FindShaderCompilingThread(EShaderCompilerWorkerType::LocalThread))
+	// Check if there already is a distributed compile thread
+	if (FShaderCompileThreadRunnableBase* DistributedCompileThread = FindShaderCompilingThread(EShaderCompilerWorkerType::Distributed))
 	{
-		return LocalCompileThread;
+		return DistributedCompileThread;
 	}
 
 	// Check if remote compiling is allowed and find distributed controller
@@ -1050,6 +1050,20 @@ FShaderCompileThreadRunnableBase* FShaderCompilingManager::LaunchRemoteShaderCom
 		return LaunchShaderCompilingThread(MoveTemp(RemoteCompileThread), bDelayThreadExecution);
 	}
 
+	// Print information why no distributed shader compiler controller was launched, since local compilation is usually only the fallback
+	if (bCanUseRemoteCompiling)
+	{
+		UE_LOG(LogShaderCompilers, Display, TEXT("No distributed shader compiler controller found"));
+	}
+	else if (!AllTargetPlatformSupportsRemoteShaderCompiling())
+	{
+		UE_LOG(LogShaderCompilers, Display, TEXT("Distributed shader compilation is not supported for all target platforms"));
+	}
+	else
+	{
+		UE_LOG(LogShaderCompilers, Display, TEXT("Distributed shader compilation is disabled"));
+	}
+
 	return nullptr;
 }
 
