Restrict GPU skin bones and cloth pools from using the faster manual versioning approach when GTSyncType is not set to a value that respects the 4 frame maximum latency.

diff --git a/Engine/Source/Runtime/Engine/Private/GPUSkinVertexFactory.cpp b/Engine/Source/Runtime/Engine/Private/GPUSkinVertexFactory.cpp
--- a/Engine/Source/Runtime/Engine/Private/GPUSkinVertexFactory.cpp
+++ b/Engine/Source/Runtime/Engine/Private/GPUSkinVertexFactory.cpp
@@ -261,9 +261,19 @@ TStatId FClothBufferPool::GetStatId() const
 	RETURN_QUICK_DECLARE_CYCLE_STAT(FClothBufferPool, STATGROUP_Tickables);
 }
 
-TConsoleVariableData<int32>* FGPUBaseSkinVertexFactory::FShaderDataType::MaxBonesVar = NULL;
 uint32 FGPUBaseSkinVertexFactory::FShaderDataType::MaxGPUSkinBones = 0;
 
+bool SupportsNoOverwriteBuffers()
+{
+	static const TConsoleVariableData<int32>* const CVarGTSyncType =
+		IConsoleManager::Get().FindTConsoleVariableDataInt(TEXT("r.GTSyncType"));
+
+	// We can only rely on the low-latency syncing approaches to respect NumSafeFrames properly.
+	const int32 SyncType = CVarGTSyncType->GetValueOnRenderThread();
+
+	return GRHIGlobals.SupportsMapWriteNoOverwrite && (SyncType == 1 || SyncType == 2);
+}
+
 TGlobalResource<FBoneBufferPool> GBoneBufferPool;
 
 void FGPUBaseSkinVertexFactory::FShaderDataType::AllocateBoneBuffer(FRHICommandList& RHICmdList, uint32 BufferSize, FVertexBufferAndSRV& OutBoneBuffer)
@@ -271,8 +281,11 @@ void FGPUBaseSkinVertexFactory::FShaderDataType::AllocateBoneBuffer(FRHICommandL
 	check(IsInParallelRenderingThread());
 	check(BufferSize > 0);
 
-	GBoneBufferPool.ReleasePooledResource(MoveTemp(OutBoneBuffer));
-	OutBoneBuffer = GBoneBufferPool.CreatePooledResource(RHICmdList, BufferSize);
+	if (!OutBoneBuffer.VertexBufferRHI || SupportsNoOverwriteBuffers())
+	{
+		GBoneBufferPool.ReleasePooledResource(MoveTemp(OutBoneBuffer));
+		OutBoneBuffer = GBoneBufferPool.CreatePooledResource(RHICmdList, BufferSize);
+	}
 }
 
 void FGPUBaseSkinVertexFactory::FShaderDataType::ReleaseBoneData()
@@ -305,7 +318,7 @@ void FGPUBaseSkinVertexFactory::FShaderDataType::UpdateBoneData(
 	const uint32 VectorArraySize = NumBones * sizeof(FVector4f) * 3;
 
 	// RLM_WriteOnly_NoOverwrite is safe here because the buffer pool waits N frames before allowing new allocations.
-	FMatrix3x4* ChunkMatrices = (FMatrix3x4*)RHICmdList.LockBuffer(VertexBufferRHI, 0, VectorArraySize, GRHIGlobals.SupportsMapWriteNoOverwrite ? RLM_WriteOnly_NoOverwrite : RLM_WriteOnly);
+	FMatrix3x4* ChunkMatrices = (FMatrix3x4*)RHICmdList.LockBuffer(VertexBufferRHI, 0, VectorArraySize, SupportsNoOverwriteBuffers() ? RLM_WriteOnly_NoOverwrite : RLM_WriteOnly);
 
 	{
 		//FMatrix3x4 is sizeof() == 48
@@ -1124,26 +1137,30 @@ void FGPUBaseSkinAPEXClothVertexFactory::ClothShaderType::UpdateClothSimulationD
 {
 	QUICK_SCOPE_CYCLE_COUNTER(STAT_FGPUBaseSkinAPEXClothVertexFactory_UpdateClothSimulationData);
 
+	const bool bSupportNoOverwrite = SupportsNoOverwriteBuffers();
+
 	uint32 NumSimulVerts = InSimulPositions.Num();
 
 	check(IsInParallelRenderingThread());
-	
+
 	SetCurrentRevisionNumber(RevisionNumber);
 	FVertexBufferAndSRV* CurrentClothBuffer = &GetClothBufferForWriting();
 
 	NumSimulVerts = FMath::Min(NumSimulVerts, (uint32)MAX_APEXCLOTH_VERTICES_FOR_VB);
 
 	uint32 VectorArraySize = NumSimulVerts * sizeof(float) * 6;
-	uint32 PooledArraySize = GClothBufferPool.PooledSizeForCreationArguments(VectorArraySize);
-	GClothBufferPool.ReleasePooledResource(MoveTemp(*CurrentClothBuffer));
-	*CurrentClothBuffer = GClothBufferPool.CreatePooledResource(RHICmdList, VectorArraySize);
+	if (bSupportNoOverwrite || CurrentClothBuffer->Size != GClothBufferPool.PooledSizeForCreationArguments(VectorArraySize))
+	{
+		GClothBufferPool.ReleasePooledResource(MoveTemp(*CurrentClothBuffer));
+		*CurrentClothBuffer = GClothBufferPool.CreatePooledResource(RHICmdList, VectorArraySize);
+	}
 	check(IsValidRef(*CurrentClothBuffer));
 	CurrentClothBuffer->VertexBufferRHI->SetOwnerName(AssetPathName);
 
 	if(NumSimulVerts)
 	{
 		// RLM_WriteOnly_NoOverwrite is safe here because the buffer pool waits N frames before allowing new allocations.
-		float* RESTRICT Data = (float* RESTRICT)RHICmdList.LockBuffer(CurrentClothBuffer->VertexBufferRHI, 0, VectorArraySize, GRHIGlobals.SupportsMapWriteNoOverwrite ? RLM_WriteOnly_NoOverwrite : RLM_WriteOnly);
+		float* RESTRICT Data = (float* RESTRICT)RHICmdList.LockBuffer(CurrentClothBuffer->VertexBufferRHI, 0, VectorArraySize, bSupportNoOverwrite ? RLM_WriteOnly_NoOverwrite : RLM_WriteOnly);
 		{
 			QUICK_SCOPE_CYCLE_COUNTER(STAT_FGPUBaseSkinAPEXClothVertexFactory_UpdateClothSimulationData_CopyData);
 			float* RESTRICT Pos = (float* RESTRICT) &InSimulPositions[0].X;
diff --git a/Engine/Source/Runtime/Engine/Public/GPUSkinVertexFactory.h b/Engine/Source/Runtime/Engine/Public/GPUSkinVertexFactory.h
--- a/Engine/Source/Runtime/Engine/Public/GPUSkinVertexFactory.h
+++ b/Engine/Source/Runtime/Engine/Public/GPUSkinVertexFactory.h
@@ -290,8 +290,7 @@ public:
 		uint32 CurrentBuffer = 0;
 		uint32 PreviousRevisionNumber = 0;
 		uint32 CurrentRevisionNumber = 0;
-		
-		static TConsoleVariableData<int32>* MaxBonesVar;
+
 		static uint32 MaxGPUSkinBones;
 
 		bool GetPreviousForRead(bool bPrevious) const
