Remove 5.0 deprecated APIs from AttributeContainer.h

diff --git a/Engine/Source/Runtime/Engine/Classes/Animation/AttributesContainer.h b/Engine/Source/Runtime/Engine/Classes/Animation/AttributesContainer.h
--- a/Engine/Source/Runtime/Engine/Classes/Animation/AttributesContainer.h
+++ b/Engine/Source/Runtime/Engine/Classes/Animation/AttributesContainer.h
@@ -742,61 +742,7 @@ namespace UE
 				
 			return NumAttributes;
 		}
-
-		public:
-			/* Deprecated API */
-			template<typename DataType> 
-			UE_DEPRECATED(5.0, "Deprecated behaviour, see new API")
-			TArray<DataType>& GetValuesArray()
-			{
-				static TArray<DataType, InAllocator> Array;
-				return Array;
-			}
-
-			template<typename DataType>
-			UE_DEPRECATED(5.0, "Deprecated behaviour, see new API")
-			const TArray<DataType, InAllocator>& GetValuesArray() const
-			{
-				static TArray<DataType, InAllocator> Array;
-				return Array;
-			}
-
-
-			template<typename DataType>
-			UE_DEPRECATED(5.0, "Deprecated behaviour, see new API")
-			int32 AddBoneAttribute(const BoneIndexType& BoneIndex, const FName& AttributeName, ECustomAttributeBlendType BlendType, const DataType& Value) { return INDEX_NONE; }
-
-			template<typename DataType>
-			UE_DEPRECATED(5.0, "Deprecated behaviour, see new API")
-			int32 AddBoneAttribute(const FAttributeId& AttributeInfo, const DataType& Value) { return INDEX_NONE; }
-
-			template<typename DataType>
-			UE_DEPRECATED(5.0, "Deprecated behaviour, see new API")
-			bool GetBoneAttribute(const BoneIndexType& BoneIndex, const FName& AttributeName, DataType& OutValue) const { return false; }
 			
-			template<typename DataType>
-			UE_DEPRECATED(5.0, "Deprecated behaviour, see new API")
-			int32 IndexOfBoneAttribute(uint32 BoneAttributeHash, int32 BoneIndexInt) const { return INDEX_NONE; }
-
-			template<typename DataType>
-			UE_DEPRECATED(5.0, "Deprecated behaviour, see new API")
-			bool ContainsBoneAttribute(uint32 BoneAttributeHash, int32 BoneIndexInt) const { return false; 	}
-
-			template<typename DataType>
-			UE_DEPRECATED(5.0, "Deprecated behaviour, see new API")
-			const TArray<FAttributeId>& GetAttributeInfo() const
-			{
-				static TArray<FAttributeId> Array;
-				return Array;
-			}
-
-			template<typename DataType>
-			UE_DEPRECATED(5.0, "Deprecated behaviour, see new API")
-			const TArray<int32>& GetUniqueBoneIndices() const
-			{
-				static TArray<int32> Array;
-				return Array;
-			}
 		protected:
 
 			/*
@@ -844,12 +790,3 @@ namespace UE
 		};
 	}
 }
-
-struct ENGINE_API UE_DEPRECATED(5.0, "FCustomAttributeInfo has been deprecated use UE::Anim::FAttributeId instead") FCustomAttributeInfo : public UE::Anim::FAttributeId
-{
-	UE_DEPRECATED(5.0, "Deprecated constructor, see  UE::Anim::FAttributeId")
-	FCustomAttributeInfo(const FName& InName, const FCompactPoseBoneIndex& InCompactBoneIndex, const ECustomAttributeBlendType& InBlendType) : UE::Anim::FAttributeId(InName, InCompactBoneIndex) {}
-	
-	FCustomAttributeInfo(const FName& InName, const FCompactPoseBoneIndex& InCompactBoneIndex) : UE::Anim::FAttributeId(InName, InCompactBoneIndex) {}
-	FCustomAttributeInfo(const FName& InName, int32 InIndex) : UE::Anim::FAttributeId(InName, InIndex, UE::Anim::BoneAttributeNamespace) {}
-};
