[ue5] - Remove "transient" from CMC server properties allowing client authoritative position. Typically they are dynamically set at runtime but this now allows them to be globally set and saved as intended, since they are editable.

#tests compiled, tried editing values

diff --git a/Engine/Source/Runtime/Engine/Classes/GameFramework/CharacterMovementComponent.h b/Engine/Source/Runtime/Engine/Classes/GameFramework/CharacterMovementComponent.h
--- a/Engine/Source/Runtime/Engine/Classes/GameFramework/CharacterMovementComponent.h
+++ b/Engine/Source/Runtime/Engine/Classes/GameFramework/CharacterMovementComponent.h
@@ -957,7 +957,7 @@ public:
 	 * break this character's server-client movement correction.
 	 * @see bServerAcceptClientAuthoritativePosition, ServerCheckClientError()
 	 */
-	UPROPERTY(Transient, Category="Character Movement", EditAnywhere, BlueprintReadWrite)
+	UPROPERTY(Category="Character Movement", EditAnywhere, BlueprintReadWrite)
 	uint8 bIgnoreClientMovementErrorChecksAndCorrection:1;
 
 	/**
@@ -967,7 +967,7 @@ public:
 	 * Also see GameNetworkManager->ClientAuthorativePosition which permanently enables this behavior.
 	 * @see bIgnoreClientMovementErrorChecksAndCorrection, ServerShouldUseAuthoritativePosition()
 	 */
-	UPROPERTY(Transient, Category="Character Movement", EditAnywhere, BlueprintReadWrite)
+	UPROPERTY(Category="Character Movement", EditAnywhere, BlueprintReadWrite)
 	uint8 bServerAcceptClientAuthoritativePosition : 1;
 
 	/**
