Added GFirstHighPriTickGroup equivalent to FParticleSystemWorldManager

#tests Locally tested

diff --git a/Engine/Source/Runtime/Engine/Classes/Particles/ParticleSystemManager.h b/Engine/Source/Runtime/Engine/Classes/Particles/ParticleSystemManager.h
--- a/Engine/Source/Runtime/Engine/Classes/Particles/ParticleSystemManager.h
+++ b/Engine/Source/Runtime/Engine/Classes/Particles/ParticleSystemManager.h
@@ -144,6 +144,13 @@ public:
 
 	void Dump();
 
+	static constexpr ETickingGroup FirstTickGroup = TG_PrePhysics;
+	static constexpr ETickingGroup LastTickGroup = TG_LastDemotable;
+	static constexpr int NumTickGroups = LastTickGroup - FirstTickGroup + 1;
+
+	void SetTickGroupPriority(const int32 ParticleSystemTickGroupIndex, const bool bHighPriority);
+	static void SetFirstHighPriTickGroup(int FirstHighPriTickGroupIndex);
+
 protected:
 	/** Clean up dead components post GC */
 	void HandlePostGarbageCollect();
diff --git a/Engine/Source/Runtime/Engine/Private/Particles/ParticleSystemManager.cpp b/Engine/Source/Runtime/Engine/Private/Particles/ParticleSystemManager.cpp
--- a/Engine/Source/Runtime/Engine/Private/Particles/ParticleSystemManager.cpp
+++ b/Engine/Source/Runtime/Engine/Private/Particles/ParticleSystemManager.cpp
@@ -36,6 +36,30 @@ FAutoConsoleVariableRef CVarParticleManagerAsyncBatchSize(
 	ECVF_Scalability
 );
 
+namespace ParticleSystemWorldManagerInternal
+{
+	static int32 GFirstHighPriTickGroup = 0;
+	static FAutoConsoleVariableRef CVarFirstHighPriTickGroup(
+		TEXT("fx.ParticleManagerFirstHighPriTickGroup"),
+		GFirstHighPriTickGroup,
+		TEXT("Defines which tick groups should be set to high priority for the world manager.\n")
+		TEXT("0 - (Default) all tick groups will run high priority.\n")
+		TEXT("1 - The first tick group will be normal priority, all others high, etc."),
+		FConsoleVariableDelegate::CreateLambda(
+			[](IConsoleVariable*)
+			{
+				FParticleSystemWorldManager::SetFirstHighPriTickGroup(GFirstHighPriTickGroup);
+			}
+		),
+		ECVF_Default
+	);
+
+	static bool GetParticleSystemTickGroupPriority(int ParticleSystemTickGroup)
+	{
+		return ParticleSystemTickGroup >= GFirstHighPriTickGroup;
+	}
+}
+
 //////////////////////////////////////////////////////////////////////////
 
 class FParticleManagerFinalizeTask
@@ -234,7 +258,7 @@ FParticleSystemWorldManager::FParticleSystemWorldManager(UWorld* InWorld)
 		TickFunc.bCanEverTick = true;
 		TickFunc.bStartWithTickEnabled = true;
 		TickFunc.bAllowTickOnDedicatedServer = false;
-		TickFunc.bHighPriority = true;
+		TickFunc.bHighPriority = ParticleSystemWorldManagerInternal::GetParticleSystemTickGroupPriority(TickGroup);
 		TickFunc.Owner = this;
 		TickFunc.RegisterTickFunction(InWorld->PersistentLevel);
 
@@ -845,6 +869,28 @@ void FParticleSystemWorldManager::Dump()
 #endif
 }
 
+void FParticleSystemWorldManager::SetTickGroupPriority(const int32 ParticleSystemTickGroupIndex, const bool bHighPriority)
+{
+	check((ParticleSystemTickGroupIndex >= 0) && (ParticleSystemTickGroupIndex < NumTickGroups));
+	TickFunctions[ParticleSystemTickGroupIndex].SetPriorityIncludingPrerequisites(bHighPriority);
+}
+
+void FParticleSystemWorldManager::SetFirstHighPriTickGroup(int32 FirstHighPriTickGroupIndex)
+{
+	for (auto& Pair : WorldManagers)
+	{
+		FParticleSystemWorldManager* WorldManager = Pair.Value;
+		if (WorldManager)
+		{
+			for (int32 i = 0; i < NumTickGroups; i++)
+			{
+				bool HighPriority = i >= FirstHighPriTickGroupIndex;
+				WorldManager->SetTickGroupPriority(i, HighPriority);
+			}
+		}
+	}
+}
+
 
 //////////////////////////////////////////////////////////////////////////
 
