[shaders] move worker file read/write and some other SCW-related code into a new internal utilities class in RenderCore in preparation for an upcoming change (which will use this code in FShaderCompileJob::OnComplete and FShaderPipelineCompileJob::OnComplete which are both in RenderCore)

diff --git a/Engine/Source/Developer/Windows/ShaderFormatD3D/Private/D3DShaderCompiler.cpp b/Engine/Source/Developer/Windows/ShaderFormatD3D/Private/D3DShaderCompiler.cpp
--- a/Engine/Source/Developer/Windows/ShaderFormatD3D/Private/D3DShaderCompiler.cpp
+++ b/Engine/Source/Developer/Windows/ShaderFormatD3D/Private/D3DShaderCompiler.cpp
@@ -9,6 +9,7 @@
 #include "ShaderFormatD3D.h"
 #include "ShaderCompilerCommon.h"
 #include "ShaderCompilerDefinitions.h"
+#include "ShaderCompileWorkerUtil.h"
 #include "ShaderMinifier.h"
 #include "ShaderParameterParser.h"
 #include "ShaderPreprocessTypes.h"
diff --git a/Engine/Source/Developer/Windows/ShaderFormatD3D/Private/D3DShaderCompilerDXC.cpp b/Engine/Source/Developer/Windows/ShaderFormatD3D/Private/D3DShaderCompilerDXC.cpp
--- a/Engine/Source/Developer/Windows/ShaderFormatD3D/Private/D3DShaderCompilerDXC.cpp
+++ b/Engine/Source/Developer/Windows/ShaderFormatD3D/Private/D3DShaderCompilerDXC.cpp
@@ -3,6 +3,7 @@
 #include "ShaderFormatD3D.h"
 #include "ShaderPreprocessor.h"
 #include "ShaderCompilerCommon.h"
+#include "ShaderCompileWorkerUtil.h"
 #include "ShaderParameterParser.h"
 #include "D3D12RHI.h"
 #include "Misc/Paths.h"
diff --git a/Engine/Source/Programs/ShaderCompileWorker/Private/ShaderCompileWorker.cpp b/Engine/Source/Programs/ShaderCompileWorker/Private/ShaderCompileWorker.cpp
--- a/Engine/Source/Programs/ShaderCompileWorker/Private/ShaderCompileWorker.cpp
+++ b/Engine/Source/Programs/ShaderCompileWorker/Private/ShaderCompileWorker.cpp
@@ -19,6 +19,7 @@
 #include "Interfaces/ITargetPlatformManagerModule.h"
 #include "RHIShaderFormatDefinitions.inl"
 #include "ShaderCompilerCommon.h"
+#include "ShaderCompileWorkerUtil.h"
 #include "Serialization/MemoryReader.h"
 #include "SocketSubsystem.h"
 
diff --git a/Engine/Source/Runtime/Engine/Private/ShaderCompiler/ShaderCompilerDistributed.cpp b/Engine/Source/Runtime/Engine/Private/ShaderCompiler/ShaderCompilerDistributed.cpp
--- a/Engine/Source/Runtime/Engine/Private/ShaderCompiler/ShaderCompilerDistributed.cpp
+++ b/Engine/Source/Runtime/Engine/Private/ShaderCompiler/ShaderCompilerDistributed.cpp
@@ -7,6 +7,7 @@
 #include "Misc/StringBuilder.h"
 #include "ShaderCompiler.h"
 #include "ShaderCompilerJobTypes.h"
+#include "ShaderCompileWorkerUtil.h"
 
 namespace DistributedShaderCompilerVariables
 {
@@ -185,7 +186,7 @@ void FShaderCompileDistributedThreadRunnable_Interface::DispatchShaderCompileJob
 	// Serialize the jobs to the input file
 	GShaderCompilerStats->RegisterJobBatch(JobsToSerialize.Num(), FShaderCompilerStats::EExecutionType::Distributed);
 	FArchive* InputFileAr = IFileManager::Get().CreateFileWriter(*InputFilePath, FILEWRITE_EvenIfReadOnly | FILEWRITE_NoFail);
-	FShaderCompileUtilities::DoWriteTasks(JobsToSerialize, *InputFileAr, &CachedController, CachedController.RequiresRelativePaths());
+	FShaderCompileWorkerUtil::WriteTasks(JobsToSerialize, *InputFileAr);
 	delete InputFileAr;
 
 	// Kick off the job
@@ -498,7 +499,11 @@ int32 FShaderCompileDistributedThreadRunnable_Interface::CompilingLoop()
 				{
 					bOutputFileReadFailed = false;
 					FShaderCompileWorkerDiagnostics WorkerDiagnostics;
-					if (FShaderCompileUtilities::DoReadTaskResults(Task->ShaderJobs, *OutputFileAr, GShaderCompilerDumpWorkerDiagnostics ? &WorkerDiagnostics : nullptr) == FSCWErrorCode::Success)
+					if (FShaderCompileWorkerUtil::ReadTasks(
+						Task->ShaderJobs, 
+						*OutputFileAr, 
+						GShaderCompilerDumpWorkerDiagnostics ? &WorkerDiagnostics : nullptr, 
+						FShaderCompileWorkerUtil::EReadTasksFlags::WillRetry) == FSCWErrorCode::Success)
 					{
 						bCompileJobsSucceeded = true;
 
diff --git a/Engine/Source/Runtime/Engine/Private/ShaderCompiler/ShaderCompilerPrivate.h b/Engine/Source/Runtime/Engine/Private/ShaderCompiler/ShaderCompilerPrivate.h
--- a/Engine/Source/Runtime/Engine/Private/ShaderCompiler/ShaderCompilerPrivate.h
+++ b/Engine/Source/Runtime/Engine/Private/ShaderCompiler/ShaderCompilerPrivate.h
@@ -24,8 +24,6 @@ extern const ITargetPlatform* GGlobalShaderTargetPlatform[SP_NumPlatforms];
 extern FGlobalShaderMap* GGlobalShaderMap_DeferredDeleteCopy[SP_NumPlatforms];
 
 bool AreShaderErrorsFatal();
-void DumpWorkerInputs(TConstArrayView<FShaderCommonCompileJobPtr> QueuedJobs);
-void LogQueuedCompileJobs(const TArray<FShaderCommonCompileJobPtr>& QueuedJobs, int32 NumProcessedJobs);
 FString GetGlobalShaderMapKeyString(const FGlobalShaderMapId& ShaderMapId, EShaderPlatform Platform, TArray<FShaderTypeDependency> const& Dependencies);
 
 UE::DerivedData::FCacheKey GetGlobalShaderMapKey(const FGlobalShaderMapId& ShaderMapId, EShaderPlatform Platform, const ITargetPlatform* TargetPlatform, TArray<FShaderTypeDependency> const& Dependencies);
diff --git a/Engine/Source/Runtime/Engine/Private/ShaderCompiler/ShaderCompilerThreadRunnable.cpp b/Engine/Source/Runtime/Engine/Private/ShaderCompiler/ShaderCompilerThreadRunnable.cpp
--- a/Engine/Source/Runtime/Engine/Private/ShaderCompiler/ShaderCompilerThreadRunnable.cpp
+++ b/Engine/Source/Runtime/Engine/Private/ShaderCompiler/ShaderCompilerThreadRunnable.cpp
@@ -18,6 +18,7 @@
 #include "ProfilingDebugging/LoadTimeTracker.h"
 #include "ProfilingDebugging/ScopedTimers.h"
 #include "ProfilingDebugging/StallDetector.h"
+#include "ShaderCompileWorkerUtil.h"
 
 
 float GShaderCompilerTooLongIOThresholdSeconds = 0.3;
@@ -569,7 +570,7 @@ void FShaderCompileThreadRunnable::WriteNewTasks()
 			check(TransferFile);
 
 			GShaderCompilerStats->RegisterJobBatch(CurrentWorkerInfo.QueuedJobs.Num(), FShaderCompilerStats::EExecutionType::Local);
-			if (!FShaderCompileUtilities::DoWriteTasks(CurrentWorkerInfo.QueuedJobs, *TransferFile))
+			if (!FShaderCompileWorkerUtil::WriteTasks(CurrentWorkerInfo.QueuedJobs, *TransferFile))
 			{
 				uint64 TotalDiskSpace = 0;
 				uint64 FreeDiskSpace = 0;
@@ -684,7 +685,7 @@ bool FShaderCompileThreadRunnable::LaunchWorkersIfNeeded()
 					else
 					{
 						UE_LOG(LogShaderCompilers, Error, TEXT("ShaderCompileWorker terminated unexpectedly, return code %d! Falling back to directly compiling which will be very slow.  Thread %u."), ReturnCode, WorkerIndex);
-						LogQueuedCompileJobs(CurrentWorkerInfo.QueuedJobs, -1);
+						FShaderCompileWorkerUtil::LogQueuedCompileJobs(CurrentWorkerInfo.QueuedJobs, -1);
 
 						bAbandonWorkers = true;
 						break;
@@ -767,7 +768,7 @@ int32 FShaderCompileThreadRunnable::ReadAvailableResults()
 					check(!CurrentWorkerInfo.bComplete);
 					FShaderCompileWorkerDiagnostics WorkerDiagnostics;
 
-					FShaderCompileUtilities::DoReadTaskResults(CurrentWorkerInfo.QueuedJobs, *OutputFile, GShaderCompilerDumpWorkerDiagnostics ? &WorkerDiagnostics : nullptr);
+					FShaderCompileWorkerUtil::ReadTasks(CurrentWorkerInfo.QueuedJobs, *OutputFile, GShaderCompilerDumpWorkerDiagnostics ? &WorkerDiagnostics : nullptr);
 
 					if (GShaderCompilerDumpWorkerDiagnostics)
 					{
@@ -835,7 +836,7 @@ void FShaderCompileThreadRunnable::CompileDirectlyThroughDll()
 
 		if (CurrentWorkerInfo.QueuedJobs.Num() > 0)
 		{
-			DumpWorkerInputs(CurrentWorkerInfo.QueuedJobs);
+			FShaderCompileWorkerUtil::DumpDebugCompileInput(CurrentWorkerInfo.QueuedJobs);
 
 			for (int32 JobIndex = 0; JobIndex < CurrentWorkerInfo.QueuedJobs.Num(); JobIndex++)
 			{
@@ -883,7 +884,7 @@ void FShaderCompileThreadRunnable::PrintWorkerMemoryUsageWithLockTaken()
 			TotalMemoryStats.UsedVirtual += MemoryStats.PeakUsedVirtual;
 			TotalMemoryStats.PeakUsedVirtual += MemoryStats.PeakUsedVirtual;
 		}
-		LogQueuedCompileJobs(WorkerInfo->QueuedJobs, -1);
+		FShaderCompileWorkerUtil::LogQueuedCompileJobs(WorkerInfo->QueuedJobs, -1);
 	}
 
 	if (NumValidWorkers > 0)
diff --git a/Engine/Source/Runtime/Engine/Private/ShaderCompiler/ShaderCompilerUtil.cpp b/Engine/Source/Runtime/Engine/Private/ShaderCompiler/ShaderCompilerUtil.cpp
--- a/Engine/Source/Runtime/Engine/Private/ShaderCompiler/ShaderCompilerUtil.cpp
+++ b/Engine/Source/Runtime/Engine/Private/ShaderCompiler/ShaderCompilerUtil.cpp
@@ -17,843 +17,6 @@
 #include "Serialization/MemoryWriter.h"
 #include "ShaderDiagnostics.h"
 
-
-static TAutoConsoleVariable<bool> CVarDebugDumpWorkerInputs(
-	TEXT("r.ShaderCompiler.DebugDumpWorkerInputs"),
-	false,
-	TEXT("If true, worker input files will be saved for each individual compile job alongside other debug data (note that r.DumpShaderDebugInfo must also be enabled for this to function)"),
-	ECVF_ReadOnly);
-
-static TAutoConsoleVariable<bool> CVarShadersPropagateLocalWorkerOOMs(
-	TEXT("r.Shaders.PropagateLocalWorkerOOMs"),
-	false,
-	TEXT("When set, out-of-memory conditions in a local shader compile worker will be treated as regular out-of-memory conditions and propagated to the main process.\n")
-	TEXT("This is useful when running in environment with hard memory limits, where it does not matter which process in particular caused us to violate the memory limit."),
-	ECVF_Default);
-
-
-static void ModalErrorOrLog(const FString& Title, const FString& Text, int64 CurrentFilePos = 0, int64 ExpectedFileSize = 0, bool bIsErrorFatal = true)
-{
-	static FThreadSafeBool bModalReported;
-
-	FString BadFile;
-	if (CurrentFilePos > ExpectedFileSize)
-	{
-		// Corrupt file
-		BadFile = FString::Printf(TEXT(" (Truncated or corrupt output file! Current file pos %lld, file size %lld)"), CurrentFilePos, ExpectedFileSize);
-	}
-
-	if (bIsErrorFatal)
-	{
-		// Ensure errors are logged before exiting
-		GLog->Panic();
-
-		if (FPlatformProperties::SupportsWindowedMode() && !FApp::IsUnattended())
-		{
-			if (!bModalReported.AtomicSet(true))
-			{
-				UE_LOG(LogShaderCompilers, Error, TEXT("%s\n%s"), *Text, *BadFile);
-
-				// Show dialog box with error message and request exit
-				FMessageDialog::Open(EAppMsgType::Ok, FText::FromString(Text), FText::FromString(Title));
-				constexpr bool bForceExit = true;
-				FPlatformMisc::RequestExit(bForceExit, TEXT("ShaderCompiler.ModalErrorOrLog"));
-			}
-			else
-			{
-				// Another thread already opened a dialog box and requests exit
-				FPlatformProcess::SleepInfinite();
-			}
-		}
-		else
-		{
-			UE_LOG(LogShaderCompilers, Fatal, TEXT("%s\n%s\n%s"), *Title, *Text, *BadFile);
-		}
-	}
-	else
-	{
-		UE_LOG(LogShaderCompilers, Error, TEXT("%s\n%s\n%s"), *Title, *Text, *BadFile);
-	}
-}
-
-// Make functions so the crash reporter can disambiguate the actual error because of the different callstacks
-namespace ShaderCompileWorkerError
-{
-	void HandleGeneralCrash(const TCHAR* ExceptionInfo, const TCHAR* Callstack)
-	{
-		ModalErrorOrLog(TEXT("ShaderCompileWorker crashed"), *FString::Printf(TEXT("Exception:\n%s\n\nCallstack:\n%s"), ExceptionInfo, Callstack));
-	}
-
-	void HandleBadShaderFormatVersion(const TCHAR* Data)
-	{
-		ModalErrorOrLog(TEXT("ShaderCompileWorker failed"), Data);
-	}
-
-	void HandleBadInputVersion(const TCHAR* Data)
-	{
-		ModalErrorOrLog(TEXT("ShaderCompileWorker failed"), Data);
-	}
-
-	void HandleBadSingleJobHeader(const TCHAR* Data)
-	{
-		ModalErrorOrLog(TEXT("ShaderCompileWorker failed"), Data);
-	}
-
-	void HandleBadPipelineJobHeader(const TCHAR* Data)
-	{
-		ModalErrorOrLog(TEXT("ShaderCompileWorker failed"), Data);
-	}
-
-	void HandleCantDeleteInputFile(const TCHAR* Data)
-	{
-		ModalErrorOrLog(TEXT("ShaderCompileWorker failed"), Data);
-	}
-
-	void HandleCantSaveOutputFile(const TCHAR* Data)
-	{
-		ModalErrorOrLog(TEXT("ShaderCompileWorker failed"), Data);
-	}
-
-	void HandleNoTargetShaderFormatsFound(const TCHAR* Data)
-	{
-		ModalErrorOrLog(TEXT("ShaderCompileWorker failed"), Data);
-	}
-
-	void HandleCantCompileForSpecificFormat(const TCHAR* Data)
-	{
-		ModalErrorOrLog(TEXT("ShaderCompileWorker failed"), Data);
-	}
-
-	void HandleOutputFileEmpty(const TCHAR* Filename)
-	{
-		ModalErrorOrLog(TEXT("ShaderCompileWorker failed"), FString::Printf(TEXT("Output file %s size is 0. Are you out of disk space?"), Filename));
-	}
-
-	void HandleOutputFileCorrupted(const TCHAR* Filename, int64 ExpectedSize, int64 ActualSize)
-	{
-		ModalErrorOrLog(TEXT("ShaderCompileWorker failed"), FString::Printf(TEXT("Output file corrupted (expected %I64d bytes, but only got %I64d): %s"), ExpectedSize, ActualSize, Filename));
-	}
-
-	void HandleCrashInsidePlatformCompiler(const TCHAR* Data)
-	{
-		// If the crash originates from a platform compiler, the error code must have been reported and we don't have to assume a corrupted output file.
-		// In that case, don't crash the cooker with a fatal error, just report the error so the cooker can dump debug info.
-		constexpr bool bIsErrorFatal = false;
-		ModalErrorOrLog(TEXT("ShaderCompileWorker failed"), FString::Printf(TEXT("Crash inside the platform compiler:\n%s"), Data), 0, 0, bIsErrorFatal);
-	}
-
-	void HandleBadInputFile(const TCHAR* Data)
-	{
-		ModalErrorOrLog(TEXT("ShaderCompileWorker failed"), FString::Printf(TEXT("Bad-input-file exception:\n%s"), Data));
-	}
-
-	bool HandleOutOfMemory(const TCHAR* ExceptionInfo, const TCHAR* Hostname, const FPlatformMemoryStats& MemoryStats, const TArray<FShaderCommonCompileJobPtr>& QueuedJobs)
-	{
-		constexpr int64 Gibibyte = 1024 * 1024 * 1024;
-		const FString ErrorReport = FString::Printf(
-			TEXT("ShaderCompileWorker failed with out-of-memory (OOM) exception on machine \"%s\" (%s); MemoryStats:")
-			TEXT("\n\tAvailablePhysical %llu (%.2f GiB)")
-			TEXT("\n\t AvailableVirtual %llu (%.2f GiB)")
-			TEXT("\n\t     UsedPhysical %llu (%.2f GiB)")
-			TEXT("\n\t PeakUsedPhysical %llu (%.2f GiB)")
-			TEXT("\n\t      UsedVirtual %llu (%.2f GiB)")
-			TEXT("\n\t  PeakUsedVirtual %llu (%.2f GiB)"),
-			Hostname,
-			(ExceptionInfo[0] == TEXT('\0') ? TEXT("No exception information") : ExceptionInfo),
-			MemoryStats.AvailablePhysical, double(MemoryStats.AvailablePhysical) / Gibibyte,
-			MemoryStats.AvailableVirtual, double(MemoryStats.AvailableVirtual) / Gibibyte,
-			MemoryStats.UsedPhysical, double(MemoryStats.UsedPhysical) / Gibibyte,
-			MemoryStats.PeakUsedPhysical, double(MemoryStats.PeakUsedPhysical) / Gibibyte,
-			MemoryStats.UsedVirtual, double(MemoryStats.UsedVirtual) / Gibibyte,
-			MemoryStats.PeakUsedVirtual, double(MemoryStats.PeakUsedVirtual) / Gibibyte
-		);
-
-		if (GShaderCompilingManager->IsRemoteCompilingEnabled())
-		{
-			// Remote shader compiler supports re-compiling jobs on local machine
-			UE_LOG(LogShaderCompilers, Warning, TEXT("%s\nRecompile %d shader compile %s locally"), *ErrorReport, QueuedJobs.Num(), (QueuedJobs.Num() == 1 ? TEXT("job") : TEXT("jobs")));
-			return true;
-		}
-		else
-		{
-			if (CVarShadersPropagateLocalWorkerOOMs.GetValueOnAnyThread())
-			{
-				FPlatformMemory::OnOutOfMemory(0, 64);
-			}
-			ModalErrorOrLog(TEXT("ShaderCompileWorker failed"), ErrorReport);
-			return false;
-		}
-	}
-}
-
-static TMap<FString, uint32> GetFormatVersionMap()
-{
-	TMap<FString, uint32> FormatVersionMap;
-
-	const TArray<const class IShaderFormat*>& ShaderFormats = GetTargetPlatformManagerRef().GetShaderFormats();
-	check(ShaderFormats.Num());
-	for (int32 Index = 0; Index < ShaderFormats.Num(); Index++)
-	{
-		TArray<FName> OutFormats;
-		ShaderFormats[Index]->GetSupportedFormats(OutFormats);
-		check(OutFormats.Num());
-		for (int32 InnerIndex = 0; InnerIndex < OutFormats.Num(); InnerIndex++)
-		{
-			uint32 Version = ShaderFormats[Index]->GetVersion(OutFormats[InnerIndex]);
-			FormatVersionMap.Add(OutFormats[InnerIndex].ToString(), Version);
-		}
-	}
-
-	return FormatVersionMap;
-}
-
-#if 0 //UNUSED
-static int32 GetNumTotalJobs(const TArray<FShaderCommonCompileJobPtr>& Jobs)
-{
-	int32 NumJobs = 0;
-	for (int32 Index = 0; Index < Jobs.Num(); ++Index)
-	{
-		auto* PipelineJob = Jobs[Index]->GetShaderPipelineJob();
-		NumJobs += PipelineJob ? PipelineJob->StageJobs.Num() : 1;
-	}
-
-	return NumJobs;
-}
-#endif
-
-static void SplitJobsByType(const TArray<FShaderCommonCompileJobPtr>& QueuedJobs, TArray<FShaderCompileJob*>& OutQueuedSingleJobs, TArray<FShaderPipelineCompileJob*>& OutQueuedPipelineJobs)
-{
-	for (int32 Index = 0; Index < QueuedJobs.Num(); ++Index)
-	{
-		FShaderCommonCompileJobPtr CommonJob = QueuedJobs[Index];
-		if (FShaderCompileJob* SingleJob = CommonJob->GetSingleShaderJob())
-		{
-			OutQueuedSingleJobs.Add(SingleJob);
-
-		}
-		else if (FShaderPipelineCompileJob* PipelineJob = CommonJob->GetShaderPipelineJob())
-		{
-			OutQueuedPipelineJobs.Add(PipelineJob);
-		}
-		else
-		{
-			checkf(0, TEXT("FShaderCommonCompileJob::Type=%d is not a valid type for a shader compile job"), (int32)CommonJob->Type);
-		}
-	}
-}
-
-static bool DoWriteTasksInner(const TArray<FShaderCommonCompileJobPtr>& QueuedJobs, FArchive& InTransferFile, IDistributedBuildController* BuildDistributionController, bool bUseRelativePaths, bool bCompressTaskFile)
-{
-	int32 InputVersion = ShaderCompileWorkerInputVersion;
-	InTransferFile << InputVersion;
-
-	TArray<uint8> UncompressedArray;
-	FMemoryWriter TransferMemory(UncompressedArray);
-	FArchive& TransferFile = bCompressTaskFile ? TransferMemory : InTransferFile;
-	if (!bCompressTaskFile)
-	{
-		// still write NAME_None as string
-		FString FormatNone = FName(NAME_None).ToString();
-		TransferFile << FormatNone;
-	}
-
-	static TMap<FString, uint32> FormatVersionMap = GetFormatVersionMap();
-
-	TransferFile << FormatVersionMap;
-
-	TArray<FShaderCompileJob*> QueuedSingleJobs;
-	TArray<FShaderPipelineCompileJob*> QueuedPipelineJobs;
-	SplitJobsByType(QueuedJobs, QueuedSingleJobs, QueuedPipelineJobs);
-
-	TArray<TRefCountPtr<FSharedShaderCompilerEnvironment>> SharedEnvironments;
-	TArray<const FShaderParametersMetadata*> RequestShaderParameterStructures;
-
-	// gather shared environments and parameter structures, these tend to be shared between jobs
-	{
-
-		for (int32 JobIndex = 0; JobIndex < QueuedSingleJobs.Num(); JobIndex++)
-		{
-			QueuedSingleJobs[JobIndex]->Input.GatherSharedInputs(SharedEnvironments, RequestShaderParameterStructures);
-		}
-
-		for (int32 JobIndex = 0; JobIndex < QueuedPipelineJobs.Num(); JobIndex++)
-		{
-			auto* PipelineJob = QueuedPipelineJobs[JobIndex];
-			int32 NumStageJobs = PipelineJob->StageJobs.Num();
-
-			for (int32 Index = 0; Index < NumStageJobs; Index++)
-			{
-				PipelineJob->StageJobs[Index]->Input.GatherSharedInputs(SharedEnvironments, RequestShaderParameterStructures);
-			}
-		}
-
-		int32 NumSharedEnvironments = SharedEnvironments.Num();
-		TransferFile << NumSharedEnvironments;
-
-		for (int32 EnvironmentIndex = 0; EnvironmentIndex < SharedEnvironments.Num(); EnvironmentIndex++)
-		{
-			TransferFile << *SharedEnvironments[EnvironmentIndex];
-		}
-	}
-
-	// Write shader parameter structures
-	TArray<const FShaderParametersMetadata*> AllShaderParameterStructures;
-	{
-		// List all dependencies.
-		for (int32 StructId = 0; StructId < RequestShaderParameterStructures.Num(); StructId++)
-		{
-			RequestShaderParameterStructures[StructId]->IterateStructureMetadataDependencies(
-				[&](const FShaderParametersMetadata* Struct)
-			{
-				AllShaderParameterStructures.AddUnique(Struct);
-			});
-		}
-
-		// Write all shader parameter structure.
-		int32 NumParameterStructures = AllShaderParameterStructures.Num();
-		TransferFile << NumParameterStructures;
-		for (const FShaderParametersMetadata* Struct : AllShaderParameterStructures)
-		{
-			FString LayoutName = Struct->GetLayout().GetDebugName();
-			FString StructTypeName = Struct->GetStructTypeName();
-			FString ShaderVariableName = Struct->GetShaderVariableName();
-			uint8 UseCase = uint8(Struct->GetUseCase());
-			FString StructFileName = FString(ANSI_TO_TCHAR(Struct->GetFileName()));
-			int32 StructFileLine = Struct->GetFileLine();
-			uint32 Size = Struct->GetSize();
-			int32 MemberCount = Struct->GetMembers().Num();
-
-			static_assert(sizeof(UseCase) == sizeof(FShaderParametersMetadata::EUseCase), "Cast failure.");
-
-			TransferFile << LayoutName;
-			TransferFile << StructTypeName;
-			TransferFile << ShaderVariableName;
-			TransferFile << UseCase;
-			TransferFile << StructFileName;
-			TransferFile << StructFileLine;
-			TransferFile << Size;
-			TransferFile << MemberCount;
-
-			for (const FShaderParametersMetadata::FMember& Member : Struct->GetMembers())
-			{
-				FString Name = Member.GetName();
-				FString ShaderType = Member.GetShaderType();
-				int32 FileLine = Member.GetFileLine();
-				uint32 Offset = Member.GetOffset();
-				uint8 BaseType = uint8(Member.GetBaseType());
-				uint8 PrecisionModifier = uint8(Member.GetPrecision());
-				uint32 NumRows = Member.GetNumRows();
-				uint32 NumColumns = Member.GetNumColumns();
-				uint32 NumElements = Member.GetNumElements();
-				int32 StructMetadataIndex = INDEX_NONE;
-				if (Member.GetStructMetadata())
-				{
-					StructMetadataIndex = AllShaderParameterStructures.Find(Member.GetStructMetadata());
-					check(StructMetadataIndex != INDEX_NONE);
-				}
-
-				static_assert(sizeof(BaseType) == sizeof(EUniformBufferBaseType), "Cast failure.");
-				static_assert(sizeof(PrecisionModifier) == sizeof(EShaderPrecisionModifier::Type), "Cast failure.");
-
-				TransferFile << Name;
-				TransferFile << ShaderType;
-				TransferFile << FileLine;
-				TransferFile << Offset;
-				TransferFile << BaseType;
-				TransferFile << PrecisionModifier;
-				TransferFile << NumRows;
-				TransferFile << NumColumns;
-				TransferFile << NumElements;
-				TransferFile << StructMetadataIndex;
-			}
-		}
-	}
-
-	// Write individual shader jobs
-	{
-		int32 SingleJobHeader = ShaderCompileWorkerSingleJobHeader;
-		TransferFile << SingleJobHeader;
-
-		int32 NumBatches = QueuedSingleJobs.Num();
-		TransferFile << NumBatches;
-
-		// Serialize all the batched jobs
-		for (int32 JobIndex = 0; JobIndex < QueuedSingleJobs.Num(); JobIndex++)
-		{
-			QueuedSingleJobs[JobIndex]->SerializeWorkerInput(TransferFile);
-			QueuedSingleJobs[JobIndex]->Input.SerializeSharedInputs(TransferFile, SharedEnvironments, AllShaderParameterStructures);
-		}
-	}
-
-	// Write shader pipeline jobs
-	{
-		int32 PipelineJobHeader = ShaderCompileWorkerPipelineJobHeader;
-		TransferFile << PipelineJobHeader;
-
-		int32 NumBatches = QueuedPipelineJobs.Num();
-		TransferFile << NumBatches;
-		for (int32 JobIndex = 0; JobIndex < QueuedPipelineJobs.Num(); JobIndex++)
-		{
-			auto* PipelineJob = QueuedPipelineJobs[JobIndex];
-			FString PipelineName = PipelineJob->Key.ShaderPipeline->GetName();
-			TransferFile << PipelineName;
-			int32 NumStageJobs = PipelineJob->StageJobs.Num();
-			TransferFile << NumStageJobs;
-			for (int32 Index = 0; Index < NumStageJobs; Index++)
-			{
-				PipelineJob->StageJobs[Index]->SerializeWorkerInput(TransferFile);
-				PipelineJob->StageJobs[Index]->Input.SerializeSharedInputs(TransferFile, SharedEnvironments, AllShaderParameterStructures);
-			}
-		}
-	}
-
-	if (bCompressTaskFile)
-	{
-		TransferFile.Close();
-
-		FName CompressionFormatToUse = NAME_LZ4;
-
-		FString FormatName = CompressionFormatToUse.ToString();
-		InTransferFile << FormatName;
-
-		// serialize uncompressed data size
-		int32 UncompressedDataSize = UncompressedArray.Num();
-		checkf(UncompressedDataSize != 0, TEXT("Did not write any data to the task file for the compression."));
-		InTransferFile << UncompressedDataSize;
-
-		// not using SerializeCompressed because it splits into smaller chunks
-		int32 CompressedSizeBound = FCompression::CompressMemoryBound(CompressionFormatToUse, static_cast<int32>(UncompressedDataSize));
-		TArray<uint8> CompressedBuffer;
-		CompressedBuffer.SetNumUninitialized(CompressedSizeBound);
-
-		int32 ActualCompressedSize = CompressedSizeBound;
-		bool bSucceeded = FCompression::CompressMemory(CompressionFormatToUse, CompressedBuffer.GetData(), ActualCompressedSize, UncompressedArray.GetData(), UncompressedDataSize, COMPRESS_BiasSpeed);
-		checkf(ActualCompressedSize <= CompressedSizeBound, TEXT("Compressed size was larger than the bound - we stomped the memory."));
-		CompressedBuffer.SetNum(ActualCompressedSize, EAllowShrinking::No);
-
-		InTransferFile << CompressedBuffer;
-		UE_LOG(LogShaderCompilers, Verbose, TEXT("Compressed the task file from %d bytes to %d bytes (%.2f%% savings)"), UncompressedDataSize, ActualCompressedSize,
-			100.0 * (UncompressedDataSize - ActualCompressedSize) / static_cast<double>(UncompressedDataSize));
-	}
-
-	return InTransferFile.Close();
-}
-
-const TCHAR* DebugWorkerInputFileName = TEXT("DebugSCW.in");
-const TCHAR* DebugWorkerOutputFileName = TEXT("DebugSCW.out");
-
-static FString CreateShaderCompilerWorkerDebugCommandLine(FString DebugWorkerInputFilePath)
-{
-	// 0 is parent PID, pass zero TTL and KeepInput to make SCW process the single job then exit without deleting the input file
-	return FString::Printf(TEXT("\"%s\" 0 \"DebugSCW\" %s %s -TimeToLive=0.0f -KeepInput"),
-		*DebugWorkerInputFilePath, // working directory for SCW
-		DebugWorkerInputFileName,
-		DebugWorkerOutputFileName);
-}
-
-void DumpWorkerInputs(TConstArrayView<FShaderCommonCompileJobPtr> QueuedJobs)
-{
-	if (CVarDebugDumpWorkerInputs.GetValueOnAnyThread())
-	{
-		for (const FShaderCommonCompileJobPtr& CommonJob : QueuedJobs)
-		{
-			FShaderPipelineCompileJob* PipelineJob = CommonJob->GetShaderPipelineJob();
-			FString DebugWorkerInputFilePath;
-			if (PipelineJob)
-			{
-				// for pipeline jobs, write out the worker input for the whole pipeline, but only for the first stage
-				// would be better to put in a parent folder probably...
-				DebugWorkerInputFilePath = PipelineJob->StageJobs[0]->Input.DumpDebugInfoPath;
-			}
-			else
-			{
-				DebugWorkerInputFilePath = CommonJob->GetSingleShaderJob()->Input.DumpDebugInfoPath;
-			}
-			if (!DebugWorkerInputFilePath.IsEmpty())
-			{
-				TArray<FShaderCommonCompileJobPtr> SingleJobArray;
-				SingleJobArray.Add(CommonJob);
-
-				FArchive* DebugWorkerInputFileWriter = IFileManager::Get().CreateFileWriter(*(DebugWorkerInputFilePath / DebugWorkerInputFileName), FILEWRITE_NoFail);
-				DoWriteTasksInner(
-					SingleJobArray,
-					*DebugWorkerInputFileWriter,
-					nullptr,	// Don't pass a IDistributedBuildController, this is only used for conversion to relative paths which we do not want for debug files
-					false,		// As above, use absolute paths not relative
-					true);		// Always compress the debug files; they are rather large so this saves some disk space
-				DebugWorkerInputFileWriter->Close();
-				delete DebugWorkerInputFileWriter;
-
-				FFileHelper::SaveStringToFile(
-					CreateShaderCompilerWorkerDebugCommandLine(DebugWorkerInputFilePath),
-					*(DebugWorkerInputFilePath / TEXT("DebugCompileArgs.txt")));
-			}
-		}
-	}
-}
-
-// Serialize Queued Job information
-bool FShaderCompileUtilities::DoWriteTasks(const TArray<FShaderCommonCompileJobPtr>& QueuedJobs, FArchive& InTransferFile, IDistributedBuildController* BuildDistributionController, bool bUseRelativePaths, bool bCompressTaskFile)
-{
-	DumpWorkerInputs(QueuedJobs);
-
-	return DoWriteTasksInner(QueuedJobs, InTransferFile, BuildDistributionController, bUseRelativePaths, bCompressTaskFile);
-}
-
-static void ReadSingleJob(FShaderCompileJob* CurrentJob, FArchive& WorkerOutputFileReader)
-{
-	TRACE_CPUPROFILER_EVENT_SCOPE(ReadSingleJob);
-
-	check(!CurrentJob->bFinalized);
-	CurrentJob->bFinalized = true;
-
-	// Deserialize the shader compilation output.
-	CurrentJob->SerializeWorkerOutput(WorkerOutputFileReader);
-
-	// The job should already have a non-zero output hash
-	checkf(CurrentJob->Output.OutputHash != FSHAHash() || !CurrentJob->bSucceeded, TEXT("OutputHash for a successful job was not set in the shader compile worker!"));
-}
-
-static const TCHAR* GetCompileJobSuccessText(FShaderCompileJob* SingleJob)
-{
-	if (SingleJob)
-	{
-		return SingleJob->Output.bSucceeded ? TEXT("Succeeded") : TEXT("Failed");
-	}
-	return TEXT("");
-}
-
-void LogQueuedCompileJobs(const TArray<FShaderCommonCompileJobPtr>& QueuedJobs, int32 NumProcessedJobs)
-{
-	if (NumProcessedJobs == -1)
-	{
-		UE_LOG(LogShaderCompilers, Error, TEXT("SCW %d Queued Jobs, Unknown number of processed jobs!"), QueuedJobs.Num());
-	}
-	else
-	{
-		UE_LOG(LogShaderCompilers, Error, TEXT("SCW %d Queued Jobs, Finished %d single jobs"), QueuedJobs.Num(), NumProcessedJobs);
-	}
-
-	for (int32 Index = 0; Index < QueuedJobs.Num(); ++Index)
-	{
-		if (FShaderCompileJob* SingleJob = QueuedJobs[Index]->GetSingleShaderJob())
-		{
-			UE_LOG(LogShaderCompilers, Error, TEXT("Job %d [Single] %s: %s"), Index, GetCompileJobSuccessText(SingleJob), *GetSingleJobCompilationDump(SingleJob));
-		}
-		else
-		{
-			FShaderPipelineCompileJob* PipelineJob = QueuedJobs[Index]->GetShaderPipelineJob();
-			UE_LOG(LogShaderCompilers, Error, TEXT("Job %d: Pipeline %s "), Index, PipelineJob->Key.ShaderPipeline->GetName());
-			for (int32 JobIndex = 0; JobIndex < PipelineJob->StageJobs.Num(); ++JobIndex)
-			{
-				FShaderCompileJob* StageJob = PipelineJob->StageJobs[JobIndex]->GetSingleShaderJob();
-				UE_LOG(LogShaderCompilers, Error, TEXT("PipelineJob %d %s: %s"), JobIndex, GetCompileJobSuccessText(StageJob), *GetSingleJobCompilationDump(StageJob));
-			}
-		}
-	}
-
-	// Force a log flush so we can track the crash before the cooker potentially crashes before the output shows up
-	GLog->Flush();
-}
-
-// Disable optimization for this crash handler to get full access to the entire stack frame when debugging a crash dump
-UE_DISABLE_OPTIMIZATION_SHIP
-static bool HandleWorkerCrash(const TArray<FShaderCommonCompileJobPtr>& QueuedJobs, FArchive& OutputFile, int32 OutputVersion, int64 FileSize, FSCWErrorCode::ECode ErrorCode, int32 NumProcessedJobs, int32 CallstackLength, int32 ExceptionInfoLength, int32 HostnameLength)
-{
-	TArray<TCHAR> Callstack;
-	Callstack.AddUninitialized(CallstackLength + 1);
-	OutputFile.Serialize(Callstack.GetData(), CallstackLength * sizeof(TCHAR));
-	Callstack[CallstackLength] = TEXT('\0');
-
-	TArray<TCHAR> ExceptionInfo;
-	ExceptionInfo.AddUninitialized(ExceptionInfoLength + 1);
-	OutputFile.Serialize(ExceptionInfo.GetData(), ExceptionInfoLength * sizeof(TCHAR));
-	ExceptionInfo[ExceptionInfoLength] = TEXT('\0');
-
-	TArray<TCHAR> Hostname;
-	Hostname.AddUninitialized(HostnameLength + 1);
-	OutputFile.Serialize(Hostname.GetData(), HostnameLength * sizeof(TCHAR));
-	Hostname[HostnameLength] = TEXT('\0');
-
-	// Read available and used physical memory from worker machine on OOM error
-	FPlatformMemoryStats MemoryStats;
-	if (ErrorCode == FSCWErrorCode::OutOfMemory)
-	{
-		OutputFile
-			<< MemoryStats.AvailablePhysical
-			<< MemoryStats.AvailableVirtual
-			<< MemoryStats.UsedPhysical
-			<< MemoryStats.PeakUsedPhysical
-			<< MemoryStats.UsedVirtual
-			<< MemoryStats.PeakUsedVirtual
-			;
-	}
-
-	// Store primary job information onto stack to make it part of a crash dump
-	static const int32 MaxNumCharsForSourcePaths = 8192;
-	int32 JobInputSourcePathsLength = 0;
-	ANSICHAR JobInputSourcePaths[MaxNumCharsForSourcePaths];
-	JobInputSourcePaths[0] = 0;
-
-	auto WriteInputSourcePathOntoStack = [&JobInputSourcePathsLength, &JobInputSourcePaths](const ANSICHAR* InputSourcePath)
-	{
-		if (InputSourcePath != nullptr && JobInputSourcePathsLength + 3 < MaxNumCharsForSourcePaths)
-		{
-			// Copy input source path into stack buffer
-			int32 InputSourcePathLength = FMath::Min(FCStringAnsi::Strlen(InputSourcePath), (MaxNumCharsForSourcePaths - JobInputSourcePathsLength - 2));
-			FMemory::Memcpy(JobInputSourcePaths + JobInputSourcePathsLength, InputSourcePath, InputSourcePathLength);
-
-			// Write newline character and put NUL character at the end
-			JobInputSourcePathsLength += InputSourcePathLength;
-			JobInputSourcePaths[JobInputSourcePathsLength] = TEXT('\n');
-			++JobInputSourcePathsLength;
-			JobInputSourcePaths[JobInputSourcePathsLength] = 0;
-		}
-	};
-
-	auto StoreInputDebugInfo = [&WriteInputSourcePathOntoStack, &JobInputSourcePathsLength, &JobInputSourcePaths](const FShaderCompilerInput& Input)
-	{
-		FString DebugInfo = FString::Printf(TEXT("%s:%s"), *Input.VirtualSourceFilePath, *Input.EntryPointName);
-		WriteInputSourcePathOntoStack(TCHAR_TO_UTF8(*DebugInfo));
-	};
-
-	for (auto CommonJob : QueuedJobs)
-	{
-		if (FShaderCompileJob* SingleJob = CommonJob->GetSingleShaderJob())
-		{
-			StoreInputDebugInfo(SingleJob->Input);
-		}
-		else if (FShaderPipelineCompileJob* PipelineJob = CommonJob->GetShaderPipelineJob())
-		{
-			for (int32 Job = 0; Job < PipelineJob->StageJobs.Num(); ++Job)
-			{
-				if (FShaderCompileJob* SingleStageJob = PipelineJob->StageJobs[Job])
-				{
-					StoreInputDebugInfo(SingleStageJob->Input);
-				}
-			}
-		}
-	}
-
-	// One entry per error code as we want to have different callstacks for crash reporter...
-	switch (ErrorCode)
-	{
-	default:
-	case FSCWErrorCode::GeneralCrash:
-		LogQueuedCompileJobs(QueuedJobs, NumProcessedJobs);
-		ShaderCompileWorkerError::HandleGeneralCrash(ExceptionInfo.GetData(), Callstack.GetData());
-		break;
-	case FSCWErrorCode::BadShaderFormatVersion:
-		ShaderCompileWorkerError::HandleBadShaderFormatVersion(ExceptionInfo.GetData());
-		break;
-	case FSCWErrorCode::BadInputVersion:
-		ShaderCompileWorkerError::HandleBadInputVersion(ExceptionInfo.GetData());
-		break;
-	case FSCWErrorCode::BadSingleJobHeader:
-		ShaderCompileWorkerError::HandleBadSingleJobHeader(ExceptionInfo.GetData());
-		break;
-	case FSCWErrorCode::BadPipelineJobHeader:
-		ShaderCompileWorkerError::HandleBadPipelineJobHeader(ExceptionInfo.GetData());
-		break;
-	case FSCWErrorCode::CantDeleteInputFile:
-		ShaderCompileWorkerError::HandleCantDeleteInputFile(ExceptionInfo.GetData());
-		break;
-	case FSCWErrorCode::CantSaveOutputFile:
-		ShaderCompileWorkerError::HandleCantSaveOutputFile(ExceptionInfo.GetData());
-		break;
-	case FSCWErrorCode::NoTargetShaderFormatsFound:
-		ShaderCompileWorkerError::HandleNoTargetShaderFormatsFound(ExceptionInfo.GetData());
-		break;
-	case FSCWErrorCode::CantCompileForSpecificFormat:
-		ShaderCompileWorkerError::HandleCantCompileForSpecificFormat(ExceptionInfo.GetData());
-		break;
-	case FSCWErrorCode::CrashInsidePlatformCompiler:
-		LogQueuedCompileJobs(QueuedJobs, NumProcessedJobs);
-		ShaderCompileWorkerError::HandleCrashInsidePlatformCompiler(ExceptionInfo.GetData());
-		break;
-	case FSCWErrorCode::BadInputFile:
-		ShaderCompileWorkerError::HandleBadInputFile(ExceptionInfo.GetData());
-		break;
-	case FSCWErrorCode::OutOfMemory:
-		return ShaderCompileWorkerError::HandleOutOfMemory(ExceptionInfo.GetData(), Hostname.GetData(), MemoryStats, QueuedJobs);
-	case FSCWErrorCode::Success:
-		// Can't get here...
-		return true;
-	}
-	return false;
-}
-UE_ENABLE_OPTIMIZATION_SHIP
-
-// Helper struct to provide consistent error report with detailed information about corrupted ShaderCompileWorker output file.
-struct FSCWOutputFileContext
-{
-	FArchive& OutputFile;
-	int64 FileSize = 0;
-
-	FSCWOutputFileContext(FArchive& OutputFile) :
-		OutputFile(OutputFile)
-	{
-	}
-
-	template <typename... Types>
-	void ModalErrorOrLog(UE::Core::TCheckedFormatString<FString::FmtCharType, Types...> Format, Types... Args)
-	{
-		FString Text = FString::Printf(Format, Args...);
-		Text = FString::Printf(TEXT("File path: \"%s\"\n%s\nForgot to build ShaderCompileWorker or delete invalidated DerivedDataCache?"), *OutputFile.GetArchiveName(), *Text);
-		const TCHAR* Title = TEXT("Corrupted ShaderCompileWorker output file");
-		if (FileSize > 0)
-		{
-			::ModalErrorOrLog(Title, Text, OutputFile.Tell(), FileSize);
-		}
-		else
-		{
-			::ModalErrorOrLog(Title, Text, 0, 0);
-		}
-	}
-};
-
-// Process results from Worker Process.
-// Returns false if reading the tasks failed but we were able to recover from handing a crash report. In this case, all jobs must be submitted/processed again.
-FSCWErrorCode::ECode FShaderCompileUtilities::DoReadTaskResults(const TArray<FShaderCommonCompileJobPtr>& QueuedJobs, FArchive& OutputFile, FShaderCompileWorkerDiagnostics* OutWorkerDiagnostics)
-{
-	FSCWOutputFileContext OutputFileContext(OutputFile);
-
-	if (OutputFile.TotalSize() == 0)
-	{
-		ShaderCompileWorkerError::HandleOutputFileEmpty(*OutputFile.GetArchiveName());
-	}
-
-	int32 OutputVersion = ShaderCompileWorkerOutputVersion;
-	OutputFile << OutputVersion;
-
-	if (ShaderCompileWorkerOutputVersion != OutputVersion)
-	{
-		OutputFileContext.ModalErrorOrLog(TEXT("Expecting output version %d, got %d instead!"), ShaderCompileWorkerOutputVersion, OutputVersion);
-	}
-
-	OutputFile << OutputFileContext.FileSize;
-
-	// Check for corrupted output file
-	if (OutputFileContext.FileSize > OutputFile.TotalSize())
-	{
-		ShaderCompileWorkerError::HandleOutputFileCorrupted(*OutputFile.GetArchiveName(), OutputFileContext.FileSize, OutputFile.TotalSize());
-	}
-
-	FShaderCompileWorkerDiagnostics WorkerDiagnostics;
-	OutputFile << WorkerDiagnostics;
-
-	if (OutWorkerDiagnostics)
-	{
-		*OutWorkerDiagnostics = WorkerDiagnostics;
-	}
-
-	int32 NumProcessedJobs = 0;
-	OutputFile << NumProcessedJobs;
-
-	int32 CallstackLength = 0;
-	OutputFile << CallstackLength;
-
-	int32 ExceptionInfoLength = 0;
-	OutputFile << ExceptionInfoLength;
-
-	int32 HostnameLength = 0;
-	OutputFile << HostnameLength;
-
-	if (WorkerDiagnostics.ErrorCode != FSCWErrorCode::Success)
-	{
-		// If worker crashed in a way we were able to recover from, return and expect the compile jobs to be reissued already
-		if (HandleWorkerCrash(QueuedJobs, OutputFile, OutputVersion, OutputFileContext.FileSize, (FSCWErrorCode::ECode)WorkerDiagnostics.ErrorCode, NumProcessedJobs, CallstackLength, ExceptionInfoLength, HostnameLength))
-		{
-			FSCWErrorCode::Reset();
-			return (FSCWErrorCode::ECode)WorkerDiagnostics.ErrorCode;
-		}
-	}
-
-	TArray<FShaderCompileJob*> QueuedSingleJobs;
-	TArray<FShaderPipelineCompileJob*> QueuedPipelineJobs;
-	SplitJobsByType(QueuedJobs, QueuedSingleJobs, QueuedPipelineJobs);
-
-	// Read single jobs
-	{
-		int32 SingleJobHeader = -1;
-		OutputFile << SingleJobHeader;
-		if (SingleJobHeader != ShaderCompileWorkerSingleJobHeader)
-		{
-			OutputFileContext.ModalErrorOrLog(TEXT("Expecting single job header ID 0x%08X, got 0x%08X instead!"), ShaderCompileWorkerSingleJobHeader, SingleJobHeader);
-		}
-
-		int32 NumJobs;
-		OutputFile << NumJobs;
-		if (NumJobs != QueuedSingleJobs.Num())
-		{
-			OutputFileContext.ModalErrorOrLog(TEXT("Expecting %d single %s, got %d instead!"), QueuedSingleJobs.Num(), (QueuedSingleJobs.Num() == 1 ? TEXT("job") : TEXT("jobs")), NumJobs);
-		}
-		else
-		{
-			for (int32 JobIndex = 0; JobIndex < NumJobs; JobIndex++)
-			{
-				auto* CurrentJob = QueuedSingleJobs[JobIndex];
-				ReadSingleJob(CurrentJob, OutputFile);
-			}
-		}
-	}
-
-	// Pipeline jobs
-	{
-		int32 PipelineJobHeader = -1;
-		OutputFile << PipelineJobHeader;
-		if (PipelineJobHeader != ShaderCompileWorkerPipelineJobHeader)
-		{
-			OutputFileContext.ModalErrorOrLog(TEXT("Expecting pipeline jobs header ID 0x%08X, got 0x%08X instead!"), ShaderCompileWorkerPipelineJobHeader, PipelineJobHeader);
-		}
-
-		int32 NumJobs;
-		OutputFile << NumJobs;
-		if (NumJobs != QueuedPipelineJobs.Num())
-		{
-			OutputFileContext.ModalErrorOrLog(TEXT("Expecting %d pipeline %s, got %d instead!"), QueuedPipelineJobs.Num(), (QueuedPipelineJobs.Num() == 1 ? TEXT("job") : TEXT("jobs")), NumJobs);
-		}
-		else
-		{
-			for (int32 JobIndex = 0; JobIndex < NumJobs; JobIndex++)
-			{
-				FShaderPipelineCompileJob* CurrentJob = QueuedPipelineJobs[JobIndex];
-
-				FString PipelineName;
-				OutputFile << PipelineName;
-				bool bSucceeded = false;
-				OutputFile << bSucceeded;
-				CurrentJob->bSucceeded = bSucceeded;
-				if (PipelineName != CurrentJob->Key.ShaderPipeline->GetName())
-				{
-					OutputFileContext.ModalErrorOrLog(TEXT("Expecting pipeline job \"%s\", got \"%s\" instead!"), CurrentJob->Key.ShaderPipeline->GetName(), *PipelineName);
-				}
-
-				check(!CurrentJob->bFinalized);
-				CurrentJob->bFinalized = true;
-
-				int32 NumStageJobs = -1;
-				OutputFile << NumStageJobs;
-
-				if (NumStageJobs != CurrentJob->StageJobs.Num())
-				{
-					OutputFileContext.ModalErrorOrLog(TEXT("Expecting %d stage pipeline %s, got %d instead!"), CurrentJob->StageJobs.Num(), (CurrentJob->StageJobs.Num() == 1 ? TEXT("job") : TEXT("jobs")), NumStageJobs);
-				}
-				else
-				{
-					for (int32 Index = 0; Index < NumStageJobs; Index++)
-					{
-						FShaderCompileJob* SingleJob = CurrentJob->StageJobs[Index];
-						ReadSingleJob(SingleJob, OutputFile);
-					}
-				}
-			}
-		}
-	}
-	
-	return FSCWErrorCode::Success;
-}
-
 void FShaderCompileUtilities::ExecuteShaderCompileJob(FShaderCommonCompileJob& Job)
 {
 	TRACE_CPUPROFILER_EVENT_SCOPE(FShaderCompileUtilities::ExecuteShaderCompileJob);
diff --git a/Engine/Source/Runtime/Engine/Public/ShaderCompiler.h b/Engine/Source/Runtime/Engine/Public/ShaderCompiler.h
--- a/Engine/Source/Runtime/Engine/Public/ShaderCompiler.h
+++ b/Engine/Source/Runtime/Engine/Public/ShaderCompiler.h
@@ -332,9 +332,6 @@ private:
 class FShaderCompileUtilities
 {
 public:
-	static bool DoWriteTasks(const TArray<FShaderCommonCompileJobPtr>& QueuedJobs, FArchive& TransferFile, IDistributedBuildController* BuildDistributionController = nullptr, bool bUseRelativePaths = false, bool bCompressTaskFile = false);
-	static FSCWErrorCode::ECode DoReadTaskResults(const TArray<FShaderCommonCompileJobPtr>& QueuedJobs, FArchive& OutputFile, FShaderCompileWorkerDiagnostics* OutWorkerDiagnostics = nullptr);
-
 	/** Execute the specified (single or pipeline) shader compile job. */
 	static void ExecuteShaderCompileJob(FShaderCommonCompileJob& Job);
 
diff --git a/Engine/Source/Runtime/RenderCore/Internal/ShaderCompileWorkerUtil.h b/Engine/Source/Runtime/RenderCore/Internal/ShaderCompileWorkerUtil.h
new file mode 100644
--- /dev/null
+++ b/Engine/Source/Runtime/RenderCore/Internal/ShaderCompileWorkerUtil.h
@@ -0,0 +1,85 @@
+// Copyright Epic Games, Inc. All Rights Reserved.
+
+#pragma once
+
+#include "ShaderCompilerJobTypes.h"
+#include "Containers/Array.h"
+
+class FArchive;
+
+struct FSCWErrorCode
+{
+	enum ECode : int32
+	{
+		NotSet = -1,
+		Success,
+		GeneralCrash,
+		BadShaderFormatVersion,
+		BadInputVersion,
+		BadSingleJobHeader,
+		BadPipelineJobHeader,
+		CantDeleteInputFile,
+		CantSaveOutputFile,
+		NoTargetShaderFormatsFound,
+		CantCompileForSpecificFormat,
+		CrashInsidePlatformCompiler,
+		BadInputFile,
+		OutOfMemory,
+	};
+
+	/**
+	Sets the global SCW error code if it hasn't been set before.
+	Call Reset first before setting a new value.
+	Returns true on success, otherwise the error code has already been set.
+	*/
+	static RENDERCORE_API void Report(ECode Code, const FStringView& Info = {});
+
+	/** Resets the global SCW error code to NotSet. */
+	static RENDERCORE_API void Reset();
+
+	/** Returns the global SCW error code. */
+	static RENDERCORE_API ECode Get();
+
+	/** Returns the global SCW error code information string. Empty string if not set. */
+	static RENDERCORE_API const FString& GetInfo();
+
+	/** Returns true if the SCW global error code has been set. Equivalent to 'Get() != NotSet'. */
+	static RENDERCORE_API bool IsSet();
+};
+
+class FShaderCompileWorkerUtil
+{
+public:
+	enum class EWriteTasksFlags : uint8
+	{
+		None = 0,
+		CompressTaskFile = 1 << 0,
+	};
+
+	enum class EReadTasksFlags : uint8
+	{
+		None = 0,
+		WillRetry = 1 << 0,
+	};
+
+	RENDERCORE_API static void LogQueuedCompileJobs(const TArray<FShaderCommonCompileJobPtr>& QueuedJobs, int32 NumProcessedJobs);
+	RENDERCORE_API static void DumpDebugCompileInput(TConstArrayView<FShaderCommonCompileJobPtr> QueuedJobs);
+	RENDERCORE_API static bool WriteTasks(const TArray<FShaderCommonCompileJobPtr>& QueuedJobs, FArchive& TransferFile, EWriteTasksFlags Flags = EWriteTasksFlags::None);
+	RENDERCORE_API static FSCWErrorCode::ECode ReadTasks(const TArray<FShaderCommonCompileJobPtr>& QueuedJobs, FArchive& OutputFile, FShaderCompileWorkerDiagnostics* OutWorkerDiagnostics, EReadTasksFlags Flags = EReadTasksFlags::None);
+
+private:
+	static bool HandleWorkerCrash(
+		const TArray<FShaderCommonCompileJobPtr>& QueuedJobs,
+		FArchive& OutputFile,
+		int32 OutputVersion,
+		int64 FileSize,
+		FSCWErrorCode::ECode ErrorCode,
+		int32 NumProcessedJobs,
+		int32 CallstackLength,
+		int32 ExceptionInfoLength,
+		int32 HostnameLength,
+		bool bWillRetry);
+};
+
+ENUM_CLASS_FLAGS(FShaderCompileWorkerUtil::EWriteTasksFlags);
+ENUM_CLASS_FLAGS(FShaderCompileWorkerUtil::EReadTasksFlags);
diff --git a/Engine/Source/Runtime/RenderCore/Private/ShaderCompileWorkerUtil.cpp b/Engine/Source/Runtime/RenderCore/Private/ShaderCompileWorkerUtil.cpp
new file mode 100644
--- /dev/null
+++ b/Engine/Source/Runtime/RenderCore/Private/ShaderCompileWorkerUtil.cpp
@@ -0,0 +1,844 @@
+// Copyright Epic Games, Inc. All Rights Reserved.
+
+#include "ShaderCompileWorkerUtil.h"
+
+#include "Misc/Compression.h"
+#include "Misc/FileHelper.h"
+#include "Misc/MessageDialog.h"
+#include "Misc/OutputDeviceRedirector.h"
+#include "Serialization/MemoryWriter.h"
+#include "ShaderDiagnostics.h"
+#include "Interfaces/ITargetPlatformManagerModule.h"
+#include "Interfaces/IShaderFormat.h"
+
+static TAutoConsoleVariable<bool> CVarDebugDumpWorkerInputs(
+	TEXT("r.ShaderCompiler.DebugDumpWorkerInputs"),
+	false,
+	TEXT("If true, worker input files will be saved for each individual compile job alongside other debug data (note that r.DumpShaderDebugInfo must also be enabled for this to function)"),
+	ECVF_ReadOnly);
+
+static TAutoConsoleVariable<bool> CVarShadersPropagateLocalWorkerOOMs(
+	TEXT("r.Shaders.PropagateLocalWorkerOOMs"),
+	false,
+	TEXT("When set, out-of-memory conditions in a local shader compile worker will be treated as regular out-of-memory conditions and propagated to the main process.\n")
+	TEXT("This is useful when running in environment with hard memory limits, where it does not matter which process in particular caused us to violate the memory limit."),
+	ECVF_Default);
+
+static void ModalErrorOrLog(const FString& Title, const FString& Text, int64 CurrentFilePos = 0, int64 ExpectedFileSize = 0, bool bIsErrorFatal = true)
+{
+	static FThreadSafeBool bModalReported;
+
+	FString BadFile;
+	if (CurrentFilePos > ExpectedFileSize)
+	{
+		// Corrupt file
+		BadFile = FString::Printf(TEXT(" (Truncated or corrupt output file! Current file pos %lld, file size %lld)"), CurrentFilePos, ExpectedFileSize);
+	}
+
+	if (bIsErrorFatal)
+	{
+		// Ensure errors are logged before exiting
+		GLog->Panic();
+
+		if (FPlatformProperties::SupportsWindowedMode() && !FApp::IsUnattended())
+		{
+			if (!bModalReported.AtomicSet(true))
+			{
+				UE_LOG(LogShaders, Error, TEXT("%s\n%s"), *Text, *BadFile);
+
+				// Show dialog box with error message and request exit
+				FMessageDialog::Open(EAppMsgType::Ok, FText::FromString(Text), FText::FromString(Title));
+				constexpr bool bForceExit = true;
+				FPlatformMisc::RequestExit(bForceExit, TEXT("ShaderCompiler.ModalErrorOrLog"));
+			}
+			else
+			{
+				// Another thread already opened a dialog box and requests exit
+				FPlatformProcess::SleepInfinite();
+			}
+		}
+		else
+		{
+			UE_LOG(LogShaders, Fatal, TEXT("%s\n%s\n%s"), *Title, *Text, *BadFile);
+		}
+	}
+	else
+	{
+		UE_LOG(LogShaders, Error, TEXT("%s\n%s\n%s"), *Title, *Text, *BadFile);
+	}
+}
+
+static TMap<FString, uint32> GetFormatVersionMap()
+{
+	TMap<FString, uint32> FormatVersionMap;
+
+	const TArray<const class IShaderFormat*>& ShaderFormats = GetTargetPlatformManagerRef().GetShaderFormats();
+	check(ShaderFormats.Num());
+	for (int32 Index = 0; Index < ShaderFormats.Num(); Index++)
+	{
+		TArray<FName> OutFormats;
+		ShaderFormats[Index]->GetSupportedFormats(OutFormats);
+		check(OutFormats.Num());
+		for (int32 InnerIndex = 0; InnerIndex < OutFormats.Num(); InnerIndex++)
+		{
+			uint32 Version = ShaderFormats[Index]->GetVersion(OutFormats[InnerIndex]);
+			FormatVersionMap.Add(OutFormats[InnerIndex].ToString(), Version);
+		}
+	}
+
+	return FormatVersionMap;
+}
+
+static const TCHAR* GetCompileJobSuccessText(FShaderCompileJob* SingleJob)
+{
+	if (SingleJob)
+	{
+		return SingleJob->Output.bSucceeded ? TEXT("Succeeded") : TEXT("Failed");
+	}
+	return TEXT("");
+}
+
+void FShaderCompileWorkerUtil::LogQueuedCompileJobs(const TArray<FShaderCommonCompileJobPtr>& QueuedJobs, int32 NumProcessedJobs)
+{
+	if (NumProcessedJobs == -1)
+	{
+		UE_LOG(LogShaders, Error, TEXT("SCW %d Queued Jobs, Unknown number of processed jobs!"), QueuedJobs.Num());
+	}
+	else
+	{
+		UE_LOG(LogShaders, Error, TEXT("SCW %d Queued Jobs, Finished %d single jobs"), QueuedJobs.Num(), NumProcessedJobs);
+	}
+
+	for (int32 Index = 0; Index < QueuedJobs.Num(); ++Index)
+	{
+		if (FShaderCompileJob* SingleJob = QueuedJobs[Index]->GetSingleShaderJob())
+		{
+			UE_LOG(LogShaders, Error, TEXT("Job %d [Single] %s: %s"), Index, GetCompileJobSuccessText(SingleJob), *GetSingleJobCompilationDump(SingleJob));
+		}
+		else
+		{
+			FShaderPipelineCompileJob* PipelineJob = QueuedJobs[Index]->GetShaderPipelineJob();
+			UE_LOG(LogShaders, Error, TEXT("Job %d: Pipeline %s "), Index, PipelineJob->Key.ShaderPipeline->GetName());
+			for (int32 JobIndex = 0; JobIndex < PipelineJob->StageJobs.Num(); ++JobIndex)
+			{
+				FShaderCompileJob* StageJob = PipelineJob->StageJobs[JobIndex]->GetSingleShaderJob();
+				UE_LOG(LogShaders, Error, TEXT("PipelineJob %d %s: %s"), JobIndex, GetCompileJobSuccessText(StageJob), *GetSingleJobCompilationDump(StageJob));
+			}
+		}
+	}
+
+	// Force a log flush so we can track the crash before the cooker potentially crashes before the output shows up
+	GLog->Flush();
+}
+
+// Make functions so the crash reporter can disambiguate the actual error because of the different callstacks
+namespace ShaderCompileWorkerError
+{
+	void HandleGeneralCrash(const TCHAR* ExceptionInfo, const TCHAR* Callstack)
+	{
+		ModalErrorOrLog(TEXT("ShaderCompileWorker crashed"), *FString::Printf(TEXT("Exception:\n%s\n\nCallstack:\n%s"), ExceptionInfo, Callstack));
+	}
+
+	void HandleBadShaderFormatVersion(const TCHAR* Data)
+	{
+		ModalErrorOrLog(TEXT("ShaderCompileWorker failed"), Data);
+	}
+
+	void HandleBadInputVersion(const TCHAR* Data)
+	{
+		ModalErrorOrLog(TEXT("ShaderCompileWorker failed"), Data);
+	}
+
+	void HandleBadSingleJobHeader(const TCHAR* Data)
+	{
+		ModalErrorOrLog(TEXT("ShaderCompileWorker failed"), Data);
+	}
+
+	void HandleBadPipelineJobHeader(const TCHAR* Data)
+	{
+		ModalErrorOrLog(TEXT("ShaderCompileWorker failed"), Data);
+	}
+
+	void HandleCantDeleteInputFile(const TCHAR* Data)
+	{
+		ModalErrorOrLog(TEXT("ShaderCompileWorker failed"), Data);
+	}
+
+	void HandleCantSaveOutputFile(const TCHAR* Data)
+	{
+		ModalErrorOrLog(TEXT("ShaderCompileWorker failed"), Data);
+	}
+
+	void HandleNoTargetShaderFormatsFound(const TCHAR* Data)
+	{
+		ModalErrorOrLog(TEXT("ShaderCompileWorker failed"), Data);
+	}
+
+	void HandleCantCompileForSpecificFormat(const TCHAR* Data)
+	{
+		ModalErrorOrLog(TEXT("ShaderCompileWorker failed"), Data);
+	}
+
+	void HandleOutputFileEmpty(const TCHAR* Filename)
+	{
+		ModalErrorOrLog(TEXT("ShaderCompileWorker failed"), FString::Printf(TEXT("Output file %s size is 0. Are you out of disk space?"), Filename));
+	}
+
+	void HandleOutputFileCorrupted(const TCHAR* Filename, int64 ExpectedSize, int64 ActualSize)
+	{
+		ModalErrorOrLog(TEXT("ShaderCompileWorker failed"), FString::Printf(TEXT("Output file corrupted (expected %I64d bytes, but only got %I64d): %s"), ExpectedSize, ActualSize, Filename));
+	}
+
+	void HandleCrashInsidePlatformCompiler(const TCHAR* Data)
+	{
+		// If the crash originates from a platform compiler, the error code must have been reported and we don't have to assume a corrupted output file.
+		// In that case, don't crash the cooker with a fatal error, just report the error so the cooker can dump debug info.
+		constexpr bool bIsErrorFatal = false;
+		ModalErrorOrLog(TEXT("ShaderCompileWorker failed"), FString::Printf(TEXT("Crash inside the platform compiler:\n%s"), Data), 0, 0, bIsErrorFatal);
+	}
+
+	void HandleBadInputFile(const TCHAR* Data)
+	{
+		ModalErrorOrLog(TEXT("ShaderCompileWorker failed"), FString::Printf(TEXT("Bad-input-file exception:\n%s"), Data));
+	}
+
+	bool HandleOutOfMemory(const TCHAR* ExceptionInfo, const TCHAR* Hostname, const FPlatformMemoryStats& MemoryStats, const TArray<FShaderCommonCompileJobPtr>& QueuedJobs, bool bWillRetry)
+	{
+		constexpr int64 Gibibyte = 1024 * 1024 * 1024;
+		const FString ErrorReport = FString::Printf(
+			TEXT("ShaderCompileWorker failed with out-of-memory (OOM) exception on machine \"%s\" (%s); MemoryStats:")
+			TEXT("\n\tAvailablePhysical %llu (%.2f GiB)")
+			TEXT("\n\t AvailableVirtual %llu (%.2f GiB)")
+			TEXT("\n\t     UsedPhysical %llu (%.2f GiB)")
+			TEXT("\n\t PeakUsedPhysical %llu (%.2f GiB)")
+			TEXT("\n\t      UsedVirtual %llu (%.2f GiB)")
+			TEXT("\n\t  PeakUsedVirtual %llu (%.2f GiB)"),
+			Hostname,
+			(ExceptionInfo[0] == TEXT('\0') ? TEXT("No exception information") : ExceptionInfo),
+			MemoryStats.AvailablePhysical, double(MemoryStats.AvailablePhysical) / Gibibyte,
+			MemoryStats.AvailableVirtual, double(MemoryStats.AvailableVirtual) / Gibibyte,
+			MemoryStats.UsedPhysical, double(MemoryStats.UsedPhysical) / Gibibyte,
+			MemoryStats.PeakUsedPhysical, double(MemoryStats.PeakUsedPhysical) / Gibibyte,
+			MemoryStats.UsedVirtual, double(MemoryStats.UsedVirtual) / Gibibyte,
+			MemoryStats.PeakUsedVirtual, double(MemoryStats.PeakUsedVirtual) / Gibibyte
+		);
+
+		if (bWillRetry)
+		{
+			// assume caller will retry the failed jobs rather than aborting
+			return true;
+		}
+		else
+		{
+			if (CVarShadersPropagateLocalWorkerOOMs.GetValueOnAnyThread())
+			{
+				FPlatformMemory::OnOutOfMemory(0, 64);
+			}
+			ModalErrorOrLog(TEXT("ShaderCompileWorker failed"), ErrorReport);
+			return false;
+		}
+	}
+}
+
+
+// Disable optimization for this crash handler to get full access to the entire stack frame when debugging a crash dump
+UE_DISABLE_OPTIMIZATION_SHIP
+bool FShaderCompileWorkerUtil::HandleWorkerCrash(
+	const TArray<FShaderCommonCompileJobPtr>& QueuedJobs, 
+	FArchive& OutputFile, 
+	int32 OutputVersion, 
+	int64 FileSize, 
+	FSCWErrorCode::ECode ErrorCode, 
+	int32 NumProcessedJobs, 
+	int32 CallstackLength, 
+	int32 ExceptionInfoLength, 
+	int32 HostnameLength,
+	bool bWillRetry)
+{
+	TArray<TCHAR> Callstack;
+	Callstack.AddUninitialized(CallstackLength + 1);
+	OutputFile.Serialize(Callstack.GetData(), CallstackLength * sizeof(TCHAR));
+	Callstack[CallstackLength] = TEXT('\0');
+
+	TArray<TCHAR> ExceptionInfo;
+	ExceptionInfo.AddUninitialized(ExceptionInfoLength + 1);
+	OutputFile.Serialize(ExceptionInfo.GetData(), ExceptionInfoLength * sizeof(TCHAR));
+	ExceptionInfo[ExceptionInfoLength] = TEXT('\0');
+
+	TArray<TCHAR> Hostname;
+	Hostname.AddUninitialized(HostnameLength + 1);
+	OutputFile.Serialize(Hostname.GetData(), HostnameLength * sizeof(TCHAR));
+	Hostname[HostnameLength] = TEXT('\0');
+
+	// Read available and used physical memory from worker machine on OOM error
+	FPlatformMemoryStats MemoryStats;
+	if (ErrorCode == FSCWErrorCode::OutOfMemory)
+	{
+		OutputFile
+			<< MemoryStats.AvailablePhysical
+			<< MemoryStats.AvailableVirtual
+			<< MemoryStats.UsedPhysical
+			<< MemoryStats.PeakUsedPhysical
+			<< MemoryStats.UsedVirtual
+			<< MemoryStats.PeakUsedVirtual
+			;
+	}
+
+	// Store primary job information onto stack to make it part of a crash dump
+	static const int32 MaxNumCharsForSourcePaths = 8192;
+	int32 JobInputSourcePathsLength = 0;
+	ANSICHAR JobInputSourcePaths[MaxNumCharsForSourcePaths];
+	JobInputSourcePaths[0] = 0;
+
+	auto WriteInputSourcePathOntoStack = [&JobInputSourcePathsLength, &JobInputSourcePaths](const ANSICHAR* InputSourcePath)
+	{
+		if (InputSourcePath != nullptr && JobInputSourcePathsLength + 3 < MaxNumCharsForSourcePaths)
+		{
+			// Copy input source path into stack buffer
+			int32 InputSourcePathLength = FMath::Min(FCStringAnsi::Strlen(InputSourcePath), (MaxNumCharsForSourcePaths - JobInputSourcePathsLength - 2));
+			FMemory::Memcpy(JobInputSourcePaths + JobInputSourcePathsLength, InputSourcePath, InputSourcePathLength);
+
+			// Write newline character and put NUL character at the end
+			JobInputSourcePathsLength += InputSourcePathLength;
+			JobInputSourcePaths[JobInputSourcePathsLength] = TEXT('\n');
+			++JobInputSourcePathsLength;
+			JobInputSourcePaths[JobInputSourcePathsLength] = 0;
+		}
+	};
+
+	auto StoreInputDebugInfo = [&WriteInputSourcePathOntoStack, &JobInputSourcePathsLength, &JobInputSourcePaths](const FShaderCompilerInput& Input)
+	{
+		FString DebugInfo = FString::Printf(TEXT("%s:%s"), *Input.VirtualSourceFilePath, *Input.EntryPointName);
+		WriteInputSourcePathOntoStack(TCHAR_TO_UTF8(*DebugInfo));
+	};
+
+	for (auto CommonJob : QueuedJobs)
+	{
+		if (FShaderCompileJob* SingleJob = CommonJob->GetSingleShaderJob())
+		{
+			StoreInputDebugInfo(SingleJob->Input);
+		}
+		else if (FShaderPipelineCompileJob* PipelineJob = CommonJob->GetShaderPipelineJob())
+		{
+			for (int32 Job = 0; Job < PipelineJob->StageJobs.Num(); ++Job)
+			{
+				if (FShaderCompileJob* SingleStageJob = PipelineJob->StageJobs[Job])
+				{
+					StoreInputDebugInfo(SingleStageJob->Input);
+				}
+			}
+		}
+	}
+
+	// One entry per error code as we want to have different callstacks for crash reporter...
+	switch (ErrorCode)
+	{
+	default:
+	case FSCWErrorCode::GeneralCrash:
+		LogQueuedCompileJobs(QueuedJobs, NumProcessedJobs);
+		ShaderCompileWorkerError::HandleGeneralCrash(ExceptionInfo.GetData(), Callstack.GetData());
+		break;
+	case FSCWErrorCode::BadShaderFormatVersion:
+		ShaderCompileWorkerError::HandleBadShaderFormatVersion(ExceptionInfo.GetData());
+		break;
+	case FSCWErrorCode::BadInputVersion:
+		ShaderCompileWorkerError::HandleBadInputVersion(ExceptionInfo.GetData());
+		break;
+	case FSCWErrorCode::BadSingleJobHeader:
+		ShaderCompileWorkerError::HandleBadSingleJobHeader(ExceptionInfo.GetData());
+		break;
+	case FSCWErrorCode::BadPipelineJobHeader:
+		ShaderCompileWorkerError::HandleBadPipelineJobHeader(ExceptionInfo.GetData());
+		break;
+	case FSCWErrorCode::CantDeleteInputFile:
+		ShaderCompileWorkerError::HandleCantDeleteInputFile(ExceptionInfo.GetData());
+		break;
+	case FSCWErrorCode::CantSaveOutputFile:
+		ShaderCompileWorkerError::HandleCantSaveOutputFile(ExceptionInfo.GetData());
+		break;
+	case FSCWErrorCode::NoTargetShaderFormatsFound:
+		ShaderCompileWorkerError::HandleNoTargetShaderFormatsFound(ExceptionInfo.GetData());
+		break;
+	case FSCWErrorCode::CantCompileForSpecificFormat:
+		ShaderCompileWorkerError::HandleCantCompileForSpecificFormat(ExceptionInfo.GetData());
+		break;
+	case FSCWErrorCode::CrashInsidePlatformCompiler:
+		LogQueuedCompileJobs(QueuedJobs, NumProcessedJobs);
+		ShaderCompileWorkerError::HandleCrashInsidePlatformCompiler(ExceptionInfo.GetData());
+		break;
+	case FSCWErrorCode::BadInputFile:
+		ShaderCompileWorkerError::HandleBadInputFile(ExceptionInfo.GetData());
+		break;
+	case FSCWErrorCode::OutOfMemory:
+		return ShaderCompileWorkerError::HandleOutOfMemory(ExceptionInfo.GetData(), Hostname.GetData(), MemoryStats, QueuedJobs, bWillRetry);
+	case FSCWErrorCode::Success:
+		// Can't get here...
+		return true;
+	}
+	return false;
+}
+UE_ENABLE_OPTIMIZATION_SHIP
+
+
+static void SplitJobsByType(const TArray<FShaderCommonCompileJobPtr>& QueuedJobs, TArray<FShaderCompileJob*>& OutQueuedSingleJobs, TArray<FShaderPipelineCompileJob*>& OutQueuedPipelineJobs)
+{
+	for (int32 Index = 0; Index < QueuedJobs.Num(); ++Index)
+	{
+		FShaderCommonCompileJobPtr CommonJob = QueuedJobs[Index];
+		if (FShaderCompileJob* SingleJob = CommonJob->GetSingleShaderJob())
+		{
+			OutQueuedSingleJobs.Add(SingleJob);
+
+		}
+		else if (FShaderPipelineCompileJob* PipelineJob = CommonJob->GetShaderPipelineJob())
+		{
+			OutQueuedPipelineJobs.Add(PipelineJob);
+		}
+		else
+		{
+			checkf(0, TEXT("FShaderCommonCompileJob::Type=%d is not a valid type for a shader compile job"), (int32)CommonJob->Type);
+		}
+	}
+}
+
+static bool DoWriteTasksInner(const TArray<FShaderCommonCompileJobPtr>& QueuedJobs, FArchive& InTransferFile, FShaderCompileWorkerUtil::EWriteTasksFlags Flags)
+{
+	int32 InputVersion = ShaderCompileWorkerInputVersion;
+	InTransferFile << InputVersion;
+
+	TArray<uint8> UncompressedArray;
+	FMemoryWriter TransferMemory(UncompressedArray);
+	bool bCompressTaskFile = EnumHasAnyFlags(Flags, FShaderCompileWorkerUtil::EWriteTasksFlags::CompressTaskFile);
+	FArchive& TransferFile = bCompressTaskFile ? TransferMemory : InTransferFile;
+	if (!bCompressTaskFile)
+	{
+		// still write NAME_None as string
+		FString FormatNone = FName(NAME_None).ToString();
+		TransferFile << FormatNone;
+	}
+
+	static TMap<FString, uint32> FormatVersionMap = GetFormatVersionMap();
+
+	TransferFile << FormatVersionMap;
+
+	TArray<FShaderCompileJob*> QueuedSingleJobs;
+	TArray<FShaderPipelineCompileJob*> QueuedPipelineJobs;
+	SplitJobsByType(QueuedJobs, QueuedSingleJobs, QueuedPipelineJobs);
+
+	TArray<TRefCountPtr<FSharedShaderCompilerEnvironment>> SharedEnvironments;
+	TArray<const FShaderParametersMetadata*> RequestShaderParameterStructures;
+
+	// gather shared environments and parameter structures, these tend to be shared between jobs
+	{
+
+		for (int32 JobIndex = 0; JobIndex < QueuedSingleJobs.Num(); JobIndex++)
+		{
+			QueuedSingleJobs[JobIndex]->Input.GatherSharedInputs(SharedEnvironments, RequestShaderParameterStructures);
+		}
+
+		for (int32 JobIndex = 0; JobIndex < QueuedPipelineJobs.Num(); JobIndex++)
+		{
+			auto* PipelineJob = QueuedPipelineJobs[JobIndex];
+			int32 NumStageJobs = PipelineJob->StageJobs.Num();
+
+			for (int32 Index = 0; Index < NumStageJobs; Index++)
+			{
+				PipelineJob->StageJobs[Index]->Input.GatherSharedInputs(SharedEnvironments, RequestShaderParameterStructures);
+			}
+		}
+
+		int32 NumSharedEnvironments = SharedEnvironments.Num();
+		TransferFile << NumSharedEnvironments;
+
+		for (int32 EnvironmentIndex = 0; EnvironmentIndex < SharedEnvironments.Num(); EnvironmentIndex++)
+		{
+			TransferFile << *SharedEnvironments[EnvironmentIndex];
+		}
+	}
+
+	// Write shader parameter structures
+	TArray<const FShaderParametersMetadata*> AllShaderParameterStructures;
+	{
+		// List all dependencies.
+		for (int32 StructId = 0; StructId < RequestShaderParameterStructures.Num(); StructId++)
+		{
+			RequestShaderParameterStructures[StructId]->IterateStructureMetadataDependencies(
+				[&](const FShaderParametersMetadata* Struct)
+			{
+				AllShaderParameterStructures.AddUnique(Struct);
+			});
+		}
+
+		// Write all shader parameter structure.
+		int32 NumParameterStructures = AllShaderParameterStructures.Num();
+		TransferFile << NumParameterStructures;
+		for (const FShaderParametersMetadata* Struct : AllShaderParameterStructures)
+		{
+			FString LayoutName = Struct->GetLayout().GetDebugName();
+			FString StructTypeName = Struct->GetStructTypeName();
+			FString ShaderVariableName = Struct->GetShaderVariableName();
+			uint8 UseCase = uint8(Struct->GetUseCase());
+			FString StructFileName = FString(ANSI_TO_TCHAR(Struct->GetFileName()));
+			int32 StructFileLine = Struct->GetFileLine();
+			uint32 Size = Struct->GetSize();
+			int32 MemberCount = Struct->GetMembers().Num();
+
+			static_assert(sizeof(UseCase) == sizeof(FShaderParametersMetadata::EUseCase), "Cast failure.");
+
+			TransferFile << LayoutName;
+			TransferFile << StructTypeName;
+			TransferFile << ShaderVariableName;
+			TransferFile << UseCase;
+			TransferFile << StructFileName;
+			TransferFile << StructFileLine;
+			TransferFile << Size;
+			TransferFile << MemberCount;
+
+			for (const FShaderParametersMetadata::FMember& Member : Struct->GetMembers())
+			{
+				FString Name = Member.GetName();
+				FString ShaderType = Member.GetShaderType();
+				int32 FileLine = Member.GetFileLine();
+				uint32 Offset = Member.GetOffset();
+				uint8 BaseType = uint8(Member.GetBaseType());
+				uint8 PrecisionModifier = uint8(Member.GetPrecision());
+				uint32 NumRows = Member.GetNumRows();
+				uint32 NumColumns = Member.GetNumColumns();
+				uint32 NumElements = Member.GetNumElements();
+				int32 StructMetadataIndex = INDEX_NONE;
+				if (Member.GetStructMetadata())
+				{
+					StructMetadataIndex = AllShaderParameterStructures.Find(Member.GetStructMetadata());
+					check(StructMetadataIndex != INDEX_NONE);
+				}
+
+				static_assert(sizeof(BaseType) == sizeof(EUniformBufferBaseType), "Cast failure.");
+				static_assert(sizeof(PrecisionModifier) == sizeof(EShaderPrecisionModifier::Type), "Cast failure.");
+
+				TransferFile << Name;
+				TransferFile << ShaderType;
+				TransferFile << FileLine;
+				TransferFile << Offset;
+				TransferFile << BaseType;
+				TransferFile << PrecisionModifier;
+				TransferFile << NumRows;
+				TransferFile << NumColumns;
+				TransferFile << NumElements;
+				TransferFile << StructMetadataIndex;
+			}
+		}
+	}
+
+	// Write individual shader jobs
+	{
+		int32 SingleJobHeader = ShaderCompileWorkerSingleJobHeader;
+		TransferFile << SingleJobHeader;
+
+		int32 NumBatches = QueuedSingleJobs.Num();
+		TransferFile << NumBatches;
+
+		// Serialize all the batched jobs
+		for (int32 JobIndex = 0; JobIndex < QueuedSingleJobs.Num(); JobIndex++)
+		{
+			QueuedSingleJobs[JobIndex]->SerializeWorkerInput(TransferFile);
+			QueuedSingleJobs[JobIndex]->Input.SerializeSharedInputs(TransferFile, SharedEnvironments, AllShaderParameterStructures);
+		}
+	}
+
+	// Write shader pipeline jobs
+	{
+		int32 PipelineJobHeader = ShaderCompileWorkerPipelineJobHeader;
+		TransferFile << PipelineJobHeader;
+
+		int32 NumBatches = QueuedPipelineJobs.Num();
+		TransferFile << NumBatches;
+		for (int32 JobIndex = 0; JobIndex < QueuedPipelineJobs.Num(); JobIndex++)
+		{
+			auto* PipelineJob = QueuedPipelineJobs[JobIndex];
+			FString PipelineName = PipelineJob->Key.ShaderPipeline->GetName();
+			TransferFile << PipelineName;
+			int32 NumStageJobs = PipelineJob->StageJobs.Num();
+			TransferFile << NumStageJobs;
+			for (int32 Index = 0; Index < NumStageJobs; Index++)
+			{
+				PipelineJob->StageJobs[Index]->SerializeWorkerInput(TransferFile);
+				PipelineJob->StageJobs[Index]->Input.SerializeSharedInputs(TransferFile, SharedEnvironments, AllShaderParameterStructures);
+			}
+		}
+	}
+
+	if (bCompressTaskFile)
+	{
+		TransferFile.Close();
+
+		FName CompressionFormatToUse = NAME_LZ4;
+
+		FString FormatName = CompressionFormatToUse.ToString();
+		InTransferFile << FormatName;
+
+		// serialize uncompressed data size
+		int32 UncompressedDataSize = UncompressedArray.Num();
+		checkf(UncompressedDataSize != 0, TEXT("Did not write any data to the task file for the compression."));
+		InTransferFile << UncompressedDataSize;
+
+		// not using SerializeCompressed because it splits into smaller chunks
+		int32 CompressedSizeBound = FCompression::CompressMemoryBound(CompressionFormatToUse, static_cast<int32>(UncompressedDataSize));
+		TArray<uint8> CompressedBuffer;
+		CompressedBuffer.SetNumUninitialized(CompressedSizeBound);
+
+		int32 ActualCompressedSize = CompressedSizeBound;
+		bool bSucceeded = FCompression::CompressMemory(CompressionFormatToUse, CompressedBuffer.GetData(), ActualCompressedSize, UncompressedArray.GetData(), UncompressedDataSize, COMPRESS_BiasSpeed);
+		checkf(ActualCompressedSize <= CompressedSizeBound, TEXT("Compressed size was larger than the bound - we stomped the memory."));
+		CompressedBuffer.SetNum(ActualCompressedSize, EAllowShrinking::No);
+
+		InTransferFile << CompressedBuffer;
+	}
+
+	return InTransferFile.Close();
+}
+
+const TCHAR* DebugWorkerInputFileName = TEXT("DebugSCW.in");
+const TCHAR* DebugWorkerOutputFileName = TEXT("DebugSCW.out");
+
+static FString CreateShaderCompilerWorkerDebugCommandLine(FString DebugWorkerInputFilePath)
+{
+	// 0 is parent PID, pass zero TTL and KeepInput to make SCW process the single job then exit without deleting the input file
+	return FString::Printf(TEXT("\"%s\" 0 \"DebugSCW\" %s %s -TimeToLive=0.0f -KeepInput"),
+		*DebugWorkerInputFilePath, // working directory for SCW
+		DebugWorkerInputFileName,
+		DebugWorkerOutputFileName);
+}
+
+void FShaderCompileWorkerUtil::DumpDebugCompileInput(TConstArrayView<FShaderCommonCompileJobPtr> QueuedJobs)
+{
+	if (CVarDebugDumpWorkerInputs.GetValueOnAnyThread())
+	{
+		for (const FShaderCommonCompileJobPtr& CommonJob : QueuedJobs)
+		{
+			FShaderPipelineCompileJob* PipelineJob = CommonJob->GetShaderPipelineJob();
+			FString DebugWorkerInputFilePath;
+			if (PipelineJob)
+			{
+				// for pipeline jobs, write out the worker input for the whole pipeline, but only for the first stage
+				// would be better to put in a parent folder probably...
+				DebugWorkerInputFilePath = PipelineJob->StageJobs[0]->Input.DumpDebugInfoPath;
+			}
+			else
+			{
+				DebugWorkerInputFilePath = CommonJob->GetSingleShaderJob()->Input.DumpDebugInfoPath;
+			}
+			if (!DebugWorkerInputFilePath.IsEmpty())
+			{
+				TArray<FShaderCommonCompileJobPtr> SingleJobArray;
+				SingleJobArray.Add(CommonJob);
+
+				FArchive* DebugWorkerInputFileWriter = IFileManager::Get().CreateFileWriter(*(DebugWorkerInputFilePath / DebugWorkerInputFileName), FILEWRITE_NoFail);
+				DoWriteTasksInner(
+					SingleJobArray,
+					*DebugWorkerInputFileWriter,
+					EWriteTasksFlags::CompressTaskFile);	// Always compress the debug files; they are rather large so this saves some disk space
+				DebugWorkerInputFileWriter->Close();
+				delete DebugWorkerInputFileWriter;
+
+				FFileHelper::SaveStringToFile(
+					CreateShaderCompilerWorkerDebugCommandLine(DebugWorkerInputFilePath),
+					*(DebugWorkerInputFilePath / TEXT("DebugCompileArgs.txt")));
+			}
+		}
+	}
+}
+
+// Serialize Queued Job information
+bool FShaderCompileWorkerUtil::WriteTasks(const TArray<FShaderCommonCompileJobPtr>& QueuedJobs, FArchive& InTransferFile, EWriteTasksFlags Flags)
+{
+	DumpDebugCompileInput(QueuedJobs);
+
+	return DoWriteTasksInner(QueuedJobs, InTransferFile, Flags);
+}
+
+static void ReadSingleJob(FShaderCompileJob* CurrentJob, FArchive& WorkerOutputFileReader)
+{
+	TRACE_CPUPROFILER_EVENT_SCOPE(ReadSingleJob);
+
+	check(!CurrentJob->bFinalized);
+	CurrentJob->bFinalized = true;
+
+	// Deserialize the shader compilation output.
+	CurrentJob->SerializeWorkerOutput(WorkerOutputFileReader);
+
+	// The job should already have a non-zero output hash
+	checkf(CurrentJob->Output.OutputHash != FSHAHash() || !CurrentJob->bSucceeded, TEXT("OutputHash for a successful job was not set in the shader compile worker!"));
+}
+
+// Helper struct to provide consistent error report with detailed information about corrupted ShaderCompileWorker output file.
+struct FSCWOutputFileContext
+{
+	FArchive& OutputFile;
+	int64 FileSize = 0;
+
+	FSCWOutputFileContext(FArchive& OutputFile) :
+		OutputFile(OutputFile)
+	{
+	}
+
+	template <typename... Types>
+	void ModalErrorOrLog(UE::Core::TCheckedFormatString<FString::FmtCharType, Types...> Format, Types... Args)
+	{
+		FString Text = FString::Printf(Format, Args...);
+		Text = FString::Printf(TEXT("File path: \"%s\"\n%s\nForgot to build ShaderCompileWorker or delete invalidated DerivedDataCache?"), *OutputFile.GetArchiveName(), *Text);
+		const TCHAR* Title = TEXT("Corrupted ShaderCompileWorker output file");
+		if (FileSize > 0)
+		{
+			::ModalErrorOrLog(Title, Text, OutputFile.Tell(), FileSize);
+		}
+		else
+		{
+			::ModalErrorOrLog(Title, Text, 0, 0);
+		}
+	}
+};
+
+// Process results from Worker Process.
+// Returns false if reading the tasks failed but we were able to recover from handing a crash report. In this case, all jobs must be submitted/processed again.
+FSCWErrorCode::ECode FShaderCompileWorkerUtil::ReadTasks(const TArray<FShaderCommonCompileJobPtr>& QueuedJobs, FArchive& OutputFile, FShaderCompileWorkerDiagnostics* OutWorkerDiagnostics, EReadTasksFlags Flags)
+{
+	FSCWOutputFileContext OutputFileContext(OutputFile);
+
+	if (OutputFile.TotalSize() == 0)
+	{
+		ShaderCompileWorkerError::HandleOutputFileEmpty(*OutputFile.GetArchiveName());
+	}
+
+	int32 OutputVersion = ShaderCompileWorkerOutputVersion;
+	OutputFile << OutputVersion;
+
+	if (ShaderCompileWorkerOutputVersion != OutputVersion)
+	{
+		OutputFileContext.ModalErrorOrLog(TEXT("Expecting output version %d, got %d instead!"), ShaderCompileWorkerOutputVersion, OutputVersion);
+	}
+
+	OutputFile << OutputFileContext.FileSize;
+
+	// Check for corrupted output file
+	if (OutputFileContext.FileSize > OutputFile.TotalSize())
+	{
+		ShaderCompileWorkerError::HandleOutputFileCorrupted(*OutputFile.GetArchiveName(), OutputFileContext.FileSize, OutputFile.TotalSize());
+	}
+
+	FShaderCompileWorkerDiagnostics WorkerDiagnostics;
+	OutputFile << WorkerDiagnostics;
+
+	if (OutWorkerDiagnostics)
+	{
+		*OutWorkerDiagnostics = WorkerDiagnostics;
+	}
+
+	int32 NumProcessedJobs = 0;
+	OutputFile << NumProcessedJobs;
+
+	int32 CallstackLength = 0;
+	OutputFile << CallstackLength;
+
+	int32 ExceptionInfoLength = 0;
+	OutputFile << ExceptionInfoLength;
+
+	int32 HostnameLength = 0;
+	OutputFile << HostnameLength;
+
+	bool bWillRetry = EnumHasAnyFlags(Flags, EReadTasksFlags::WillRetry);
+
+	if (WorkerDiagnostics.ErrorCode != FSCWErrorCode::Success)
+	{
+		// If worker crashed in a way we were able to recover from, return and expect the compile jobs to be reissued already
+		if (HandleWorkerCrash(QueuedJobs, OutputFile, OutputVersion, OutputFileContext.FileSize, (FSCWErrorCode::ECode)WorkerDiagnostics.ErrorCode, NumProcessedJobs, CallstackLength, ExceptionInfoLength, HostnameLength, bWillRetry))
+		{
+			FSCWErrorCode::Reset();
+			return (FSCWErrorCode::ECode)WorkerDiagnostics.ErrorCode;
+		}
+	}
+
+	TArray<FShaderCompileJob*> QueuedSingleJobs;
+	TArray<FShaderPipelineCompileJob*> QueuedPipelineJobs;
+	SplitJobsByType(QueuedJobs, QueuedSingleJobs, QueuedPipelineJobs);
+
+	// Read single jobs
+	{
+		int32 SingleJobHeader = -1;
+		OutputFile << SingleJobHeader;
+		if (SingleJobHeader != ShaderCompileWorkerSingleJobHeader)
+		{
+			OutputFileContext.ModalErrorOrLog(TEXT("Expecting single job header ID 0x%08X, got 0x%08X instead!"), ShaderCompileWorkerSingleJobHeader, SingleJobHeader);
+		}
+
+		int32 NumJobs;
+		OutputFile << NumJobs;
+		if (NumJobs != QueuedSingleJobs.Num())
+		{
+			OutputFileContext.ModalErrorOrLog(TEXT("Expecting %d single %s, got %d instead!"), QueuedSingleJobs.Num(), (QueuedSingleJobs.Num() == 1 ? TEXT("job") : TEXT("jobs")), NumJobs);
+		}
+		else
+		{
+			for (int32 JobIndex = 0; JobIndex < NumJobs; JobIndex++)
+			{
+				auto* CurrentJob = QueuedSingleJobs[JobIndex];
+				ReadSingleJob(CurrentJob, OutputFile);
+			}
+		}
+	}
+
+	// Pipeline jobs
+	{
+		int32 PipelineJobHeader = -1;
+		OutputFile << PipelineJobHeader;
+		if (PipelineJobHeader != ShaderCompileWorkerPipelineJobHeader)
+		{
+			OutputFileContext.ModalErrorOrLog(TEXT("Expecting pipeline jobs header ID 0x%08X, got 0x%08X instead!"), ShaderCompileWorkerPipelineJobHeader, PipelineJobHeader);
+		}
+
+		int32 NumJobs;
+		OutputFile << NumJobs;
+		if (NumJobs != QueuedPipelineJobs.Num())
+		{
+			OutputFileContext.ModalErrorOrLog(TEXT("Expecting %d pipeline %s, got %d instead!"), QueuedPipelineJobs.Num(), (QueuedPipelineJobs.Num() == 1 ? TEXT("job") : TEXT("jobs")), NumJobs);
+		}
+		else
+		{
+			for (int32 JobIndex = 0; JobIndex < NumJobs; JobIndex++)
+			{
+				FShaderPipelineCompileJob* CurrentJob = QueuedPipelineJobs[JobIndex];
+
+				FString PipelineName;
+				OutputFile << PipelineName;
+				bool bSucceeded = false;
+				OutputFile << bSucceeded;
+				CurrentJob->bSucceeded = bSucceeded;
+				if (PipelineName != CurrentJob->Key.ShaderPipeline->GetName())
+				{
+					OutputFileContext.ModalErrorOrLog(TEXT("Expecting pipeline job \"%s\", got \"%s\" instead!"), CurrentJob->Key.ShaderPipeline->GetName(), *PipelineName);
+				}
+
+				check(!CurrentJob->bFinalized);
+				CurrentJob->bFinalized = true;
+
+				int32 NumStageJobs = -1;
+				OutputFile << NumStageJobs;
+
+				if (NumStageJobs != CurrentJob->StageJobs.Num())
+				{
+					OutputFileContext.ModalErrorOrLog(TEXT("Expecting %d stage pipeline %s, got %d instead!"), CurrentJob->StageJobs.Num(), (CurrentJob->StageJobs.Num() == 1 ? TEXT("job") : TEXT("jobs")), NumStageJobs);
+				}
+				else
+				{
+					for (int32 Index = 0; Index < NumStageJobs; Index++)
+					{
+						FShaderCompileJob* SingleJob = CurrentJob->StageJobs[Index];
+						ReadSingleJob(SingleJob, OutputFile);
+					}
+				}
+			}
+		}
+	}
+	
+	return FSCWErrorCode::Success;
+}
+
diff --git a/Engine/Source/Runtime/RenderCore/Private/ShaderCore.cpp b/Engine/Source/Runtime/RenderCore/Private/ShaderCore.cpp
--- a/Engine/Source/Runtime/RenderCore/Private/ShaderCore.cpp
+++ b/Engine/Source/Runtime/RenderCore/Private/ShaderCore.cpp
@@ -37,6 +37,7 @@
 #include "ShaderCompilerCore.h"
 #include "ShaderCompilerDefinitions.h"
 #include "ShaderCompilerJobTypes.h"
+#include "ShaderCompileWorkerUtil.h"
 #include "ShaderDiagnostics.h"
 #include "Stats/StatsMisc.h"
 #include "String/Find.h"
diff --git a/Engine/Source/Runtime/RenderCore/Public/ShaderCompilerCore.h b/Engine/Source/Runtime/RenderCore/Public/ShaderCompilerCore.h
--- a/Engine/Source/Runtime/RenderCore/Public/ShaderCompilerCore.h
+++ b/Engine/Source/Runtime/RenderCore/Public/ShaderCompilerCore.h
@@ -630,46 +630,6 @@ struct FShaderCompilerOutput
 	}
 };
 
-struct FSCWErrorCode
-{
-	enum ECode : int32
-	{
-		NotSet = -1,
-		Success,
-		GeneralCrash,
-		BadShaderFormatVersion,
-		BadInputVersion,
-		BadSingleJobHeader,
-		BadPipelineJobHeader,
-		CantDeleteInputFile,
-		CantSaveOutputFile,
-		NoTargetShaderFormatsFound,
-		CantCompileForSpecificFormat,
-		CrashInsidePlatformCompiler,
-		BadInputFile,
-		OutOfMemory,
-	};
-
-	/**
-	Sets the global SCW error code if it hasn't been set before.
-	Call Reset first before setting a new value.
-	Returns true on success, otherwise the error code has already been set.
-	*/
-	static RENDERCORE_API void Report(ECode Code, const FStringView& Info = {});
-
-	/** Resets the global SCW error code to NotSet. */
-	static RENDERCORE_API void Reset();
-
-	/** Returns the global SCW error code. */
-	static RENDERCORE_API ECode Get();
-
-	/** Returns the global SCW error code information string. Empty string if not set. */
-	static RENDERCORE_API const FString& GetInfo();
-
-	/** Returns true if the SCW global error code has been set. Equivalent to 'Get() != NotSet'. */
-	static RENDERCORE_API bool IsSet();
-};
-
 /** Serializable structure of diagnostic output from a SCW process. Include error code and timing statistics for the duration of a job batch. */
 struct FShaderCompileWorkerDiagnostics
 {
