NavMovementComponent implements it's own basic versioning. NavMovementProperties on the NavMovementComponent now properly copy over previously deprecated values in a better and inheritance friendly way.

diff --git a/Engine/Source/Runtime/Engine/Classes/GameFramework/NavMovementComponent.h b/Engine/Source/Runtime/Engine/Classes/GameFramework/NavMovementComponent.h
--- a/Engine/Source/Runtime/Engine/Classes/GameFramework/NavMovementComponent.h
+++ b/Engine/Source/Runtime/Engine/Classes/GameFramework/NavMovementComponent.h
@@ -19,6 +19,53 @@
 
 class UCapsuleComponent;
 
+/** The version of the UNavMovementComponent. Used for upgrade paths. */
+UENUM()
+enum class ENavMovementComponentVersion : uint8
+{
+	IndependentNavMovementProperties,	// Legacy version from Pre-UE5.6 (before we moved nav movement properties into a struct through this serialization)
+	NavMovementPropertiesInStruct,	// New version available in UE5.6 with properly upgrading NavMovementProperties
+
+	Current = NavMovementPropertiesInStruct
+};
+
+/**
+ * UNavMovementComponent Data needs to be versioned
+ * Unfortunately, the Custom Package Versioning doesn't work for this case. For that reason we want to
+ * store a UProperty that says what version we're on. Unfortunately that propagates from Parent to Child
+ * (inheritance rules) so to overcome that, we have a special UStruct that always serializes its data
+ * (so it will always be loaded, not inherited).
+ * Here is how to do that:
+ *	1. Return false in Identical (which effectively disables delta serialization).
+ *  2. Reset the value in PostInitProperties, thereby not using the inherited value on the initial preload.
+ *  3. Ensure the CurrentVersion field isn't marked up as a UPROPERTY to avoid automatic copying between parent/child.
+ *  4. Implement Serialize to ensure the CurrentVersion is still serialized.
+ */
+USTRUCT()
+struct FNavMovementComponentVersion
+{
+	GENERATED_BODY()
+
+	/** The version the owning NavMovementComponent is currently set to */
+	ENavMovementComponentVersion CurrentVersion;
+
+	/** By always returning false here, we can disable delta serialization */
+	bool Identical(const FNavMovementComponentVersion* Other, uint32 PortFlags) const { return false; }
+
+	/** We must use a custom serializer to make sure CurrentVersion serializes properly (without markup to avoid unwanted copying) */
+	bool Serialize(FStructuredArchive::FSlot Slot);
+};
+
+template<>
+struct TStructOpsTypeTraits<FNavMovementComponentVersion> : public TStructOpsTypeTraitsBase2<FNavMovementComponentVersion>
+{
+	enum
+	{
+		WithIdentical = true,
+		WithStructuredSerializer = true,
+	};
+};
+
 /**
  * NavMovementComponent defines base functionality for MovementComponents that move any 'agent' that may be involved in AI pathfinding.
  */
@@ -70,6 +117,12 @@ private:
 	UPROPERTY()
 	TObjectPtr<UObject> PathFollowingComp;
 
+#if WITH_EDITORONLY_DATA
+	/** The saved version of this package (the value is not inherited from its parents). @see SetVersion and GetVersion. */
+	UPROPERTY()
+	FNavMovementComponentVersion DataVersion;
+#endif // WITH_EDITORONLY_DATA
+	
 public:
 	/** Get the owner of the object consuming nav movement */
 	virtual UObject* GetOwnerAsObject() const override { return GetOwner(); }
@@ -115,8 +168,6 @@ public:
 	virtual FVector GetFeetLocation() const override { return GetActorFeetLocation(); }
 	/** Returns based location of controlled actor */
 	virtual FBasedPosition GetFeetLocationBased() const override { return GetActorFeetLocationBased(); };
-	
-	ENGINE_API virtual void Serialize(FArchive& Ar) override;
 
 	/** Set nav agent properties from an object */
 	ENGINE_API virtual void UpdateNavAgent(const UObject& ObjectToUpdateFrom) override;
@@ -179,6 +230,21 @@ public:
 	/** Returns true if currently flying (moving through a non-fluid volume without resting on the ground) */
 	UFUNCTION(BlueprintCallable, Category="AI|Components|NavMovement")
 	virtual bool IsFlying() const override;
+	
+	/** Needed to properly disable inheriting the version value from its parent. */
+	ENGINE_API virtual void PostInitProperties() override;
+
+	/** PostLoad gets called once after the asset has been loaded.  It does not get called again on Blueprint recompile */
+	ENGINE_API virtual void PostLoad() override;
+	
+#if WITH_EDITOR
+protected:
+	/** Let's keep track of the version so that we can upgrade the Components properly.  This will be set properly in PostLoad after upgrades. */
+	ENavMovementComponentVersion GetVersion() const;
+
+	/** Sets the Version of the class to denote it's been upgraded */
+	void SetVersion(ENavMovementComponentVersion Version);
+#endif // WITH_EDITOR
 };
 
 
diff --git a/Engine/Source/Runtime/Engine/Private/Components/NavMovementComponent.cpp b/Engine/Source/Runtime/Engine/Private/Components/NavMovementComponent.cpp
--- a/Engine/Source/Runtime/Engine/Private/Components/NavMovementComponent.cpp
+++ b/Engine/Source/Runtime/Engine/Private/Components/NavMovementComponent.cpp
@@ -7,6 +7,29 @@
 
 #include UE_INLINE_GENERATED_CPP_BY_NAME(NavMovementComponent)
 
+namespace UE::NavMovementComponent
+{
+	TAutoConsoleVariable<int32> CVarNavMovementComponentMaxVersion(TEXT("Movement.NavMovement.MaxVersion"), (int32)ENavMovementComponentVersion::Current, TEXT("Override the Nav Movement Component Current Version (disabling upgrade code paths)"), ECVF_Default);
+
+#if WITH_EDITOR
+	namespace EditorOnly
+	{
+		// Helper function that allows us to guard against upgrading versions (in case a bug is discovered while the product is live)
+		static bool ShouldUpgradeVersion(ENavMovementComponentVersion FromVersion, ENavMovementComponentVersion ToVersion)
+		{
+			ENavMovementComponentVersion MaxVersion = static_cast<ENavMovementComponentVersion>(CVarNavMovementComponentMaxVersion.GetValueOnAnyThread() & 0xff);
+			if (MaxVersion < ToVersion)
+			{
+				return false;
+			}
+
+			return FromVersion < ToVersion;
+		}
+	}
+#endif // WITH_EDITOR
+}
+
+
 
 //----------------------------------------------------------------------//
 // UMovementComponent
@@ -29,26 +52,6 @@ FBasedPosition UNavMovementComponent::GetActorFeetLocationBased() const
 	return FBasedPosition(NULL, GetActorFeetLocation());
 }
 
-void UNavMovementComponent::Serialize(FArchive& Ar)
-{
-	Super::Serialize(Ar);
-
-#if WITH_EDITOR
-	if (Ar.IsLoading() && GetLinkerCustomVersion(FFortniteReleaseBranchCustomObjectVersion::GUID) < FFortniteReleaseBranchCustomObjectVersion::NavMovementComponentMovingPropertiesToStruct)
-	{
-		PRAGMA_DISABLE_DEPRECATION_WARNINGS
-
-		NavMovementProperties.FixedPathBrakingDistance = FixedPathBrakingDistance_DEPRECATED;
-		NavMovementProperties.bUpdateNavAgentWithOwnersCollision = bUpdateNavAgentWithOwnersCollision_DEPRECATED;
-		NavMovementProperties.bUseAccelerationForPaths = bUseAccelerationForPaths_DEPRECATED;
-		NavMovementProperties.bUseFixedBrakingDistanceForPaths = bUseFixedBrakingDistanceForPaths_DEPRECATED;
-		NavMovementProperties.bStopMovementAbortPaths = bStopMovementAbortPaths_DEPRECATED;
-
-		PRAGMA_ENABLE_DEPRECATION_WARNINGS
-	}
-#endif // WITH_EDITORONLY_DATA	
-}
-
 void UNavMovementComponent::UpdateNavAgent(const UObject& ObjectToUpdateFrom)
 {
 	if (!ShouldUpdateNavAgentWithOwnersCollision())
@@ -95,6 +98,77 @@ void UNavMovementComponent::ClearFixedBrakingDistance()
 	NavMovementProperties.bUseFixedBrakingDistanceForPaths = false;
 }
 
+void UNavMovementComponent::PostInitProperties()
+{
+	Super::PostInitProperties();
+	
+#if WITH_EDITORONLY_DATA
+	if (const UPackage* Package = GetPackage())
+	{
+		const bool bAlreadyExists = (Package->GetLinker() != nullptr);
+		if (!bAlreadyExists)
+		{
+			constexpr bool bForceGameThreadValue = true;
+			SetVersion(static_cast<ENavMovementComponentVersion>(UE::NavMovementComponent::CVarNavMovementComponentMaxVersion.GetValueOnAnyThread(bForceGameThreadValue)));
+		}
+		else
+		{
+			SetVersion(ENavMovementComponentVersion::IndependentNavMovementProperties);
+		}
+	}
+#endif // WITH_EDITORONLY_DATA
+}
+
+void UNavMovementComponent::PostLoad()
+{
+	Super::PostLoad();
+
+#if WITH_EDITOR
+	if (UE::NavMovementComponent::EditorOnly::ShouldUpgradeVersion(GetVersion(), ENavMovementComponentVersion::NavMovementPropertiesInStruct))
+	{
+		const bool bOldIndependentNavMovePropsFixUpApplied = GetLinkerCustomVersion(FFortniteReleaseBranchCustomObjectVersion::GUID) >= FFortniteReleaseBranchCustomObjectVersion::NavMovementComponentMovingPropertiesToStruct;
+
+		// Only set new nav movement properties to the deprecated ones if we haven't applied old fix up
+		// (Old fix up should have copied over values already and doing it again could result in overwriting user set values with the deprecated values)
+		if (!bOldIndependentNavMovePropsFixUpApplied)
+		{
+			PRAGMA_DISABLE_DEPRECATION_WARNINGS
+
+			NavMovementProperties.FixedPathBrakingDistance = FixedPathBrakingDistance_DEPRECATED;
+			NavMovementProperties.bUpdateNavAgentWithOwnersCollision = bUpdateNavAgentWithOwnersCollision_DEPRECATED;
+			NavMovementProperties.bUseAccelerationForPaths = bUseAccelerationForPaths_DEPRECATED;
+			NavMovementProperties.bUseFixedBrakingDistanceForPaths = bUseFixedBrakingDistanceForPaths_DEPRECATED;
+			NavMovementProperties.bStopMovementAbortPaths = bStopMovementAbortPaths_DEPRECATED;
+
+			PRAGMA_ENABLE_DEPRECATION_WARNINGS
+		}
+	}
+	
+	// We're done loading (and therefore upgrading), boost the version.
+	SetVersion( static_cast<ENavMovementComponentVersion>(UE::NavMovementComponent::CVarNavMovementComponentMaxVersion.GetValueOnGameThread()) );
+#endif // WITH_EDITOR
+}
+
+#if WITH_EDITOR
+
+ENavMovementComponentVersion UNavMovementComponent::GetVersion() const
+{
+	return DataVersion.CurrentVersion;
+}
+
+void UNavMovementComponent::SetVersion(ENavMovementComponentVersion Version)
+{
+	DataVersion.CurrentVersion = Version;
+}
+
+#endif // WITH_EDITOR
+
+bool FNavMovementComponentVersion::Serialize(FStructuredArchive::FSlot Slot)
+{
+	Slot << CurrentVersion;
+	return true;
+}
+
 void UNavMovementComponent::GetSimpleCollisionCylinder(float& CollisionRadius, float& CollisionHalfHeight) const
 {
 	GetOwner()->GetSimpleCollisionCylinder(CollisionRadius, CollisionHalfHeight);
