- Move values from public header to C++

#rnx

diff --git a/Engine/Source/Runtime/Engine/Classes/Particles/ParticleSystemManager.h b/Engine/Source/Runtime/Engine/Classes/Particles/ParticleSystemManager.h
--- a/Engine/Source/Runtime/Engine/Classes/Particles/ParticleSystemManager.h
+++ b/Engine/Source/Runtime/Engine/Classes/Particles/ParticleSystemManager.h
@@ -144,10 +144,6 @@ public:
 
 	void Dump();
 
-	static constexpr ETickingGroup FirstTickGroup = TG_PrePhysics;
-	static constexpr ETickingGroup LastTickGroup = TG_LastDemotable;
-	static constexpr int NumTickGroups = LastTickGroup - FirstTickGroup + 1;
-
 	void SetTickGroupPriority(const int32 ParticleSystemTickGroupIndex, const bool bHighPriority);
 	static void SetFirstHighPriTickGroup(int FirstHighPriTickGroupIndex);
 
diff --git a/Engine/Source/Runtime/Engine/Private/Particles/ParticleSystemManager.cpp b/Engine/Source/Runtime/Engine/Private/Particles/ParticleSystemManager.cpp
--- a/Engine/Source/Runtime/Engine/Private/Particles/ParticleSystemManager.cpp
+++ b/Engine/Source/Runtime/Engine/Private/Particles/ParticleSystemManager.cpp
@@ -38,6 +38,10 @@ FAutoConsoleVariableRef CVarParticleManagerAsyncBatchSize(
 
 namespace ParticleSystemWorldManagerInternal
 {
+	static constexpr ETickingGroup FirstTickGroup = TG_PrePhysics;
+	static constexpr ETickingGroup LastTickGroup = TG_LastDemotable;
+	static constexpr int NumTickGroups = LastTickGroup - FirstTickGroup + 1;
+
 	static int32 GFirstHighPriTickGroup = 0;
 	static FAutoConsoleVariableRef CVarFirstHighPriTickGroup(
 		TEXT("fx.ParticleManagerFirstHighPriTickGroup"),
@@ -871,7 +875,7 @@ void FParticleSystemWorldManager::Dump()
 
 void FParticleSystemWorldManager::SetTickGroupPriority(const int32 ParticleSystemTickGroupIndex, const bool bHighPriority)
 {
-	check((ParticleSystemTickGroupIndex >= 0) && (ParticleSystemTickGroupIndex < NumTickGroups));
+	check((ParticleSystemTickGroupIndex >= 0) && (ParticleSystemTickGroupIndex < ParticleSystemWorldManagerInternal::NumTickGroups));
 	TickFunctions[ParticleSystemTickGroupIndex].SetPriorityIncludingPrerequisites(bHighPriority);
 }
 
@@ -882,7 +886,7 @@ void FParticleSystemWorldManager::SetFirstHighPriTickGroup(int32 FirstHighPriTic
 		FParticleSystemWorldManager* WorldManager = Pair.Value;
 		if (WorldManager)
 		{
-			for (int32 i = 0; i < NumTickGroups; i++)
+			for (int32 i = 0; i < ParticleSystemWorldManagerInternal::NumTickGroups; i++)
 			{
 				bool HighPriority = i >= FirstHighPriTickGroupIndex;
 				WorldManager->SetTickGroupPriority(i, HighPriority);
