Remove ensure around SoundWaves compiling correctly

diff --git a/Engine/Source/Runtime/Engine/Private/AudioDerivedData.cpp b/Engine/Source/Runtime/Engine/Private/AudioDerivedData.cpp
--- a/Engine/Source/Runtime/Engine/Private/AudioDerivedData.cpp
+++ b/Engine/Source/Runtime/Engine/Private/AudioDerivedData.cpp
@@ -489,10 +489,9 @@ class FStreamedAudioCacheDerivedDataWorker : public FNonAbandonableTask
 			{
 				FMemoryReader Reader(ByteStream);
 				ensure(OutData.Serialize(Reader));
-				ensure(OutData.EncodedData.Num() > 0);
 			}
 			CompressedAudioHandle = 0;
-			return bResult;
+			return bResult && OutData.EncodedData.Num() > 0;
 		}
 
 		return false;
