Take 2: Restore default value but without double initializing which can leak memory
#rnx

diff --git a/Engine/Source/Runtime/Engine/Private/BlueprintMapLibrary.cpp b/Engine/Source/Runtime/Engine/Private/BlueprintMapLibrary.cpp
--- a/Engine/Source/Runtime/Engine/Private/BlueprintMapLibrary.cpp
+++ b/Engine/Source/Runtime/Engine/Private/BlueprintMapLibrary.cpp
@@ -1,6 +1,7 @@
 // Copyright Epic Games, Inc. All Rights Reserved.
 
 #include "Kismet/BlueprintMapLibrary.h"
+#include "Kismet/BlueprintPropertyHelpers.h"
 #include "Kismet/KismetArrayLibrary.h"
 
 #include UE_INLINE_GENERATED_CPP_BY_NAME(BlueprintMapLibrary)
@@ -45,7 +46,7 @@ bool UBlueprintMapLibrary::GenericMap_Find(const void* TargetMap, const FMapProp
 			}
 			else
 			{
-				MapProperty->ValueProp->InitializeValue(OutValuePtr);
+				UE::BlueprintPropertyHelpers::ResetToDefault(OutValuePtr, MapProperty->ValueProp);
 			}
 		}
 		return FoundValuePtr != nullptr;
diff --git a/Engine/Source/Runtime/Engine/Private/Kismet/BlueprintPropertyHelpers.cpp b/Engine/Source/Runtime/Engine/Private/Kismet/BlueprintPropertyHelpers.cpp
new file mode 100644
--- /dev/null
+++ b/Engine/Source/Runtime/Engine/Private/Kismet/BlueprintPropertyHelpers.cpp
@@ -0,0 +1,24 @@
+// Copyright Epic Games, Inc. All Rights Reserved.
+
+#include "Kismet/BlueprintPropertyHelpers.h"
+#include "UObject/UnrealType.h"
+
+void UE::BlueprintPropertyHelpers::ResetToDefault(void* Instance, const FProperty* Type)
+{
+	// the goal of this function is match the semantic of InitializeValue. It's very similar
+	// to ClearValue, except that it honors default values from UserDefinedStruct.
+ 	const int32 PropertySize = Type->GetElementSize() * Type->ArrayDim;
+	if(Type->HasAnyPropertyFlags(CPF_ZeroConstructor))
+	{
+		// common case.. default value is zero, just zero the memory:
+		FMemory::Memzero(Instance, static_cast<size_t>(PropertySize));
+	}
+	else
+	{
+		// copy from default value:
+		void* StorageSpace = FMemory_Alloca(PropertySize);
+		Type->InitializeValue(StorageSpace);
+		Type->CopySingleValueToScriptVM(Instance, StorageSpace);
+		Type->DestroyValue(StorageSpace);
+	}
+}
diff --git a/Engine/Source/Runtime/Engine/Private/Kismet/BlueprintPropertyHelpers.h b/Engine/Source/Runtime/Engine/Private/Kismet/BlueprintPropertyHelpers.h
new file mode 100644
--- /dev/null
+++ b/Engine/Source/Runtime/Engine/Private/Kismet/BlueprintPropertyHelpers.h
@@ -0,0 +1,10 @@
+// Copyright Epic Games, Inc. All Rights Reserved.
+
+#pragma once
+
+class FProperty;
+
+namespace UE::BlueprintPropertyHelpers
+{
+	void ResetToDefault(void* Instance, const FProperty* Type);
+}
diff --git a/Engine/Source/Runtime/Engine/Private/KismetArrayLibrary.cpp b/Engine/Source/Runtime/Engine/Private/KismetArrayLibrary.cpp
--- a/Engine/Source/Runtime/Engine/Private/KismetArrayLibrary.cpp
+++ b/Engine/Source/Runtime/Engine/Private/KismetArrayLibrary.cpp
@@ -1,6 +1,7 @@
 // Copyright Epic Games, Inc. All Rights Reserved.
 
 #include "Kismet/KismetArrayLibrary.h"
+#include "Kismet/BlueprintPropertyHelpers.h"
 #include "GameFramework/Actor.h"
 #include "Blueprint/BlueprintSupport.h"
 #include "Math/RandomStream.h"
@@ -474,7 +475,7 @@ void UKismetArrayLibrary::GenericArray_Get(void* TargetArray, const FArrayProper
 				*ArrayProp->GetOwnerVariant().GetPathName()),
 				ELogVerbosity::Warning,
 				GetOutOfBoundsWarning);
-			InnerProp->InitializeValue(Item);
+			UE::BlueprintPropertyHelpers::ResetToDefault(Item, InnerProp);
 		}
 	}
 }
@@ -623,7 +624,7 @@ void UKismetArrayLibrary::GenericArray_Random(void* TargetArray, const FArrayPro
 			FFrame::KismetExecutionMessage(*FString::Printf(TEXT("Attempted to access random index from empty array!")),
 				ELogVerbosity::Warning,
 				RandomAccessToEmptyArrayWarning);
-			InnerProp->InitializeValue(OutItem);
+			UE::BlueprintPropertyHelpers::ResetToDefault(OutItem, InnerProp);
 		}
 	}
 }
@@ -648,7 +649,7 @@ void UKismetArrayLibrary::GenericArray_RandomFromStream(void* TargetArray, const
 			FFrame::KismetExecutionMessage(*FString::Printf(TEXT("Attempted to access random index from empty array!")),
 				ELogVerbosity::Warning,
 				RandomAccessToEmptyArrayWarning);
-			InnerProp->InitializeValue(OutItem);
+			UE::BlueprintPropertyHelpers::ResetToDefault(OutItem, InnerProp);
 		}
 	}
 }
