guard against trying to create soundwave proxy data when running with `-nosound` or other contexts where the audio data is never loaded.

#swarm
#rnx

diff --git a/Engine/Plugins/Runtime/Metasound/Source/MetasoundFrontend/Private/MetasoundFrontendDataTypeRegistry.cpp b/Engine/Plugins/Runtime/Metasound/Source/MetasoundFrontend/Private/MetasoundFrontendDataTypeRegistry.cpp
--- a/Engine/Plugins/Runtime/Metasound/Source/MetasoundFrontend/Private/MetasoundFrontendDataTypeRegistry.cpp
+++ b/Engine/Plugins/Runtime/Metasound/Source/MetasoundFrontend/Private/MetasoundFrontendDataTypeRegistry.cpp
@@ -4,6 +4,7 @@
 #include "MetasoundLiteral.h"
 #include "MetasoundTrace.h"
 #include "Misc/ScopeLock.h"
+#include "Misc/App.h"
 
 namespace Metasound
 {
@@ -998,7 +999,7 @@ PRAGMA_ENABLE_DEPRECATION_WARNINGS
 					}
 					else
 					{
-						if (InObject)
+						if (InObject && FApp::CanEverRenderAudio())
 						{
 							UE_LOG(LogMetaSound, Error, TEXT("Failed to create a valid proxy from UObject '%s'."), *InObject->GetName());
 						}
@@ -1101,7 +1102,7 @@ PRAGMA_ENABLE_DEPRECATION_WARNINGS
 				{
 					TSharedPtr<Audio::IProxyData> ProxyPtr = CreateProxyFromUObject(ElementDataType, InObject);
 					ProxyPtr = ElementEntry->CreateProxy(InObject);
-					if (!ProxyPtr && InObject)
+					if (!ProxyPtr && InObject && FApp::CanEverRenderAudio())
 					{
 						UE_LOG(LogMetaSound, Error, TEXT("Failed to create a valid proxy from UObject '%s'."), *InObject->GetName());
 					}
diff --git a/Engine/Source/Runtime/Engine/Private/SoundWave.cpp b/Engine/Source/Runtime/Engine/Private/SoundWave.cpp
--- a/Engine/Source/Runtime/Engine/Private/SoundWave.cpp
+++ b/Engine/Source/Runtime/Engine/Private/SoundWave.cpp
@@ -2163,7 +2163,7 @@ void USoundWave::PostLoad()
 #endif // #if WITH_EDITORONLY_DATA
 
 		// update shared state
-		if (!HasAnyFlags(RF_ClassDefaultObject) && FApp::CanEverRenderAudio())
+		if (!HasAnyFlags(RF_ClassDefaultObject))
 		{
 			SoundWaveDataPtr->InitializeDataFromSoundWave(*this);
 		}
@@ -4083,6 +4083,11 @@ int32 USoundWave::GetCompressionQuality() const
 
 FSoundWaveProxyPtr USoundWave::CreateSoundWaveProxy()
 {
+	if (!FApp::CanEverRenderAudio())
+	{
+		return Proxy;
+	}
+	
 	LLM_SCOPE(ELLMTag::AudioSoundWaveProxies);
 
 #if WITH_EDITORONLY_DATA
@@ -4099,7 +4104,7 @@ TSharedPtr<Audio::IProxyData> USoundWave::CreateProxyData(const Audio::FProxyDat
 	check(SoundWaveDataPtr);
 	LLM_SCOPE(ELLMTag::AudioSoundWaveProxies);
 
-	if (!Proxy.IsValid())
+	if (!Proxy.IsValid() && FApp::CanEverRenderAudio())
 	{
 		Proxy = CreateSoundWaveProxy();
 	}
@@ -4672,12 +4677,12 @@ void USoundWave::OverrideLoadingBehavior(ESoundWaveLoadingBehavior InLoadingBeha
 				// In editor, just make sure that data is available in the local DDC cache for quick access.
 				CachePlatformData(true /* bAsyncCache */);
 
-				if (!Proxy.IsValid())
+				if (!Proxy.IsValid() && FApp::CanEverRenderAudio())
 				{
 					Proxy = CreateSoundWaveProxy();
 				}
 #else
-				if (GetNumChunks() > 1)
+				if (GetNumChunks() > 1 && FApp::CanEverRenderAudio())
 				{
 					if (!Proxy.IsValid())
 					{
