CopyPropertiesForUnrelatedObjects now will generate a new ActorGuid when duplicating objects when FCopyPropertiesForUnrelatedObjectsParams::bPerformDuplication is set.
This fixes UMovieSceneSpawnableVCamBinding to crash when recording a VCam into a world-partioned world (where you'll get a check when an AActor::ActorGuid is unset).

This change sets the PPF_Duplicate in the FCPFUOReader if FCopyPropertiesForUnrelatedObjectsParams::bPerformDuplication is set. This causes AActor::Serialize to generate a new ActorGuid.
FCPFUOWriter already sets the PPF_Duplicate flag when collecting the data that is duplicated. It seems only logical that FCPFUOReader should also set it.
The use case of this is when you want to call StaticDuplicateObject but but cannot, e.g.  you are duplicating to a parent class; in that case, you'd call NewObject and copy compatible properties using CopyPropertiesForUnrelatedObjects.

diff --git a/Engine/Source/Runtime/Engine/Private/UnrealEngine.cpp b/Engine/Source/Runtime/Engine/Private/UnrealEngine.cpp
--- a/Engine/Source/Runtime/Engine/Private/UnrealEngine.cpp
+++ b/Engine/Source/Runtime/Engine/Private/UnrealEngine.cpp
@@ -17758,7 +17758,10 @@ public:
 struct FCPFUOReader : public FObjectReaderWithReplacement, public FCPFUOArchive
 {
 public:
-	FCPFUOReader(FCPFUOWriter& DataSrc, UObject* DstObject, const TMap<UObject*, UObject*>* InReferenceReplacementMap = nullptr, const TMap<UClass*, UClass*>* InClassReferenceReplacementMap = nullptr)
+	FCPFUOReader(
+		FCPFUOWriter& DataSrc, UObject* DstObject, const UEngine::FCopyPropertiesForUnrelatedObjectsParams& Params,
+		const TMap<UObject*, UObject*>* InReferenceReplacementMap = nullptr, const TMap<UClass*, UClass*>* InClassReferenceReplacementMap = nullptr
+		)
 		: FObjectReaderWithReplacement(DataSrc.SavedPropertyData, InReferenceReplacementMap, InClassReferenceReplacementMap)
 		, FCPFUOArchive(DataSrc)
 	{
@@ -17771,6 +17774,11 @@ public:
 			SetLocalizationNamespace(TextNamespaceUtil::EnsurePackageNamespace(DstObject));
 		}
 #endif // USE_STABLE_LOCALIZATION_KEYS
+		
+		if (Params.bPerformDuplication)
+		{
+			ArPortFlags |= PPF_Duplicate;
+		}
 
 		UE::FScopedObjectSerializeContext ObjectSerializeContext(DstObject, *this);
 		DstObject->Serialize(*this);
@@ -18242,7 +18250,7 @@ void UEngine::CopyPropertiesForUnrelatedObjects(UObject* OldObject, UObject* New
 		// Serialize in the modified properties from the old CDO to the new CDO
 		if (Writer.SavedPropertyData.Num() > 0)
 		{
-			FCPFUOReader Reader(Writer, NewObject, OptionalReferenceReplacementMap, OptionalClassReferenceReplacementMap);
+			FCPFUOReader Reader(Writer, NewObject, Params, OptionalReferenceReplacementMap, OptionalClassReferenceReplacementMap);
 		}
 
 		TArray<UObject*> EditInlineSubobjectsOfComponents;
