Introduced the Enumeration material parameters.

This change introduces a new material expression called "UMaterialExpressionEnumerationParameter" that works similarly to HLSL enumerations, allowing the user to select one among a well-defined list of entries. Each named entry maps to an integer value.

Enumeration parameters are "uniform" values, that is they leverage the preshader buffer to store the value. That means that changing the entry of an enum parameter in a Material Instance won't cause a new shader permutation.

The list of entries in the enumerations can be provided two ways:
1. Using a new UMaterialEnumeration data asset, which the user can edit in Unreal to specify a list of entries and their integer values.
2. Implement the IMaterialEnumerationProvider interface in some other UObject, so that any UObject can be used as the provider of an enumeration in the material editor.

#rn Added a new Material Enumeration Parameter to enable selecting an entry from a predefined list by name, mapping the symbol to an integer value.

diff --git a/Engine/Config/BaseMaterialExpressions.ini b/Engine/Config/BaseMaterialExpressions.ini
--- a/Engine/Config/BaseMaterialExpressions.ini
+++ b/Engine/Config/BaseMaterialExpressions.ini
@@ -1051,3 +1051,5 @@ bShaderInputData=true
 +ExternalCodeDeclarations=(ReturnType=Float3,	Name="GetCloudEmptySpaceSkippingSphereCenterWorldPosition",	Definition="MaterialExpressionCloudEmptySpaceSkippingSphereCenterWorldPosition(Parameters)")
 +ExternalCodeDeclarations=(ReturnType=Float1,	Name="GetCloudEmptySpaceSkippingSphereRadius",				Definition="MaterialExpressionCloudEmptySpaceSkippingSphereRadius(Parameters)")
 
+[/Script/Engine.MaterialExpressionEnumerationParameter]
++MenuCategories=NSLOCTEXT("MaterialExpression", "Parameters", "Parameters")
diff --git a/Engine/Source/Editor/UnrealEd/Classes/MaterialEditor/DEditorEnumerationParameterValue.h b/Engine/Source/Editor/UnrealEd/Classes/MaterialEditor/DEditorEnumerationParameterValue.h
new file mode 100644
--- /dev/null
+++ b/Engine/Source/Editor/UnrealEd/Classes/MaterialEditor/DEditorEnumerationParameterValue.h
@@ -0,0 +1,50 @@
+// Copyright Epic Games, Inc. All Rights Reserved.
+
+#pragma once
+
+#include "CoreMinimal.h"
+#include "UObject/ObjectMacros.h"
+#include "MaterialEditor/DEditorParameterValue.h"
+#include "Materials/MaterialEnumeration.h"
+
+#include "DEditorEnumerationParameterValue.generated.h"
+
+UCLASS(hidecategories = Object, collapsecategories, editinlinenew, MinimalAPI)
+class UDEditorEnumerationParameterValue : public UDEditorParameterValue
+{
+	GENERATED_UCLASS_BODY()
+
+public:
+	UPROPERTY(Transient)
+	TSoftObjectPtr<UObject> Enumeration = nullptr;
+
+	UPROPERTY(EditAnywhere, Category = DEnumerationParameter, Meta = (GetOptions = "GetEntryOptions"))
+	FName ParameterValue{};
+
+	UFUNCTION()
+	TArray<FString> GetEntryOptions() const;
+
+	virtual FName GetDefaultGroupName() const override
+	{
+		return TEXT("Enumeration Parameter Values");
+	}
+
+	virtual bool GetValue(FMaterialParameterMetadata& OutResult) const override
+	{
+		UDEditorParameterValue::GetValue(OutResult);
+		OutResult.Enumeration = Enumeration.Get();
+		OutResult.Value = FMaterialParameterValue(ParameterValue);
+		return true;
+	}
+
+	virtual bool SetValue(const FMaterialParameterValue& Value) override
+	{
+		if (Value.Type == EMaterialParameterType::Enumeration)
+		{
+			ParameterValue = Value.EnumEntry;
+			return true;
+		}
+		return false;
+	}
+};
+
diff --git a/Engine/Source/Editor/UnrealEd/Private/EditorParameters.cpp b/Engine/Source/Editor/UnrealEd/Private/EditorParameters.cpp
--- a/Engine/Source/Editor/UnrealEd/Private/EditorParameters.cpp
+++ b/Engine/Source/Editor/UnrealEd/Private/EditorParameters.cpp
@@ -13,6 +13,8 @@
 #include "MaterialEditor/DEditorTextureParameterValue.h"
 #include "MaterialEditor/DEditorTextureCollectionParameterValue.h"
 #include "MaterialEditor/DEditorVectorParameterValue.h"
+#include "MaterialEditor/DEditorEnumerationParameterValue.h"
+#include "Materials/MaterialEnumeration.h"
 #include "MaterialTypes.h"
 #include "Math/Color.h"
 #include "Math/Vector4.h"
@@ -83,6 +85,28 @@ UDEditorSparseVolumeTextureParameterValue::UDEditorSparseVolumeTextureParameterV
 {
 }
 
+UDEditorEnumerationParameterValue::UDEditorEnumerationParameterValue(const FObjectInitializer& ObjectInitializer)
+	: Super(ObjectInitializer)
+{
+}
+
+TArray<FString> UDEditorEnumerationParameterValue::GetEntryOptions() const
+{
+	TArray<FString> Options;
+	if (Enumeration)
+	{
+		Cast<IMaterialEnumerationProvider>(Enumeration.Get())->ForEachEntry([&Options](FName InEntryName, int32 InEntryValue)
+		{
+			Options.Add(InEntryName.ToString());
+		});
+	}
+	if (Options.IsEmpty())
+	{
+		Options.Add(TEXT("None"));
+	}
+	return Options;
+}
+
 static UDEditorParameterValue* CreateParameter_Scalar(UObject* Owner, const FMaterialParameterMetadata& Meta)
 {
 	UDEditorScalarParameterValue* Parameter = NewObject<UDEditorScalarParameterValue>(Owner);
@@ -197,6 +221,17 @@ static UDEditorParameterValue* CreateParameter_StaticComponentMask(UObject* Owne
 	return Parameter;
 }
 
+static UDEditorParameterValue* CreateParameter_Enumeration(UObject* Owner, const FMaterialParameterMetadata& Meta)
+{
+	UDEditorEnumerationParameterValue* Parameter = NewObject<UDEditorEnumerationParameterValue>(Owner);
+	if (Meta.Value.Type == EMaterialParameterType::Enumeration)
+	{
+		Parameter->Enumeration = Meta.Enumeration;
+		Parameter->ParameterValue = Meta.Value.EnumEntry;
+	}
+	return Parameter;
+}
+
 UDEditorParameterValue* UDEditorParameterValue::Create(UObject* Owner,
 	EMaterialParameterType Type,
 	const FMaterialParameterInfo& ParameterInfo,
@@ -215,6 +250,7 @@ UDEditorParameterValue* UDEditorParameterValue::Create(UObject* Owner,
 	case EMaterialParameterType::Font: Parameter = CreateParameter_Font(Owner, Meta); break;
 	case EMaterialParameterType::StaticSwitch: Parameter = CreateParameter_StaticSwitch(Owner, Meta); break;
 	case EMaterialParameterType::StaticComponentMask: Parameter = CreateParameter_StaticComponentMask(Owner, Meta); break;
+	case EMaterialParameterType::Enumeration: Parameter = CreateParameter_Enumeration(Owner, Meta); break;
 	default: break;
 	}
 
diff --git a/Engine/Source/Runtime/Engine/Private/Materials/HLSLMaterialTranslator.cpp b/Engine/Source/Runtime/Engine/Private/Materials/HLSLMaterialTranslator.cpp
--- a/Engine/Source/Runtime/Engine/Private/Materials/HLSLMaterialTranslator.cpp
+++ b/Engine/Source/Runtime/Engine/Private/Materials/HLSLMaterialTranslator.cpp
@@ -538,6 +538,7 @@ static inline EMaterialValueType GetMaterialValueType(EMaterialParameterType Typ
 	case EMaterialParameterType::Vector: return MCT_Float4;
 	case EMaterialParameterType::DoubleVector: return MCT_LWCVector4;
 	case EMaterialParameterType::StaticSwitch: return MCT_Bool;
+	case EMaterialParameterType::Enumeration: return MCT_UInt1;
 	default: checkNoEntry(); return MCT_Unknown;
 	}
 }
@@ -4316,7 +4317,7 @@ int32 FHLSLMaterialTranslator::AccessUniformExpression(int32 Index)
 			RegisterOffset,
 			NumBoolComponents);
 	}
-	else if (IsFloatNumericType(CodeChunk.Type))
+	else if (IsFloatNumericType(CodeChunk.Type) || IsUIntNumericType(CodeChunk.Type))
 	{
 		check(CodeChunk.DerivativeStatus != EDerivativeStatus::Valid);
 		const uint32 NumComponents = GetNumComponents(CodeChunk.Type);
@@ -4388,7 +4389,17 @@ int32 FHLSLMaterialTranslator::AccessUniformExpression(int32 Index)
 		}
 		else
 		{
+			if (IsUIntNumericType(CodeChunk.Type))
+			{
+				FormattedCode.Append(TEXT("asuint("));
+			}
+
 			WriteMaterialUniformAccess(UE::Shader::EValueComponentType::Float, NumComponents, UniformOffset, FormattedCode);
+
+			if (IsUIntNumericType(CodeChunk.Type))
+			{
+				FormattedCode.Append(TEXT(")"));
+			}
 		}
 	}
 	else if(CodeChunk.Type & MCT_Texture)
diff --git a/Engine/Source/Runtime/Engine/Private/Materials/MaterialCachedData.cpp b/Engine/Source/Runtime/Engine/Private/Materials/MaterialCachedData.cpp
--- a/Engine/Source/Runtime/Engine/Private/Materials/MaterialCachedData.cpp
+++ b/Engine/Source/Runtime/Engine/Private/Materials/MaterialCachedData.cpp
@@ -48,6 +48,7 @@
 #include "MaterialShared.h"
 #include "MaterialCache/MaterialCacheMaterial.h"
 #include "Materials/MaterialExpressionMaterialCache.h"
+#include "Materials/MaterialEnumeration.h"
 
 #include UE_INLINE_GENERATED_CPP_BY_NAME(MaterialCachedData)
 
@@ -267,6 +268,11 @@ bool FMaterialCachedExpressionData::AddParameter(const FMaterialParameterInfo& P
 			EditorOnlyData->StaticComponentMaskValues.Insert(ParameterMeta.Value.AsStaticComponentMask(), Index);
 			break;
 
+		case EMaterialParameterType::Enumeration:
+			Enumerations.Insert(ParameterMeta.Enumeration, Index);
+			EnumerationEntryValues.Insert(ParameterMeta.Value.EnumEntry, Index);
+			break;
+
 		default:
 			checkNoEntry();
 			break;
@@ -390,6 +396,17 @@ bool FMaterialCachedExpressionData::AddParameter(const FMaterialParameterInfo& P
 			}
 			break;
 
+		case EMaterialParameterType::Enumeration:
+			bSameValue = Enumerations[Index] == ParameterMeta.Enumeration && EnumerationEntryValues[Index] == ParameterMeta.Value.EnumEntry;
+			if (!bSameValue)
+			{
+				OutErrorMessage = ReportError(
+					FString::Printf(TEXT("%s(%s)"), *Enumerations[Index].ToString(), *EnumerationEntryValues[Index].ToString()),
+					FString::Printf(TEXT("%s(%s)"), *TSoftObjectPtr<UObject>(ParameterMeta.Enumeration).ToString(), *ParameterMeta.Value.EnumEntry.ToString()),
+					*PreviousEditorInfo, EditorInfo);
+			}
+			break;
+
 		default:
 			bSameValue = true;
 			break;
@@ -904,6 +921,10 @@ void FMaterialCachedExpressionData::GetParameterValueByIndex(EMaterialParameterT
 		}
 		break;
 #endif // WITH_EDITORONLY_DATA
+	case EMaterialParameterType::Enumeration:
+		OutResult.Enumeration = Enumerations[ParameterIndex].LoadSynchronous();
+		OutResult.Value = FMaterialParameterValue(EnumerationEntryValues[ParameterIndex]);
+		break;
 	default:
 		checkNoEntry();
 		break;
diff --git a/Engine/Source/Runtime/Engine/Private/Materials/MaterialEnumeration.cpp b/Engine/Source/Runtime/Engine/Private/Materials/MaterialEnumeration.cpp
new file mode 100644
--- /dev/null
+++ b/Engine/Source/Runtime/Engine/Private/Materials/MaterialEnumeration.cpp
@@ -0,0 +1,67 @@
+#include "Materials/MaterialEnumeration.h"
+
+int32 IMaterialEnumerationProvider::GetValueOrDefault(FName EntryName, int32 DefaultValue) const
+{
+	int32 Value;
+	ResolveValue(EntryName, Value, DefaultValue);
+	return Value;
+}
+
+bool UMaterialEnumeration::ResolveValue(FName EntryName, int32& OutValue, int32 DefaultValue) const
+{
+	if (SortedEntries.Num() != Entries.Num())
+	{
+		for (const FMaterialEnumerationEntry& Entry : Entries)
+		{
+			if (Entry.Name == EntryName)
+			{
+				OutValue = Entry.Value;
+				return true;
+			}
+		}
+	}
+	else
+	{
+		int32 Index = Algo::BinarySearchBy(SortedEntries, EntryName, &FMaterialEnumerationEntry::Name, FNameFastLess{});
+		if (Index != INDEX_NONE)
+		{
+			OutValue = SortedEntries[Index].Value;
+			return true;
+		}
+	}
+	OutValue = DefaultValue;
+	return false;
+}
+
+void UMaterialEnumeration::ForEachEntry(TFunctionRef<void (FName Name, int32 Value)> Iterator) const
+{
+	for (const FMaterialEnumerationEntry& Entry : Entries)
+	{
+		Iterator(Entry.Name, Entry.Value);
+	}
+}
+
+void UMaterialEnumeration::SortEntries()
+{
+	SortedEntries = Entries;
+	SortedEntries.Sort([](const FMaterialEnumerationEntry& A, const FMaterialEnumerationEntry& B)
+	{
+		return A.Name.FastLess(B.Name);
+	});
+}
+
+void UMaterialEnumeration::PostLoad()
+{
+	Super::PostLoad();
+	SortEntries();
+}
+
+#if WITH_EDITOR
+
+void UMaterialEnumeration::PostEditChangeProperty(FPropertyChangedEvent& PropertyChangedEvent)
+{
+	Super::PostEditChangeProperty(PropertyChangedEvent);
+	SortEntries();
+}
+
+#endif
diff --git a/Engine/Source/Runtime/Engine/Private/Materials/MaterialExpressions.cpp b/Engine/Source/Runtime/Engine/Private/Materials/MaterialExpressions.cpp
--- a/Engine/Source/Runtime/Engine/Private/Materials/MaterialExpressions.cpp
+++ b/Engine/Source/Runtime/Engine/Private/Materials/MaterialExpressions.cpp
@@ -172,6 +172,7 @@
 #include "Materials/MaterialExpressionReroute.h"
 #include "Materials/MaterialExpressionNamedReroute.h"
 #include "Materials/MaterialExpressionScalarParameter.h"
+#include "Materials/MaterialExpressionEnumerationParameter.h"
 #include "Materials/MaterialExpressionSetMaterialAttributes.h"
 #include "Materials/MaterialExpressionShadowReplace.h"
 #include "Materials/MaterialExpressionSign.h"
@@ -397,23 +398,28 @@ void GetMaterialValueTypeDescriptions(const uint32 MaterialValueType, TArray<FTe
 		switch (MaskedFloatType)
 		{
 			case MCT_Float:
-			case MCT_Float1:
-				OutDescriptions.Add(LOCTEXT("Float", "Float"));
-				break;
-			case MCT_Float2:
-				OutDescriptions.Add(LOCTEXT("Float2", "Float 2"));
-				break;
-			case MCT_Float3:
-				OutDescriptions.Add(LOCTEXT("Float3", "Float 3"));
-				break;
-			case MCT_Float4:
-				OutDescriptions.Add(LOCTEXT("Float4", "Float 4"));
-				break;
-			default:
-				break;
+			case MCT_Float1: OutDescriptions.Add(LOCTEXT("Float", "Float")); break;
+			case MCT_Float2: OutDescriptions.Add(LOCTEXT("Float2", "Float 2")); break;
+			case MCT_Float3: OutDescriptions.Add(LOCTEXT("Float3", "Float 3")); break;
+			case MCT_Float4: OutDescriptions.Add(LOCTEXT("Float4", "Float 4")); break;
+			default: break;
 		}
 	}
 
+	// Get exact int type if possible
+	uint32 MaskedUIntType = MaterialValueType & MCT_UInt;
+	if (MaskedUIntType)
+	{
+		switch (MaskedUIntType)
+		{
+			case MCT_UInt:
+			case MCT_UInt1: OutDescriptions.Add(LOCTEXT("UInt", "UInt")); break;
+			case MCT_UInt2: OutDescriptions.Add(LOCTEXT("UInt2", "UInt 2")); break;
+			case MCT_UInt3: OutDescriptions.Add(LOCTEXT("UInt3", "UInt 3")); break;
+			case MCT_UInt4: OutDescriptions.Add(LOCTEXT("UInt4", "UInt 4")); break;
+			default: break;
+		}
+	}
 	// Get exact texture type if possible
 	uint32 MaskedTextureType = MaterialValueType & MCT_Texture;
 	if (MaskedTextureType)
@@ -7913,6 +7919,10 @@ int32 UMaterialExpressionUIntToFloat::Compile(class FMaterialCompiler* Compiler,
 	}
 
 	int32 Value = Input.Compile(Compiler);
+	if (Value == INDEX_NONE)
+	{
+		return INDEX_NONE;
+	}
 	EMaterialValueType Type = Compiler->GetParameterType(Value);
 	int NumComponents = GetNumComponents(Type);
 	if (NumComponents <= 0 || NumComponents > 4)
@@ -24071,6 +24081,105 @@ void UMaterialExpressionOperator::PostLoad()
 	UMaterialExpression::PostLoad();
 }
 
+///////////////////////////////////////////////////////////////////////////////
+// UMaterialExpressionEnumerationParameter
+///////////////////////////////////////////////////////////////////////////////
+
+int32 UMaterialExpressionEnumerationParameter::Compile(class FMaterialCompiler* Compiler, int32 OutputIndex)
+{
+	if (!Enumeration)
+	{
+		return Compiler->Errorf(TEXT("Enum Parameter '%s': Enumeration asset missing."), *ParameterName.ToString());
+	}
+	int32 ResolvedDefaultValue = Enumeration->GetValueOrDefault(DefaultEntry, 0);
+	return Compiler->NumericParameter(EMaterialParameterType::Enumeration, ParameterName, ResolvedDefaultValue);
+}
+
+void UMaterialExpressionEnumerationParameter::GetCaption(TArray<FString>& OutCaptions) const
+{
+	OutCaptions.Add(FString::Printf(TEXT("Param (%s)"), *DefaultEntry.ToString())); 
+	OutCaptions.Add(FString::Printf(TEXT("'%s'"), *ParameterName.ToString())); 
+}
+
+EMaterialValueType UMaterialExpressionEnumerationParameter::GetOutputValueType(int32 InIndex)
+{
+	return MCT_UInt1;
+}
+
+bool UMaterialExpressionEnumerationParameter::GetParameterValue(FMaterialParameterMetadata& OutMeta) const
+{
+	OutMeta.Value = FMaterialParameterValue(DefaultEntry);
+	OutMeta.Enumeration = Enumeration.GetObject();
+	return Super::GetParameterValue(OutMeta);
+}
+
+bool UMaterialExpressionEnumerationParameter::SetParameterValue(const FName& Name, const FMaterialParameterMetadata& Meta, EMaterialExpressionSetParameterValueFlags Flags)
+{
+	if (Meta.Value.Type == EMaterialParameterType::Enumeration)
+	{
+		if (SetParameterValue(Name, Meta.Enumeration.Get(), Meta.Value.EnumEntry, Flags))
+		{
+			if (EnumHasAnyFlags(Flags, EMaterialExpressionSetParameterValueFlags::AssignGroupAndSortPriority))
+			{
+				Group = Meta.Group;
+				SortPriority = Meta.SortPriority;
+			}
+			return true;
+		}
+	}
+	return false;
+}
+
+void UMaterialExpressionEnumerationParameter::PostEditChangeProperty(FPropertyChangedEvent& PropertyChangedEvent)
+{
+	const FName PropertyName = PropertyChangedEvent.GetMemberPropertyName();
+	if (PropertyName == GET_MEMBER_NAME_STRING_CHECKED(ThisClass, Enumeration)
+		|| PropertyName == GET_MEMBER_NAME_STRING_CHECKED(ThisClass, DefaultEntry))
+	{
+		int32 DefaultEntryValue;
+		if (!Enumeration || !Enumeration->ResolveValue(DefaultEntry, DefaultEntryValue))
+		{
+			DefaultEntryValue = 0;
+			DefaultEntry = NAME_None;
+		}
+		FEditorSupportDelegates::NumericParameterDefaultChanged.Broadcast(this, EMaterialParameterType::Enumeration, ParameterName, DefaultEntryValue);
+	}
+	Super::PostEditChangeProperty(PropertyChangedEvent);
+}
+
+bool UMaterialExpressionEnumerationParameter::SetParameterValue(const FName& InParameterName, UObject* InEnumeration, FName InDefaultEntry, EMaterialExpressionSetParameterValueFlags Flags)
+{
+	if (InParameterName == ParameterName)
+	{
+		Enumeration = InEnumeration;
+		DefaultEntry = InDefaultEntry;
+		if (EnumHasAnyFlags(Flags, EMaterialExpressionSetParameterValueFlags::SendPostEditChangeProperty))
+		{
+			SendPostEditChangeProperty(this, GET_MEMBER_NAME_STRING_CHECKED(ThisClass, Enumeration));
+			SendPostEditChangeProperty(this, GET_MEMBER_NAME_STRING_CHECKED(ThisClass, DefaultEntry));
+		}
+		return true;
+	}
+	return false;
+}
+
+TArray<FString> UMaterialExpressionEnumerationParameter::GetEntryOptions() const
+{
+	TArray<FString> Options;
+	if (Enumeration)
+	{
+		Enumeration->ForEachEntry([&Options](FName InEntryName, int32 InEntryValue)
+		{
+			Options.Add(InEntryName.ToString());
+		});
+	}
+	if (Options.IsEmpty())
+	{
+		Options.Add(TEXT("None"));
+	}
+	return Options;
+}
+
 #endif // WITH_EDITOR
 
 #undef LOCTEXT_NAMESPACE
diff --git a/Engine/Source/Runtime/Engine/Private/Materials/MaterialInstance.cpp b/Engine/Source/Runtime/Engine/Private/Materials/MaterialInstance.cpp
--- a/Engine/Source/Runtime/Engine/Private/Materials/MaterialInstance.cpp
+++ b/Engine/Source/Runtime/Engine/Private/Materials/MaterialInstance.cpp
@@ -29,6 +29,7 @@
 #include "Materials/MaterialExpressionMaterialAttributeLayers.h"
 #include "Materials/MaterialExpressionRuntimeVirtualTextureSampleParameter.h"
 #include "Materials/MaterialExpressionSparseVolumeTextureSample.h"
+#include "Materials/MaterialExpressionEnumerationParameter.h"
 #include "Materials/MaterialExpressionStaticComponentMaskParameter.h"
 #include "Materials/MaterialExpressionMaterialFunctionCall.h"
 #include "Materials/MaterialInstanceConstant.h"
@@ -339,15 +340,16 @@ bool FMaterialInstanceResource::GetParameterValue(EMaterialParameterType Type, c
 		// Check for instances overrides
 		switch (Type)
 		{
-		case EMaterialParameterType::StaticSwitch: bResult = RenderThread_GetParameterValue<bool>(ParameterInfo, OutValue); break;
-		case EMaterialParameterType::Scalar: bResult = RenderThread_GetParameterValue<float>(ParameterInfo, OutValue); break;
-		case EMaterialParameterType::Vector: bResult = RenderThread_GetParameterValue<FLinearColor>(ParameterInfo, OutValue); break;
-		case EMaterialParameterType::DoubleVector: bResult = RenderThread_GetParameterValue<FVector4d>(ParameterInfo, OutValue); break;
-		case EMaterialParameterType::Texture: bResult = RenderThread_GetParameterValue<const UTexture*>(ParameterInfo, OutValue); break;
-		case EMaterialParameterType::TextureCollection: bResult = RenderThread_GetParameterValue<const UTextureCollection*>(ParameterInfo, OutValue); break;
-		case EMaterialParameterType::RuntimeVirtualTexture: bResult = RenderThread_GetParameterValue<const URuntimeVirtualTexture*>(ParameterInfo, OutValue); break;
-		case EMaterialParameterType::SparseVolumeTexture: bResult = RenderThread_GetParameterValue<const USparseVolumeTexture*>(ParameterInfo, OutValue); break;
-		default: ensure(false); break; // other parameter types are not expected on the render thread
+			case EMaterialParameterType::StaticSwitch: bResult = RenderThread_GetParameterValue<bool>(ParameterInfo, OutValue); break;
+			case EMaterialParameterType::Scalar: bResult = RenderThread_GetParameterValue<float>(ParameterInfo, OutValue); break;
+			case EMaterialParameterType::Vector: bResult = RenderThread_GetParameterValue<FLinearColor>(ParameterInfo, OutValue); break;
+			case EMaterialParameterType::DoubleVector: bResult = RenderThread_GetParameterValue<FVector4d>(ParameterInfo, OutValue); break;
+			case EMaterialParameterType::Texture: bResult = RenderThread_GetParameterValue<const UTexture*>(ParameterInfo, OutValue); break;
+			case EMaterialParameterType::TextureCollection: bResult = RenderThread_GetParameterValue<const UTextureCollection*>(ParameterInfo, OutValue); break;
+			case EMaterialParameterType::RuntimeVirtualTexture: bResult = RenderThread_GetParameterValue<const URuntimeVirtualTexture*>(ParameterInfo, OutValue); break;
+			case EMaterialParameterType::SparseVolumeTexture: bResult = RenderThread_GetParameterValue<const USparseVolumeTexture*>(ParameterInfo, OutValue); break;
+			case EMaterialParameterType::Enumeration: bResult = RenderThread_GetParameterValue<FName>(ParameterInfo, OutValue); break;
+			default: ensure(false); break; // other parameter types are not expected on the render thread
 		}
 	}
 
@@ -511,6 +513,7 @@ void FMaterialInstanceResource::InitMIParameters(FMaterialInstanceParameterSet&
 	ParameterSet.TextureCollectionParameters.Sort(SortMaterialInstanceParametersPredicate<const UTextureCollection*>);
 	ParameterSet.RuntimeVirtualTextureParameters.Sort(SortMaterialInstanceParametersPredicate<const URuntimeVirtualTexture*>);
 	ParameterSet.SparseVolumeTextureParameters.Sort(SortMaterialInstanceParametersPredicate<const USparseVolumeTexture*>);
+	ParameterSet.EnumerationParameters.Sort(SortMaterialInstanceParametersPredicate<FName>);
 
 	StaticSwitchParameterArray.Array = MoveTemp(ParameterSet.StaticSwitchParameters);
 	ScalarParameterArray.Array = MoveTemp(ParameterSet.ScalarParameters);
@@ -520,10 +523,10 @@ void FMaterialInstanceResource::InitMIParameters(FMaterialInstanceParameterSet&
 	TextureCollectionParameterArray.Array = MoveTemp(ParameterSet.TextureCollectionParameters);
 	RuntimeVirtualTextureParameterArray.Array = MoveTemp(ParameterSet.RuntimeVirtualTextureParameters);
 	SparseVolumeTextureParameterArray.Array = MoveTemp(ParameterSet.SparseVolumeTextureParameters);
+	EnumerationParameterArray.Array = MoveTemp(ParameterSet.EnumerationParameters);
 	UserSceneTextureOverrides = MoveTemp(ParameterSet.UserSceneTextureOverrides);
 	PostProcessBlendableOverrides = ParameterSet.PostProcessBlendableOverrides;
 
-
 	// Build hash tables.
 	StaticSwitchParameterArray.HashAddAllItems();
 	ScalarParameterArray.HashAddAllItems();
@@ -533,6 +536,7 @@ void FMaterialInstanceResource::InitMIParameters(FMaterialInstanceParameterSet&
 	TextureCollectionParameterArray.HashAddAllItems();
 	RuntimeVirtualTextureParameterArray.HashAddAllItems();
 	SparseVolumeTextureParameterArray.HashAddAllItems();
+	EnumerationParameterArray.HashAddAllItems();
 }
 
 /**
@@ -689,6 +693,7 @@ void UMaterialInstance::SwapLayerParameterIndices(int32 OriginalIndex, int32 New
 		SwapLayerParameterIndicesArray(SparseVolumeTextureParameterValues, OriginalIndex, NewIndex);
 		SwapLayerParameterIndicesArray(FontParameterValues, OriginalIndex, NewIndex);
 		SwapLayerParameterIndicesArray(StaticParametersRuntime.StaticSwitchParameters, OriginalIndex, NewIndex);
+		SwapLayerParameterIndicesArray(EnumerationParameterValues, OriginalIndex, NewIndex);
 		if (EditorOnly)
 		{
 			SwapLayerParameterIndicesArray(EditorOnly->StaticParameters.StaticComponentMaskParameters, OriginalIndex, NewIndex);
@@ -708,6 +713,7 @@ void UMaterialInstance::RemoveLayerParameterIndex(int32 Index)
 	RemoveLayerParameterIndicesArray(SparseVolumeTextureParameterValues, Index);
 	RemoveLayerParameterIndicesArray(FontParameterValues, Index);
 	RemoveLayerParameterIndicesArray(StaticParametersRuntime.StaticSwitchParameters, Index);
+	RemoveLayerParameterIndicesArray(EnumerationParameterValues, Index);
 	if (EditorOnly)
 	{
 		RemoveLayerParameterIndicesArray(EditorOnly->StaticParameters.StaticComponentMaskParameters, Index);
@@ -770,6 +776,9 @@ bool UMaterialInstance::UpdateParameters()
 
 			// Static component mask parameters
 			bDirty = UpdateParameterSet<FStaticComponentMaskParameter, UMaterialExpressionStaticComponentMaskParameter>(EditorOnly->StaticParameters.StaticComponentMaskParameters, ParentMaterial) || bDirty;
+			
+			// Enumeration parameters
+			bDirty = UpdateParameterSet<FEnumerationParameterValue, UMaterialExpressionEnumerationParameter>(EnumerationParameterValues, ParentMaterial) || bDirty;
 		}
 
 		if (StaticParametersRuntime.bHasMaterialLayers && Parent)
@@ -794,6 +803,7 @@ bool UMaterialInstance::UpdateParameters()
 					RemapLayerParameterIndicesArray(FontParameterValues, RemapLayerIndices);
 					RemapLayerParameterIndicesArray(StaticParametersRuntime.StaticSwitchParameters, RemapLayerIndices);
 					RemapLayerParameterIndicesArray(EditorOnly->StaticParameters.StaticComponentMaskParameters, RemapLayerIndices);
+					RemapLayerParameterIndicesArray(EnumerationParameterValues, RemapLayerIndices);
 					bDirty = true;
 				}
 			}
@@ -919,6 +929,14 @@ void GameThread_InitMIParameters(const UMaterialInstance& Instance)
 		ParamRef.Info = Parameter.ParameterInfo;
 		ParamRef.Value = FSparseVolumeTextureParameterValue::GetValue(Parameter);
 	}
+
+	// Enumeration parameters
+	for (const FEnumerationParameterValue& Parameter : Instance.EnumerationParameterValues)
+	{
+		auto& ParamRef = ParameterSet.EnumerationParameters.AddDefaulted_GetRef();
+		ParamRef.Info = Parameter.ParameterInfo;
+		ParamRef.Value = FEnumerationParameterValue::GetValue(Parameter);
+	}
 	
 	FStaticParameterSet StaticParamSet = Instance.GetStaticParameters();
 	ParameterSet.StaticSwitchParameters.Reserve(StaticParamSet.StaticSwitchParameters.Num());
@@ -1097,6 +1115,7 @@ bool UMaterialInstance::GetParameterOverrideValue(EMaterialParameterType Type, c
 		bResult = GameThread_GetParameterValue(GetEditorOnlyData()->StaticParameters.StaticComponentMaskParameters, ParameterInfo, OutResult);
 		break;
 #endif // WITH_EDITORONLY_DATA
+	case EMaterialParameterType::Enumeration: bResult = GameThread_GetParameterValue(EnumerationParameterValues, ParameterInfo, OutResult); break;
 	default: checkNoEntry(); break;
 	}
 	return bResult;
@@ -2215,6 +2234,7 @@ void UMaterialInstance::GetAllParametersOfType(EMaterialParameterType Type, TMap
 		case EMaterialParameterType::SparseVolumeTexture: GameThread_ApplyParameterOverrides(Instance->SparseVolumeTextureParameterValues, LayerIndexRemap, bSetOverride, OverridenParameters, OutParameters); break;
 		case EMaterialParameterType::Font: GameThread_ApplyParameterOverrides(Instance->FontParameterValues, LayerIndexRemap, bSetOverride, OverridenParameters, OutParameters); break;
 		case EMaterialParameterType::StaticSwitch: GameThread_ApplyParameterOverrides(Instance->StaticParametersRuntime.StaticSwitchParameters, LayerIndexRemap, bSetOverride, OverridenParameters, OutParameters); break;
+		case EMaterialParameterType::Enumeration: GameThread_ApplyParameterOverrides(Instance->EnumerationParameterValues, LayerIndexRemap, bSetOverride, OverridenParameters, OutParameters); break;
 #if WITH_EDITORONLY_DATA
 		case EMaterialParameterType::StaticComponentMask:
 			GameThread_ApplyParameterOverrides(Instance->GetEditorOnlyData()->StaticParameters.StaticComponentMaskParameters, LayerIndexRemap, bSetOverride, OverridenParameters, OutParameters);
@@ -3906,6 +3926,7 @@ void UMaterialInstance::ReserveParameterValuesInternal(EMaterialParameterType Ty
 	case EMaterialParameterType::Font: FontParameterValues.Reserve(Capacity); break;
 	case EMaterialParameterType::RuntimeVirtualTexture: RuntimeVirtualTextureParameterValues.Reserve(Capacity); break;
 	case EMaterialParameterType::SparseVolumeTexture: SparseVolumeTextureParameterValues.Reserve(Capacity); break;
+	case EMaterialParameterType::Enumeration: EnumerationParameterValues.Reserve(Capacity); break;
 	default: checkNoEntry();
 	}
 }
@@ -3936,6 +3957,7 @@ void UMaterialInstance::AddParameterValueInternal(const FMaterialParameterInfo&
 	case EMaterialParameterType::RuntimeVirtualTexture: RuntimeVirtualTextureParameterValues.Emplace(ParameterInfo, Value.RuntimeVirtualTexture); break;
 	case EMaterialParameterType::SparseVolumeTexture: SparseVolumeTextureParameterValues.Emplace(ParameterInfo, Value.SparseVolumeTexture); break;
 	case EMaterialParameterType::StaticSwitch: break;
+	case EMaterialParameterType::Enumeration: EnumerationParameterValues.Emplace(ParameterInfo, Value.EnumEntry); break;
 	default: checkNoEntry();
 	}
 }
@@ -3965,6 +3987,7 @@ void UMaterialInstance::SetParameterValueInternal(const FMaterialParameterInfo&
 	case EMaterialParameterType::Font: SetFontParameterValueInternal(ParameterInfo, Value.Font.Value, Value.Font.Page); break;
 	case EMaterialParameterType::RuntimeVirtualTexture: SetRuntimeVirtualTextureParameterValueInternal(ParameterInfo, Value.RuntimeVirtualTexture); break;
 	case EMaterialParameterType::SparseVolumeTexture: SetSparseVolumeTextureParameterValueInternal(ParameterInfo, Value.SparseVolumeTexture); break;
+	case EMaterialParameterType::Enumeration: SetEnumerationParameterValueInternal(ParameterInfo, Value.EnumEntry); break;
 	default: checkNoEntry();
 	}
 }
@@ -4305,6 +4328,32 @@ void UMaterialInstance::SetFontParameterValueInternal(const FMaterialParameterIn
 	}
 }
 
+void UMaterialInstance::SetEnumerationParameterValueInternal(const FMaterialParameterInfo& ParameterInfo, FName Value)
+{
+	LLM_SCOPE(ELLMTag::MaterialInstance);
+
+	FEnumerationParameterValue* ParameterValue = GameThread_FindParameterByName(EnumerationParameterValues, ParameterInfo);
+
+	bool bForceUpdate = false;
+	if (!ParameterValue)
+	{
+			// If there's no element for the named parameter in array yet, add one.
+			ParameterValue = new(EnumerationParameterValues) FEnumerationParameterValue;
+			ParameterValue->ParameterInfo = ParameterInfo;
+			ParameterValue->ExpressionGUID.Invalidate();
+			bForceUpdate = true;
+	}
+
+	// Don't enqueue an update if it isn't needed
+	if (bForceUpdate || ParameterValue->EnumEntry != Value)
+	{
+		ParameterValue->EnumEntry = Value;
+
+		// Update the material instance data in the rendering thread.
+		GameThread_UpdateMIParameter(this, *ParameterValue);
+	}
+}
+
 void UMaterialInstance::ClearParameterValuesInternal(EMaterialInstanceClearParameterFlag Flags)
 {
 	bool bUpdateResource = false;
@@ -4313,6 +4362,7 @@ void UMaterialInstance::ClearParameterValuesInternal(EMaterialInstanceClearParam
 		ScalarParameterValues.Empty();
 		VectorParameterValues.Empty();
 		DoubleVectorParameterValues.Empty();
+		EnumerationParameterValues.Empty();
 		bUpdateResource = true;
 	}
 
@@ -5248,7 +5298,10 @@ bool UMaterialInstance::Equivalent(const UMaterialInstance* CompareTo) const
 	{
 		return false;
 	}
-
+	if (!CompareValueArraysByExpressionGUID(EnumerationParameterValues, CompareTo->EnumerationParameterValues))
+	{
+		return false;
+	}
 	const FStaticParameterSet LocalStaticParameters = GetStaticParameters();
 	if (!LocalStaticParameters.Equivalent(CompareTo->GetStaticParameters()))
 	{
diff --git a/Engine/Source/Runtime/Engine/Private/Materials/MaterialInstanceSupport.h b/Engine/Source/Runtime/Engine/Private/Materials/MaterialInstanceSupport.h
--- a/Engine/Source/Runtime/Engine/Private/Materials/MaterialInstanceSupport.h
+++ b/Engine/Source/Runtime/Engine/Private/Materials/MaterialInstanceSupport.h
@@ -249,6 +249,7 @@ public:
 		RuntimeVirtualTextureParameterArray.Empty();
 		SparseVolumeTextureParameterArray.Empty();
 		UserSceneTextureOverrides.Empty();
+		EnumerationParameterArray.Empty();
 	}
 
 	/**
@@ -320,6 +321,7 @@ private:
 	static bool IsValidParameterValue(const UTextureCollection* Value) { return Value != nullptr; }
 	static bool IsValidParameterValue(const URuntimeVirtualTexture* Value) { return Value != nullptr; }
 	static bool IsValidParameterValue(const USparseVolumeTexture* Value) { return Value != nullptr; }
+	static bool IsValidParameterValue(FName Value) { return true; }
 
 	virtual void StartCacheUniformExpressions() const override;
 	virtual void FinishCacheUniformExpressions() const override;
@@ -411,6 +413,8 @@ private:
 	THashedMaterialParameterMap<const URuntimeVirtualTexture*> RuntimeVirtualTextureParameterArray;
 	/** Sparse Volume Texture parameters for this material instance. */
 	THashedMaterialParameterMap<const USparseVolumeTexture*> SparseVolumeTextureParameterArray;
+	/** Enumeration parameters for this material instance. */
+	THashedMaterialParameterMap<FName> EnumerationParameterArray;
 	/** Remap layer indices for parent */
 	TArray<int32> ParentLayerIndexRemap;
 	/** User Scene Texture overrides for this material instance. */
@@ -427,6 +431,7 @@ template <> FORCEINLINE THashedMaterialParameterMap<const UTexture*>& FMaterialI
 template <> FORCEINLINE THashedMaterialParameterMap<const UTextureCollection*>& FMaterialInstanceResource::GetValueArray() { return TextureCollectionParameterArray; }
 template <> FORCEINLINE THashedMaterialParameterMap<const URuntimeVirtualTexture*>& FMaterialInstanceResource::GetValueArray() { return RuntimeVirtualTextureParameterArray; }
 template <> FORCEINLINE THashedMaterialParameterMap<const USparseVolumeTexture*>& FMaterialInstanceResource::GetValueArray() { return SparseVolumeTextureParameterArray; }
+template <> FORCEINLINE THashedMaterialParameterMap<FName>& FMaterialInstanceResource::GetValueArray() { return EnumerationParameterArray; }
 template <> FORCEINLINE const THashedMaterialParameterMap<bool>& FMaterialInstanceResource::GetValueArray() const { return StaticSwitchParameterArray; }
 template <> FORCEINLINE const THashedMaterialParameterMap<float>& FMaterialInstanceResource::GetValueArray() const { return ScalarParameterArray; }
 template <> FORCEINLINE const THashedMaterialParameterMap<FLinearColor>& FMaterialInstanceResource::GetValueArray() const { return VectorParameterArray; }
@@ -435,6 +440,7 @@ template <> FORCEINLINE const THashedMaterialParameterMap<const UTexture*>& FMat
 template <> FORCEINLINE const THashedMaterialParameterMap<const UTextureCollection*>& FMaterialInstanceResource::GetValueArray() const { return TextureCollectionParameterArray; }
 template <> FORCEINLINE const THashedMaterialParameterMap<const URuntimeVirtualTexture*>& FMaterialInstanceResource::GetValueArray() const { return RuntimeVirtualTextureParameterArray; }
 template <> FORCEINLINE const THashedMaterialParameterMap<const USparseVolumeTexture*>& FMaterialInstanceResource::GetValueArray() const { return SparseVolumeTextureParameterArray; }
+template <> FORCEINLINE const THashedMaterialParameterMap<FName>& FMaterialInstanceResource::GetValueArray() const { return EnumerationParameterArray; }
 
 struct FMaterialInstanceParameterSet
 {
@@ -446,6 +452,7 @@ struct FMaterialInstanceParameterSet
 	TArray<THashedMaterialParameterMap<const UTextureCollection*>::TNamedParameter>		TextureCollectionParameters;
 	TArray<THashedMaterialParameterMap<const URuntimeVirtualTexture*>::TNamedParameter>	RuntimeVirtualTextureParameters;
 	TArray<THashedMaterialParameterMap<const USparseVolumeTexture*>::TNamedParameter>	SparseVolumeTextureParameters;
+	TArray<THashedMaterialParameterMap<FName>::TNamedParameter>							EnumerationParameters;
 	TArray<FUserSceneTextureOverride>													UserSceneTextureOverrides;
 	FPostProcessBlendableOverrides														PostProcessBlendableOverrides;
 };
diff --git a/Engine/Source/Runtime/Engine/Private/Materials/MaterialRenderProxy.cpp b/Engine/Source/Runtime/Engine/Private/Materials/MaterialRenderProxy.cpp
--- a/Engine/Source/Runtime/Engine/Private/Materials/MaterialRenderProxy.cpp
+++ b/Engine/Source/Runtime/Engine/Private/Materials/MaterialRenderProxy.cpp
@@ -1,6 +1,8 @@
 // Copyright Epic Games, Inc. All Rights Reserved.
 
 #include "Materials/MaterialRenderProxy.h"
+#include "Materials/MaterialEnumeration.h"
+#include "MaterialCachedData.h"
 #include "Engine/Texture.h"
 #include "EngineModule.h"
 #include "ExternalTexture.h"
@@ -255,6 +257,28 @@ bool FMaterialRenderProxy::GetTextureCollectionValue(const FHashedMaterialParame
 	return false;
 }
 
+bool FMaterialRenderProxy::GetParameterShaderValue(EMaterialParameterType Type, const FHashedMaterialParameterInfo& ParameterInfo, UE::Shader::FValue& OutValue, const FMaterialRenderContext& Context) const
+{
+	FMaterialParameterValue Value;
+	if (GetParameterValue(Type, ParameterInfo, Value, Context))
+	{
+		if (Type == EMaterialParameterType::Enumeration)
+		{
+			FMaterialParameterMetadata ParameterMetadata;
+			bool bResult = GetMaterialInterface()->GetCachedExpressionData().GetParameterValue(Type, ParameterInfo, ParameterMetadata);
+			check(bResult);
+
+			OutValue = Cast<IMaterialEnumerationProvider>(ParameterMetadata.Enumeration)->GetValueOrDefault(Value.EnumEntry);
+		}
+		else
+		{
+			OutValue = Value.AsShaderValue();
+		}
+		return true;
+	}
+	return false;
+}
+
 FName FMaterialRenderProxy::GetUserSceneTextureOutput(const FMaterial* Base) const
 {
 	FName Result = NAME_None;
diff --git a/Engine/Source/Runtime/Engine/Private/Materials/MaterialShared.cpp b/Engine/Source/Runtime/Engine/Private/Materials/MaterialShared.cpp
--- a/Engine/Source/Runtime/Engine/Private/Materials/MaterialShared.cpp
+++ b/Engine/Source/Runtime/Engine/Private/Materials/MaterialShared.cpp
@@ -5497,6 +5497,8 @@ UE::Shader::FType GetShaderValueType(EMaterialParameterType Type)
 		return FMaterialTextureValue::GetTypeName();
 	case EMaterialParameterType::SparseVolumeTexture:
 		return UE::Shader::EValueType::Void; // TODO
+	case EMaterialParameterType::Enumeration:
+		return UE::Shader::EValueType::Int1;
 	default:
 		checkNoEntry();
 		return UE::Shader::EValueType::Void;
diff --git a/Engine/Source/Runtime/Engine/Private/Shader/Preshader.cpp b/Engine/Source/Runtime/Engine/Private/Shader/Preshader.cpp
--- a/Engine/Source/Runtime/Engine/Private/Shader/Preshader.cpp
+++ b/Engine/Source/Runtime/Engine/Private/Shader/Preshader.cpp
@@ -600,10 +600,10 @@ static void EvaluateParameter(FPreshaderStack& Stack, const FUniformExpressionSe
 	// First allow proxy the chance to override parameter
 	if (Context.MaterialRenderProxy)
 	{
-		FMaterialParameterValue ParameterValue;
-		if (Context.MaterialRenderProxy->GetParameterValue(Parameter.ParameterType, Parameter.ParameterInfo, ParameterValue, Context))
+		UE::Shader::FValue ParameterShaderValue;
+		if (Context.MaterialRenderProxy->GetParameterShaderValue(Parameter.ParameterType, Parameter.ParameterInfo, ParameterShaderValue, Context))
 		{
-			Stack.PushValue(ParameterValue.AsShaderValue());
+			Stack.PushValue(ParameterShaderValue);
 			bFoundParameter = true;
 		}
 	}
diff --git a/Engine/Source/Runtime/Engine/Public/MaterialCachedData.h b/Engine/Source/Runtime/Engine/Public/MaterialCachedData.h
--- a/Engine/Source/Runtime/Engine/Public/MaterialCachedData.h
+++ b/Engine/Source/Runtime/Engine/Public/MaterialCachedData.h
@@ -25,6 +25,7 @@ class ULandscapeGrassType;
 class UMaterialFunctionInterface;
 class UMaterialInterface;
 struct FStaticParameterSet;
+class IMaterialEnumerationProvider;
 
 /** Stores information about a function that this material references, used to know when the material needs to be recompiled. */
 USTRUCT()
@@ -315,6 +316,12 @@ struct FMaterialCachedExpressionData
 
 	UPROPERTY()
 	TArray<int32> FontPageValues;
+	
+	UPROPERTY()
+	TArray<TSoftObjectPtr<UObject>> Enumerations;
+
+	UPROPERTY()
+	TArray<FName> EnumerationEntryValues;
 
 	UPROPERTY()
 	TArray<TSoftObjectPtr<URuntimeVirtualTexture>> RuntimeVirtualTextureValues;
diff --git a/Engine/Source/Runtime/Engine/Public/MaterialTypes.h b/Engine/Source/Runtime/Engine/Public/MaterialTypes.h
--- a/Engine/Source/Runtime/Engine/Public/MaterialTypes.h
+++ b/Engine/Source/Runtime/Engine/Public/MaterialTypes.h
@@ -193,6 +193,7 @@ enum class EMaterialParameterType : uint8
 	RuntimeVirtualTexture,
 	SparseVolumeTexture,
 	StaticSwitch,
+	Enumeration,
 
 	NumRuntime, // Runtime parameter types must go above here, and editor-only ones below
 
@@ -215,6 +216,7 @@ inline bool IsNumericMaterialParameter(EMaterialParameterType InType)
 	case EMaterialParameterType::Scalar:
 	case EMaterialParameterType::Vector:
 	case EMaterialParameterType::DoubleVector:
+	case EMaterialParameterType::Enumeration:
 		return true;
 	default:
 		return false;
@@ -365,7 +367,7 @@ struct FMaterialParameterValue
 	FMaterialParameterValue(bool InValue) : Type(EMaterialParameterType::StaticSwitch) { Bool[0] = InValue; }
 	FMaterialParameterValue(const FStaticComponentMaskValue& InValue) : Type(EMaterialParameterType::StaticComponentMask) { Bool[0] = InValue.R; Bool[1] = InValue.G; Bool[2] = InValue.B; Bool[3] = InValue.A; }
 	FMaterialParameterValue(bool bMaskR, bool bMaskG, bool bMaskB, bool bMaskA) : Type(EMaterialParameterType::StaticComponentMask) { Bool[0] = bMaskR; Bool[1] = bMaskG; Bool[2] = bMaskB; Bool[3] = bMaskA; }
-	
+	FMaterialParameterValue(FName InEnumEntry) : Type(EMaterialParameterType::Enumeration) { EnumEntry = InEnumEntry; }
 	ENGINE_API FMaterialParameterValue(EMaterialParameterType Type, const UE::Shader::FValue& InValue);
 
 	inline float AsScalar() const { check(Type == EMaterialParameterType::Scalar); return Float[0]; }
@@ -390,6 +392,7 @@ struct FMaterialParameterValue
 			UFont* Value;
 			int32 Page;
 		} Font;
+		FName EnumEntry;
 	};
 	EMaterialParameterType Type;
 
@@ -425,6 +428,7 @@ struct FMaterialParameterValue
 			Lhs.Bool[1] == Rhs.Bool[1] &&
 			Lhs.Bool[2] == Rhs.Bool[2] &&
 			Lhs.Bool[3] == Rhs.Bool[3];
+		case EMaterialParameterType::Enumeration: return Lhs.EnumEntry == Rhs.EnumEntry;
 		default: checkNoEntry(); return false;
 		}
 	}
@@ -482,6 +486,10 @@ struct FMaterialParameterMetadata
 	bool bOverride = false;
 #endif // WITH_EDITORONLY_DATA
 
+	/** The enumeration definition, valid when this is an enumeration parameter.
+	  * Note: this UObject must implement the IMaterialEnumerationProvider interface */
+	TObjectPtr<UObject> Enumeration;
+
 	/** Is a Static Switch parameter using a dynamic branch? */
 	bool bDynamicSwitchParameter = false;
 };
diff --git a/Engine/Source/Runtime/Engine/Public/Materials/Material.h b/Engine/Source/Runtime/Engine/Public/Materials/Material.h
--- a/Engine/Source/Runtime/Engine/Public/Materials/Material.h
+++ b/Engine/Source/Runtime/Engine/Public/Materials/Material.h
@@ -2058,12 +2058,6 @@ private:
 	UPROPERTY()
 	TArray<TObjectPtr<class UMaterialExpressionComment>> EditorComments_DEPRECATED;
 
-	// UPROPERTY()
-	// TObjectPtr<class UMaterialExpression> ExpressionExecBegin_DEPRECATED; // maxtodo
-	
-	// UPROPERTY()
-	// TObjectPtr<class UMaterialExpression> ExpressionExecEnd_DEPRECATED; // // maxtodo
-
 	UPROPERTY()
 	TArray<TObjectPtr<class UMaterialExpression>> Expressions_DEPRECATED;
 
diff --git a/Engine/Source/Runtime/Engine/Public/Materials/MaterialEnumeration.h b/Engine/Source/Runtime/Engine/Public/Materials/MaterialEnumeration.h
new file mode 100644
--- /dev/null
+++ b/Engine/Source/Runtime/Engine/Public/Materials/MaterialEnumeration.h
@@ -0,0 +1,72 @@
+// Copyright Epic Games, Inc. All Rights Reserved.
+
+#pragma once
+
+#include "CoreMinimal.h"
+#include "Engine/DataAsset.h"
+#include "UObject/Interface.h"
+#include "MaterialEnumeration.generated.h"
+
+class IMaterialEnumerationProvider
+{
+	GENERATED_BODY()
+
+public:
+	// Tries to resolve the value for the given name.
+	// Returns true if found, false otherwise. Writes result to OutValue.
+	virtual bool ResolveValue(FName EntryName, int32& OutValue, int32 DefaultValue = 0) const = 0;
+
+	// Iterates over all name-value entries.
+	virtual void ForEachEntry(TFunctionRef<void (FName Name, int32 Value)> Iterator) const = 0;
+
+	// Helper to get the value for a name, or return default if not found.
+	int32 GetValueOrDefault(FName EntryName, int32 DefaultValue = 0) const;
+};
+
+UINTERFACE(MinimalAPI)
+class UMaterialEnumerationProvider : public UInterface
+{
+	GENERATED_BODY()
+};
+
+USTRUCT(MinimalAPI)
+struct FMaterialEnumerationEntry
+{
+	GENERATED_BODY()
+
+	// Entry identifier.
+	UPROPERTY(EditAnywhere, Category=MaterialEnumerationEntry)
+	FName Name; 
+
+	// Associated integer value.
+	UPROPERTY(EditAnywhere, Category=MaterialEnumerationEntry)
+	int32 Value;
+};
+
+UCLASS(MinimalAPI)
+class UMaterialEnumeration : public UDataAsset, public IMaterialEnumerationProvider
+{
+	GENERATED_BODY()
+
+public:
+	// Raw list of entries.
+	UPROPERTY(EditAnywhere, Category=MaterialEnumeration)
+	TArray<FMaterialEnumerationEntry> Entries;
+
+	// Cached and sorted version of Entries for faster lookup.
+	TArray<FMaterialEnumerationEntry> SortedEntries;
+
+	// IMaterialEnumerationProvider implementation:
+	virtual bool ResolveValue(FName EntryName, int32& OutValue, int32 DefaultValue = 0) const final;
+	virtual void ForEachEntry(TFunctionRef<void (FName Name, int32 Value)> Iterator) const final;
+
+private:
+	// Sorts entries by name for optimized access.
+	void SortEntries();
+
+	virtual void PostLoad() override;
+
+#if WITH_EDITOR
+	virtual void PostEditChangeProperty(struct FPropertyChangedEvent& PropertyChangedEvent) override;
+#endif
+};
diff --git a/Engine/Source/Runtime/Engine/Public/Materials/MaterialExpressionEnumerationParameter.h b/Engine/Source/Runtime/Engine/Public/Materials/MaterialExpressionEnumerationParameter.h
new file mode 100644
--- /dev/null
+++ b/Engine/Source/Runtime/Engine/Public/Materials/MaterialExpressionEnumerationParameter.h
@@ -0,0 +1,41 @@
+// Copyright Epic Games, Inc. All Rights Reserved.
+
+#pragma once
+
+#include "CoreMinimal.h"
+#include "UObject/ObjectMacros.h"
+#include "Materials/MaterialExpressionParameter.h"
+#include "Materials/MaterialEnumeration.h"
+#include "MaterialExpressionEnumerationParameter.generated.h"
+
+struct FPropertyChangedEvent;
+
+UCLASS(MinimalAPI)
+class UMaterialExpressionEnumerationParameter : public UMaterialExpressionParameter
+{
+	GENERATED_BODY()
+public:
+	// The enumeration asset defining the possible entries and their values.
+	UPROPERTY(EditAnywhere, BlueprintReadWrite, Category=MaterialExpressionEnumeratorParameter)
+	TScriptInterface<IMaterialEnumerationProvider> Enumeration;
+
+	// The default entry to use from the Enumeration asset.
+	UPROPERTY(EditAnywhere, Category=MaterialExpressionEnumeratorParameter, Meta=(ShowAsInputPin="Primary", GetOptions="GetEntryOptions"))
+	FName DefaultEntry;
+
+#if WITH_EDITOR
+	//~ Begin UMaterialExpression Interface
+	virtual int32 Compile(class FMaterialCompiler* Compiler, int32 OutputIndex) override;
+	virtual void GetCaption(TArray<FString>& OutCaptions) const override;
+	virtual EMaterialValueType GetOutputValueType(int32 InIndex) override;
+	virtual bool GetParameterValue(FMaterialParameterMetadata& OutMeta) const override;
+	virtual bool SetParameterValue(const FName& Name, const FMaterialParameterMetadata& Meta, EMaterialExpressionSetParameterValueFlags Flags) override;
+	//~ End UMaterialExpression Interface
+
+	virtual void PostEditChangeProperty(FPropertyChangedEvent& PropertyChangedEvent) override;
+	bool SetParameterValue(const FName& InParameterName, UObject* InEnumeration, FName InDefaultEntry, EMaterialExpressionSetParameterValueFlags Flags);
+
+	UFUNCTION()
+	TArray<FString> GetEntryOptions() const;
+#endif
+};
diff --git a/Engine/Source/Runtime/Engine/Public/Materials/MaterialInstance.h b/Engine/Source/Runtime/Engine/Public/Materials/MaterialInstance.h
--- a/Engine/Source/Runtime/Engine/Public/Materials/MaterialInstance.h
+++ b/Engine/Source/Runtime/Engine/Public/Materials/MaterialInstance.h
@@ -476,6 +476,54 @@ struct FFontParameterValue
 	}
 };
 
+/** Editable enumeration parameter. */
+USTRUCT(BlueprintType)
+struct FEnumerationParameterValue
+{
+	GENERATED_USTRUCT_BODY()
+
+	UPROPERTY(EditAnywhere, BlueprintReadWrite, Category=EnumerationParameterValue)
+	FMaterialParameterInfo ParameterInfo;
+
+	UPROPERTY(EditAnywhere, BlueprintReadWrite, Category=EnumerationParameterValue)
+	FName EnumEntry;
+
+	UPROPERTY()
+	FGuid ExpressionGUID;
+
+	explicit FEnumerationParameterValue(const FMaterialParameterInfo& InParameterInfo = FMaterialParameterInfo(), FName InEnumEntry = {})
+		: ParameterInfo(InParameterInfo), EnumEntry(InEnumEntry)
+	{
+	}
+
+	bool IsOverride() const { return true; }
+
+	bool IsValid() const { return GetValue(*this).IsValid(); }
+
+	bool operator==(const FEnumerationParameterValue& Other) const
+	{
+		return
+			ParameterInfo == Other.ParameterInfo &&
+			EnumEntry == Other.EnumEntry &&
+			ExpressionGUID == Other.ExpressionGUID;
+	}
+	bool operator!=(const FEnumerationParameterValue& Other) const
+	{
+		return !((*this) == Other);
+	}
+	
+	typedef FName ValueType;
+	static ValueType GetValue(const FEnumerationParameterValue& Parameter) { return Parameter.EnumEntry; }
+
+	void GetValue(FMaterialParameterMetadata& OutResult) const
+	{
+		OutResult.Value = FMaterialParameterValue(EnumEntry);
+#if WITH_EDITORONLY_DATA
+		OutResult.ExpressionGuid = ExpressionGUID;
+#endif
+	}
+};
+
 USTRUCT(BlueprintType)
 struct FUserSceneTextureOverride
 {
@@ -728,6 +776,10 @@ public:
 	/** Font parameters. */
 	UPROPERTY(EditAnywhere, BlueprintReadOnly, Category=MaterialInstance, meta = (EditFixedOrder))
 	TArray<struct FFontParameterValue> FontParameterValues;
+	
+	/** Enumeration parameters. */
+	UPROPERTY(EditAnywhere, BlueprintReadOnly, Category = MaterialInstance, meta = (EditFixedOrder))
+	TArray<FEnumerationParameterValue> EnumerationParameterValues;
 
 	/** User scene texture overrides.  Applies to post process domain materials only. */
 	UPROPERTY(EditAnywhere, Category = MaterialInstance)
@@ -1149,6 +1201,7 @@ protected:
 	void SetRuntimeVirtualTextureParameterValueInternal(const FMaterialParameterInfo& ParameterInfo, class URuntimeVirtualTexture* Value);
 	void SetSparseVolumeTextureParameterValueInternal(const FMaterialParameterInfo& ParameterInfo, class USparseVolumeTexture* Value);
 	void SetFontParameterValueInternal(const FMaterialParameterInfo& ParameterInfo, class UFont* FontValue, int32 FontPage);
+	void SetEnumerationParameterValueInternal(const FMaterialParameterInfo& ParameterInfo, FName Value);
 	void ClearParameterValuesInternal(EMaterialInstanceClearParameterFlag Flags = EMaterialInstanceClearParameterFlag::AllNonStatic);
 
 	/** Initialize the material instance's resources. */
diff --git a/Engine/Source/Runtime/Engine/Public/Materials/MaterialRenderProxy.h b/Engine/Source/Runtime/Engine/Public/Materials/MaterialRenderProxy.h
--- a/Engine/Source/Runtime/Engine/Public/Materials/MaterialRenderProxy.h
+++ b/Engine/Source/Runtime/Engine/Public/Materials/MaterialRenderProxy.h
@@ -36,6 +36,10 @@ struct FMaterialRenderContext;
 struct FMemoryImageMaterialParameterInfo;
 using FHashedMaterialParameterInfo = FMemoryImageMaterialParameterInfo;
 
+namespace UE::Shader
+{
+	struct FValue;
+}
 
 /**
  * Cached uniform expression values.
@@ -167,6 +171,7 @@ public:
 	ENGINE_API bool GetTextureValue(const FHashedMaterialParameterInfo& ParameterInfo, const URuntimeVirtualTexture** OutValue, const FMaterialRenderContext& Context) const;
 	ENGINE_API bool GetTextureValue(const FHashedMaterialParameterInfo& ParameterInfo, const USparseVolumeTexture** OutValue, const FMaterialRenderContext& Context) const;
 	ENGINE_API bool GetTextureCollectionValue(const FHashedMaterialParameterInfo& ParameterInfo, const UTextureCollection** OutValue, const FMaterialRenderContext& Context) const;
+	ENGINE_API bool GetParameterShaderValue(EMaterialParameterType Type, const FHashedMaterialParameterInfo& ParameterInfo, UE::Shader::FValue& OutValue, const FMaterialRenderContext& Context) const;
 	virtual bool GetParameterValue(EMaterialParameterType Type, const FHashedMaterialParameterInfo& ParameterInfo, FMaterialParameterValue& OutValue, const FMaterialRenderContext& Context) const = 0;
 	virtual bool GetUserSceneTextureOverride(FName& InOutValue) const { return false; }
 	ENGINE_API FName GetUserSceneTextureOutput(const FMaterial* Base) const;
