Added missing boilerplate to MaterialEnumeration.cpp

diff --git a/Engine/Source/Runtime/Engine/Private/Materials/MaterialEnumeration.cpp b/Engine/Source/Runtime/Engine/Private/Materials/MaterialEnumeration.cpp
--- a/Engine/Source/Runtime/Engine/Private/Materials/MaterialEnumeration.cpp
+++ b/Engine/Source/Runtime/Engine/Private/Materials/MaterialEnumeration.cpp
@@ -1,3 +1,5 @@
+// Copyright Epic Games, Inc. All Rights Reserved.
+
 #include "Materials/MaterialEnumeration.h"
 
 int32 IMaterialEnumerationProvider::GetValueOrDefault(FName EntryName, int32 DefaultValue) const
