Fixed material derivative autogen crash when dealing with uints.

#rnx

diff --git a/Engine/Source/Runtime/Engine/Private/Materials/HLSLMaterialDerivativeAutogen.cpp b/Engine/Source/Runtime/Engine/Private/Materials/HLSLMaterialDerivativeAutogen.cpp
--- a/Engine/Source/Runtime/Engine/Private/Materials/HLSLMaterialDerivativeAutogen.cpp
+++ b/Engine/Source/Runtime/Engine/Private/Materials/HLSLMaterialDerivativeAutogen.cpp
@@ -122,6 +122,7 @@ EDerivativeType GetDerivType(EMaterialValueType ValueType, bool bAllowNonFloat)
 	case MCT_StaticBool:
 	case MCT_Float:
 	case MCT_Float1:
+	case MCT_UInt1:
 		return EDerivativeType::Float1;
 	case MCT_Float2:
 		return EDerivativeType::Float2;
