Attributes support

diff --git a/Engine/Source/Runtime/Engine/Classes/Curves/Spline.h b/Engine/Source/Runtime/Engine/Classes/Curves/Spline.h
--- a/Engine/Source/Runtime/Engine/Classes/Curves/Spline.h
+++ b/Engine/Source/Runtime/Engine/Classes/Curves/Spline.h
@@ -67,8 +67,9 @@ struct FSpline
 	ENGINE_API FVector EvaluateScale(float Param) const;
 
 	/* Attribute Interface */
-	
-	bool SupportsAttributes() const { return false;/*IsNew();*/ }
+
+	// Non-Templated functions
+	bool SupportsAttributes() const { return false; /*return IsNew();*/ }
 	ENGINE_API bool HasAttributeChannel(FName AttributeName) const;
 	ENGINE_API TArray<FName> GetFloatPropertyChannels() const;
 	ENGINE_API TArray<FName> GetVectorPropertyChannels() const;
diff --git a/Engine/Source/Runtime/Engine/Private/Components/SplineComponent.cpp b/Engine/Source/Runtime/Engine/Private/Components/SplineComponent.cpp
--- a/Engine/Source/Runtime/Engine/Private/Components/SplineComponent.cpp
+++ b/Engine/Source/Runtime/Engine/Private/Components/SplineComponent.cpp
@@ -649,10 +649,10 @@ float USplineComponent::GetDistanceAlongSplineAtLocation(const FVector& InLocati
 
 int32 USplineComponent::GetNumberOfPropertyValues(FName PropertyName) const
 {
-	// if (UE::SplineComponent::GUseSplineCurves == false && Spline.SupportsAttributes())
-	// {
-	// 	return Spline.NumAttributeValues(PropertyName);
-	// }
+	if (UE::SplineComponent::GUseSplineCurves == false && Spline.SupportsAttributes())
+	{
+		return Spline.NumAttributeValues<float>(PropertyName);
+	}
 
 	return 0;
 }
@@ -736,20 +736,20 @@ int32 SetPropertyAtSplineInputKey(const FSpline& Spline, float InKey, const T& I
 
 float GetInputKeyAtIndex(const FSpline& Spline, int32 Index, FName PropertyName)
 {
-	// if (UE::SplineComponent::GUseSplineCurves == false && Spline.SupportsAttributes())
-	// {
-	// 	return Spline.GetAttributeParameter(Index, PropertyName);
-	// }
+	if (UE::SplineComponent::GUseSplineCurves == false && Spline.SupportsAttributes())
+	{
+		return Spline.GetAttributeParameter<float>(Index, PropertyName);
+	}
 
 	return 0.f;
 }
 
 int32 SetInputKeyAtIndex(FSpline& Spline, int32 Index, float InKey, FName PropertyName)
 {
-	// if (UE::SplineComponent::GUseSplineCurves == false && Spline.SupportsAttributes())
-	// {
-	// 	return Spline.SetAttributeParameter(Index, InKey , PropertyName);
-	// }
+	if (UE::SplineComponent::GUseSplineCurves == false && Spline.SupportsAttributes())
+	{
+		return Spline.SetAttributeParameter<float>(Index, InKey , PropertyName);
+	}
 
 	return INDEX_NONE;
 }
@@ -778,7 +778,7 @@ void USplineComponent::RemovePropertyAtIndex(int32 Index, FName PropertyName)
 {
 	if (UE::SplineComponent::GUseSplineCurves == false && Spline.SupportsAttributes())
 	{
-	//	Spline.RemoveAttributeValue(Index, PropertyName);
+		Spline.RemoveAttributeValue<float>(Index, PropertyName);
 	}
 }
 
diff --git a/Engine/Source/Runtime/Engine/Private/Curves/Spline.cpp b/Engine/Source/Runtime/Engine/Private/Curves/Spline.cpp
--- a/Engine/Source/Runtime/Engine/Private/Curves/Spline.cpp
+++ b/Engine/Source/Runtime/Engine/Private/Curves/Spline.cpp
@@ -54,6 +54,14 @@ FAutoConsoleVariableRef CVarFallbackFindNearest(
 	TEXT("True if we should implement FindNearest and FindNearestOnSegment using an intermediate spline representation. Only applies if Spline.Implementation == 2.")
 );
 
+// WARNING! THIS BREAKS EXISTING TEXT REPRESENTATION WHEN CHANGED, IT WILL ABSOLUTELY TRASH YOUR STUFF IF YOU CHANGE WILLY-NILLY
+bool GEncodeAsHex = true;
+// FAutoConsoleVariableRef CVarEncodeAsHex(
+// 	TEXT("Spline.EncodeAsHex"),
+// 	GEncodeAsHex,
+// 	TEXT("If true, splines are imported/exported using hexadecimal encoding. If false, the data is encoded in base 64.")
+// );
+	
 }
 
 using namespace UE::Geometry::Spline;
@@ -181,20 +189,23 @@ public:
     float FindNearestOnSegment(const FVector& InLocation, int32 SegmentIndex, float& OutSquaredDistance) const;
 	
 	/** Parameter based evaluation API. */
-    virtual FVector EvaluateImpl(float Parameter) const override;
+	FVector EvaluatePosition(float Parameter) const;	// Not named Evaluate because we would be shadowing non-virtual base, this maps parameters
     FVector EvaluateDerivative(float Parameter) const;
 	FQuat EvaluateRotation(float Parameter) const;
 	FVector EvaluateScale(float Parameter) const;
 
-	template <typename AttrType> TArray<FName> GetAttributeChannelNames() const;
-	template <typename AttrType> AttrType EvaluateAttribute(FName Name, float Parameter) const;
-	template <typename AttrType> int32 AddAttributeValue(FName Name, const AttrType& Value, float Parameter);
+	// Templated Channel Creation/Query functions
+	template <typename AttrType> int32 NumAttributeValues(FName Name) const;
+
+	// Templated Attribute Interaction by Index functions
+	template <typename AttrType> AttrType GetAttributeValue(const FName& Name, int32 Index) const;
 	template <typename AttrType> void SetAttributeValue(FName Name, const AttrType& Value, int32 Index);
 	template <typename AttrType> void RemoveAttributeValue(FName Name, int32 Index);
-	template <typename AttrType> AttrType GetAttributeValue(const FName& Name, int32 Index) const;
-	template <typename AttrType> int32 NumAttributeValues(FName Name) const;
 	template <typename AttrType> int32 SetAttributeParameter(const FName& Name, int32 Index, float ParentSpaceParameter);
 	template <typename AttrType> float GetAttributeParameter(const FName& Name, int32 Index) const;
+
+	// Templated Attribute Interaction by Parameter functions
+	template <typename AttrType> int32 AddAttributeValue(FName Name, const AttrType& Value, float Parameter);
 	
 	/** Updates the spline using the current configuration. */
 	void UpdateSpline();
@@ -225,6 +236,10 @@ private:
 
 private:
 
+	// In order for FSpline to access the parameter mapping functions below.
+	// The alternative is to redefine a bunch of MultiSpline functions to do the remapping internally, but this is a bit of a pain.
+	friend struct FSpline;
+	
 	/** Convert from [0, 1] to [0, NumSegments] */
 	float FromNormalizedSpace(float Parameter) const;
 	
@@ -256,48 +271,59 @@ private:
 	/**
 	 * Converts an index (possibly fractional) to a parameter value in internal spline space.
 	 */
-	float ConvertIndexToParameter(int32 Index, float Fraction = 0.0f) const;
+	float ConvertIndexToInternalParameter(int32 Index, float Fraction = 0.0f) const;
 
     /**
 	 * Converts an internal spline space parameter value to the closest index
 	 */
-	int32 ConvertParameterToIndex(float Parameter) const;
-};
-
-/**
- * Retrieves the names of all attribute channels of the given template type in the spline
- * @tparam AttrType - attribute type
- * @return Array of attribute channel names
- */
-template <typename AttrType>
-TArray<FName> FNewSpline::GetAttributeChannelNames() const
-{
-	using AttrSplineType = TTangentBezierSpline<AttrType>;
-
-	TArray<FName> ChannelNames; /*= InternalSpline.GetAttributeChannelNames<AttrType>();
-	ChannelNames.Remove(RotationAttrName);
-	ChannelNames.Remove(ScaleAttrName);
-	ChannelNames.Remove(ParamToLengthAttrName);
-	ChannelNames.Remove(LengthToParamAttrName);*/
-	return ChannelNames;
-}
+	int32 ConvertInternalParameterToNearestPointIndex(float Parameter) const;
 
-template <typename AttrType>
-AttrType FNewSpline::EvaluateAttribute(FName Name, float Parameter) const
-{
-	using AttrSplineType = TTangentBezierSpline<AttrType>;
-
-	Parameter = ToInternalSplineSpace(Parameter);
-	return AttrType();//InternalSpline.EvaluateAttribute<AttrType>(Name, Parameter);
-}
+	void ValidateRotScale() const;
+};
 
 template <typename AttrType>
 int32 FNewSpline::AddAttributeValue(FName Name, const AttrType& Value, float Parameter)
 {
 	using AttrSplineType = TTangentBezierSpline<AttrType>;
+	using ControlPointType = TTangentBezierControlPoint<AttrType>;
+	
+	if (AttrSplineType* Channel = GetTypedAttributeChannel<AttrSplineType>(Name))
+	{
+		ControlPointType ControlPoint = ControlPointType(Value);
+		int32 NumPoints = Channel->GetNumPoints();
+		
+		// Special case for empty spline
+		if (NumPoints == 0)
+		{
+			Channel->AppendPoint(ControlPoint);
+			return 0;
+		}
+		else if (NumPoints == 1)
+		{
+			if (Parameter > 0)
+			{
+				Channel->AppendPoint(ControlPoint);
+				Channel->Reparameterize(EParameterizationPolicy::Uniform);
+				return 1;
+			}
+			else
+			{
+				Channel->PrependPoint(ControlPoint);
+				Channel->Reparameterize(EParameterizationPolicy::Uniform);
+				return 0;
+			}
+		}
+		
+		float InternalParameter = ToInternalSplineSpace(Parameter);
+		float LocalT;
+		int32 SegmentIndex = GetSpline().FindSegmentIndex(InternalParameter, LocalT);
+		
+		int NewPointIndex = Channel->InsertPointAtSegmentParam(SegmentIndex, LocalT, ControlPoint);
+		Channel->Reparameterize(EParameterizationPolicy::Uniform);
+		return NewPointIndex;
+	}
 
-	Parameter = ToInternalSplineSpace(Parameter);
-	return -1;//InternalSpline.AddAttributeValue<AttrType>(Name, Value, Parameter);
+	return INDEX_NONE;
 }
 
 template <typename AttrType>
@@ -351,13 +377,26 @@ int32 FNewSpline::NumAttributeValues(FName Name) const
 template <typename AttrType>
 int32 FNewSpline::SetAttributeParameter(const FName& Name, int32 Index, float Parameter)
 {
-	return -1;//InternalSpline.SetAttributeParameter(Name, Index, ToInternalSplineSpace(Parameter));
+	using AttrSplineType = TTangentBezierSpline<AttrType>;
+	
+	if (AttrSplineType* Channel = GetTypedAttributeChannel<AttrSplineType>(Name))
+	{
+		
+		return Channel->SetParameter(Index, MapParameterToChildSpace(Name, ToInternalSplineSpace(Parameter)));
+	}
+	return INDEX_NONE;
 }
 
 template <typename AttrType>
 float FNewSpline::GetAttributeParameter(const FName& Name, int32 Index) const
 {
-	return -1;//FromInternalSplineSpace(InternalSpline.GetAttributeParameter(Name, Index));
+	using AttrSplineType = TTangentBezierSpline<AttrType>;
+	
+	if (AttrSplineType* Channel = GetTypedAttributeChannel<AttrSplineType>(Name))
+	{
+		return FromInternalSplineSpace(MapParameterFromChildSpace(Name, Channel->GetParameter(Index)));
+	}
+	return 0.0f;
 }
 
 /** FLegacySpline Definition */
@@ -474,6 +513,7 @@ FSpline::FSpline()
 : CurrentImplementation(UE::Spline::GImplementation)
 , Version(0xffffffff)
 {
+	
 #if WITH_EDITOR
 	PreviousImplementation = 0;
 #endif
@@ -741,7 +781,7 @@ FVector FSpline::Evaluate(float Param) const
 	{
 	default: return FVector();
 	case 1: return LegacyData->Evaluate(Param);
-	case 2: return NewData->Evaluate(Param);
+	case 2: return NewData->EvaluatePosition(Param);
 	}
 }
 
@@ -804,7 +844,7 @@ FVector FSpline::EvaluateScale(float Param) const
 
 bool FSpline::HasAttributeChannel(FName AttributeName) const
 {
-	if (NewData)
+	if (NewData && SupportsAttributes())
 	{
 		return NewData->HasAttributeChannel(AttributeName);
 	}
@@ -815,9 +855,13 @@ bool FSpline::HasAttributeChannel(FName AttributeName) const
 TArray<FName> FSpline::GetFloatPropertyChannels() const
 {
 	TArray<FName> ChannelNames;
-	if (NewData)
+	if (NewData && SupportsAttributes())
 	{
-		ChannelNames = NewData->GetAttributeChannelNames<float>();
+		ChannelNames = NewData->GetAttributeChannelNamesByValueType<float>();
+
+		// Don't report these internal attribute channels to the caller
+		ChannelNames.Remove(NewData->ParamToLengthAttrName);
+		ChannelNames.Remove(NewData->LengthToParamAttrName);
 	}
 	return ChannelNames;
 }
@@ -825,9 +869,12 @@ TArray<FName> FSpline::GetFloatPropertyChannels() const
 TArray<FName> FSpline::GetVectorPropertyChannels() const
 {
 	TArray<FName> ChannelNames;
-	if (NewData)
+	if (NewData && SupportsAttributes())
 	{
-		ChannelNames = NewData->GetAttributeChannelNames<FVector>();
+		ChannelNames = NewData->GetAttributeChannelNamesByValueType<FVector>();
+
+		// Don't report these internal attribute channels to the caller
+		ChannelNames.Remove(NewData->ScaleAttrName);
 	}
 	return ChannelNames;
 }
@@ -835,7 +882,7 @@ TArray<FName> FSpline::GetVectorPropertyChannels() const
 template <typename AttrType>
 float FSpline::GetAttributeParameter(int32 Index, const FName& Name) const
 {
-	if (NewData)
+	if (NewData && SupportsAttributes())
 	{
 		return NewData->GetAttributeParameter<AttrType>(Name, Index);
 	}
@@ -848,7 +895,7 @@ template ENGINE_API float FSpline::GetAttributeParameter<FVector>(int32 Index, c
 template <typename AttrType>
 int32 FSpline::SetAttributeParameter(int32 Index, float Parameter, const FName& Name)
 {
-	if (NewData)
+	if (NewData && SupportsAttributes())
 	{
 		return NewData->SetAttributeParameter<AttrType>(Name, Index, Parameter);
 	}
@@ -861,7 +908,7 @@ template ENGINE_API int32 FSpline::SetAttributeParameter<FVector>(int32 Index, f
 template <typename AttrType>
 int32 FSpline::NumAttributeValues(FName AttributeName) const
 {
-	if (NewData)
+	if (NewData && SupportsAttributes())
 	{
 		return NewData->NumAttributeValues<AttrType>(AttributeName);
 	}
@@ -874,7 +921,7 @@ template ENGINE_API int32 FSpline::NumAttributeValues<FVector>(FName AttributeNa
 template<typename AttrType>
 AttrType FSpline::GetAttributeValue(int32 Index, const FName& Name) const
 {
-	if (NewData)
+	if (NewData && SupportsAttributes())
 	{
 		return NewData->GetAttributeValue<AttrType>(Name, Index);
 	}
@@ -887,9 +934,9 @@ template ENGINE_API FVector FSpline::GetAttributeValue<FVector>(int32 Index, con
 template<typename AttrType>
 void FSpline::SetAttributeValue(int32 Index, const AttrType& Value, const FName& Name)
 {
-	if (NewData)
+	if (NewData && SupportsAttributes())
 	{
-	//	NewData->SetAttributeValue<AttrType>(Name, Index, Value);
+		NewData->SetAttributeValue<AttrType>(Name, Value, Index);
 	}
 }
 template ENGINE_API void FSpline::SetAttributeValue<float>(int32 Index, const float& Value, const FName& Name);
@@ -898,7 +945,7 @@ template ENGINE_API void FSpline::SetAttributeValue<FVector>(int32 Index, const
 template<typename AttrType>
 bool FSpline::CreateAttributeChannel(FName AttributeName) const
 {
-	if (NewData)
+	if (NewData && SupportsAttributes())
 	{
 		return NewData->CreateAttributeChannel<TTangentBezierSpline<AttrType>>(AttributeName) != nullptr;
 	}
@@ -911,7 +958,7 @@ template ENGINE_API bool FSpline::CreateAttributeChannel<FVector>(FName Attribut
 template<typename AttrType>
 int32 FSpline::AddAttributeValue(float Param, const AttrType& Value, FName AttributeName) const
 {
-	if (NewData)
+	if (NewData && SupportsAttributes())
 	{
 		return NewData->AddAttributeValue(AttributeName, Value, Param);
 	}
@@ -924,7 +971,7 @@ template ENGINE_API int32 FSpline::AddAttributeValue<FVector>(float Param, const
 template <typename AttrType>
 void FSpline::RemoveAttributeValue(int32 Index, FName AttributeName)
 {
-	if (NewData)
+	if (NewData && SupportsAttributes())
 	{
 		NewData->RemoveAttributeValue<AttrType>(AttributeName, Index);
 	}
@@ -935,9 +982,9 @@ template ENGINE_API void FSpline::RemoveAttributeValue<FVector>(int32 Index, FNa
 template<typename AttrType>
 AttrType FSpline::EvaluateAttribute(float Param, FName AttributeName) const
 {
-	if (NewData)
+	if (NewData && SupportsAttributes())
 	{
-		return NewData->EvaluateAttribute<AttrType>(AttributeName, Param);
+		return NewData->EvaluateAttribute<AttrType>(AttributeName, NewData->ToInternalSplineSpace(Param));
 	}
 	return AttrType(0);
 }
@@ -1084,14 +1131,22 @@ bool FSpline::ExportTextItem(FString& ValueStr, FSpline const& DefaultValue, cla
     FMemoryWriter MemWriter(SplineWriteBuffer);
     SerializeSave(MemWriter);
 
-    FString Base64String = FBase64::Encode(SplineWriteBuffer);
+	if (UE::Spline::GEncodeAsHex)
+	{
+		FString HexStr = FString::FromHexBlob(SplineWriteBuffer.GetData(), SplineWriteBuffer.Num());
+		ValueStr = FString::Printf(TEXT("SplineData SplineDataLen=%d SplineData=%s\r\n"), HexStr.Len(), *HexStr);
+	}
+	else
+	{
+		FString Base64String = FBase64::Encode(SplineWriteBuffer);
 
-    // Base64 encoding uses the '/' character, but T3D interprets '//' as some kind of
-    // terminator (?). If it occurs then the string passed to ImportTextItem() will
-    // come back as full of nullptrs. So we will swap in '-' here, and swap back to '/' in ImportTextItem()
-	Base64String.ReplaceCharInline('/', '-');
+		// Base64 encoding uses the '/' character, but T3D interprets '//' as some kind of
+		// terminator (?). If it occurs then the string passed to ImportTextItem() will
+		// come back as full of nullptrs. So we will swap in '-' here, and swap back to '/' in ImportTextItem()
+		Base64String.ReplaceCharInline('/', '-');
 
-	ValueStr = FString::Printf(TEXT("SplineData SplineDataLen=%d SplineData=%s\r\n"), Base64String.Len(), *Base64String);
+		ValueStr = FString::Printf(TEXT("SplineData SplineDataLen=%d SplineData=%s\r\n"), Base64String.Len(), *Base64String);
+	}
 	
 	return true;
 }
@@ -1118,7 +1173,18 @@ bool FSpline::ImportTextItem(const TCHAR*& SourceText, int32 PortFlags, UObject*
 				SplineData.ReplaceCharInline('-', '/');
 
 				TArray<uint8> SplineReadBuffer;
-				bool bDecoded = FBase64::Decode(SplineData, SplineReadBuffer);
+				bool bDecoded;
+				
+				if (UE::Spline::GEncodeAsHex)
+				{
+					SplineReadBuffer.SetNum(SplineDataLen);
+					bDecoded = FString::ToHexBlob(SplineData, SplineReadBuffer.GetData(), SplineDataLen);
+				}
+				else
+				{
+					bDecoded = FBase64::Decode(SplineData, SplineReadBuffer);
+				}
+
 				if (bDecoded)
 				{
 					FMemoryReader MemReader(SplineReadBuffer);
@@ -1302,11 +1368,20 @@ FNewSpline::FNewSpline()
 	CreateAttributeChannel<TTangentBezierSpline<float>>(LengthToParamAttrName);
     
     ResetToDefault();
+	
+	ValidateRotScale();
 }
 
 FNewSpline::FNewSpline(const FNewSpline& Other)
 {
+	CreateAttributeChannel<TTangentBezierSpline<FQuat>>(RotationAttrName);
+	CreateAttributeChannel<TTangentBezierSpline<FVector>>(ScaleAttrName);
+	CreateAttributeChannel<TTangentBezierSpline<float>>(ParamToLengthAttrName);
+	CreateAttributeChannel<TTangentBezierSpline<float>>(LengthToParamAttrName);
+	
     *this = Other;
+
+	ValidateRotScale();
 }
 
 FNewSpline::FNewSpline(const FLegacySpline& Other)
@@ -1329,13 +1404,17 @@ FNewSpline::FNewSpline(const FLegacySpline& Other)
 
 	// initialize spline
 	GetSpline().SetControlPoints(Points);
+	ResetRotation();
+	ResetScale();
 	
 	// setup rotation/scale attributes
 	for (int32 Idx = 0; Idx < GetSpline().GetNumPoints(); Idx++)
 	{    		
-		AddAttributeValue<FQuat>(RotationAttrName, Other.GetRotation(Idx), static_cast<float>(Idx));
-		AddAttributeValue<FVector>(ScaleAttrName, Other.GetScale(Idx), static_cast<float>(Idx));
+		SetAttributeValue<FQuat>(RotationAttrName, Other.GetRotation(Idx), Idx);
+		SetAttributeValue<FVector>(ScaleAttrName, Other.GetScale(Idx), Idx);
 	}
+	
+	ValidateRotScale();
 
 	FSpline::FUpdateSplineParams UpdateParams;
 	UpdateParams.bClosedLoop = Other.GetSplinePointsPosition().bIsLooped;
@@ -1365,17 +1444,21 @@ FNewSpline::FNewSpline(const FSplineCurves& Other)
 	
 	// initialize spline
 	GetSpline().SetControlPoints(Points);
+	ResetRotation();
+	ResetScale();
 	
     // setup rotation/scale attributes
     for (int32 Idx = 0; Idx < GetSpline().GetNumPoints(); Idx++)
     {
-    	const FQuat& Rotation = Other.Rotation.Points[Idx].OutVal;
-        const FVector& Scale = Other.Scale.Points[Idx].OutVal;
+		const FQuat& Rotation = Other.Rotation.Points[Idx].OutVal;
+		const FVector& Scale = Other.Scale.Points[Idx].OutVal;
     	
-    	AddAttributeValue<FQuat>(RotationAttrName, Rotation, static_cast<float>(Idx));
-    	AddAttributeValue<FVector>(ScaleAttrName, Scale, static_cast<float>(Idx));
+		SetAttributeValue<FQuat>(RotationAttrName, Rotation, Idx);
+		SetAttributeValue<FVector>(ScaleAttrName, Scale, Idx);
     }
 
+	ValidateRotScale();
+	
 	FSpline::FUpdateSplineParams UpdateParams;
 	UpdateParams.bClosedLoop = Other.Position.bIsLooped;
 	
@@ -1387,13 +1470,10 @@ FNewSpline& FNewSpline::operator=(const FNewSpline& Other)
 {
     if (this != &Other)
     {
-        InternalSpline = Other.InternalSpline;
-    	GetSpline().SetTangentModes(Other.GetSpline().GetTangentModes());
-    	ReparamStepsPerSegment = Other.ReparamStepsPerSegment;
-    	GetSpline().SetTension(Other.GetSpline().GetTension());
-    	GetSpline().SetStationaryEndpoints(Other.GetSpline().IsStationaryEndpoints());
+		TMultiSpline::operator=(Other);
+		ReparamStepsPerSegment = Other.ReparamStepsPerSegment;
 
-    	UpdateSpline();
+		UpdateSpline();
     }
 	
     return *this;
@@ -1497,7 +1577,7 @@ void FNewSpline::SetClosedLoop(bool bInClosedLoop)
 
 void FNewSpline::Reset()
 {
-	InternalSpline.Clear();
+	Clear();
 	GetSpline().SetTangentModes(TArray<ETangentMode>());
 	PositionCurve.Reset();
 	RotationCurve.Reset();
@@ -1523,6 +1603,9 @@ void FNewSpline::ResetRotation()
     {    		
     	AddAttributeValue<FQuat>(RotationAttrName, FQuat::Identity, static_cast<float>(Idx));
     }
+
+	TTangentBezierSpline<FQuat>* RotSpline = GetTypedAttributeChannel<TTangentBezierSpline<FQuat>>(RotationAttrName);
+	RotSpline->SetKnotVector(GetSpline().GetKnotVector());
 }
 
 void FNewSpline::ResetScale()
@@ -1533,32 +1616,84 @@ void FNewSpline::ResetScale()
     {    		
         AddAttributeValue<FVector>(ScaleAttrName, FVector::OneVector, static_cast<float>(Idx));
     }
+
+	TTangentBezierSpline<FVector>* ScaleSpline = GetTypedAttributeChannel<TTangentBezierSpline<FVector>>(ScaleAttrName);
+	ScaleSpline->SetKnotVector(GetSpline().GetKnotVector());
 }
 
 void FNewSpline::AddPoint(const FSplinePoint& Point)
 {
+	ValidateRotScale();
+	ON_SCOPE_EXIT
+	{
+		ValidateRotScale();
+	};
+	
 	FTangentBezierControlPoint ControlPoint = ConvertToTangentBezierControlPoint(Point);
 	const int32 NumPoints = GetSpline().GetNumPoints();
     
+	// Convert InputKey to parameter value
+	int32 Index = FMath::FloorToInt(Point.InputKey);
+	float Fraction = Point.InputKey - Index;
+	
 	// Special case for empty spline
-	if (NumPoints == 0)
+	if (NumPoints == 0 || (Index >= NumPoints && !IsClosedLoop()))
 	{
 		GetSpline().AppendPoint(ControlPoint);
-		UpdatePointAttributes(Point, 0);
+
+		TTangentBezierSpline<FVector>* ScaleChild = GetTypedAttributeChannel<TTangentBezierSpline<FVector>>(ScaleAttrName);
+		ScaleChild->AppendPoint(TTangentBezierSpline<FVector>::FTangentBezierControlPoint(FVector::OneVector));
+		ScaleChild->SetKnotVector(GetSpline().GetKnotVector());
+		
+		TTangentBezierSpline<FQuat>* RotChild = GetTypedAttributeChannel<TTangentBezierSpline<FQuat>>(RotationAttrName);
+		RotChild->AppendPoint(TTangentBezierSpline<FQuat>::FTangentBezierControlPoint(FQuat::Identity));
+		RotChild->SetKnotVector(GetSpline().GetKnotVector());
+		
+		RebuildLegacyCurves();
 		return;
 	}
-	// Convert InputKey to parameter value
-	int32 BaseIndex = FMath::FloorToInt(Point.InputKey);
-	float Fraction = Point.InputKey - BaseIndex;
-	    
+	else if (Index == 0)
+    {
+    	GetSpline().PrependPoint(ControlPoint);
+    	
+    	TTangentBezierSpline<FVector>* ScaleChild = GetTypedAttributeChannel<TTangentBezierSpline<FVector>>(ScaleAttrName);
+    	ScaleChild->PrependPoint(TTangentBezierSpline<FVector>::FTangentBezierControlPoint(FVector::OneVector));
+    	ScaleChild->SetKnotVector(GetSpline().GetKnotVector());
+    
+    	TTangentBezierSpline<FQuat>* RotChild = GetTypedAttributeChannel<TTangentBezierSpline<FQuat>>(RotationAttrName);
+    	RotChild->PrependPoint(TTangentBezierSpline<FQuat>::FTangentBezierControlPoint(FQuat::Identity));
+    	RotChild->SetKnotVector(GetSpline().GetKnotVector());
+    	
+    	RebuildLegacyCurves();
+    	return;
+    }
+	
+	
 	// Insert at calculated parameter
-	int32 InsertedIndex = GetSpline().InsertPointAtSegmentParam(BaseIndex, Fraction, ControlPoint);
+	int32 InsertedIndex = GetSpline().InsertPointAtSegmentParam(Index, Fraction, ControlPoint);
+	float NewPointInternalParam = GetSpline().GetParameter(InsertedIndex);
+	
+	float ScaleParam = MapParameterToChildSpace(ScaleAttrName, NewPointInternalParam);
+	TTangentBezierSpline<FVector>* ScaleChild = GetTypedAttributeChannel<TTangentBezierSpline<FVector>>(ScaleAttrName);
+	ScaleChild->InsertPointAtGlobalParam(ScaleParam, TTangentBezierSpline<FVector>::FTangentBezierControlPoint(FVector::OneVector));
+	ScaleChild->SetKnotVector(GetSpline().GetKnotVector());
+	
+	float RotParam = MapParameterToChildSpace(RotationAttrName, NewPointInternalParam);
+	TTangentBezierSpline<FQuat>* RotChild = GetTypedAttributeChannel<TTangentBezierSpline<FQuat>>(RotationAttrName);
+	RotChild->InsertPointAtGlobalParam(RotParam, TTangentBezierSpline<FQuat>::FTangentBezierControlPoint(FQuat::Identity));
+	RotChild->SetKnotVector(GetSpline().GetKnotVector());
 	
-	UpdatePointAttributes(Point, InsertedIndex);
+	RebuildLegacyCurves();
 }
 
 void FNewSpline::InsertPoint(const FSplinePoint& Point, int32 Index)
 {
+	ValidateRotScale();
+	ON_SCOPE_EXIT
+	{
+		ValidateRotScale();
+	};
+	
 	FTangentBezierControlPoint ControlPoint = ConvertToTangentBezierControlPoint(Point);
 	const int32 NumPoints = GetSpline().GetNumPoints();
     
@@ -1566,21 +1701,50 @@ void FNewSpline::InsertPoint(const FSplinePoint& Point, int32 Index)
 	if (NumPoints == 0 || (Index >= NumPoints && !IsClosedLoop()))
 	{
 		GetSpline().AppendPoint(ControlPoint);
-		UpdatePointAttributes(Point, NumPoints);
+		
+		TTangentBezierSpline<FVector>* ScaleChild = GetTypedAttributeChannel<TTangentBezierSpline<FVector>>(ScaleAttrName);
+		ScaleChild->AppendPoint(TTangentBezierSpline<FVector>::FTangentBezierControlPoint(FVector::OneVector));
+		ScaleChild->SetKnotVector(GetSpline().GetKnotVector());
+		
+		TTangentBezierSpline<FQuat>* RotChild = GetTypedAttributeChannel<TTangentBezierSpline<FQuat>>(RotationAttrName);
+		RotChild->AppendPoint(TTangentBezierSpline<FQuat>::FTangentBezierControlPoint(FQuat::Identity));
+		RotChild->SetKnotVector(GetSpline().GetKnotVector());
+		
+		RebuildLegacyCurves();
 		return;
 	}
 	else if (Index == 0)
 	{
 		GetSpline().PrependPoint(ControlPoint);
-		UpdatePointAttributes(Point, 0);
+		
+		TTangentBezierSpline<FVector>* ScaleChild = GetTypedAttributeChannel<TTangentBezierSpline<FVector>>(ScaleAttrName);
+		ScaleChild->PrependPoint(TTangentBezierSpline<FVector>::FTangentBezierControlPoint(FVector::OneVector));
+		ScaleChild->SetKnotVector(GetSpline().GetKnotVector());
+	
+		TTangentBezierSpline<FQuat>* RotChild = GetTypedAttributeChannel<TTangentBezierSpline<FQuat>>(RotationAttrName);
+		RotChild->PrependPoint(TTangentBezierSpline<FQuat>::FTangentBezierControlPoint(FQuat::Identity));
+		RotChild->SetKnotVector(GetSpline().GetKnotVector());
+		
+		RebuildLegacyCurves();
 		return;
 	}
 	
 	Index = FMath::Clamp(Index, 0, GetSpline().GetNumPoints());
 	
-	int32 InsertedSegmentIndex = GetSpline().InsertPointAtPosition(Index, ControlPoint);
+	int32 InsertedIndex = GetSpline().InsertPointAtPosition(Index, ControlPoint);
+	float NewPointInternalParam = GetSpline().GetParameter(InsertedIndex);
+	
+	float ScaleParam = MapParameterToChildSpace(ScaleAttrName, NewPointInternalParam);
+	TTangentBezierSpline<FVector>* ScaleChild = GetTypedAttributeChannel<TTangentBezierSpline<FVector>>(ScaleAttrName);
+	ScaleChild->InsertPointAtGlobalParam(ScaleParam, TTangentBezierSpline<FVector>::FTangentBezierControlPoint(FVector::OneVector));
+	ScaleChild->SetKnotVector(GetSpline().GetKnotVector());
+	
+	float RotParam = MapParameterToChildSpace(RotationAttrName, NewPointInternalParam);
+	TTangentBezierSpline<FQuat>* RotChild = GetTypedAttributeChannel<TTangentBezierSpline<FQuat>>(RotationAttrName);
+	RotChild->InsertPointAtGlobalParam(RotParam, TTangentBezierSpline<FQuat>::FTangentBezierControlPoint(FQuat::Identity));
+	RotChild->SetKnotVector(GetSpline().GetKnotVector());
 	
-	UpdatePointAttributes(Point, InsertedSegmentIndex);
+	RebuildLegacyCurves();
 }
 
 FSplinePoint FNewSpline::GetPoint(int32 Index) const
@@ -1606,12 +1770,27 @@ FSplinePoint FNewSpline::GetPoint(int32 Index) const
 /** Removes a point from the spline */
 void FNewSpline::RemovePoint(int32 Index)
 {
+	ValidateRotScale();
+	ON_SCOPE_EXIT
+	{
+		ValidateRotScale();
+	};
+	
 	if (Index < 0 || Index >= GetSpline().GetNumPoints())
 	{
 		return;
 	}
-
+	
 	GetSpline().RemovePoint(Index);
+
+	TTangentBezierSpline<FVector>* ScaleChild = GetTypedAttributeChannel<TTangentBezierSpline<FVector>>(ScaleAttrName);
+	ScaleChild->RemovePoint(Index);
+	ScaleChild->SetKnotVector(GetSpline().GetKnotVector());
+	
+	TTangentBezierSpline<FQuat>* RotChild = GetTypedAttributeChannel<TTangentBezierSpline<FQuat>>(RotationAttrName);
+	RotChild->RemovePoint(Index);
+	RotChild->SetKnotVector(GetSpline().GetKnotVector());
+	
     RebuildLegacyCurves();
 }
 
@@ -1687,15 +1866,9 @@ void FNewSpline::SetRotation(int32 Index, const FQuat& InRotation)
     	return;
     }
 
-	// if rotation value not added to attribute channel, add it
-	if (Index >= NumAttributeValues<FQuat>(RotationAttrName))
-	{
-		AddAttributeValue<FQuat>(RotationAttrName, InRotation, GetSpline().GetParameterAt(Index));
-	}
-	else
-	{
-		SetAttributeValue<FQuat>(RotationAttrName, InRotation, Index);
-	}
+	ValidateRotScale();
+
+	SetAttributeValue<FQuat>(RotationAttrName, InRotation, Index);
 	
     RebuildLegacyCurves();
 }
@@ -1717,15 +1890,9 @@ void FNewSpline::SetScale(int32 Index, const FVector& InScale)
     	return;
     }
 
-	// if scale value not added to attribute channel, add it
-	if (Index >= NumAttributeValues<float>(ScaleAttrName))
-	{
-		AddAttributeValue<FVector>(ScaleAttrName, InScale, GetSpline().GetParameterAt(Index));
-	}
-	else
-	{
-		SetAttributeValue<FVector>(ScaleAttrName, InScale, Index);
-	}
+	ValidateRotScale();
+
+	SetAttributeValue<FVector>(ScaleAttrName, InScale, Index);
     
     RebuildLegacyCurves();
 }
@@ -1768,8 +1935,10 @@ float FNewSpline::GetParameterAtIndex(int32 Index) const
     {
     	return 0.f;
     }
-    
-    return static_cast<float>(Index);
+
+	return FromInternalSplineSpace(GetSpline().GetParameter(Index));
+	
+    //return static_cast<float>(Index);
 }
 
 float FNewSpline::GetParameterAtDistance(float Distance) const
@@ -1816,10 +1985,9 @@ float FNewSpline::FindNearestOnSegment(const FVector& InLocation, int32 SegmentI
     }
 }
 
-FVector FNewSpline::EvaluateImpl(float Parameter) const
+FVector FNewSpline::EvaluatePosition(float Parameter) const
 {
-	Parameter = ToInternalSplineSpace(Parameter);
-    return GetSpline().GetValue(Parameter);
+	return Evaluate(ToInternalSplineSpace(Parameter));
 }
 
 FVector FNewSpline::EvaluateDerivative(float Parameter) const
@@ -1830,12 +1998,12 @@ FVector FNewSpline::EvaluateDerivative(float Parameter) const
 
 FQuat FNewSpline::EvaluateRotation(float Parameter) const
 {
-    return EvaluateAttribute<FQuat>(RotationAttrName, Parameter);
+    return EvaluateAttribute<FQuat>(RotationAttrName, ToInternalSplineSpace(Parameter));
 }
 
 FVector FNewSpline::EvaluateScale(float Parameter) const
 {
-    return EvaluateAttribute<FVector>(ScaleAttrName, Parameter);
+    return EvaluateAttribute<FVector>(ScaleAttrName, ToInternalSplineSpace(Parameter));
 }
 
 void FNewSpline::UpdateSpline()
@@ -1974,7 +2142,7 @@ void FNewSpline::UpdatePointAttributes(const FSplinePoint& Point, int32 PointInd
 	RebuildLegacyCurves();
 }
 
-float FNewSpline::ConvertIndexToParameter(int32 Index, float Fraction) const
+float FNewSpline::ConvertIndexToInternalParameter(int32 Index, float Fraction) const
 {
 	// Handle empty spline
 	if (GetSpline().GetNumPoints() <= 1)
@@ -1985,17 +2153,17 @@ float FNewSpline::ConvertIndexToParameter(int32 Index, float Fraction) const
     
 	// For exact index, just return parameter at that index
 	if (FMath::IsNearlyZero(Fraction))
-		return GetSpline().GetParameterAt(Index);
+		return GetSpline().GetParameter(Index);
     
 	// For fractional indices, interpolate
 	int32 NextIndex = FMath::Min(Index + 1, GetSpline().GetNumPoints() - 1);
-	float StartParam = GetSpline().GetParameterAt(Index);
-	float EndParam = GetSpline().GetParameterAt(NextIndex);
+	float StartParam = GetSpline().GetParameter(Index);
+	float EndParam = GetSpline().GetParameter(NextIndex);
     
 	return FMath::Lerp(StartParam, EndParam, Fraction);
 }
 
-int32 FNewSpline::ConvertParameterToIndex(float Parameter) const
+int32 FNewSpline::ConvertInternalParameterToNearestPointIndex(float Parameter) const
 {
 	// Handle empty spline
 	if (GetSpline().GetNumPoints() <= 1)
@@ -2004,8 +2172,8 @@ int32 FNewSpline::ConvertParameterToIndex(float Parameter) const
 	// Find segment containing this parameter
 	for (int32 i = 0; i < GetSpline().GetNumPoints() - 1; i++)
 	{
-		float StartParam = GetSpline().GetParameterAt(i);
-		float EndParam = GetSpline().GetParameterAt(i + 1);
+		float StartParam = GetSpline().GetParameter(i);
+		float EndParam = GetSpline().GetParameter(i + 1);
         
 		if (Parameter >= StartParam && Parameter <= EndParam)
 		{
@@ -2016,9 +2184,39 @@ int32 FNewSpline::ConvertParameterToIndex(float Parameter) const
 	}
     
 	// Out of range
-	return (Parameter < GetSpline().GetParameterAt(0)) ? 0 : (GetSpline().GetNumPoints() - 1);
+	return (Parameter < GetSpline().GetParameter(0)) ? 0 : (GetSpline().GetNumPoints() - 1);
 }
 
+void FNewSpline::ValidateRotScale() const
+{
+	// Validate child existence
+	TTangentBezierSpline<FQuat>* RotSpline = GetTypedAttributeChannel<TTangentBezierSpline<FQuat>>(RotationAttrName);
+	TTangentBezierSpline<FVector>* ScaleSpline = GetTypedAttributeChannel<TTangentBezierSpline<FVector>>(ScaleAttrName);
+	ensure (RotSpline && ScaleSpline);
+
+	if (RotSpline && ScaleSpline)
+	{
+		// Validate number of points
+		int32 NumControlPoints = GetSpline().GetNumPoints();
+		int32 NumRotPoints = RotSpline->GetNumPoints();
+		int32 NumScalePoints = ScaleSpline->GetNumPoints();
+		ensure(NumControlPoints == NumRotPoints && NumRotPoints == NumScalePoints);
+
+		// Validate control point -> attr point mapping
+		for (int32 i = 0; i < NumControlPoints; i++)
+		{
+			float InternalParam = GetSpline().GetParameter(i);
+		
+			float ExpectedRotParam = MapParameterToChildSpace(RotationAttrName, InternalParam);
+			float ActualRotParam = RotSpline->GetParameter(i);
+			ensure(FMath::IsNearlyEqual(ExpectedRotParam, ActualRotParam, UE_KINDA_SMALL_NUMBER));
+		
+			float ExpectedScaleParam = MapParameterToChildSpace(ScaleAttrName, InternalParam);
+			float ActualScaleParam = ScaleSpline->GetParameter(i);
+			ensure(FMath::IsNearlyEqual(ExpectedScaleParam, ActualScaleParam, UE_KINDA_SMALL_NUMBER));
+		}
+	}
+}
 
 /** FLegacySpline Implementation */
 
diff --git a/Engine/Source/Runtime/GeometryCore/Public/Splines/BSpline.h b/Engine/Source/Runtime/GeometryCore/Public/Splines/BSpline.h
--- a/Engine/Source/Runtime/GeometryCore/Public/Splines/BSpline.h
+++ b/Engine/Source/Runtime/GeometryCore/Public/Splines/BSpline.h
@@ -49,7 +49,7 @@ class TBSpline :
 	private TSelfRegisteringSpline<TBSpline<VALUETYPE, DEGREE>, VALUETYPE>
 {
 public:
-    static constexpr int32 Degree = DEGREE;
+	static constexpr int32 Degree = DEGREE;
 	using ValueType = typename TSplineInterface<VALUETYPE>::ValueType;
 
 	DECLARE_SPLINE_TYPE_ID(
@@ -57,117 +57,131 @@ public:
 		*TSplineValueTypeTraits<VALUETYPE>::Name.ToString()
 	);
 	
-    TBSpline() = default;
-    virtual ~TBSpline() override = default;
+	TBSpline() = default;
+	virtual ~TBSpline() override = default;
 
-    // ISplineInterface Implementation
-    virtual void Clear() override { Values.Empty(); }
-    virtual bool Serialize(FArchive& Ar) override
-    {
-    	// Call immediate parent's Serialize
-    	if (!TSplineInterface<ValueType>::Serialize(Ar))
-    	{
-    		return false;
-    	}
+	// ISplineInterface Implementation
+	virtual void Clear() override
+	{
+		Values.Empty();
+		Knots.Empty();
+	}
+	virtual bool Serialize(FArchive& Ar) override
+	{
+		// Call immediate parent's Serialize
+		if (!TSplineInterface<ValueType>::Serialize(Ar))
+		{
+			return false;
+		}
     	
-    	Ar << Values;
-    	Ar << Knots;
-    	Ar << bIsClosedLoop;
+		Ar << Values;
+		Ar << Knots;
+		Ar << bIsClosedLoop;
 		Ar << bClampEnds;
 		
-	    return true;
-    }
+		return true;
+	}
 	bool operator==(const TBSpline& Other) const
-    {
-    	return	Values == Other.Values &&
-    			Knots == Other.Knots &&
-    			bIsClosedLoop == Other.bIsClosedLoop &&
-    			bClampEnds == Other.bClampEnds;
-    }
+	{
+		return	Values == Other.Values &&
+				Knots == Other.Knots &&
+				bIsClosedLoop == Other.bIsClosedLoop &&
+				bClampEnds == Other.bClampEnds;
+	}
 	friend FArchive& operator<<(FArchive& Ar, TBSpline& BSpline)
-    {
-    	BSpline.Serialize(Ar);
-    	return Ar;	
-    }
+	{
+		BSpline.Serialize(Ar);
+		return Ar;	
+	}
 	
-    virtual FInterval1f GetParameterSpace() const override { return GetKnotRange(); }
-    virtual void SetClosedLoop(bool bClosed) override { bIsClosedLoop = bClosed; }
-    virtual bool IsClosedLoop() const override { return bIsClosedLoop; }
+	virtual FInterval1f GetParameterSpace() const override { return GetKnotRange(); }
+	virtual void SetClosedLoop(bool bClosed) override { bIsClosedLoop = bClosed; }
+	virtual bool IsClosedLoop() const override { return bIsClosedLoop; }
 	virtual TUniquePtr<ISplineInterface> Clone() const override
-    {
-    	TUniquePtr<TBSpline<VALUETYPE, DEGREE>> Clone = MakeUnique<TBSpline<VALUETYPE, DEGREE>>();
+	{
+		TUniquePtr<TBSpline<VALUETYPE, DEGREE>> Clone = MakeUnique<TBSpline<VALUETYPE, DEGREE>>();
     
-    	// Copy all member variables
-    	Clone->Values = this->Values;
-    	Clone->Knots = this->Knots;
-    	Clone->bIsClosedLoop = this->bIsClosedLoop;
-    	Clone->bClampEnds = this->bClampEnds;
-    	Clone->bEnableKnotUpdates = this->bEnableKnotUpdates;
+		// Copy all member variables
+		Clone->Values = this->Values;
+		Clone->Knots = this->Knots;
+		Clone->bIsClosedLoop = this->bIsClosedLoop;
+		Clone->bClampEnds = this->bClampEnds;
+		Clone->bEnableKnotUpdates = this->bEnableKnotUpdates;
     
-    	// Copy infinity modes
-    	Clone->PreInfinityMode = this->PreInfinityMode;
-    	Clone->PostInfinityMode = this->PostInfinityMode;
+		// Copy infinity modes
+		Clone->PreInfinityMode = this->PreInfinityMode;
+		Clone->PostInfinityMode = this->PostInfinityMode;
     
-    	return Clone;
-    }
-   // virtual bool IsEqual(const ISplineInterface& Other) const override { /*todo*/ }
+		return Clone;
+	}
+	// virtual bool IsEqual(const ISplineInterface& Other) const override { /*todo*/ }
 
-    // TSplineInterface Implementation
-    virtual ValueType EvaluateImpl(float Parameter) const override
-    {
-    	int32 NumValues = NumKeys();
-    	if (NumValues == 0) return ValueType();
-    	if (NumValues == 1) return GetValue(0);
+	// TSplineInterface Implementation
+	virtual ValueType EvaluateImpl(float Parameter) const override
+	{
+		int32 NumValues = NumKeys();
+		if (NumValues == 0) return ValueType();
+		if (NumValues == 1) return GetValue(0);
 
-    	// Get window using override or evaluator
-    	TArray<const ValueType*> Window = FindWindow(Parameter);
+		// Get window using override or evaluator
+		TArray<const ValueType*> Window = FindWindow(Parameter);
     	
-    	if (Window.Num() == 0) return ValueType();
+		if (Window.Num() == 0) return ValueType();
         
-    	return InterpolateWindow(Window, Parameter);
-    }
-    virtual float FindNearest(const ValueType& Point, float& OutSquaredDistance) const override
-        { return 0;/* todo */ }
+		return InterpolateWindow(Window, Parameter);
+	}
+	virtual float FindNearest(const ValueType& Point, float& OutSquaredDistance) const override
+	{ return 0;/* todo */ }
 
-    int32 NumKeys() const { return Values.Num(); }
-    const ValueType& GetValue(int32 Idx) const { return Values[Idx]; }
-    int32 AddValue(const ValueType& NewValue)  { return Values.Add(NewValue); }
+	int32 NumKeys() const { return Values.Num(); }
+	const ValueType& GetValue(int32 Idx) const { return Values[Idx]; }
+	int32 AddValue(const ValueType& NewValue)  { return Values.Add(NewValue); }
 	bool SetValue(int32 Idx, const ValueType& NewValue)
-    {
-        if (!Values.IsValidIndex(Idx))
-    	{
-    		return false;
-    	}
-        Values[Idx] = NewValue;
-        return true;
-    }
+	{
+		if (!Values.IsValidIndex(Idx))
+		{
+			return false;
+		}
+		Values[Idx] = NewValue;
+		return true;
+	}
 	int32 InsertValue(int32 Idx, const ValueType& NewValue)
-    {
-    	if (Idx < 0 || Idx > Values.Num())
-    	{
-    		return Values.Add(NewValue);
-    	}
-    	return Values.Insert(NewValue, Idx);
-    }
+	{
+		if (Idx < 0 || Idx > Values.Num())
+		{
+			return Values.Add(NewValue);
+		}
+		return Values.Insert(NewValue, Idx);
+	}
 	virtual bool RemoveValue(int32 Index) 
-    {
-    	if (Values.IsValidIndex(Index))
-    	{
-    		Values.RemoveAt(Index);
-    		return true;
-    	}
-    	return false;
-    }
+	{
+		if (Values.IsValidIndex(Index))
+		{
+			Values.RemoveAt(Index);
+			return true;
+		}
+		return false;
+	}
     
 	// Parameterization methods
-	// Returns parameter value for a control point index
-    virtual float GetParameterAtIndex(int32 Index) const
+	virtual bool SetParameter(int32 Index, float NewParameter)
+	{
+		UE_LOG(LogSpline, Warning, TEXT("SetParameter is disabled. Use Reparameterize() or a linear spline if you need manual control."));
+		return false;
+	}
+	
+	// Returns parameter value for a control point index using Greville Abscissae
+    virtual float GetParameter(int32 Index) const
 	{
-		// For BSplines, Knots are defined by knot vector
-		if (Index < 0 || Index >= Knots.Num() - Degree - 1)
+		if (Index < 0 || Index + Degree >= Knots.Num() - 1)
 			return 0.0f;
-            
-		return Knots[Index + Degree];
+
+		float Sum = 0.0f;
+		for (int j = 1; j <= Degree; ++j)
+		{
+			Sum += Knots[Index + j];
+		}
+		return Sum / static_cast<float>(Degree);
 	}
 
 	// Find the control point index and local parameter for global parameter
@@ -224,7 +238,9 @@ public:
     
 		return Index;
 	}
-
+	// Knot vector management
+	const TArray<float>& GetKnotVector() const { return Knots; }
+	void ResetKnotVector() { Knots.Reset(); }
 	/**
      * Gets the array of distinct knot values (without duplicates)
      * @return Array of distinct knot values in ascending order
@@ -297,10 +313,7 @@ public:
         return false;
     }
 	
-	// Knot vector management
-	void SetKnotVector(const TArray<float>& InKnots) { Knots = InKnots; }
-	const TArray<float>& GetKnotVector() const { return Knots; }
-	void ResetKnotVector() { Knots.Reset(); }
+	
     
 	/** 
      * Generates a uniform knot vector with equal spacing
@@ -582,7 +595,6 @@ public:
 
 	/**
 	 * Inserts a knot at the specified parameter value
-	 * Note: This preserves curve shape exactly but may change parameterization properties
 	 * 
 	 * @param Parameter - Parameter value where the knot should be inserted
 	 * @param Count - Number of times to insert the knot
@@ -710,7 +722,7 @@ public:
 	}
 	
 	/**
-	 * Sets a custom knot vector and switches to custom knot mode
+	 * Sets a custom knot vector 
 	 * @param NewKnots - Array of knot values. Must be non-decreasing and have correct length (n + p + 1)
 	 * @param NumValues - Number of values (control points) in the spline
 	 * @return True if the knot vector was valid and set successfully
@@ -719,9 +731,7 @@ public:
 	{
     	int32 NumValues = NumKeys();
 		// For closed loops, we repeat the first Degree points at the end
-		const int32 RequiredKnots = this->IsClosedLoop() ? 
-			(NumValues + 2 * Degree + 1) : // Extra knots for closure
-			(NumValues + Degree + 1);      // Standard case
+		const int32 RequiredKnots = GetExpectedNumKnots();    
 		
 		if (NewKnots.Num() < RequiredKnots)
 		{
diff --git a/Engine/Source/Runtime/GeometryCore/Public/Splines/LinearSpline.h b/Engine/Source/Runtime/GeometryCore/Public/Splines/LinearSpline.h
new file mode 100644
--- /dev/null
+++ b/Engine/Source/Runtime/GeometryCore/Public/Splines/LinearSpline.h
@@ -0,0 +1,54 @@
+// Copyright Epic Games, Inc. All Rights Reserved.
+
+#pragma once
+
+#include "BSpline.h"
+namespace UE
+{
+namespace Geometry
+{
+namespace Spline
+{
+
+template<typename T>
+class TLinearSpline : public TBSpline<T, 1>
+{
+public:
+	using Base = TBSpline<T, 1>;
+
+	DECLARE_SPLINE_TYPE_ID(TEXT("BSpline1"), *TSplineValueTypeTraits<T>::Name.ToString());
+
+	TLinearSpline()
+	{
+	}
+
+	virtual float GetParameter(int32 Index) const override
+	{
+		return (Index >= 0 && Index + 1 < this->Knots.Num())
+			? this->Knots[Index + 1]
+			: 0.0f;
+	}
+
+	virtual bool SetParameter(int32 Index, float NewParam) override
+	{
+		if (Index < 0 || Index + 1 >= this->Knots.Num())
+			return false;
+
+		float Prev = (Index > 0) ? this->Knots[Index] : -FLT_MAX;
+		float Next = (Index + 2 < this->Knots.Num()) ? this->Knots[Index + 2] : FLT_MAX;
+
+		if (NewParam < Prev || NewParam > Next)
+			return false;
+
+		this->Knots[Index + 1] = NewParam;
+		return true;
+	}
+
+};
+
+using FLinearSplineFloat = TLinearSpline<float>; 
+using FLinearSplineVector3d = TLinearSpline<FVector3d>;
+ 
+} // end namespace UE::Geometry::Spline
+} // end namespace UE::Geometry
+} // end namespace UE
\ No newline at end of file
diff --git a/Engine/Source/Runtime/GeometryCore/Public/Splines/MultiSpline.h b/Engine/Source/Runtime/GeometryCore/Public/Splines/MultiSpline.h
--- a/Engine/Source/Runtime/GeometryCore/Public/Splines/MultiSpline.h
+++ b/Engine/Source/Runtime/GeometryCore/Public/Splines/MultiSpline.h
@@ -33,18 +33,96 @@ private:
     /* Internal structure for child splines */
     struct FChildSpline
     {
+		// Default constructor
+    	FChildSpline() = default;
+    
+    	// Copy constructor - creates a deep copy of the spline
+    	FChildSpline(const FChildSpline& Other)
+			: MappingRangeSpace(Other.MappingRangeSpace)
+			, MappingRange(Other.MappingRange)
+    	{
+    		// Clone the spline if it exists
+    		if (Other.Spline)
+    		{
+    			Spline = Other.Spline->Clone();
+    		}
+    	}
+    
+    	// Move constructor
+    	FChildSpline(FChildSpline&& Other) = default;
+    
+    	// Copy assignment - creates a deep copy
+    	FChildSpline& operator=(const FChildSpline& Other)
+    	{
+    		if (this != &Other)
+    		{
+    			MappingRangeSpace = Other.MappingRangeSpace;
+    			MappingRange = Other.MappingRange;
+            
+    			// Clone the spline if it exists
+    			if (Other.Spline)
+    			{
+    				Spline = Other.Spline->Clone();
+    			}
+    			else
+    			{
+    				Spline.Reset();
+    			}
+    		}
+    		return *this;
+    	}
+    
+    	// Move assignment
+    	FChildSpline& operator=(FChildSpline&& Other) = default;
+		bool Serialize(FArchive& Ar)
+		{
+			Ar << MappingRangeSpace;
+		    Ar << MappingRange.Min;
+		    Ar << MappingRange.Max;
+
+		    // We must create the proper spline before we can serialize it.
+		    if (Ar.IsLoading())
+		    {
+    			FSplineTypeId::IdType TypeId;
+    			
+    			// We are about to read data from the archive that the spline itself will attempt to read,
+    			// so we have to restore the archive position after reading it.
+    			const int64 Pos = Ar.Tell();
+    			{
+    				int32 Version;
+    				Ar << Version;
+    				Ar << TypeId;
+    			}
+    			Ar.Seek(Pos);
+
+    			Spline = FSplineTypeRegistry::CreateSpline(TypeId);
+
+    			ensureAlways(Spline);	// if this failed, most likely the spline type has not been registered.
+		    }
+		    
+		    Ar << *Spline;
+		    
+		    return true;
+		}
+    	
         float MapParameterToChildSpace(float ParentParameter, const FInterval1f& ParentSpace) const
         {
             if (!Spline)
             {
                 return 0.f;
             }
+
+			if (FInterval1f ChildSpace = Spline->GetParameterSpace(); !ChildSpace.IsEmpty())
+			{
+				FInterval1f MappedParentRange = GetParentSpaceRange(ParentSpace);
     
-            FInterval1f ChildSpace = Spline->GetParameterSpace();
-            FInterval1f MappedParentRange = GetParentSpaceRange(ParentSpace);
-    
-            float T = MappedParentRange.GetT(ParentParameter);
-            return ChildSpace.Interpolate(T);
+				float T = MappedParentRange.GetT(ParentParameter);
+				return ChildSpace.Interpolate(T);
+			}
+			else
+			{
+				return 0.f;
+			}
         }
     
         float MapParameterFromChildSpace(float ChildParameter, const FInterval1f& ParentSpace) const
@@ -53,12 +131,18 @@ private:
             {
                 return 0.f;
             }
+
+			if (FInterval1f ChildSpace = Spline->GetParameterSpace(); !ChildSpace.IsEmpty())
+			{
+				FInterval1f MappedParentRange = GetParentSpaceRange(ParentSpace);
     
-            FInterval1f ChildSpace = Spline->GetParameterSpace();
-            FInterval1f MappedParentRange = GetParentSpaceRange(ParentSpace);
-    
-            float T = ChildSpace.GetT(ChildParameter);
-            return MappedParentRange.Interpolate(T);
+				float T = ChildSpace.GetT(ChildParameter);
+				return MappedParentRange.Interpolate(T);
+			}
+			else
+			{
+				return 0.f;
+			}
         }
 
         FInterval1f GetParentSpaceRange(const FInterval1f& ParentSpace) const
@@ -115,7 +199,35 @@ public:
 	
     SplineType& GetSpline() { return TSplineWrapper<SPLINETYPE>::InternalSpline; }
     const SplineType& GetSpline() const { return TSplineWrapper<SPLINETYPE>::InternalSpline; }
-    
+	
+	virtual void Clear() override
+    {
+		TSplineWrapper<SPLINETYPE>::Clear();
+
+		for (const FName& Name : GetAllAttributeChannelNames())
+		{
+			ClearAttributeChannel(Name);
+		}
+    }
+
+	virtual bool Serialize(FArchive& Ar) override
+	{
+		if (!TSplineWrapper<SPLINETYPE>::Serialize(Ar))
+		{
+			return false;
+		}
+
+		for (const FName& Name : GetAllAttributeChannelNames())
+		{
+			if (FChildSpline& Child = Children[Name]; !Child.Serialize(Ar))
+			{
+				return false;
+			}
+		}
+
+		return true;
+	}
+	
 	/**
 	 * Gets a strongly-typed attribute channel with the specific implementation type
 	 * 
@@ -335,17 +447,73 @@ public:
     {
         // Extract the value type from the implementation type
         using ValueType = typename ImplType::ValueType;
+
+    	// Check if the type is already registered
+    	uint32 TypeId = FSplineTypeRegistry::GetTypeId(ImplType::GetSplineTypeName(), 
+													  TSplineValueTypeTraits<ValueType>::Name);
+    	// If not registered, register it manually
+    	if (TypeId == 0)
+    	{
+    		// Get type information for registration
+    		TypeId = ImplType::GetStaticTypeId();
+    		FName ImplName = ImplType::GetSplineTypeName();
+    		FName ValueName = TSplineValueTypeTraits<ValueType>::Name;
+        
+    		// Register with a factory function that creates this specific type
+    		bool bSuccess = FSplineTypeRegistry::RegisterType(
+				TypeId,
+				ImplName,
+				ValueName,
+				[]() { return MakeUnique<ImplType>(); }
+			);
         
+    		if (!bSuccess)
+    		{
+    			UE_LOG(LogSpline, Error, TEXT("Failed to register spline type '%s' with value type '%s'"),
+					  *ImplName.ToString(), *ValueName.ToString());
+    			return nullptr;
+    		}
+    	}
+    	
         // Create the channel using the existing name-based method
-        if (!CreateAttributeChannelInternal<ValueType>(Name, ImplType::GetSplineTypeName()))
-        {
-            return nullptr;
-        }
+        CreateAttributeChannelInternal<ValueType>(Name, ImplType::GetSplineTypeName());
     
         // Return the typed implementation
         return GetTypedAttributeChannel<ImplType>(Name);
     }
 
+	template <typename AttrSplineType>
+	TArray<FName> GetAttributeChannelNamesBySplineType() const
+    {
+		TArray<FName> Names;
+
+		for (const TPair<FName, FChildSpline>& Pair : Children)
+		{
+			if (Pair.Value.Spline->GetTypeId() == AttrSplineType::GetStaticTypeId())
+			{
+				Names.Add(Pair.Key);
+			}
+		}
+
+		return Names;
+    }
+
+	template <typename AttrType>
+    TArray<FName> GetAttributeChannelNamesByValueType() const
+	{
+		TArray<FName> Names;
+
+		for (const TPair<FName, FChildSpline>& Pair : Children)
+		{
+			if (Pair.Value.Spline->GetValueTypeName() == TSplineValueTypeTraits<AttrType>::Name)
+			{
+				Names.Add(Pair.Key);
+			}
+		}
+
+		return Names;
+    }
+	
     /**
      * Clears all values from a specific attribute channel
      * @param Name - Name of the attribute channel to clear
@@ -368,11 +536,7 @@ public:
      */
     bool HasAttributeChannel(const FName& Name) const
     {
-        if (const FChildSpline* Channel = Children.Find(Name))
-        {
-            return true;
-        }
-        return false;
+		return Children.Contains(Name);
     }
     
     /**
diff --git a/Engine/Source/Runtime/GeometryCore/Public/Splines/PolyBezierSpline.h b/Engine/Source/Runtime/GeometryCore/Public/Splines/PolyBezierSpline.h
--- a/Engine/Source/Runtime/GeometryCore/Public/Splines/PolyBezierSpline.h
+++ b/Engine/Source/Runtime/GeometryCore/Public/Splines/PolyBezierSpline.h
@@ -450,7 +450,7 @@ public:
 	            {
 	            	
 	                // Points don't match, use AddBezierSegmentInternal with appropriately computed parameter
-	                float StartParam = GetParameterAtIndex(NumPoints - 1);
+	                float StartParam = GetParameter(NumPoints - 1);
 	            	this->AddBezierSegmentInternal(NumSegments, P0, P1, P2, P3, ParameterizationPolicy, StartParam);
 	            }
 	        }
@@ -726,7 +726,11 @@ public:
 	    }
 	    
 	    SegmentIndex = FMath::Clamp(SegmentIndex, 0, NumSegments - 1);
-	    LocalT = FMath::Clamp(LocalT, 0.0f, 1.0f);
+
+		// we clamp the LocalT to not be exactly 0 or 1, as this would cause
+		// the new point to be exactly the same as the start or end of the segment
+		// and creating possible degenerate segments with zero lengths.
+	    LocalT = FMath::Clamp(LocalT, UE_KINDA_SMALL_NUMBER, 1.0f - UE_KINDA_SMALL_NUMBER);
 	    
 	    // Get the segment's parameter range
 	    FInterval1f SegmentRange = GetSegmentParameterRange(SegmentIndex);
@@ -767,13 +771,6 @@ public:
 		const int32 BaseIndex = SegmentIndex * 4;
 		ValueType P0 = Base::GetValue(BaseIndex);
 		
-		if (P0 == Position)
-		{
-			UE_LOG(LogSpline, Warning, TEXT("Trying to insert duplicate control point."));
-			// No need to split, just return the index of the existing point
-			return BaseIndex;
-		}
-		
 		// Handle out-of-range Knots
 		if (Parameter <= ParameterRange.Min)
 		{
@@ -1133,7 +1130,7 @@ private:
 	{
 		// Append - create a point at the end
 		int32 LastIndex = Base::NumKeys() - 1;
-		float StartParam = Base::GetParameterAtIndex(LastIndex);
+		float StartParam = Base::GetParameter(LastIndex);
 		ValueType P0 = Base::GetValue(LastIndex);
         
 		// Calculate reasonable auto-tangents
@@ -1144,7 +1141,6 @@ private:
 		int32 NumSegments = Base::NumKeys() / 4;
 		// Add the segment
 		int32 NewIndex = AddBezierSegmentInternal(NumSegments, P0, P1, P2, Position, ParameterizationPolicy, StartParam);
-		check (NewIndex == NumSegments - 1)
 		return NewIndex;
 	}
 	
@@ -1331,10 +1327,13 @@ public:
 	virtual FInterval1f GetParameterSpace() const override
 	{
 		const TArray<float> Knots = Base::GetDistinctKnotValues();
-		return FInterval1f(Knots[0], Knots.Last());
+		
+		return Knots.Num() > 0
+			? FInterval1f(Knots[0], Knots.Last())
+			: FInterval1f::Empty();
 	}
 
-	virtual float GetParameterAtIndex(int32 Index) const override
+	virtual float GetParameter(int32 Index) const override
 	{
 		const TArray<float> Knots = Base::GetDistinctKnotValues();
 		if (Knots.IsEmpty())
@@ -1351,6 +1350,106 @@ public:
         return Range.Interpolate(T);
 	}
 
+	virtual bool SetParameter(int32 Index, float NewParameter) override
+	{
+		static constexpr int32 Degree = 3;
+
+		if (Index < 0 || Index >= Base::NumKeys() || Base::NumKeys() < 4)
+		{
+			return false;
+		}
+
+		const int32 PointsPerSegment = 4;
+		const int32 SegmentIndex = Index / PointsPerSegment;
+		const int32 LocalPointIndex = Index % PointsPerSegment;
+		int32 NumSegments = GetNumSegments();
+		if (LocalPointIndex == 0) // P0
+		{
+			float OldParam = GetSegmentParameterRange(SegmentIndex).Min;
+			int32 Multiplicity = (Base::IsClampedEnds() && SegmentIndex == 0) ? (Degree + 1) : Degree;
+			//Base::RemoveKnotsCompletely(OldParam, GetExpectedNumKnots());
+			Base::Knots.Remove(OldParam);
+			Base::InsertKnot(NewParameter, Multiplicity);
+
+			// Fix ordering
+			int32 Segment = SegmentIndex;
+			while (Segment > 0 && GetSegmentParameterRange(Segment).Min < GetSegmentParameterRange(Segment - 1).Min)
+			{
+				SwapSegments(Segment, Segment - 1);
+				--Segment;
+			}
+			while (Segment < NumSegments - 1 && GetSegmentParameterRange(Segment).Min > GetSegmentParameterRange(Segment + 1).Min)
+			{
+				SwapSegments(Segment, Segment + 1);
+				++Segment;
+			}
+			return true;
+		}
+		else if (LocalPointIndex == 3) // P3
+		{
+			float OldParam = GetSegmentParameterRange(SegmentIndex).Max;
+			int32 Multiplicity = (Base::IsClampedEnds() && SegmentIndex == GetNumSegments() - 1) ? (Degree + 1) : Degree;
+			// Base::RemoveKnotsCompletely(OldParam, GetExpectedNumKnots());
+			Base::Knots.Remove(OldParam);
+			Base::InsertKnot(NewParameter, Multiplicity);
+
+			int32 Segment = SegmentIndex;
+            while (Segment > 0 && GetSegmentParameterRange(Segment).Max < GetSegmentParameterRange(Segment - 1).Max)
+            {
+                SwapSegments(Segment, Segment - 1);
+                --Segment;
+            }
+            while (Segment < NumSegments - 1 && GetSegmentParameterRange(Segment).Max > GetSegmentParameterRange(Segment + 1).Max)
+            {
+                SwapSegments(Segment, Segment + 1);
+                ++Segment;
+            }
+			return true;
+		}
+
+		return false;
+	}
+	void SwapSegments(int32 SegmentA, int32 SegmentB)
+	{
+		static constexpr int32 Degree = 3; // Cubic
+		static constexpr int32 ControlPointsPerSegment = 4;
+		static constexpr int32 KnotsPerSegment = Degree; // Not degree+1 because knot overlap between segments
+
+		if (SegmentA == SegmentB || SegmentA < 0 || SegmentB < 0)
+		{
+			return;
+		}
+
+		const int32 NumSegments = GetNumSegments();
+		if (SegmentA >= NumSegments || SegmentB >= NumSegments)
+		{
+			return;
+		}
+
+		// Swap control points
+		int32 ControlIndexA = SegmentA * ControlPointsPerSegment;
+		int32 ControlIndexB = SegmentB * ControlPointsPerSegment;
+
+		for (int32 i = 0; i < ControlPointsPerSegment; ++i)
+		{
+			Base::Values.Swap(ControlIndexA + i, ControlIndexB + i);
+		}
+
+		// Swap repeated knots
+		// Important: due to clamping at endpoints, segments share some knots
+
+		// Degree + 1 knots define a segment (for Cubic: 4 knots)
+		int32 KnotStartA = Degree + SegmentA;
+		int32 KnotStartB = Degree + SegmentB;
+
+		for (int32 i = 0; i < Degree; ++i)
+		{
+			if ((KnotStartA + i) < Base::Knots.Num() && (KnotStartB + i) < Base::Knots.Num())
+			{
+				Base::Knots.Swap(KnotStartA + i, KnotStartB + i);
+			}
+		}
+	}
 	virtual int32 FindIndexForParameter(float Parameter, float& OutLocalParam) const override
 	{
 		TArray<float> Knots = Base::GetDistinctKnotValues();
@@ -1615,6 +1714,10 @@ public:
 		}
 		return 4 + (NumSegments - 1) * 3 + 4;
 	}
+	void SetKnotVector(const TArray<float>& NewKnots)
+	{
+		Base::SetCustomKnots(NewKnots);
+	}
 private:
 	
 	/**
diff --git a/Engine/Source/Runtime/GeometryCore/Public/Splines/TangentBezierSpline.h b/Engine/Source/Runtime/GeometryCore/Public/Splines/TangentBezierSpline.h
--- a/Engine/Source/Runtime/GeometryCore/Public/Splines/TangentBezierSpline.h
+++ b/Engine/Source/Runtime/GeometryCore/Public/Splines/TangentBezierSpline.h
@@ -183,12 +183,13 @@ public:
         {
             Tension = Other.Tension;
             TangentModes = Other.TangentModes;
+			bStationaryEndpoint = Other.bStationaryEndpoint;
             InternalSpline = Other.InternalSpline;
         }
         return *this;
     }
 
-    bool Serialize(FArchive& Ar)
+    virtual bool Serialize(FArchive& Ar) override
     {
 		// Call immediate parent's Serialize (TSplineWrapper)
 		if (!TSplineWrapper<TPolyBezierSpline<ValueType>>::Serialize(Ar))
@@ -731,6 +732,13 @@ public:
 		const FTangentBezierControlPoint& ControlPoint,
 		EParameterizationPolicy ParameterizationPolicy = EParameterizationPolicy::Centripetal)
 	{
+		int32 OutNewIndex;
+		// Special case: handle first two points
+		if (InitializeFirstSegment(ControlPoint, OutNewIndex))
+		{
+			return OutNewIndex;
+		}
+		
 		// Insert at that parameter with the exact position
 		// Have the internal spline handle the insertion
 		int32 InsertedPointIndex = InternalSpline.InsertPointAtSegmentParam(SegmentIndex, LocalT, ControlPoint.Position);
@@ -908,7 +916,7 @@ public:
 		{
 			// For other modes, we'll set up a reasonable default
 			// The UpdatePointTangents call below will update it properly
-			ValueType DefaultTangent = (1.0f - Tension) * (P3 - P0);
+			ValueType DefaultTangent = (P3 - P0) * (1.0f - Tension);
 			P2 = P3 - (DefaultTangent / 3.0f);
 		}
     
@@ -1013,7 +1021,7 @@ public:
 		{
 			// Removing an internal point
 			// Get the parameter of the point being removed
-			float PointParam = InternalSpline.GetParameterAtIndex((Index - 1) * 4 + 3);
+			float PointParam = InternalSpline.GetParameter((Index - 1) * 4 + 3);
         
 			// For an internal point, we need to remove exactly 4 control points:
 			// - P2 and P3 from the previous segment (where P3 is the point)
@@ -1186,7 +1194,7 @@ public:
                 const int32 LastSegment = InternalSpline.GetNumSegments() - 1;
                 ValueType P3 = InternalSpline.GetValue(LastSegment * 4 + 3);
                 ValueType P2 = InternalSpline.GetValue(LastSegment * 4 + 2);
-                return 3.0f * (P3 - P2);
+                return (P3 - P2) * 3.0f;
             }
             return ValueType();
         }
@@ -1197,7 +1205,7 @@ public:
         {
             ValueType P3 = InternalSpline.GetValue(BezierIndex + 3);
             ValueType P2 = InternalSpline.GetValue(BezierIndex + 2);
-            return 3.0f * (P3 - P2);
+            return (P3 - P2) * 3.0f;
         }
         return ValueType();
     }
@@ -1217,7 +1225,7 @@ public:
         {
             ValueType P0 = InternalSpline.GetValue(BezierIndex);
             ValueType P1 = InternalSpline.GetValue(BezierIndex + 1);
-            return 3.0f * (P1 - P0);
+            return (P1 - P0) * 3.0f;
         }
         return ValueType();
     }
@@ -1320,11 +1328,21 @@ public:
 		// Update tangents based on new mode
         UpdatePointTangents(Index, Mode);
 	}
-	
-	float GetParameterAt(int32 Index) const
+	virtual void Clear() override
 	{
-		return InternalSpline.GetParameterAtIndex(Index);
+		InternalSpline.Clear();
+		TangentModes.Empty();
 	}
+	float GetParameter(int32 Index) const
+	{
+		const int32 ControlPointIndex = (Index == 0) ? 0 : ((Index - 1) * 4 + 3);
+		return InternalSpline.GetParameter(ControlPointIndex);
+	}
+	virtual bool SetParameter(int32 Index, float NewParameter)
+    {
+		const int32 ControlPointIndex = (Index == 0) ? 0 : ((Index - 1) * 4 + 3);
+        return InternalSpline.SetParameter(ControlPointIndex, NewParameter);
+    }
 
 	FInterval1f GetParameterRange() const
 	{
@@ -1413,7 +1431,7 @@ public:
 				float PrevToNextParamDiff = FMath::Max<float>(UE_KINDA_SMALL_NUMBER, NextParam - PrevParam);
 
 				// Calculate basic tangent - matches legacy formula
-				ValueType AutoTangent = (1.0f - Tension) * ((Next - Prev) / PrevToNextParamDiff);
+				ValueType AutoTangent = ((Next - Prev) / PrevToNextParamDiff) * (1.0f - Tension);
 
 				// For AutoClamped mode, apply clamping like the legacy implementation
 				if (Mode == ETangentMode::AutoClamped)
@@ -1483,7 +1501,7 @@ public:
 		default:
 			// Default to Auto mode if unknown
 			{
-				ValueType Tangent = (1.0f - Tension) * (Next - Prev);
+				ValueType Tangent = (Next - Prev) * (1.0f - Tension);
 				InTangent = Tangent;
 				OutTangent = Tangent;
 			}
@@ -1509,6 +1527,18 @@ public:
 	    // Update tangents to maintain proper curve shape
 	    UpdateTangents();
     }
+
+	void SetKnotVector(const TArray<float>& InKnots)
+	{
+		InternalSpline.SetKnotVector(InKnots);
+
+		UpdateTangents();
+	}
+	
+	const TArray<float>& GetKnotVector() const
+	{
+		return InternalSpline.GetKnotVector();
+	}
     
     /** Getter for internal spline */
     const FPolyBezierSpline3d& GetInternalSpline() const
@@ -1611,7 +1641,7 @@ private:
 				ControlPoint.TangentMode == ETangentMode::AutoClamped)
     		{
     			// For first segment with auto tangents, use simple interpolation
-    			ValueType AutoTangent = (1.0f - Tension) * (P3 - P0);
+    			ValueType AutoTangent = (P3 - P0) * (1.0f - Tension);
     			P1 = P0 + (AutoTangent / 3.0f);
     			P2 = P3 - (AutoTangent / 3.0f);
     		}
