Fixing a case where SoundWave could attempt to stream audio data when running as a dedicated server

#rnx

diff --git a/Engine/Source/Runtime/Engine/Private/SoundWave.cpp b/Engine/Source/Runtime/Engine/Private/SoundWave.cpp
--- a/Engine/Source/Runtime/Engine/Private/SoundWave.cpp
+++ b/Engine/Source/Runtime/Engine/Private/SoundWave.cpp
@@ -1307,7 +1307,7 @@ void USoundWave::Serialize( FArchive& Ar )
 	}
 
 	bool bSupportsStreaming = false;
-	if (Ar.IsLoading() && FPlatformProperties::SupportsAudioStreaming())
+	if (Ar.IsLoading() && FPlatformProperties::SupportsAudioStreaming() && !IsRunningDedicatedServer())
 	{
 		bSupportsStreaming = true;
 	}
