Texture: make automatic settings for nonpow2 texture imports optional in project TextureImportSettings
change the NeverStream logic so the bool reflects user desire, not the pow2 geometry of the texture
there is still a constraint that nonpow2 textures can't stream
detect that in IsPossibleToStream() at runtime

diff --git a/Engine/Plugins/Interchange/Runtime/Source/Import/Private/Texture/InterchangeTextureFactory.cpp b/Engine/Plugins/Interchange/Runtime/Source/Import/Private/Texture/InterchangeTextureFactory.cpp
--- a/Engine/Plugins/Interchange/Runtime/Source/Import/Private/Texture/InterchangeTextureFactory.cpp
+++ b/Engine/Plugins/Interchange/Runtime/Source/Import/Private/Texture/InterchangeTextureFactory.cpp
@@ -1833,38 +1833,6 @@ namespace UE::Interchange::Private::InterchangeTextureFactory
 			ensure(false);
 		}
 
-#if WITH_EDITORONLY_DATA
-		// At some point, these decisions should be left to the pipeline
-		if ( Texture->Source.GetNumBlocks() == 1 && !Texture->Source.IsBlockPowerOfTwo(0))
-		{
-			// try to set some better default options for non-pow2 textures
-
-			// if Texture is not pow2 , change to TMGS_NoMipMaps (if it was default)
-			//   this used to be done by Texture2d.cpp ; it is now optional
-			//	 you can set it back to having mips if you want
-			if (Texture->MipGenSettings == TMGS_FromTextureGroup && Texture->GetTextureClass() == ETextureClass::TwoD )
-			{
-				Texture->MipGenSettings = TMGS_NoMipmaps;
-			}
-
-			if ( ! Texture->Source.IsLongLatCubemap() )
-			{
-				// if Texture is not multiple of 4, change TC to EditorIcon ("UserInterface2D")
-				//	if you do not do this, you might see "Texture forced to uncompressed because size is not a multiple of 4"
-				//  this needs to match the logic in Texture.cpp : GetDefaultTextureFormatName
-				const int32 SizeX = Texture->Source.GetSizeX();
-				const int32 SizeY = Texture->Source.GetSizeY();
-				if ((SizeX & 3) != 0 || (SizeY & 3) != 0)
-				{
-					if (Texture->CompressionSettings == TC_Default) // AutoDXT/BC1
-					{
-						Texture->CompressionSettings = TC_EditorIcon; // "UserInterface2D"
-					}
-				}
-			}
-		}
-#endif // WITH_EDITORONLY_DATA
-
 		// at this point the texture is mostly set up
 		// NormalMapIdentification is not run yet
 	}
diff --git a/Engine/Source/Editor/UnrealEd/Private/Factories/EditorFactories.cpp b/Engine/Source/Editor/UnrealEd/Private/Factories/EditorFactories.cpp
--- a/Engine/Source/Editor/UnrealEd/Private/Factories/EditorFactories.cpp
+++ b/Engine/Source/Editor/UnrealEd/Private/Factories/EditorFactories.cpp
@@ -4019,36 +4019,7 @@ UObject* UTextureFactory::FactoryCreateBinary
 	}
 	
 	Texture->bPreserveBorder		= bPreserveBorder;
-
-	if ( Texture->Source.GetNumBlocks() == 1 && ! Texture->Source.IsBlockPowerOfTwo(0) )
-	{
-		// try to set some better default options for non-pow2 textures
-
-		// if Texture is not pow2 , change to TMGS_NoMipMaps (if it was default)
-		//   this used to be done by Texture2d.cpp ; it is now optional
-		//	 you can set it back to having mips if you want
-		if ( Texture->MipGenSettings == TMGS_FromTextureGroup && Texture->GetTextureClass() == ETextureClass::TwoD )
-		{
-			Texture->MipGenSettings = TMGS_NoMipmaps;		
-		}
-		
-		if ( ! Texture->Source.IsLongLatCubemap() )
-		{
-			// if Texture is not multiple of 4, change TC to EditorIcon ("UserInterface2D")
-			//	if you do not do this, you might see "Texture forced to uncompressed because size is not a multiple of 4"
-			//  this needs to match the logic in Texture.cpp : GetDefaultTextureFormatName
-			int32 SizeX = Texture->Source.GetSizeX();
-			int32 SizeY = Texture->Source.GetSizeY();
-			if ( (SizeX&3) != 0 || (SizeY&3) != 0 )
-			{
-				if ( Texture->CompressionSettings == TC_Default ) // AutoDXT/BC1
-				{
-					Texture->CompressionSettings = TC_EditorIcon; // "UserInterface2D"
-				}
-			}
-		}
-	}
-				
+					
 	if ( UE::TextureUtilitiesCommon::ShouldTextureBeVirtualByAutoImportSize(Texture) )
 	{
 		// only UTexture2D can be VT
diff --git a/Engine/Source/Runtime/Engine/Private/Texture.cpp b/Engine/Source/Runtime/Engine/Private/Texture.cpp
--- a/Engine/Source/Runtime/Engine/Private/Texture.cpp
+++ b/Engine/Source/Runtime/Engine/Private/Texture.cpp
@@ -316,13 +316,15 @@ void UTexture::ReleaseResource()
 		CachedSRRState.Clear();
 
 		// Free the resource.
-		ENQUEUE_RENDER_COMMAND(DeleteResource)([this, ToDelete = PrivateResource](FRHICommandListImmediate& RHICmdList)
+		FTextureResource* ToDelete = PrivateResource;
+		PrivateResource = nullptr;
+
+		ENQUEUE_RENDER_COMMAND(DeleteResource)([this, ToDelete](FRHICommandListImmediate& RHICmdList)
 		{
 			PrivateResourceRenderThread = nullptr;
 			ToDelete->ReleaseResource();
 			delete ToDelete;
 		});
-		PrivateResource = nullptr;
 	}
 }
 
@@ -670,23 +672,11 @@ void UTexture::ValidateSettingsAfterImportOrEdit(bool * pRequiresNotifyMaterials
 			// longlat cube always generates pow2 output
 			bIsPowerOfTwo = true;
 		}
+		// ?? bug ?? bIsPowerOfTwo does not consider ResizeDuringBuildX
 
-		// Downscale can violate IsPow2, but it only acts when NoMipMaps, so it's moot
-
-		if ( ! bIsPowerOfTwo )
-		{
-			// streaming only supports power of 2 mips
-			// due to failure to compensate for the GPU row pitch
-			// it only works for mips that naturally have the required 256 pitch
-			// so mip levels >= 256 and power of 2 only
-			// (this used to be in Texture2D.cpp)
-			// see WarnRequiresTightPackedMip
-			// there are other issues with streaming non-pow2
-			//  all streamable levels must be valid textures, so block-of-4 alignment for BCN
-			//  is easiest to guarantee if the source is pow2
-			NeverStream = true;	
-		}
-	
+		// note: NeverStream used to be turned on here if !bIsPowerOfTwo
+		//	now we leave it off to reflect artist intent, and rely on it being set at runtime by IsPossibleToStream
+		
 		const int32 LargerSourceDimension = FMath::Max( Source.GetSizeX() , Source.GetSizeY() );
 		bool bLargeTextureMustBeVT = LargerSourceDimension > GetMaximumDimensionOfNonVT();
 
@@ -1795,11 +1785,50 @@ const TArray<UAssetUserData*>* UTexture::GetAssetUserDataArray() const
 
 bool UTexture::IsPossibleToStream() const
 {
+	// user flags requesting non-streaming :
 	if ( NeverStream || LODGroup == TEXTUREGROUP_UI )
 	{
 		return false;
 	}
 
+	// we can only stream pow2 textures at the moment, check for that :
+
+	// previously NeverStream was always set in the Editor for nonpow2 textures
+	// now it is only set if user wants it
+	// but nonpow2 textures must still not stream, so detect it here on the cooked data :
+	FTexturePlatformData** PtrPlatformData = const_cast<UTexture *>(this)->GetRunningPlatformData();
+	if ( PtrPlatformData && *PtrPlatformData )
+	{
+		FTexturePlatformData * PlatformData = *PtrPlatformData;
+
+		bool bIsAsyncWorkComplete = true;
+		#if WITH_EDITOR
+		bIsAsyncWorkComplete = PlatformData->IsAsyncWorkComplete();
+		#endif
+
+		if( bIsAsyncWorkComplete )
+		{
+			if ( PlatformData->Mips.Num() <= 1 )
+			{
+				return false;
+			}
+
+			uint32 SizeX = PlatformData->SizeX;
+			uint32 SizeY = PlatformData->SizeY;
+			uint32 SizeZ = PlatformData->GetNumSlices(); // == 6 for cubes
+			if ( !FMath::IsPowerOfTwo(SizeX) || !FMath::IsPowerOfTwo(SizeY) )
+			{
+				return false;
+			}
+			if ( GetTextureClass() == ETextureClass::Volume && !FMath::IsPowerOfTwo(SizeZ) )
+			{
+				// arrays always have NeverStream set
+				// only non-array 2d, Cube, and Volume are streamable
+				return false;
+			}
+		}
+	}
+
 	#if WITH_EDITORONLY_DATA
 	if (!GetOutermost()->HasAnyPackageFlags(PKG_FilterEditorOnly))
 	{
@@ -1815,14 +1844,9 @@ bool UTexture::IsPossibleToStream() const
 	
 	if ( Source.IsValid() )
 	{
-		// should have set NeverStream for nonpow2
-		// note: this is not the case for all old textures
-		// ValidateSettingsAfterImportOrEdit makes sure this is true
-		//check( Source.IsPowerOfTwo() || PowerOfTwoMode != ETexturePowerOfTwoSetting::None )
-		
-		// duplicate the checks done for NeverStream :
+		// turn off streaming for non-pow2 sources :
 
-		// IsPowerOfTwo only checks XY :
+		// AreAllBlocksPowerOfTwo only checks XY :
 		bool bIsPowerOfTwo = Source.AreAllBlocksPowerOfTwo();
 		if ( ! FMath::IsPowerOfTwo(Source.GetVolumeSizeZ()) )
 		{
@@ -1839,7 +1863,6 @@ bool UTexture::IsPossibleToStream() const
 		}
 		if ( ! bIsPowerOfTwo )
 		{
-			// NeverStream should have been set
 			return false;
 		}
 	}
diff --git a/Engine/Source/Runtime/TextureUtilitiesCommon/Private/TextureImportSettings.cpp b/Engine/Source/Runtime/TextureUtilitiesCommon/Private/TextureImportSettings.cpp
--- a/Engine/Source/Runtime/TextureUtilitiesCommon/Private/TextureImportSettings.cpp
+++ b/Engine/Source/Runtime/TextureUtilitiesCommon/Private/TextureImportSettings.cpp
@@ -126,6 +126,12 @@ namespace UE::TextureUtilitiesCommon
 		{
 			Texture->UpdateOodleTextureSdkVersionToLatest();
 
+			// @todo Oodle : optional : reimport ; if texture was previously nonpow2 but is now pow2,
+			//		change some flags (UI, NeverStream, NoMipMaps, etc.) that may have been set by bDoAutomaticTextureSettingsForNonPow2Textures
+			//		?this may stomp intentional user settings, but maybe that's just okay?
+			//	ideally we'd like to just pop up a query box to say "yes/no" to that
+			//	but we aren't allowed to do that here because we may be running on task threads in an Interchange mass-reimport
+
 			return;
 		}
 
@@ -191,6 +197,38 @@ namespace UE::TextureUtilitiesCommon
 				checkNoEntry();
 			}
 		}
+
+		if ( Settings->bDoAutomaticTextureSettingsForNonPow2Textures )
+		{
+			if ( Texture->Source.GetNumBlocks() == 1 && ! Texture->Source.IsBlockPowerOfTwo(0) )
+			{
+				// try to set some better default options for non-pow2 textures
+
+				// if Texture is not pow2 , change to TMGS_NoMipMaps (if it was default)
+				//   this used to be done by Texture2d.cpp ; it is now optional
+				//	 you can set it back to having mips if you want
+				if ( Texture->MipGenSettings == TMGS_FromTextureGroup && Texture->GetTextureClass() == ETextureClass::TwoD )
+				{
+					Texture->MipGenSettings = TMGS_NoMipmaps;		
+				}
+		
+				if ( ! Texture->Source.IsLongLatCubemap() )
+				{
+					// if Texture is not multiple of 4, change TC to EditorIcon ("UserInterface2D")
+					//	if you do not do this, you might see "Texture forced to uncompressed because size is not a multiple of 4"
+					//  this needs to match the logic in Texture.cpp : GetDefaultTextureFormatName
+					int32 SizeX = Texture->Source.GetSizeX();
+					int32 SizeY = Texture->Source.GetSizeY();
+					if ( (SizeX&3) != 0 || (SizeY&3) != 0 )
+					{
+						if ( Texture->CompressionSettings == TC_Default ) // AutoDXT/BC1
+						{
+							Texture->CompressionSettings = TC_EditorIcon; // "UserInterface2D"
+						}
+					}
+				}
+			}
+		}
 	}
 	
 	bool ShouldTextureBeVirtualByAutoImportSize(const UTexture * Texture)
diff --git a/Engine/Source/Runtime/TextureUtilitiesCommon/Public/TextureImportSettings.h b/Engine/Source/Runtime/TextureUtilitiesCommon/Public/TextureImportSettings.h
--- a/Engine/Source/Runtime/TextureUtilitiesCommon/Public/TextureImportSettings.h
+++ b/Engine/Source/Runtime/TextureUtilitiesCommon/Public/TextureImportSettings.h
@@ -83,6 +83,11 @@ public:
 		DisplayName = "When to infill RGB in transparent white PNG",
 		ToolTip = "Default behavior is to infill only for binary transparency; this setting may change that to always or never.  Will check TextureImporter/FillPNGZeroAlpha if this is not changed from Default.  This setting is applied to newly imported textures, it does not affect existing textures in the project. This setting is project-global, prefer the per-user variant in Editor Preferences."))
 	ETextureImportPNGInfill PNGInfill = ETextureImportPNGInfill::Default;
+	
+	UPROPERTY(config, EditAnywhere, Category=ImportSettings, meta = (
+		DisplayName = "Detect nonpow2 textures and set defaults",
+		ToolTip = "Sets them to NoMipMaps and uncompressed UI compression. If this is off, artists must imported UI textures up correctly by hand."))
+	bool bDoAutomaticTextureSettingsForNonPow2Textures = true;
 
 	//~ Begin UObject Interface
 
