Automatically estimate subtitle duration when attached to a SoundWave.

diff --git a/Engine/Source/Runtime/Engine/Private/SoundWave.cpp b/Engine/Source/Runtime/Engine/Private/SoundWave.cpp
--- a/Engine/Source/Runtime/Engine/Private/SoundWave.cpp
+++ b/Engine/Source/Runtime/Engine/Private/SoundWave.cpp
@@ -3150,6 +3150,7 @@ void USoundWave::PostEditChangeProperty(FPropertyChangedEvent& PropertyChangedEv
 	static const FName InlinedAudioInSecondsFName = GET_MEMBER_NAME_CHECKED(USoundWave, SizeOfFirstAudioChunkInSeconds);
 	static const FName CloudStreamingFName = GET_MEMBER_NAME_CHECKED(USoundWave, bEnableCloudStreaming);
 	static const FName CloudStreamingPlatformSettingsFName = GET_MEMBER_NAME_CHECKED(USoundWave, PlatformSettings);
+	static const FName AssetUserDataFName = GET_MEMBER_NAME_CHECKED(USoundWave, AssetUserData);
 
 	if (FProperty* PropertyThatChanged = PropertyChangedEvent.Property)
 	{
@@ -3160,6 +3161,23 @@ void USoundWave::PostEditChangeProperty(FPropertyChangedEvent& PropertyChangedEv
 		{
 			UpdateTransformations();
 		}
+		else if (Name == AssetUserDataFName && !AssetUserData.IsEmpty())
+		{
+			// If a new Subtitle asset user datum was added, default its duration to the length of the SoundWave it was added to.
+			// Look at the last item because a new user data added in the editor will always go at the end of the array- earlier data will have already had this operation applied.
+
+			TObjectPtr<UAssetUserData>& LastModifiedUserData = AssetUserData.Last();
+			if (IsValid(LastModifiedUserData))
+			{
+				USubtitleAssetUserData* UserDataAsSubtitle = Cast<USubtitleAssetUserData>(LastModifiedUserData.Get() );
+				if (IsValid(UserDataAsSubtitle) && UserDataAsSubtitle->Duration == SubtitleDefaultDuration && UserDataAsSubtitle->Text.IsEmpty())
+				{
+					// If the sound is looping indefinitely, the subtitle duration will default to the length of the first loop.
+					// It is still editable by the user after this default is set.
+					UserDataAsSubtitle->Duration = Duration;
+				}
+			}
+		}
 	
 		// Prevent constant re-compression of SoundWave while properties are being changed interactively
 		if (PropertyChangedEvent.ChangeType != EPropertyChangeType::Interactive)
diff --git a/Engine/Source/Runtime/Engine/Public/Subtitles/SubtitlesAndClosedCaptionsDelegates.h b/Engine/Source/Runtime/Engine/Public/Subtitles/SubtitlesAndClosedCaptionsDelegates.h
--- a/Engine/Source/Runtime/Engine/Public/Subtitles/SubtitlesAndClosedCaptionsDelegates.h
+++ b/Engine/Source/Runtime/Engine/Public/Subtitles/SubtitlesAndClosedCaptionsDelegates.h
@@ -42,7 +42,10 @@ enum class ESubtitleType : uint8
 };
 
 // Minimum duration to display subtitle.
-static constexpr float SubtitleMinDuration = 0.05f;
+static constexpr float SubtitleMinDuration		= 0.05f;
+
+// Default value to initialize subtitle duration to. Used by SoundWaves to check if they should manually set the duration.
+static constexpr float SubtitleDefaultDuration	= 3.f;
 
 /**
  *  Base class for subtitle data being attached to assets
@@ -66,7 +69,7 @@ public:
 	// Defaulted to 3 seconds so when adding new subtitles it's not required to enter a placeholder Duration.
 	// Duration can be be set by ingestion pipelines when importing Subtitles in bulk
 	UPROPERTY(EditAnywhere, BlueprintReadWrite, Category = Subtitles, meta = (ClampMin = 0.05f))
-	float Duration = 3.f;
+	float Duration = SubtitleDefaultDuration;
 
 	// Some subtitles have a delay before they're allowed to be displayed (primarily from the legacy system).
 	// StartOffset measures how long in Seconds, after queuing, before the subtitle is allowed to enter the active subtitles queue.
