Add material control for local light shadow casting and local/directional light shadow casting quality for translucent material.

diff --git a/Engine/Shaders/Private/ForwardLightingCommon.ush b/Engine/Shaders/Private/ForwardLightingCommon.ush
--- a/Engine/Shaders/Private/ForwardLightingCommon.ush
+++ b/Engine/Shaders/Private/ForwardLightingCommon.ush
@@ -28,11 +28,22 @@
 #define VIRTUAL_SHADOW_MAP 0
 #endif
 
+#ifndef MATERIAL_ENABLE_TRANSLUCENT_LOCAL_LIGHT_SHADOW
+#define MATERIAL_ENABLE_TRANSLUCENT_LOCAL_LIGHT_SHADOW 0
+#endif
+#ifndef MATERIAL_ENABLE_TRANSLUCENT_HIGH_QUALITY_LOCAL_LIGHT_SHADOW
+#define MATERIAL_ENABLE_TRANSLUCENT_HIGH_QUALITY_LOCAL_LIGHT_SHADOW 0
+#endif
+#ifndef MATERIAL_ENABLE_TRANSLUCENT_HIGH_QUALITY_DIRECTIONAL_LIGHT_SHADOW
+#define MATERIAL_ENABLE_TRANSLUCENT_HIGH_QUALITY_DIRECTIONAL_LIGHT_SHADOW 0
+#endif
+
 #if VIRTUAL_SHADOW_MAP
 #include "VirtualShadowMaps/VirtualShadowMapProjectionCommon.ush"
-#if SUPPORT_VSM_FOWARD_QUALITY > 0
+#if SUPPORT_VSM_FOWARD_QUALITY > 0 || MATERIAL_ENABLE_TRANSLUCENT_HIGH_QUALITY_LOCAL_LIGHT_SHADOW || MATERIAL_ENABLE_TRANSLUCENT_HIGH_QUALITY_DIRECTIONAL_LIGHT_SHADOW
 #include "VirtualShadowMaps/VirtualShadowMapProjectionFilter.ush"
 #include "VirtualShadowMaps/VirtualShadowMapProjectionDirectional.ush"
+#include "VirtualShadowMaps/VirtualShadowMapProjectionSpot.ush"
 #endif // SUPPORT_VSM_FOWARD_QUALITY
 #endif // VIRTUAL_SHADOW_MAP
 
@@ -195,7 +206,7 @@ FDeferredLightingSplit GetForwardDirectLightingSplit(
 			BRANCH
 			if ( !bSkipDirLightVirtualShadowMapEvaluation && ForwardLightStruct.DirectionalLightVSM != INDEX_NONE )
 			{
-				#if SUPPORT_VSM_FOWARD_QUALITY == 1
+				#if SUPPORT_VSM_FOWARD_QUALITY == 1 || MATERIAL_ENABLE_TRANSLUCENT_HIGH_QUALITY_DIRECTIONAL_LIGHT_SHADOW
 				{
 					const FLightShaderParameters DirectionalLight = ConvertToLightShaderParameters(LightData);
 
@@ -297,13 +308,13 @@ FDeferredLightingSplit GetForwardDirectLightingSplit(
 		float DynamicShadowing = dot(PreviewShadowMapChannelMask, DynamicShadowFactors);
 		float4 LightAttenuation = float4(1, 1, DynamicShadowing.x, DynamicShadowing.x);
 
-		#if VIRTUAL_SHADOW_MAP && SUPPORT_SHADOWED_LOCAL_LIGHT_ON_FORWARD_LIT_TRANSLUCENT
+		#if VIRTUAL_SHADOW_MAP && (SUPPORT_SHADOWED_LOCAL_LIGHT_ON_FORWARD_LIT_TRANSLUCENT || MATERIAL_ENABLE_TRANSLUCENT_LOCAL_LIGHT_SHADOW)
 		BRANCH
 		if (LocalLight.Internal.VirtualShadowMapId != INDEX_NONE)
 		{
 			float DynamicShadowFactor = 1.f;
 			LightData.ShadowedBits = 1;
-			#if SUPPORT_VSM_FOWARD_QUALITY == 1
+			#if SUPPORT_VSM_FOWARD_QUALITY == 1 || MATERIAL_ENABLE_TRANSLUCENT_HIGH_QUALITY_LOCAL_LIGHT_SHADOW
 			{
 				const FLightShaderParameters LocalLightShaderParameters = ConvertToLightShaderParameters(LightData);
 
@@ -311,10 +322,10 @@ FDeferredLightingSplit GetForwardDirectLightingSplit(
 				float SMRTRayOffset = ScreenRayLengthWorld;
 
 				FVirtualShadowMapSampleResult VirtualShadowMapSample = TraceLocalLight(
-					LocalLight.VirtualShadowMapId,
+					LocalLight.Internal.VirtualShadowMapId,
 					LocalLightShaderParameters,
 					PixelPos,
-					GBufferData.SceneDepth,
+					GBufferData.Depth,
 					TranslatedWorldPosition,
 					SMRTRayOffset,
 					Dither,
diff --git a/Engine/Source/Runtime/Core/Private/UObject/UE5SpecialProjectStreamObjectVersion.cpp b/Engine/Source/Runtime/Core/Private/UObject/UE5SpecialProjectStreamObjectVersion.cpp
--- a/Engine/Source/Runtime/Core/Private/UObject/UE5SpecialProjectStreamObjectVersion.cpp
+++ b/Engine/Source/Runtime/Core/Private/UObject/UE5SpecialProjectStreamObjectVersion.cpp
@@ -12,6 +12,7 @@ TMap<FGuid, FGuid> FUE5SpecialProjectStreamObjectVersion::GetSystemGuids()
 
 	SystemGuids.Add(DevGuids.NANITE_DERIVEDDATA_VER, FGuid("C73C42D8-4CA0-4EAF-8968-722EB9EEF812"));
 	SystemGuids.Add(DevGuids.SkeletalMeshDerivedDataVersion, FGuid("56B88323-F44D-46B8-A3A3-B3441BC8F448"));
+	SystemGuids.Add(DevGuids.MaterialTranslationDDCVersion, FGuid("BFB09A107AD8451294FFD3B321C3565C"));
 
 	return SystemGuids;
 }
diff --git a/Engine/Source/Runtime/Engine/Private/Materials/HLSLMaterialTranslator.cpp b/Engine/Source/Runtime/Engine/Private/Materials/HLSLMaterialTranslator.cpp
--- a/Engine/Source/Runtime/Engine/Private/Materials/HLSLMaterialTranslator.cpp
+++ b/Engine/Source/Runtime/Engine/Private/Materials/HLSLMaterialTranslator.cpp
@@ -321,6 +321,9 @@ struct FHLSLMaterialTranslator::FEnvironmentDefines
 	uint8 NumCustomizedUVs;
 	bool bTextureSampleDebug;
 	bool bSubstrateRoughnessTracking;
+	bool bMaterialEnableTranslucentLocalLightShadow;
+	bool bMaterialEnableTranslucentHighQualityLocalLightShadow;
+	bool bMaterialEnableTranslucentHighQualityDirectionalLightShadow;
 	TArray<TPair<FString, int32>> SubstrateDefines;
 	TArray<TObjectPtr<UMaterialParameterCollection>> ParameterCollections;
 
@@ -423,7 +426,9 @@ struct FHLSLMaterialTranslator::FEnvironmentDefines
 		Ar << bTextureSampleDebug;
 		Ar << SubstrateDefines;
 		Ar << bSubstrateRoughnessTracking;
-
+		Ar << bMaterialEnableTranslucentLocalLightShadow;
+		Ar << bMaterialEnableTranslucentHighQualityLocalLightShadow;
+		Ar << bMaterialEnableTranslucentHighQualityDirectionalLightShadow;
 		TArray<UObject*> ParameterCollectionObjects;
 		if (!Ar.IsLoading())
 		{
@@ -2932,6 +2937,10 @@ void FHLSLMaterialTranslator::GetMaterialEnvironment(EShaderPlatform InPlatform,
 	OutEnvironment.SetDefine(TEXT("NUM_VIRTUALTEXTURE_FEEDBACK_REQUESTS"), EnvironmentDefines->NumVirtualTextureFeedbackRequests);
 	OutEnvironment.SetDefine(TEXT("MATERIAL_VIRTUALTEXTURE_FEEDBACK"), EnvironmentDefines->bMaterialVirtualTextureFeedback);
 
+	OutEnvironment.SetDefine(TEXT("MATERIAL_ENABLE_TRANSLUCENT_LOCAL_LIGHT_SHADOW"), EnvironmentDefines->bMaterialEnableTranslucentLocalLightShadow);
+	OutEnvironment.SetDefine(TEXT("MATERIAL_ENABLE_TRANSLUCENT_HIGH_QUALITY_LOCAL_LIGHT_SHADOW"), EnvironmentDefines->bMaterialEnableTranslucentHighQualityLocalLightShadow);
+	OutEnvironment.SetDefine(TEXT("MATERIAL_ENABLE_TRANSLUCENT_HIGH_QUALITY_DIRECTIONAL_LIGHT_SHADOW"), EnvironmentDefines->bMaterialEnableTranslucentHighQualityDirectionalLightShadow);
+
 	for (int i = 0; i < EnvironmentDefines->VirtualPageTypes.Num(); ++i)
 	{
 		if (EnvironmentDefines->VirtualPageTypes[i] & FEnvironmentDefines::TABLE_MESH_PAINT)
@@ -16297,6 +16306,9 @@ void FHLSLMaterialTranslator::PrepareEnvironmentDefines()
 	EnvironmentDefines->MaterialNeuralPostProcess = (MaterialCompilationOutput.bUsedWithNeuralNetworks || Material->IsUsedWithNeuralNetworks()) && Material->IsPostProcessMaterial();
 	EnvironmentDefines->PixelDepthOffsetMode = Material->GetPixelDepthOffsetMode();
 	EnvironmentDefines->NumCustomizedUVs = Material->GetNumCustomizedUVs();
+	EnvironmentDefines->bMaterialEnableTranslucentLocalLightShadow = Material->AllowTranslucentLocalLightShadow();
+	EnvironmentDefines->bMaterialEnableTranslucentHighQualityLocalLightShadow = Material->GetTranslucentLocalLightShadowQuality() > 0;
+	EnvironmentDefines->bMaterialEnableTranslucentHighQualityDirectionalLightShadow = Material->GetTranslucentDirectionalLightShadowQuality() > 0;
 
 	// Count the number of VTStacks (each stack will allocate a feedback slot)
 	EnvironmentDefines->NumVirtualTextureSamples = VTStacks.Num();
diff --git a/Engine/Source/Runtime/Engine/Private/Materials/Material.cpp b/Engine/Source/Runtime/Engine/Private/Materials/Material.cpp
--- a/Engine/Source/Runtime/Engine/Private/Materials/Material.cpp
+++ b/Engine/Source/Runtime/Engine/Private/Materials/Material.cpp
@@ -1096,6 +1096,11 @@ void UMaterial::SetInitialValues()
 	bUsedWithHeterogeneousVolumes = false;
 	bUsedWithStaticMesh = true;
 
+	// Shadow on translucent material
+	bAllowTranslucentLocalLightShadow = false;
+	TranslucentLocalLightShadowQuality = 0;
+	TranslucentDirectionalLightShadowQuality = 0;
+
 	// We default to false=incompatible, and rely onto the automatic detection executed from the HLSLTranslator.
 	bForceCompatibleWithLightFunctionAtlas = false;
 
diff --git a/Engine/Source/Runtime/Engine/Private/Materials/MaterialShared.cpp b/Engine/Source/Runtime/Engine/Private/Materials/MaterialShared.cpp
--- a/Engine/Source/Runtime/Engine/Private/Materials/MaterialShared.cpp
+++ b/Engine/Source/Runtime/Engine/Private/Materials/MaterialShared.cpp
@@ -1997,6 +1997,9 @@ bool FMaterialResource::ShouldApplyFogging() const { return Material->bUseTransl
 bool FMaterialResource::ShouldApplyCloudFogging() const { return Material->bApplyCloudFogging; }
 bool FMaterialResource::ShouldAlwaysEvaluateWorldPositionOffset() const { return Material->bAlwaysEvaluateWorldPositionOffset; }
 bool FMaterialResource::IsSky() const { return Material->bIsSky; }
+bool FMaterialResource::AllowTranslucentLocalLightShadow() const { return Material->bAllowTranslucentLocalLightShadow; }
+float FMaterialResource::GetTranslucentLocalLightShadowQuality() const { return Material->TranslucentLocalLightShadowQuality; }
+float FMaterialResource::GetTranslucentDirectionalLightShadowQuality() const { return Material->TranslucentDirectionalLightShadowQuality; }
 bool FMaterialResource::ComputeFogPerPixel() const {return Material->bComputeFogPerPixel;}
 FString FMaterialResource::GetFriendlyName() const { return GetNameSafe(Material); } //avoid using the material instance name here, we want materials that share a shadermap to also share a friendly name.
 FString FMaterialResource::GetAssetName() const { return MaterialInstance ? GetNameSafe(MaterialInstance) : GetNameSafe(Material); }
diff --git a/Engine/Source/Runtime/Engine/Public/MaterialShared.h b/Engine/Source/Runtime/Engine/Public/MaterialShared.h
--- a/Engine/Source/Runtime/Engine/Public/MaterialShared.h
+++ b/Engine/Source/Runtime/Engine/Public/MaterialShared.h
@@ -2258,6 +2258,9 @@ public:
 	virtual bool ComputeFogPerPixel() const { return false; }
 	virtual bool IsSky() const { return false; }
 	virtual bool HasSubstrateRoughnessTracking() const { return true; }
+	virtual bool AllowTranslucentLocalLightShadow() const { return false;  }
+	virtual float GetTranslucentLocalLightShadowQuality() const { return 0.f;  }
+	virtual float GetTranslucentDirectionalLightShadowQuality() const { return 0.f;  }
 	virtual FString GetFriendlyName() const = 0;
 	/** Similar to GetFriendlyName, but but avoids historical behavior of the former, returning the exact asset name for material instances instead of just the material. */
 	virtual FString GetAssetName() const { return GetFriendlyName(); }
@@ -3050,6 +3053,9 @@ public:
 	ENGINE_API virtual UMaterialInterface* GetMaterialInterface() const override;
 	ENGINE_API virtual int32 GetNeuralProfileId() const override;
 	ENGINE_API virtual bool HasSubstrateRoughnessTracking() const override;
+	ENGINE_API virtual bool AllowTranslucentLocalLightShadow() const override;
+	ENGINE_API virtual float GetTranslucentLocalLightShadowQuality() const override;
+	ENGINE_API virtual float GetTranslucentDirectionalLightShadowQuality() const override;
 
 	/**
 	 * Should shaders compiled for this material be saved to disk?
diff --git a/Engine/Source/Runtime/Engine/Public/Materials/Material.h b/Engine/Source/Runtime/Engine/Public/Materials/Material.h
--- a/Engine/Source/Runtime/Engine/Public/Materials/Material.h
+++ b/Engine/Source/Runtime/Engine/Public/Materials/Material.h
@@ -936,6 +936,18 @@ public:
 	UPROPERTY(EditAnywhere, Category=Material, AdvancedDisplay)
 	uint8 Wireframe : 1;
 
+	/** Allows a translucent material to receive local light shadows. */
+	UPROPERTY(EditAnywhere, Category = Translucency, AdvancedDisplay, meta = (DisplayName = "Allow Local Light Shadow"))
+	uint8 bAllowTranslucentLocalLightShadow : 1;
+
+	/** Set local light shadow quality received by translucent material. */
+	UPROPERTY(EditAnywhere, Category = Translucency, AdvancedDisplay, meta = (DisplayName = "Local Light Shadow Quality", UIMin = "0", UIMax = "1", ClampMin = "0", ClampMax = "1", SliderExponent = 1.0))
+	float TranslucentLocalLightShadowQuality = 0;
+
+	/** Set directional light shadow quality received by translucent material. */
+	UPROPERTY(EditAnywhere, Category = Translucency, AdvancedDisplay, meta = (DisplayName = "Directional Light Shadow Quality", UIMin = "0", UIMax = "1", ClampMin = "0", ClampMax = "1", SliderExponent = 1.0))
+	float TranslucentDirectionalLightShadowQuality = 0;
+
 	/**
 	 * Select what shading rate to apply, on platforms that support variable rate shading.
 	 * Non-1x1 rates will reduce the rasterization fidelity for the material; they will not super-sample the material.
