Prevent crash in APlayerCameraManager when the owning player controller has been GC'ed.

diff --git a/Engine/Source/Runtime/Engine/Private/PlayerCameraManager.cpp b/Engine/Source/Runtime/Engine/Private/PlayerCameraManager.cpp
--- a/Engine/Source/Runtime/Engine/Private/PlayerCameraManager.cpp
+++ b/Engine/Source/Runtime/Engine/Private/PlayerCameraManager.cpp
@@ -1513,7 +1513,7 @@ void FTViewTarget::CheckViewTarget(APlayerController* OwningController)
 				if (AController* PlayerStateOwner = Cast<AController>(PlayerState->GetOwner()))
 				{
 					AActor* PlayerStateViewTarget = PlayerStateOwner->GetPawn();
-					if( IsValid(PlayerStateViewTarget) )
+					if (IsValid(OwningController) && IsValid(PlayerStateViewTarget))
 					{
 						OwningController->PlayerCameraManager->AssignViewTarget(PlayerStateViewTarget, *this);
 					}
@@ -1530,7 +1530,7 @@ void FTViewTarget::CheckViewTarget(APlayerController* OwningController)
 		}
 	}
 
-	if (!IsValid(Target))
+	if (IsValid(OwningController) && !IsValid(Target))
 	{
 		if (OwningController->GetPawn() && !OwningController->GetPawn()->IsPendingKillPending() )
 		{
