Fix materials with parameter collections not waiting for parameter collections to finish loading

#tests Loaded Materials with MPC in CitySample

diff --git a/Engine/Source/Runtime/Engine/Private/Materials/HLSLMaterialTranslator.cpp b/Engine/Source/Runtime/Engine/Private/Materials/HLSLMaterialTranslator.cpp
--- a/Engine/Source/Runtime/Engine/Private/Materials/HLSLMaterialTranslator.cpp
+++ b/Engine/Source/Runtime/Engine/Private/Materials/HLSLMaterialTranslator.cpp
@@ -2989,14 +2989,24 @@ void FHLSLMaterialTranslator::GetMaterialEnvironment(EShaderPlatform InPlatform,
 		// Add uniform buffer declarations for any parameter collections referenced
 		const FString CollectionName = FString::Printf(TEXT("MaterialCollection%u"), CollectionIndex);
 
-		// Check that the parameter collection loaded succesfully.
+		// Check that the parameter collection loaded successfully.
 		UMaterialParameterCollection* ParameterCollection = EnvironmentDefines->ParameterCollections[CollectionIndex];
-		if (!ParameterCollection || !ParameterCollection->HasValidUniformBufferStruct())
+		if (!ParameterCollection)
 		{
 			UE_LOG(LogMaterial, Warning, TEXT("Null parameter collection found in environment defines while translating material."));
 			continue;
 		}
 
+		// Ensure PostLoad is called so the uniform buffers are created in case the parameter collection was loaded async 
+		ParameterCollection->ConditionalPostLoad();
+
+		// Check that the parameter collection uniform buffer structure is valid
+		if (!ParameterCollection->HasValidUniformBufferStruct())
+		{
+			UE_LOG(LogMaterial, Warning, TEXT("Invalid parameter collection uniform buffer struct found in environment defines while translating material."));
+			continue;
+		}
+
 		// This can potentially become an issue for MaterialCollection Uniform Buffers if they ever get non-numeric resources (eg Textures), as
 		// OutEnvironment.ResourceTableMap has a map by name, and the N ParameterCollection Uniform Buffers ALL are names "MaterialCollection"
 		// (and the hlsl cbuffers are named MaterialCollection0, etc, so the names don't match the layout)
