[Engine] Extra checks on the AnimInstance montage for GC.

diff --git a/Engine/Source/Runtime/Engine/Private/Animation/AnimInstance.cpp b/Engine/Source/Runtime/Engine/Private/Animation/AnimInstance.cpp
--- a/Engine/Source/Runtime/Engine/Private/Animation/AnimInstance.cpp
+++ b/Engine/Source/Runtime/Engine/Private/Animation/AnimInstance.cpp
@@ -1937,9 +1937,8 @@ void UAnimInstance::Montage_Advance(float DeltaSeconds)
 	for (int32 InstanceIndex = 0; InstanceIndex < MontageInstances.Num(); InstanceIndex++)
 	{
 		FAnimMontageInstance* const MontageInstance = MontageInstances[InstanceIndex];
-		// should never be NULL
-		ensure(MontageInstance);
-		if (MontageInstance && MontageInstance->IsValid())
+		// should never be NULL and the montage should never be null or garbage
+		if (ensure(MontageInstance && IsValid(MontageInstance->Montage)))
 		{
 			bool const bUsingBlendedRootMotion = (RootMotionMode == ERootMotionMode::RootMotionFromEverything);
 			bool const bNoRootMotionExtraction = (RootMotionMode == ERootMotionMode::NoRootMotionExtraction);
