Animation - NonUnity fix

diff --git a/Engine/Source/Runtime/Engine/Private/Animation/AnimNode_StateMachine.cpp b/Engine/Source/Runtime/Engine/Private/Animation/AnimNode_StateMachine.cpp
--- a/Engine/Source/Runtime/Engine/Private/Animation/AnimNode_StateMachine.cpp
+++ b/Engine/Source/Runtime/Engine/Private/Animation/AnimNode_StateMachine.cpp
@@ -18,6 +18,7 @@
 #include "Animation/AnimNode_StateResult.h"
 #include "Animation/SkeletonRemapping.h"
 #include "Animation/SkeletonRemappingRegistry.h"
+#include "HAL/IConsoleManager.h"
 
 #include UE_INLINE_GENERATED_CPP_BY_NAME(AnimNode_StateMachine)
 
