PR #11599: Fix some incorrect documentation for Malloc Leak Reporter.

#rnx

- No change in functionality, just fixes places where either the code documentation or logging was outputting incorrect value types (bytes vs KiB etc)

diff --git a/Engine/Source/Runtime/Engine/Private/ProfilingDebugging/MallocLeakReporter.cpp b/Engine/Source/Runtime/Engine/Private/ProfilingDebugging/MallocLeakReporter.cpp
--- a/Engine/Source/Runtime/Engine/Private/ProfilingDebugging/MallocLeakReporter.cpp
+++ b/Engine/Source/Runtime/Engine/Private/ProfilingDebugging/MallocLeakReporter.cpp
@@ -43,7 +43,7 @@ void FMallocLeakReporter::Start(int32 FilterSize/*=0*/, float ReportOnTime /*= 0
 		MALLOC_LEAKDETECTION, PLATFORM_USES_FIXED_GMalloc_CLASS);
 
 #else
-	UE_LOG(LogLeakDetector, Log, TEXT("Started Tracking Allocations > %d KB"), FilterSize / 1024);
+	UE_LOG(LogLeakDetector, Log, TEXT("Started Tracking Allocations > %d KiB"), FilterSize / 1024);
 	FMallocLeakDetection::Get().SetAllocationCollection(true, FilterSize);
 
 	// Create a ticker to issue checkpoints
@@ -162,17 +162,18 @@ int32 FMallocLeakReporter::WriteReport(const TCHAR* ReportName, const FMallocLea
 
 static FAutoConsoleCommand LeakReporterStartCommand(
 	TEXT("mallocleak.start"),
-	TEXT("Starts tracking allocations. Args -report=[secs] -size=[filter]"),
+	TEXT("Starts tracking allocations. Args -report=[secs] -size=[filter (KiB)]"),
 	FConsoleCommandWithArgsDelegate::CreateLambda([](const TArray<FString>& Args) {
 		FString ArgString = FString::Join(Args, TEXT(" "));
 		int32 Size = 0;
 		float ReportTime = 0;
-		FParse::Value(*ArgString, TEXT("size="), Size);
-		FParse::Value(*ArgString, TEXT("report="), ReportTime);
+		FParse::Value(*ArgString, TEXT("size="), Size); // KiB
+		FParse::Value(*ArgString, TEXT("report="), ReportTime); // Seconds
 
+		// FMallocLeakReporter expects FilterSize in bytes
 		FMallocLeakReporter::Get().Start(Size * 1024, ReportTime);
 
-		UE_LOG(LogConsoleResponse, Display, TEXT("Tracking allocations >=  %d KB and reporting every %f seconds"), Size / 1024, ReportTime);
+		UE_LOG(LogConsoleResponse, Display, TEXT("Tracking allocations >= %d KiB and reporting every %f seconds"), Size, ReportTime);
 	})
 );
 
diff --git a/Engine/Source/Runtime/Engine/Public/ProfilingDebugging/MallocLeakReporter.h b/Engine/Source/Runtime/Engine/Public/ProfilingDebugging/MallocLeakReporter.h
--- a/Engine/Source/Runtime/Engine/Public/ProfilingDebugging/MallocLeakReporter.h
+++ b/Engine/Source/Runtime/Engine/Public/ProfilingDebugging/MallocLeakReporter.h
@@ -15,7 +15,7 @@
 	reports of low-level memory leaks.
 
 	The reporter can be operated from the console where it can be started/stopped and used to periodically
-	or manually report leaks, or it can be called from game code to implemnent this funtionality at desired times.
+	or manually report leaks, or it can be called from game code to implement this functionality at desired times.
 
 	Example Console Usage:
 
@@ -25,7 +25,7 @@
 
 	Example Code usage
 
-	FMallocLeakReporter::Get().Start(0, 300) - start tracking allocs > 0KB and generate a report every 300 secs
+	FMallocLeakReporter::Get().Start(0, 300) - start tracking allocs > 0 bytes and generate a report every 300 secs
 	FMallocLeakReporter::Get().WriteReports() - Writes reports and returns the number of suspected leaks
 	FMallocLeakReporter::Get().WriteReport(TEXT("BigAllocs"), Options); - Write custom report of allocations that matches "Options"
 	FMallocLeakReporter::Get().Stop() - stop tracking allocations
@@ -57,7 +57,7 @@ public:
 	/**
 	 * Starts tracking allocations.
 	 *
-	 * @param: 	FilterSize		Only track allocations >= this value in KB. Higher values affect performance less
+	 * @param: 	FilterSize		Only track allocations >= this value in bytes. Higher values affect performance less
 	 * @param: 	ReportOnTime	Write out a report every N seconds
 	 * @return: void		
 	*/ENGINE_API void		Start(int32 FilterSize = 0, float ReportOnTime = 0.0f);
@@ -125,13 +125,11 @@ protected:
 
 	bool						Enabled;
 	int32						ReportCount;
-	FTSTicker::FDelegateHandle				CheckpointTicker;
-	FTSTicker::FDelegateHandle				ReportTicker;
+	FTSTicker::FDelegateHandle	CheckpointTicker;
+	FTSTicker::FDelegateHandle	ReportTicker;
 	FMallocLeakReportDelegate	ReportDelegate;
 
 	// Report vars
 	FMallocLeakReportOptions	DefaultLeakReportOptions;
 	FMallocLeakReportOptions	DefaultAllocReportOptions;
-	
 };
-
