[Physics] Hook up API for generic sweep traces on  the physics thread.
- Sphere, Boxt, Capsule and Raycast

diff --git a/Engine/Source/Runtime/Engine/Private/Collision/SceneQuery.cpp b/Engine/Source/Runtime/Engine/Private/Collision/SceneQuery.cpp
--- a/Engine/Source/Runtime/Engine/Private/Collision/SceneQuery.cpp
+++ b/Engine/Source/Runtime/Engine/Private/Collision/SceneQuery.cpp
@@ -1462,7 +1462,47 @@ namespace Chaos::Private
 			using PTTraits = TSQTraits<ChaosInterface::FPTRaycastHit, ESweepOrRay::Raycast, ESingleMultiOrTest::Multi>;
 			return TSceneCastCommonImp<PTTraits, FRaycastSQAdditionalInputs>(World, OutHits, FRaycastSQAdditionalInputs(), Start, End, TraceChannel, Params, ResponseParams, ObjectParams, FDefaultAccelContainer{});
 		}
-	};
+	}
+
+	bool FGenericPhysicsInterface_Internal::GeomSweepMulti(const UWorld* World, const FCollisionShape& InGeom, const FQuat& InGeomRot, TArray<FHitResult>& OutHits, const FVector& Start, const FVector& End, ECollisionChannel TraceChannel, const FCollisionQueryParams& Params, const FCollisionResponseParams& ResponseParams, const FCollisionObjectQueryParams& ObjectParams)
+	{
+		if (InGeom.IsNearlyZero())
+		{
+			// Fall back to a raycast
+			return RaycastMulti(World, OutHits, Start, End, TraceChannel, Params, ResponseParams, ObjectParams);
+		}
+
+		SCOPE_CYCLE_COUNTER(STAT_Collision_SceneQueryTotal);
+		SCOPE_CYCLE_COUNTER(STAT_Collision_GeomSweepMultiple_Internal);
+		CSV_SCOPED_TIMING_STAT(SceneQuery, GeomSweepMultiple_Internal);
+
+		Chaos::EnsureIsInPhysicsThreadContext();
+
+		if ((World == NULL) || (World->GetPhysicsScene() == NULL))
+		{
+			return false;
+		}
+
+		using PTTraits = TSQTraits<ChaosInterface::FPTSweepHit, ESweepOrRay::Sweep, ESingleMultiOrTest::Multi>;
+		return TSceneCastCommonImp<PTTraits, FGeomSQAdditionalInputs>(World, OutHits, FGeomSQAdditionalInputs(InGeom, InGeomRot), Start, End, TraceChannel, Params, ResponseParams, ObjectParams, FDefaultAccelContainer{});
+	}
+
+	bool FGenericPhysicsInterface_Internal::RaycastMulti(const UWorld* World, TArray<FHitResult>& OutHits, const FVector& Start, const FVector& End, ECollisionChannel TraceChannel, const FCollisionQueryParams& Params, const FCollisionResponseParams& ResponseParams, const FCollisionObjectQueryParams& ObjectParams)
+	{
+		SCOPE_CYCLE_COUNTER(STAT_Collision_SceneQueryTotal);
+		SCOPE_CYCLE_COUNTER(STAT_Collision_RaycastMultiple_Internal);
+		CSV_SCOPED_TIMING_STAT(SceneQuery, RaycastMultiple_Internal);
+
+		Chaos::EnsureIsInPhysicsThreadContext();
+
+		if ((World == NULL) || (World->GetPhysicsScene() == NULL))
+		{
+			return false;
+		}
+
+		using PTTraits = TSQTraits<ChaosInterface::FPTRaycastHit, ESweepOrRay::Raycast, ESingleMultiOrTest::Multi>;
+		return TSceneCastCommonImp<PTTraits, FRaycastSQAdditionalInputs>(World, OutHits, FRaycastSQAdditionalInputs(), Start, End, TraceChannel, Params, ResponseParams, ObjectParams, FDefaultAccelContainer{});
+	}
 
 	template <typename OverlapHitType, EQueryInfo InfoType>
 	bool OverlapImpl(const UWorld* World, const Chaos::FOverlapQueryData& OverlapData, const FCommonQueryData& CommonData, TArray<FOverlapResult>& OutOverlaps)
diff --git a/Engine/Source/Runtime/Engine/Private/Collision/WorldCollision.cpp b/Engine/Source/Runtime/Engine/Private/Collision/WorldCollision.cpp
--- a/Engine/Source/Runtime/Engine/Private/Collision/WorldCollision.cpp
+++ b/Engine/Source/Runtime/Engine/Private/Collision/WorldCollision.cpp
@@ -37,6 +37,8 @@ DEFINE_STAT(STAT_Collision_FBodyInstance_OverlapTest);
 DEFINE_STAT(STAT_Collision_FBodyInstance_LineTrace);
 DEFINE_STAT(STAT_Collision_PreFilter);
 DEFINE_STAT(STAT_Collision_SpherecastMultiple_Internal);
+DEFINE_STAT(STAT_Collision_GeomSweepMultiple_Internal);
+DEFINE_STAT(STAT_Collision_RaycastMultiple_Internal);
 
 /** default collision response container - to be used without reconstructing every time**/
 FCollisionResponseContainer FCollisionResponseContainer::DefaultResponseContainer(ECR_Block);
diff --git a/Engine/Source/Runtime/Engine/Public/Collision.h b/Engine/Source/Runtime/Engine/Public/Collision.h
--- a/Engine/Source/Runtime/Engine/Public/Collision.h
+++ b/Engine/Source/Runtime/Engine/Public/Collision.h
@@ -36,6 +36,8 @@ DECLARE_CYCLE_STAT_EXTERN(TEXT("BodyInstanceOverlapTest"), STAT_Collision_FBodyI
 DECLARE_CYCLE_STAT_EXTERN(TEXT("BodyInstanceLineTrace"), STAT_Collision_FBodyInstance_LineTrace, STATGROUP_Collision, );
 DECLARE_CYCLE_STAT_EXTERN(TEXT("PreFilter"), STAT_Collision_PreFilter, STATGROUP_CollisionVerbose, );
 DECLARE_CYCLE_STAT_EXTERN(TEXT("SpherecastMultipleInternal"), STAT_Collision_SpherecastMultiple_Internal, STATGROUP_Collision, );
+DECLARE_CYCLE_STAT_EXTERN(TEXT("GeomSweepMultipleInternal"), STAT_Collision_GeomSweepMultiple_Internal, STATGROUP_Collision, );
+DECLARE_CYCLE_STAT_EXTERN(TEXT("RaycastMultipleInternal"), STAT_Collision_RaycastMultiple_Internal, STATGROUP_Collision, );
 
 /** Enable collision analyzer support */
 #if (1 && !(UE_BUILD_SHIPPING || UE_BUILD_TEST) && WITH_EDITOR && WITH_UNREAL_DEVELOPER_TOOLS)
diff --git a/Engine/Source/Runtime/Engine/Public/Physics/GenericPhysicsInterface.h b/Engine/Source/Runtime/Engine/Public/Physics/GenericPhysicsInterface.h
--- a/Engine/Source/Runtime/Engine/Public/Physics/GenericPhysicsInterface.h
+++ b/Engine/Source/Runtime/Engine/Public/Physics/GenericPhysicsInterface.h
@@ -118,6 +118,27 @@ namespace Chaos::Private
 		*  Falls back to a raycast is the query radius is less than or equal to zero.
 		*/
 		static ENGINE_API bool SpherecastMulti(const UWorld* World, float QueryRadius, TArray<FHitResult>& OutHits, FVector Start, FVector End, ECollisionChannel TraceChannel, const FCollisionQueryParams& Params, const FCollisionResponseParams& ResponseParams, const FCollisionObjectQueryParams& ObjectParams = FCollisionObjectQueryParams::DefaultObjectQueryParam);
+
+		/**
+		*  INTERNAL USE ONLY
+		*  Physics thread shape sweep:
+		*  Trace a FCollisionShape against the world and return touching hits and then first blocking hit
+		*  Results are sorted, so a blocking hit (if found) will be the last element of the array
+		*  Only the single closest blocking result will be generated, no tests will be done after that.
+		*  Falls back to a raycast is the shape size is less than or equal to UE_KINDA_SMALL_NUMBER.
+		*/
+		UE_EXPERIMENTAL(5.6, "Physics Thread sweep trace, might have API changes")
+		static ENGINE_API bool GeomSweepMulti(const UWorld* World, const FCollisionShape& InGeom, const FQuat& InGeomRot, TArray<FHitResult>& OutHits, const FVector& Start, const FVector& End, ECollisionChannel TraceChannel, const FCollisionQueryParams& Params, const FCollisionResponseParams& ResponseParams, const FCollisionObjectQueryParams& ObjectParams = FCollisionObjectQueryParams::DefaultObjectQueryParam);
+
+		/**
+		*  INTERNAL USE ONLY
+		*  Physics thread raycast query:
+		*  Trace a ray against the world and return touching hits and then first blocking hit
+		*  Results are sorted, so a blocking hit (if found) will be the last element of the array
+		*  Only the single closest blocking result will be generated, no tests will be done after that.
+		*/
+		UE_EXPERIMENTAL(5.6, "Physics Thread raycast, might have API changes")
+		static ENGINE_API bool RaycastMulti(const UWorld* World, TArray<FHitResult>& OutHits, const FVector& Start, const FVector& End, ECollisionChannel TraceChannel, const FCollisionQueryParams& Params, const FCollisionResponseParams& ResponseParams, const FCollisionObjectQueryParams& ObjectParams = FCollisionObjectQueryParams::DefaultObjectQueryParam);
 	};
 
 	// This is for internal use so that the QueryInterface API can resume queries.
