[HWRT] Optimize FRayTracingGeometryManager::Tick() by avoiding some pointer dereferencing.

(Project with ~5K geometries / ~2K groups: 1.1ms -> 0.64ms)

diff --git a/Engine/Source/Runtime/Engine/Private/Rendering/RayTracingGeometryManager.cpp b/Engine/Source/Runtime/Engine/Private/Rendering/RayTracingGeometryManager.cpp
--- a/Engine/Source/Runtime/Engine/Private/Rendering/RayTracingGeometryManager.cpp
+++ b/Engine/Source/Runtime/Engine/Private/Rendering/RayTracingGeometryManager.cpp
@@ -127,6 +127,16 @@ static FAutoConsoleVariableRef CVarRayTracingShowOnScreenWarnings(
 	ECVF_RenderThreadSafe
 );
 
+#if DO_CHECK
+static bool GRayTracingTestCheckIntegrity = false;
+static FAutoConsoleVariableRef CVarRayTracingTestCheckIntegrity(
+	TEXT("r.RayTracing.Test.CheckIntegrity"),
+	GRayTracingTestCheckIntegrity,
+	TEXT("Whether to check integrity of cached state related to ray tracing."),
+	ECVF_RenderThreadSafe
+);
+#endif
+
 DECLARE_STATS_GROUP(TEXT("Ray Tracing Geometry"), STATGROUP_RayTracingGeometry, STATCAT_Advanced);
 
 DECLARE_DWORD_ACCUMULATOR_STAT(TEXT("Geometry Count"), STAT_RayTracingGeometryCount, STATGROUP_RayTracingGeometry);
@@ -544,6 +554,16 @@ void FRayTracingGeometryManager::Tick(FRHICommandList& RHICmdList)
 	checkf(bUsingReferenceBasedResidency || (ReferencedGeometryHandles.IsEmpty() && ReferencedGeometryGroups.IsEmpty() && ReferencedGeometryGroupsForDynamicUpdate.IsEmpty()),
 		TEXT("ReferencedGeometryHandles, ReferencedGeometryGroups and ReferencedGeometryGroupsForDynamicUpdate are expected to be empty when not using reference based residency"));
 
+#if DO_CHECK
+	if (GRayTracingTestCheckIntegrity)
+	{
+		for (FRegisteredGeometry& RegisteredGeometry : RegisteredGeometries)
+		{
+			check(RegisteredGeometry.bEvicted == RegisteredGeometry.Geometry->IsEvicted());
+		}
+	}
+#endif
+
 	if (bRefreshAlwaysResidentRayTracingGeometries)
 	{
 		bRefreshAlwaysResidentRayTracingGeometries = false;
@@ -566,7 +586,7 @@ void FRayTracingGeometryManager::Tick(FRHICommandList& RHICmdList)
 
 				if (RegisteredGeometry.Geometry->IsEvicted())
 				{
-					RegisteredGeometry.Geometry->MakeResident(RHICmdList);
+					MakeGeometryResident(RHICmdList, RegisteredGeometry);
 				}
 
 				if (!RequestRayTracingGeometryStreamIn(RHICmdList, RegisteredGeometry.Geometry->RayTracingGeometryHandle))
@@ -647,7 +667,7 @@ void FRayTracingGeometryManager::Tick(FRHICommandList& RHICmdList)
 
 				if (RegisteredGeometry.Geometry->IsEvicted())
 				{
-					RegisteredGeometry.Geometry->MakeResident(RHICmdList);
+					MakeGeometryResident(RHICmdList, RegisteredGeometry);
 				}
 				if (!RequestRayTracingGeometryStreamIn(RHICmdList, GeometryHandle))
 				{
@@ -705,16 +725,16 @@ void FRayTracingGeometryManager::Tick(FRHICommandList& RHICmdList)
 
 					RequestedSize += RegisteredGeometry.Size;
 
-					if (RegisteredGeometry.Geometry->LODIndex >= Group.GeometryHandles.Num() - GRayTracingNumAlwaysResidentLODs)
+					if (LODIndex >= Group.GeometryHandles.Num() - GRayTracingNumAlwaysResidentLODs)
 					{
-						checkf(!RegisteredGeometry.Geometry->IsEvicted(), TEXT("Always resident ray tracing geometry was unexpectely evicted."));
+						checkf(!RegisteredGeometry.bEvicted, TEXT("Always resident ray tracing geometry was unexpectely evicted."));
 					}
 					else
 					{
 						ReferencedGeometries.Add(GeometryHandle);
 						NotReferencedResidentGeometries.Remove(GeometryHandle);
 
-						if (RegisteredGeometry.Geometry->IsEvicted())
+						if (RegisteredGeometry.bEvicted)
 						{
 							RequestedButEvictedSize += RegisteredGeometry.Size;
 						}
@@ -799,7 +819,7 @@ void FRayTracingGeometryManager::Tick(FRHICommandList& RHICmdList)
 
 				if (RegisteredGeometry.Geometry->IsEvicted())
 				{
-					RegisteredGeometry.Geometry->MakeResident(RHICmdList);
+					MakeGeometryResident(RHICmdList, RegisteredGeometry);
 				}
 
 				RequestRayTracingGeometryStreamIn(RHICmdList, GeometryHandle);
@@ -842,7 +862,7 @@ void FRayTracingGeometryManager::Tick(FRHICommandList& RHICmdList)
 			{
 				if (RegisteredGeometry.Geometry->IsEvicted())
 				{
-					RegisteredGeometry.Geometry->MakeResident(RHICmdList);
+					MakeGeometryResident(RHICmdList, RegisteredGeometry);
 				}
 
 				if (!RequestRayTracingGeometryStreamIn(RHICmdList, RegisteredGeometry.Geometry->RayTracingGeometryHandle))
@@ -1114,6 +1134,12 @@ void FRayTracingGeometryManager::StreamOutGeometry(FRHIResourceReplaceBatcher& B
 	}
 }
 
+void FRayTracingGeometryManager::MakeGeometryResident(FRHICommandList& RHICmdList, FRegisteredGeometry& RegisteredGeometry)
+{
+	RegisteredGeometry.Geometry->MakeResident(RHICmdList);
+	RegisteredGeometry.bEvicted = false;
+}
+
 void FRayTracingGeometryManager::EvictGeometry(FRHICommandListBase& RHICmdList, FRegisteredGeometry& RegisteredGeometry)
 {
 	// Cancel associated streaming request if currently in-flight
@@ -1128,6 +1154,7 @@ void FRayTracingGeometryManager::EvictGeometry(FRHICommandListBase& RHICmdList,
 	StreamOutGeometry(Batcher, RegisteredGeometry);
 
 	RegisteredGeometry.Geometry->Evict();
+	RegisteredGeometry.bEvicted = true;
 }
 
 void FRayTracingGeometryManager::BoostPriority(FBuildRequestIndex InRequestIndex, float InBoostValue)
diff --git a/Engine/Source/Runtime/Engine/Public/Rendering/RayTracingGeometryManager.h b/Engine/Source/Runtime/Engine/Public/Rendering/RayTracingGeometryManager.h
--- a/Engine/Source/Runtime/Engine/Public/Rendering/RayTracingGeometryManager.h
+++ b/Engine/Source/Runtime/Engine/Public/Rendering/RayTracingGeometryManager.h
@@ -130,6 +130,7 @@ private:
 		};
 
 		FStatus Status = FStatus::StreamedOut;
+		bool bEvicted = false;
 	};
 
 	struct FStreamingRequest
@@ -172,6 +173,10 @@ private:
 	// Must call this instead of FRayTracingGeometry::ReleaseRHIForStreaming() directly
 	void StreamOutGeometry(FRHIResourceReplaceBatcher& Batcher, FRegisteredGeometry& RegisteredGeometry);
 
+	// Helper function to make geometry resident + necessary management
+	// Must call this instead of FRayTracingGeometry::MakeResident() directly
+	void MakeGeometryResident(FRHICommandList& RHICmdList, FRegisteredGeometry& RegisteredGeometry);
+
 	// Helper function to evict geometry + necessary streaming management
 	// Must call this instead of FRayTracingGeometry::Evict() directly
 	void EvictGeometry(FRHICommandListBase& RHICmdList, FRegisteredGeometry& RegisteredGeometry);
