NetTestHelpers
* Added float cvar override support
* Added nodiscard to scope cvar constructor

#rnx

diff --git a/Engine/Source/Runtime/Engine/Private/Tests/NetTestHelpers.cpp b/Engine/Source/Runtime/Engine/Private/Tests/NetTestHelpers.cpp
--- a/Engine/Source/Runtime/Engine/Private/Tests/NetTestHelpers.cpp
+++ b/Engine/Source/Runtime/Engine/Private/Tests/NetTestHelpers.cpp
@@ -588,6 +588,27 @@ FScopedCVarOverrideInt::~FScopedCVarOverrideInt()
 	}
 }
 
+//------------------------------------------------------------------------
+
+FScopedCVarOverrideFloat::FScopedCVarOverrideFloat(const TCHAR* VariableName, float Value)
+: Variable(IConsoleManager::Get().FindConsoleVariable(VariableName))
+{
+	if (Variable)
+	{
+		SavedValue = Variable->GetFloat();
+		Variable->Set(Value, ECVF_SetByConsole /* Use highest priority so we override even console setters */);
+	}
+}
+
+FScopedCVarOverrideFloat::~FScopedCVarOverrideFloat()
+{
+	if (Variable)
+	{
+		Variable->Set(SavedValue, ECVF_SetByConsole);
+	}
+}
+
+
 //------------------------------------------------------------------------
 
 FScopedTestSettings::FScopedTestSettings()
diff --git a/Engine/Source/Runtime/Engine/Public/Tests/NetTestHelpers.h b/Engine/Source/Runtime/Engine/Public/Tests/NetTestHelpers.h
--- a/Engine/Source/Runtime/Engine/Public/Tests/NetTestHelpers.h
+++ b/Engine/Source/Runtime/Engine/Public/Tests/NetTestHelpers.h
@@ -23,7 +23,7 @@ namespace UE::Net
 /*
  * Used to set a cvar and restore it to it's original value when destroyed
  */
-class FScopedCVarOverrideInt
+class [[nodiscard]] FScopedCVarOverrideInt
 {
 public:
 	ENGINE_API FScopedCVarOverrideInt(const TCHAR* VariableName, int32 Value);
@@ -39,6 +39,22 @@ private:
 	int32 SavedValue = 0;
 };
 
+class [[nodiscard]] FScopedCVarOverrideFloat
+{
+public:
+	ENGINE_API FScopedCVarOverrideFloat(const TCHAR* VariableName, float Value);
+	ENGINE_API ~FScopedCVarOverrideFloat();
+
+	FScopedCVarOverrideFloat(FScopedCVarOverrideFloat&&) = delete;
+	FScopedCVarOverrideFloat(const FScopedCVarOverrideFloat&) = delete;
+	FScopedCVarOverrideFloat& operator=(FScopedCVarOverrideFloat&&) = delete;
+	FScopedCVarOverrideFloat& operator=(const FScopedCVarOverrideFloat&) = delete;
+
+private:
+	IConsoleVariable* Variable = nullptr;
+	float SavedValue = 0;
+};
+
 /**
 * Sets and restores  cvars needed to use FNetTestWorldInstances within a scope.
 * Meant to be used within a single function.
