Fix UnsafeTypeCastWarningLevel for WorldPartitionHLODUtilities and HierarchicalLODOutliner

#tests Compiled + Preflight
#rnx

diff --git a/Engine/Plugins/AI/MassAI/Source/MassNavigation/Private/MassNavigationDebug.cpp b/Engine/Plugins/AI/MassAI/Source/MassNavigation/Private/MassNavigationDebug.cpp
--- a/Engine/Plugins/AI/MassAI/Source/MassNavigation/Private/MassNavigationDebug.cpp
+++ b/Engine/Plugins/AI/MassAI/Source/MassNavigation/Private/MassNavigationDebug.cpp
@@ -1,4 +1,4 @@
-﻿// Copyright Epic Games, Inc. All Rights Reserved.
+// Copyright Epic Games, Inc. All Rights Reserved.
 
 #include "MassNavigationDebug.h"
 
@@ -133,7 +133,7 @@ namespace UE::MassNavigation::Debug
 
 		if (UseDrawDebugHelper() && Context.World)
 		{
-			DrawDebugCircle(Context.World, Bottom, Radius, /*Segments = */24, Color);
+			DrawDebugCircle(Context.World, Bottom, InRadius, /*Segments = */24, Color);
 		}
 	}
 } // UE::MassNavigation::Debug
diff --git a/Engine/Plugins/Editor/WorldPartitionHLODUtilities/Source/Private/WorldPartition/HLOD/Utilities/WorldPartitionHLODUtilities.cpp b/Engine/Plugins/Editor/WorldPartitionHLODUtilities/Source/Private/WorldPartition/HLOD/Utilities/WorldPartitionHLODUtilities.cpp
--- a/Engine/Plugins/Editor/WorldPartitionHLODUtilities/Source/Private/WorldPartition/HLOD/Utilities/WorldPartitionHLODUtilities.cpp
+++ b/Engine/Plugins/Editor/WorldPartitionHLODUtilities/Source/Private/WorldPartition/HLOD/Utilities/WorldPartitionHLODUtilities.cpp
@@ -655,12 +655,12 @@ ULevelStreaming* LoadSourceActors(AWorldPartitionHLOD* InHLODActor, bool& bOutIs
 
 		// Ensure streaming requests are completed
 		// Some requests require the main thread to be ticked in order to run to completion.
-		const double STREAMING_WAIT_DT = 0.1;
+		const float STREAMING_WAIT_DT = 0.1f;
 		while (IStreamingManager::Get().StreamAllResources(STREAMING_WAIT_DT) > 0)
 		{
 			// Application tick.
 			FTaskGraphInterface::Get().ProcessThreadUntilIdle(ENamedThreads::GameThread);
-			FTSTicker::GetCoreTicker().Tick(FApp::GetDeltaTime());
+			FTSTicker::GetCoreTicker().Tick(static_cast<float>(FApp::GetDeltaTime()));
 		}
 	}
 
diff --git a/Engine/Plugins/Editor/WorldPartitionHLODUtilities/Source/WorldPartitionHLODUtilities.Build.cs b/Engine/Plugins/Editor/WorldPartitionHLODUtilities/Source/WorldPartitionHLODUtilities.Build.cs
--- a/Engine/Plugins/Editor/WorldPartitionHLODUtilities/Source/WorldPartitionHLODUtilities.Build.cs
+++ b/Engine/Plugins/Editor/WorldPartitionHLODUtilities/Source/WorldPartitionHLODUtilities.Build.cs
@@ -6,6 +6,8 @@ public class WorldPartitionHLODUtilities : ModuleRules
 {
     public WorldPartitionHLODUtilities(ReadOnlyTargetRules Target) : base(Target)
 	{
+		CppCompileWarningSettings.UnsafeTypeCastWarningLevel = WarningLevel.Error;
+
 		// Some files were initially right under /Public and were moved to sub directories
 		// Ensure we don't break old-style includes
 		PublicIncludePaths.Add(Path.Combine(ModuleDirectory, "Public/WorldPartition/HLOD/Builders"));
diff --git a/Engine/Source/Editor/HierarchicalLODOutliner/HierarchicalLODOutliner.Build.cs b/Engine/Source/Editor/HierarchicalLODOutliner/HierarchicalLODOutliner.Build.cs
--- a/Engine/Source/Editor/HierarchicalLODOutliner/HierarchicalLODOutliner.Build.cs
+++ b/Engine/Source/Editor/HierarchicalLODOutliner/HierarchicalLODOutliner.Build.cs
@@ -6,7 +6,9 @@ public class HierarchicalLODOutliner : ModuleRules
 {
     public HierarchicalLODOutliner(ReadOnlyTargetRules Target) : base(Target)
 	{
-        PublicDependencyModuleNames.AddRange
+		CppCompileWarningSettings.UnsafeTypeCastWarningLevel = WarningLevel.Error;
+
+		PublicDependencyModuleNames.AddRange
         (
             new string[] {
 				"Core",
diff --git a/Engine/Source/Editor/HierarchicalLODOutliner/Private/HLODOutliner.cpp b/Engine/Source/Editor/HierarchicalLODOutliner/Private/HLODOutliner.cpp
--- a/Engine/Source/Editor/HierarchicalLODOutliner/Private/HLODOutliner.cpp
+++ b/Engine/Source/Editor/HierarchicalLODOutliner/Private/HLODOutliner.cpp
@@ -467,7 +467,7 @@ namespace HLODOutliner
 					[
 						SNew(SButton)
 						.ButtonStyle(FAppStyle::Get(), "SimpleButton")
-						.ContentPadding(0)
+						.ContentPadding(0.0f)
 						.HAlign(HAlign_Center)
 						.OnClicked(this, &SHLODOutliner::GenerateClustersFromUI)
 						.ToolTipText(LOCTEXT("GenerateClusterToolTip", "Generates clusters (but not proxy meshes) for meshes in the level"))
@@ -485,7 +485,7 @@ namespace HLODOutliner
 					[
 						SNew(SButton)
 						.ButtonStyle(FAppStyle::Get(), "SimpleButton")
-						.ContentPadding(0)
+						.ContentPadding(0.0f)
 						.HAlign(HAlign_Center)
 						.OnClicked_Lambda([this]() { return HandleDeleteHLODs(); })
 						.IsEnabled(this, &SHLODOutliner::CanDeleteHLODs)
@@ -1159,7 +1159,7 @@ namespace HLODOutliner
 			// This call came from a context menu
 			auto SelectedItems = TreeView->GetSelectedItems();
 			
-			FScopedSlowTask SlowTask(SelectedItems.Num(), FText::Format(LOCTEXT("SHLODOutliner_BuildProxy", "Building Proxy {0}|plural(one=Mesh,other=Meshes)"), SelectedItems.Num()));
+			FScopedSlowTask SlowTask((float)SelectedItems.Num(), FText::Format(LOCTEXT("SHLODOutliner_BuildProxy", "Building Proxy {0}|plural(one=Mesh,other=Meshes)"), SelectedItems.Num()));
 			SlowTask.MakeDialog();
 
 			// Loop over all selected items (context menu can't be called with multiple items selected that aren't of the same types)
@@ -1189,7 +1189,7 @@ namespace HLODOutliner
 		}
 		
 		// Show message log if there was an HLOD message
-		FMessageLog("HLODResults").Open();		
+		FMessageLog("HLODResults").Open();
 	}
 
 	void SHLODOutliner::RebuildLODActor()
@@ -1201,7 +1201,7 @@ namespace HLODOutliner
 			// This call came from a context menu
 			auto SelectedItems = TreeView->GetSelectedItems();
 
-			FScopedSlowTask SlowTask(SelectedItems.Num(), FText::Format(LOCTEXT("SHLODOutliner_RebuildProxy", "Rebuild Proxy {0}|plural(one=Mesh,other=Meshes)"), SelectedItems.Num()));
+			FScopedSlowTask SlowTask((float)SelectedItems.Num(), FText::Format(LOCTEXT("SHLODOutliner_RebuildProxy", "Rebuild Proxy {0}|plural(one=Mesh,other=Meshes)"), SelectedItems.Num()));
 			SlowTask.MakeDialog();
 
 			// Loop over all selected items (context menu can't be called with multiple items selected that aren't of the same types)
diff --git a/Engine/Source/Editor/HierarchicalLODOutliner/Private/HLODTreeWidgetItem.cpp b/Engine/Source/Editor/HierarchicalLODOutliner/Private/HLODTreeWidgetItem.cpp
--- a/Engine/Source/Editor/HierarchicalLODOutliner/Private/HLODTreeWidgetItem.cpp
+++ b/Engine/Source/Editor/HierarchicalLODOutliner/Private/HLODTreeWidgetItem.cpp
@@ -196,7 +196,7 @@ BEGIN_SLATE_FUNCTION_BUILD_OPTIMIZATION
 		WeakTableViewBase = InOwner;
 
 		auto Args = FSuperRowType::FArguments()
-			.Padding(1)
+			.Padding(1.0f)
 			.OnDragDetected_Static(HLODOutliner::OnDragDetected, TWeakPtr<STableViewBase>(InOwner));
 
 		SMultiColumnTableRow< FTreeItemPtr >::Construct(Args, InOwner);
diff --git a/Engine/Source/Runtime/Engine/Private/DrawDebugHelpers.cpp b/Engine/Source/Runtime/Engine/Private/DrawDebugHelpers.cpp
--- a/Engine/Source/Runtime/Engine/Private/DrawDebugHelpers.cpp
+++ b/Engine/Source/Runtime/Engine/Private/DrawDebugHelpers.cpp
@@ -437,7 +437,7 @@ static void InternalDrawDebugCircle(const UWorld* InWorld, const FMatrix& Transf
 	}
 }
 
-void DrawDebugCircle(const UWorld* InWorld, const FMatrix& TransformMatrix, float Radius, int32 Segments, const FColor& Color, bool bPersistentLines, float LifeTime, uint8 DepthPriority, float Thickness, bool bDrawAxis)
+void DrawDebugCircle(const UWorld* InWorld, const FMatrix& TransformMatrix, double Radius, int32 Segments, const FColor& Color, bool bPersistentLines, float LifeTime, uint8 DepthPriority, float Thickness, bool bDrawAxis)
 {
 	RETURN_IF_DEBUG_DRAW_DISABLED;
 	
@@ -470,7 +470,7 @@ void DrawDebugCircle(const UWorld* InWorld, const FMatrix& TransformMatrix, floa
 	}
 }
 
-void DrawDebugCircle(const UWorld* InWorld, FVector Center, float Radius, int32 Segments, const FColor& Color, bool PersistentLines, float LifeTime, uint8 DepthPriority, float Thickness, FVector YAxis, FVector ZAxis, bool bDrawAxis)
+void DrawDebugCircle(const UWorld* InWorld, FVector Center, double Radius, int32 Segments, const FColor& Color, bool PersistentLines, float LifeTime, uint8 DepthPriority, float Thickness, FVector YAxis, FVector ZAxis, bool bDrawAxis)
 {
 	RETURN_IF_DEBUG_DRAW_DISABLED;
 	
@@ -790,7 +790,7 @@ void DrawDebugFrustum(const UWorld* InWorld, const FMatrix& FrustumToWorld, FCol
 	}
 }
 
-void DrawCircle(const UWorld* InWorld, const FVector& Base, const FVector& X, const FVector& Y, const FColor& Color, float Radius, int32 NumSides, bool bPersistentLines, float LifeTime, uint8 DepthPriority, float Thickness)
+void DrawCircle(const UWorld* InWorld, const FVector& Base, const FVector& X, const FVector& Y, const FColor& Color, double Radius, int32 NumSides, bool bPersistentLines, float LifeTime, uint8 DepthPriority, float Thickness)
 {
 	RETURN_IF_DEBUG_DRAW_DISABLED;
 	
diff --git a/Engine/Source/Runtime/Engine/Public/DrawDebugHelpers.h b/Engine/Source/Runtime/Engine/Public/DrawDebugHelpers.h
--- a/Engine/Source/Runtime/Engine/Public/DrawDebugHelpers.h
+++ b/Engine/Source/Runtime/Engine/Public/DrawDebugHelpers.h
@@ -33,9 +33,9 @@ ENGINE_API void DrawDebugCoordinateSystem(const UWorld* InWorld, FVector const&
 /** Draw Debug crosshair */
 ENGINE_API void DrawDebugCrosshairs(const UWorld* InWorld, FVector const& AxisLoc, FRotator const& AxisRot, float Scale, const FColor& Color = FColor::White, bool bPersistentLines = false, float LifeTime = -1.f, uint8 DepthPriority = 0);
 /** Draw Debug Circle */
-ENGINE_API void DrawDebugCircle(const UWorld* InWorld, const FMatrix& TransformMatrix, float Radius, int32 Segments, const FColor& Color, bool bPersistentLines = false, float LifeTime=-1.f, uint8 DepthPriority = 0, float Thickness = 0.f, bool bDrawAxis=true);
+ENGINE_API void DrawDebugCircle(const UWorld* InWorld, const FMatrix& TransformMatrix, double Radius, int32 Segments, const FColor& Color, bool bPersistentLines = false, float LifeTime=-1.f, uint8 DepthPriority = 0, float Thickness = 0.f, bool bDrawAxis=true);
 /** Draw Debug Circle */
-ENGINE_API void DrawDebugCircle(const UWorld* InWorld, FVector Center, float Radius, int32 Segments, const FColor& Color, bool bPersistentLines = false, float LifeTime=-1.f, uint8 DepthPriority = 0, float Thickness = 0.f, FVector YAxis=FVector(0.f,1.f,0.f), FVector ZAxis=FVector(0.f,0.f,1.f), bool bDrawAxis=true);
+ENGINE_API void DrawDebugCircle(const UWorld* InWorld, FVector Center, double Radius, int32 Segments, const FColor& Color, bool bPersistentLines = false, float LifeTime=-1.f, uint8 DepthPriority = 0, float Thickness = 0.f, FVector YAxis=FVector(0.f,1.f,0.f), FVector ZAxis=FVector(0.f,0.f,1.f), bool bDrawAxis=true);
 /** Draw Debug Circle Arc */
 ENGINE_API void DrawDebugCircleArc(const UWorld* InWorld, const FVector& Center, float Radius, const FVector& Direction, float AngleWidth, int32 Segments, const FColor& Color, bool PersistentLines = false, float LifeTime = -1.f, uint8 DepthPriority = 0, float Thickness = 0.f);
 /** Draw Debug 2D donut */
@@ -51,7 +51,7 @@ ENGINE_API void DrawDebugAltCone(const UWorld* InWorld, FVector const& Origin, F
 ENGINE_API void DrawDebugString(const UWorld* InWorld, FVector const& TextLocation, const FString& Text, class AActor* TestBaseActor = NULL, FColor const& TextColor = FColor::White, float Duration = -1.000000, bool bDrawShadow = false, float FontScale = 1.f);
 ENGINE_API void DrawDebugFrustum(const UWorld* InWorld, const FMatrix& FrustumToWorld, FColor const& Color, bool bPersistentLines = false, float LifeTime=-1.f, uint8 DepthPriority = 0, float Thickness = 0.f);
 /** Draw a circle using user define axis and radius */
-ENGINE_API void DrawCircle(const UWorld* InWorld, const FVector& Base, const FVector& X, const FVector& Y, const FColor& Color, float Radius, int32 NumSides, bool bPersistentLines = false, float LifeTime = -1.f, uint8 DepthPriority = 0, float Thickness = 0);
+ENGINE_API void DrawCircle(const UWorld* InWorld, const FVector& Base, const FVector& X, const FVector& Y, const FColor& Color, double Radius, int32 NumSides, bool bPersistentLines = false, float LifeTime = -1.f, uint8 DepthPriority = 0, float Thickness = 0);
 /** Draw a capsule using the LineBatcher */
 ENGINE_API void DrawDebugCapsule(const UWorld* InWorld, FVector const& Center, float HalfHeight, float Radius, const FQuat& Rotation, FColor const& Color, bool bPersistentLines = false, float LifeTime = -1.f, uint8 DepthPriority = 0, float Thickness = 0);
 /** Draw a debug camera shape.  FOV is full angle in degrees. */
@@ -188,8 +188,8 @@ FORCEINLINE void DrawDebugBox(const UWorld* InWorld, FVector const& Center, FVec
 FORCEINLINE void DrawDebugBox(const UWorld* InWorld, FVector const& Center, FVector const& Extent, const FQuat& Rotation, FColor const& Color, bool bPersistentLines = false, float LifeTime = -1.f, uint8 DepthPriority = 0, float Thickness = 0.f) {}
 FORCEINLINE void DrawDebugCoordinateSystem(const UWorld* InWorld, FVector const& AxisLoc, FRotator const& AxisRot, float Scale, bool bPersistentLines = false, float LifeTime = -1.f, uint8 DepthPriority = 0, float Thickness = 0.f) {}
 FORCEINLINE void DrawDebugCrosshairs(const UWorld* InWorld, FVector const& AxisLoc, FRotator const& AxisRot, float Scale, const FColor& Color = FColor::White, bool bPersistentLines = false, float LifeTime = -1.f, uint8 DepthPriority = 0) {}
-FORCEINLINE void DrawDebugCircle(const UWorld* InWorld, const FMatrix& TransformMatrix, float Radius, int32 Segments, const FColor& Color, bool bPersistentLines = false, float LifeTime = -1.f, uint8 DepthPriority = 0, float Thickness = 0.f, bool bDrawAxis = true) {}
-FORCEINLINE void DrawDebugCircle(const UWorld* InWorld, FVector Center, float Radius, int32 Segments, const FColor& Color, bool bPersistentLines = false, float LifeTime = -1.f, uint8 DepthPriority = 0, float Thickness = 0.f, FVector YAxis = FVector(0.f, 1.f, 0.f), FVector ZAxis = FVector(0.f, 0.f, 1.f), bool bDrawAxis = true) {}
+FORCEINLINE void DrawDebugCircle(const UWorld* InWorld, const FMatrix& TransformMatrix, double Radius, int32 Segments, const FColor& Color, bool bPersistentLines = false, float LifeTime = -1.f, uint8 DepthPriority = 0, float Thickness = 0.f, bool bDrawAxis = true) {}
+FORCEINLINE void DrawDebugCircle(const UWorld* InWorld, FVector Center, double Radius, int32 Segments, const FColor& Color, bool bPersistentLines = false, float LifeTime = -1.f, uint8 DepthPriority = 0, float Thickness = 0.f, FVector YAxis = FVector(0.f, 1.f, 0.f), FVector ZAxis = FVector(0.f, 0.f, 1.f), bool bDrawAxis = true) {}
 FORCEINLINE void DrawDebugCircleArc(const UWorld* InWorld, FVector const& Center, float Radius, FVector const& Direction, float AngleWidth, int32 Segments, const FColor& Color, bool bPersistentLines = false, float LifeTime = -1.f, uint8 DepthPriority = 0, float Thickness = 0.f) {}
 FORCEINLINE void DrawDebug2DDonut(const UWorld* InWorld, const FMatrix& TransformMatrix, float InnerRadius, float OuterRadius, int32 Segments, const FColor& Color, bool bPersistentLines = false, float LifeTime = -1.f, uint8 DepthPriority = 0, float Thickness = 0.f) {}
 FORCEINLINE void DrawDebugSphere(const UWorld* InWorld, FVector const& Center, float Radius, int32 Segments, FColor const& Color, bool bPersistentLines = false, float LifeTime = -1.f, uint8 DepthPriority = 0, float Thickness = 0.f) {}
@@ -198,7 +198,7 @@ FORCEINLINE void DrawDebugCone(const UWorld* InWorld, FVector const& Origin, FVe
 FORCEINLINE void DrawDebugAltCone(const UWorld* InWorld, FVector const& Origin, FRotator const& Rotation, float Length, float AngleWidth, float AngleHeight, FColor const& DrawColor, bool bPersistentLines = false, float LifeTime = -1.f, uint8 DepthPriority = 0, float Thickness = 0.f) {}
 FORCEINLINE void DrawDebugString(const UWorld* InWorld, FVector const& TextLocation, const FString& Text, class AActor* TestBaseActor = NULL, FColor const& TextColor = FColor::White, float Duration = -1.000000, bool bDrawShadow = false, float FontScale = 1.f) {}
 FORCEINLINE void DrawDebugFrustum(const UWorld* InWorld, const FMatrix& FrustumToWorld, FColor const& Color, bool bPersistentLines = false, float LifeTime = -1.f, uint8 DepthPriority = 0, float Thickness = 0.f) {}
-FORCEINLINE void DrawCircle(const UWorld* InWorld, const FVector& Base, const FVector& X, const FVector& Y, const FColor& Color, float Radius, int32 NumSides, bool bPersistentLines = false, float LifeTime = -1.f, uint8 DepthPriority = 0, float Thickness = 0) {}
+FORCEINLINE void DrawCircle(const UWorld* InWorld, const FVector& Base, const FVector& X, const FVector& Y, const FColor& Color, double Radius, int32 NumSides, bool bPersistentLines = false, float LifeTime = -1.f, uint8 DepthPriority = 0, float Thickness = 0) {}
 FORCEINLINE void DrawDebugCapsule(const UWorld* InWorld, FVector const& Center, float HalfHeight, float Radius, const FQuat& Rotation, FColor const& Color, bool bPersistentLines = false, float LifeTime = -1.f, uint8 DepthPriority = 0, float Thickness = 0) {}
 FORCEINLINE void DrawDebugCamera(const UWorld* InWorld, FVector const& Location, FRotator const& Rotation, float FOVDeg, float Scale = 1.f, FColor const& Color = FColor::White, bool bPersistentLines = false, float LifeTime = -1.f, uint8 DepthPriority = 0) {}
 FORCEINLINE void DrawCentripetalCatmullRomSpline(const UWorld* InWorld, TConstArrayView<FVector> Points, FColor const& Color, float Alpha = 0.5f, int32 NumSamplesPerSegment = 8, bool bPersistentLines = false, float LifeTime = -1.f, uint8 DepthPriority = 0, float Thickness = 0.f) {}
