UMaterial : make StateId deterministic for class default object. I need this so blueprint dependencies are deterministic and don't recook all the time.

#rnx
#incrementalcook

diff --git a/Engine/Source/Runtime/Engine/Private/Materials/Material.cpp b/Engine/Source/Runtime/Engine/Private/Materials/Material.cpp
--- a/Engine/Source/Runtime/Engine/Private/Materials/Material.cpp
+++ b/Engine/Source/Runtime/Engine/Private/Materials/Material.cpp
@@ -1165,13 +1165,19 @@ void UMaterial::PostInitProperties()
 	LLM_SCOPE(ELLMTag::Materials);
 
 	Super::PostInitProperties();
-	if(!HasAnyFlags(RF_ClassDefaultObject))
+	
+	if (HasAnyFlags(RF_ClassDefaultObject))
 	{
-		DefaultMaterialInstance = new FDefaultMaterialInstance(this);
+		// CDO needs a deterministic guid for incremental cook of blueprints.
+		StateId = FGuid();
 	}
+	else
+	{
+		DefaultMaterialInstance = new FDefaultMaterialInstance(this);
 
-	// Initialize StateId to something unique, in case this is a new material
-	FPlatformMisc::CreateGuid(StateId);
+		// Initialize StateId to something unique, in case this is a new material
+		FPlatformMisc::CreateGuid(StateId);
+	}
 }
 
 FMaterialResource* UMaterial::AllocateResource()
