Bring back setting TextureReferenceRHI for FVirtualTexture2DResource.
This allows us to use TextureReferenceRHI for cases where we can't do virtual texture sampling. For example in thumbnail generation, or when using editor preview modes for platforms with virtual textures disabled.
This was previously removed from FVirtualTexture2DResource because, if virtual textures were all recreated at the same time, at scale, we would see OOM from texture heaps. This happened in high resolution MRQ virtual production scenarios, where triggering a resize of the virtual texture pools could force a full recreate of all virtual textures.
The change reuses the referenced TextureRHI if the texture description didn't change, which is the common case. That avoids the texture recreate and the set texture reference.
#rnx

diff --git a/Engine/Source/Runtime/Engine/Private/Texture2D.cpp b/Engine/Source/Runtime/Engine/Private/Texture2D.cpp
--- a/Engine/Source/Runtime/Engine/Private/Texture2D.cpp
+++ b/Engine/Source/Runtime/Engine/Private/Texture2D.cpp
@@ -1689,13 +1689,32 @@ void FVirtualTexture2DResource::InitializeEditorResources(IVirtualTexture* InVir
 			const uint32 TileSizeInPixels = bCopyUnwantedBordersForAlignment ? GetTileSize() + 2 * GetBorderSize() : GetTileSize();
 			const uint32 LayerMask = 1u;
 
-			// Create the target texture to composite tiles to.
-			FRDGTextureRef Texture = GraphBuilder.CreateTexture(FRDGTextureDesc::Create2D(
-				FIntPoint(MipWidthInTiles * TileSizeInPixels, MipHeightInTiles * TileSizeInPixels),
-				PixelFormat, FClearValueBinding::Black,	TexCreateFlags), TEXT("VTEditorResource"));
+			// Get or create the target texture to composite tiles to.
+			TRefCountPtr<IPooledRenderTarget> PooledRenderTarget;
+			
+			// First try reusing the current referenced texture.
+			// We do that because swapping the referenced texture can be expensive if we are doing a recreate of all textures.
+			FRHITexture* ReferencedTexture = TextureReferenceRHI->GetReferencedTexture();
+			if (ReferencedTexture != nullptr)
+			{
+				FRHITextureDesc const& ExistingDesc = ReferencedTexture->GetDesc();
+				if (ExistingDesc.Format == PixelFormat && ExistingDesc.GetSize() == FIntVector(MipWidthInTiles * TileSizeInPixels, MipHeightInTiles * TileSizeInPixels, 1))
+				{
+					PooledRenderTarget = CreateRenderTarget(ReferencedTexture, TEXT("VTEditorResource"));
+					Texture2DRHI = ReferencedTexture;
+				}
+			}
 
-			TRefCountPtr<IPooledRenderTarget> PooledRenderTarget = GraphBuilder.ConvertToExternalTexture(Texture);
-			Texture2DRHI = PooledRenderTarget->GetRHI();
+			// If we can't reuse the referenced texture then create it.
+			if (Texture2DRHI == nullptr)
+			{
+				FRDGTextureRef Texture = GraphBuilder.CreateTexture(FRDGTextureDesc::Create2D(
+					FIntPoint(MipWidthInTiles * TileSizeInPixels, MipHeightInTiles * TileSizeInPixels),
+					PixelFormat, FClearValueBinding::Black,	TexCreateFlags), TEXT("VTEditorResource"));
+
+				PooledRenderTarget = GraphBuilder.ConvertToExternalTexture(Texture);
+				Texture2DRHI = PooledRenderTarget->GetRHI();
+			}
 
 			// Produce all the tiles.
 			TArray<IVirtualTextureFinalizer*> Finalizers;
@@ -1743,6 +1762,12 @@ void FVirtualTexture2DResource::InitializeEditorResources(IVirtualTexture* InVir
 		TextureRHI = Texture2DRHI;
 		TextureRHI->SetName(TextureName);
 		RHIBindDebugLabelName(TextureRHI, *Name);
+		
+		// Only update texture reference if necessary.
+		if (TextureRHI != TextureReferenceRHI->GetReferencedTexture())
+		{
+			RHIUpdateTextureReference(TextureReferenceRHI, TextureRHI);
+		}
 	}
 }
 #endif // WITH_EDITOR
