Prevent crash when diffing two empty strings

#rnx

diff --git a/Engine/Source/Runtime/Engine/Private/KismetStringLibrary.cpp b/Engine/Source/Runtime/Engine/Private/KismetStringLibrary.cpp
--- a/Engine/Source/Runtime/Engine/Private/KismetStringLibrary.cpp
+++ b/Engine/Source/Runtime/Engine/Private/KismetStringLibrary.cpp
@@ -755,6 +755,10 @@ FString UKismetStringLibrary::DiffString(const FString& First, const FString& Se
 		{
 			return TAttribute<TArray<FTokenizedStringNode>>::CreateLambda([String]()->TArray<FTokenizedStringNode>
 			{
+				if(String.IsEmpty())
+				{
+					return {};
+				}
 				return {FTokenizedStringNode{String}};
 			});
 		}
