Add blueprint function to check number of outstanding PSO precompilation tasks

diff --git a/Engine/Source/Runtime/Engine/Classes/Kismet/KismetRenderingLibrary.h b/Engine/Source/Runtime/Engine/Classes/Kismet/KismetRenderingLibrary.h
--- a/Engine/Source/Runtime/Engine/Classes/Kismet/KismetRenderingLibrary.h
+++ b/Engine/Source/Runtime/Engine/Classes/Kismet/KismetRenderingLibrary.h
@@ -258,4 +258,8 @@ class UKismetRenderingLibrary : public UBlueprintFunctionLibrary
 	 */
 	UFUNCTION(BlueprintCallable, Category = "Rendering", meta = (Keywords = "PathTracing"))
 	static ENGINE_API void RefreshPathTracingOutput();
+
+	/** Returns the number of active PSO precompilation tasks which are still running - can be used to extend loading screens to wait for PSO precaching to be finished */
+	UFUNCTION(BlueprintCallable, Category = "Rendering | Pipeline Cache")
+	static ENGINE_API int32 NumPrecompilingPSOsRemaining();
 };
diff --git a/Engine/Source/Runtime/Engine/Private/KismetRenderingLibrary.cpp b/Engine/Source/Runtime/Engine/Private/KismetRenderingLibrary.cpp
--- a/Engine/Source/Runtime/Engine/Private/KismetRenderingLibrary.cpp
+++ b/Engine/Source/Runtime/Engine/Private/KismetRenderingLibrary.cpp
@@ -31,6 +31,7 @@
 #include "UObject/Package.h"
 #include "EngineModule.h"
 #include "SceneInterface.h"
+#include "ShaderPipelineCache.h"
 
 #include UE_INLINE_GENERATED_CPP_BY_NAME(KismetRenderingLibrary)
 
@@ -966,5 +967,10 @@ ENGINE_API void UKismetRenderingLibrary::RefreshPathTracingOutput()
 	}
 }
 
+ENGINE_API int32 UKismetRenderingLibrary::NumPrecompilingPSOsRemaining()
+{
+	return (int32)FShaderPipelineCache::NumPrecompilesRemaining();
+}
+
 #undef LOCTEXT_NAMESPACE
 
