Fully change over AnimSequence compressed data locking to UE::FSharedRecursiveMutex (ran in parallel for a while to ensure there weren't any locking issues)

diff --git a/Engine/Source/Runtime/Engine/Classes/Animation/AnimSequence.h b/Engine/Source/Runtime/Engine/Classes/Animation/AnimSequence.h
--- a/Engine/Source/Runtime/Engine/Classes/Animation/AnimSequence.h
+++ b/Engine/Source/Runtime/Engine/Classes/Animation/AnimSequence.h
@@ -809,7 +809,6 @@ protected:
 	mutable FCompressedAnimSequence* CurrentPlatformData = nullptr; 
 	TMap<FIoHash, TUniquePtr<FCompressedAnimSequence>> DataByPlatformKeyHash;
 	TMap<FIoHash, TPimplPtr<UE::Anim::FAnimationSequenceAsyncCacheTask>> CacheTasksByKeyHash;
-	mutable FRWLock CompressedDataLock;
 	mutable UE::FSharedRecursiveMutex SharedCompressedDataMutex;
 
 	// Cached FIoHash keys stored according to ITargetPlatform type-hash
@@ -904,23 +903,23 @@ public:
 #if WITH_EDITOR
 		~FScopedCompressedAnimSequence()
 		{
-			ReadScopeLock.Reset();
+			SharedLock.Reset();
 		}
 #endif // WITH_EDITOR
 		
 	protected:
 		FScopedCompressedAnimSequence(const UAnimSequence* InAnimSequence, const FCompressedAnimSequence* InCompressedData) :
 			CompressedData(InCompressedData)
+		{
+			check(CompressedData && InAnimSequence);			
 #if WITH_EDITOR
-			, ReadScopeLock(MakeUnique<FReadScopeLock>(InAnimSequence->CompressedDataLock))
+			SharedLock = MakeUnique<UE::TSharedLock<UE::FSharedRecursiveMutex>>(InAnimSequence->SharedCompressedDataMutex);
 #endif // WITH_EDITOR
-		{
-			check(CompressedData && InAnimSequence);
 		}
 
 		const FCompressedAnimSequence* CompressedData = nullptr;
 #if WITH_EDITOR
-		TUniquePtr<FReadScopeLock> ReadScopeLock;
+		TUniquePtr<UE::TSharedLock<UE::FSharedRecursiveMutex>> SharedLock;
 #endif // WITH_EDITOR
 		friend class UAnimSequence;
 	};
@@ -933,9 +932,9 @@ public:
 	ENGINE_API FScopedCompressedAnimSequence GetCompressedData(const ITargetPlatform* InTargetPlatform) const;	
 protected:
 	/** Helper struct to enter/leave compressed data lock in read-mode */
-	struct FCompressedAnimationDataReadScope : public FReadScopeLock
+	struct FCompressedAnimationDataReadScope
 	{
-		FCompressedAnimationDataReadScope(const UAnimSequence* InAnimSequence) : FReadScopeLock(InAnimSequence->CompressedDataLock)
+		FCompressedAnimationDataReadScope(const UAnimSequence* InAnimSequence)
 		{
 			check(InAnimSequence);
 			SharedLock = MakeUnique<UE::TSharedLock<UE::FSharedRecursiveMutex>>(InAnimSequence->SharedCompressedDataMutex);
@@ -944,9 +943,9 @@ protected:
 	};
 	
 	/** Helper struct to enter/leave compressed data lock in write-mode */
-	struct FCompressedAnimationDataWriteScope : public FWriteScopeLock
+	struct FCompressedAnimationDataWriteScope
 	{
-		FCompressedAnimationDataWriteScope(const UAnimSequence* InAnimSequence) : FWriteScopeLock(InAnimSequence->CompressedDataLock), AnimSequence(InAnimSequence)
+		FCompressedAnimationDataWriteScope(const UAnimSequence* InAnimSequence) : AnimSequence(InAnimSequence)
 		{
 			check(InAnimSequence);
 			InAnimSequence->SharedCompressedDataMutex.Lock();
diff --git a/Engine/Source/Runtime/Engine/Private/Animation/AnimSequence.cpp b/Engine/Source/Runtime/Engine/Private/Animation/AnimSequence.cpp
--- a/Engine/Source/Runtime/Engine/Private/Animation/AnimSequence.cpp
+++ b/Engine/Source/Runtime/Engine/Private/Animation/AnimSequence.cpp
@@ -758,14 +758,16 @@ void UAnimSequence::Serialize(FArchive& Ar)
 				{
 					const bool bIsLoading = Ar.IsLoading();
 					const bool bIsSaving = Ar.IsSaving();
+					
+					TUniquePtr<UE::TSharedLock<UE::FSharedRecursiveMutex>> SharedLock;
 
 					if (bIsLoading)
 					{
-						CompressedDataLock.WriteLock();
+						SharedCompressedDataMutex.Lock();
 					}
 					else
 					{
-						CompressedDataLock.ReadLock();
+						SharedLock = MakeUnique<UE::TSharedLock<UE::FSharedRecursiveMutex>>(SharedCompressedDataMutex);
 					}
 
 					if (bCookingTargetNeedsCompressedData)
@@ -829,11 +831,11 @@ void UAnimSequence::Serialize(FArchive& Ar)
 
 					if (bIsLoading)
 					{
-						CompressedDataLock.WriteUnlock();
+						SharedCompressedDataMutex.Unlock();
 					}
 					else
 					{
-						CompressedDataLock.ReadUnlock();
+						SharedLock.Reset();
 					}
 				}
 				else
