Add a VirtualTextureStreamingPriority setting to TextureLODGroup and Texture assets.
The default for both settings is to use the "Normal" priority.
The LOD group setting is used by default unless explicitly overriden by the texture.
The setting is applied when collecting virtual texture pages to populate. At this stage we now prioritize the higher priority textures.

diff --git a/Engine/Source/Runtime/Engine/Classes/Engine/Texture.h b/Engine/Source/Runtime/Engine/Classes/Engine/Texture.h
--- a/Engine/Source/Runtime/Engine/Classes/Engine/Texture.h
+++ b/Engine/Source/Runtime/Engine/Classes/Engine/Texture.h
@@ -30,6 +30,7 @@
 #include "Texture.generated.h"
 
 namespace FOodleDataCompression {enum class ECompressor : uint8; enum class ECompressionLevel : int8; }
+enum class EVTProducerPriority : uint8;
 
 class FTextureReference;
 class FTextureResource;
@@ -1292,7 +1293,7 @@ public:
 	/** If enabled, defer compression of the texture until save or manually compressed in the texture editor. */
 	UPROPERTY(EditAnywhere, BlueprintReadWrite, Transient, SkipSerialization, Category=Compression, meta=(NoResetToDefault), meta=(DisplayName="Editor Defer Compression"))
 	uint32 DeferCompression:1;
-	
+
 	/** How aggressively should any relevant lossy compression be applied. For compressors that support EncodeSpeed (i.e. Oodle), this is only
 	*	applied if enabled (see Project Settings -> Texture Encoding). Note that this is *in addition* to any
 	*	unavoidable loss due to the target format - selecting "No Lossy Compression" will not result in zero distortion for BCn formats.
@@ -1473,6 +1474,10 @@ PRAGMA_ENABLE_DEPRECATION_WARNINGS
 	UPROPERTY(EditAnywhere, BlueprintReadWrite, Category = Texture, AssetRegistrySearchable, AdvancedDisplay)
 	ETextureMipLoadOptions MipLoadOptions;
 
+	/** Priority to use when using virtual texture streaming for this texture. */
+	UPROPERTY(EditAnywhere, BlueprintReadWrite, Category = Texture, AdvancedDisplay, meta=(EditCondition = "bUseVirtualTextureStreamingPriority"))
+	EVTProducerPriority VirtualTextureStreamingPriority;
+
 	/** If the platform supports it, tile the texture when cooking, or keep it linear and tile it when it's actually submitted to the GPU. */
 	UPROPERTY(EditAnywhere, BlueprintReadWrite, Category = Compression, AdvancedDisplay)
 	TEnumAsByte<enum TextureCookPlatformTilingSettings> CookPlatformTilingSettings;
@@ -1552,6 +1557,10 @@ PRAGMA_ENABLE_DEPRECATION_WARNINGS
 	UPROPERTY(EditAnywhere, BlueprintReadOnly, Category = Texture, AssetRegistrySearchable, AdvancedDisplay)
 	uint8 VirtualTextureStreaming : 1;
 
+	/** Use a virtual texture streaming priority set on this texture instead of the setting from the texture LOD group. */
+	UPROPERTY(EditAnywhere, Category = Texture, AdvancedDisplay, meta=(EditCondition = "VirtualTextureStreaming", InlineEditConditionToggle))
+	uint32 bUseVirtualTextureStreamingPriority : 1;
+
 #if WITH_EDITOR
 	UFUNCTION(BlueprintCallable, Category = Texture)
 	void SetVirtualTextureStreaming(bool bVirtualTextureStreaming);
diff --git a/Engine/Source/Runtime/Engine/Classes/Engine/TextureLODSettings.h b/Engine/Source/Runtime/Engine/Classes/Engine/TextureLODSettings.h
--- a/Engine/Source/Runtime/Engine/Classes/Engine/TextureLODSettings.h
+++ b/Engine/Source/Runtime/Engine/Classes/Engine/TextureLODSettings.h
@@ -9,6 +9,7 @@
 #include "TextureLODSettings.generated.h"
 
 class UTexture;
+enum class EVTProducerPriority : uint8;
 
 /** LOD settings for a single texture group. */
 USTRUCT()
@@ -115,6 +116,9 @@ struct FTextureLODGroup
 	UPROPERTY()
 	int32 VirtualTextureTileSizeBias = 0;
 
+	UPROPERTY()
+	EVTProducerPriority VirtualTextureStreamingPriority;
+
 	/** Allows us to override max anisotropy. If unspecified, uses r.MaxAnisotropy */
 	UPROPERTY()
 	int32 MaxAniso = 0;
diff --git a/Engine/Source/Runtime/Engine/Private/Texture.cpp b/Engine/Source/Runtime/Engine/Private/Texture.cpp
--- a/Engine/Source/Runtime/Engine/Private/Texture.cpp
+++ b/Engine/Source/Runtime/Engine/Private/Texture.cpp
@@ -43,6 +43,7 @@
 #include "DeviceProfiles/DeviceProfile.h"
 #include "DeviceProfiles/DeviceProfileManager.h"
 #include "Async/ParallelFor.h"
+#include "VirtualTextureEnum.h"
 
 #if WITH_EDITOR
 #include "Cooker/CookDeterminismHelper.h"
@@ -152,6 +153,7 @@ UTexture::UTexture(const FObjectInitializer& ObjectInitializer)
 	SRGB = true;
 	Filter = TF_Default;
 	MipLoadOptions = ETextureMipLoadOptions::Default;
+	VirtualTextureStreamingPriority = EVTProducerPriority::Normal;
 #if WITH_EDITORONLY_DATA
 	LoadedMainStreamObjectVersion = FUE5MainStreamObjectVersion::LatestVersion;
 	SourceColorSettings = FTextureSourceColorSettings();
diff --git a/Engine/Source/Runtime/Engine/Private/Texture2D.cpp b/Engine/Source/Runtime/Engine/Private/Texture2D.cpp
--- a/Engine/Source/Runtime/Engine/Private/Texture2D.cpp
+++ b/Engine/Source/Runtime/Engine/Private/Texture2D.cpp
@@ -1511,6 +1511,8 @@ FVirtualTexture2DResource::FVirtualTexture2DResource()
 
 	// Initialize this resource FeatureLevel, so it gets re-created on FeatureLevel changes
 	SetFeatureLevel(GMaxRHIFeatureLevel);
+
+	VirtualTextureStreamingPriority = EVTProducerPriority::Normal;
 }
 
 FVirtualTexture2DResource::FVirtualTexture2DResource(const UTexture2D* InOwner, FVirtualTextureBuiltData* InVTData, int32 InFirstMipToUse)
@@ -1523,7 +1525,8 @@ FVirtualTexture2DResource::FVirtualTexture2DResource(const UTexture2D* InOwner,
 	TextureName = InOwner->GetFName();
 	PackageName = InOwner->GetOutermost()->GetFName();
 
-	Filter = (ESamplerFilter)UDeviceProfileManager::Get().GetActiveProfile()->GetTextureLODSettings()->GetSamplerFilter(InOwner);
+	UTextureLODSettings* TextureLODSettings = UDeviceProfileManager::Get().GetActiveProfile()->GetTextureLODSettings();
+	Filter = (ESamplerFilter)TextureLODSettings->GetSamplerFilter(InOwner);
 	AddressU = InOwner->AddressX == TA_Wrap ? AM_Wrap : (InOwner->AddressX == TA_Clamp ? AM_Clamp : AM_Mirror);
 	AddressV = InOwner->AddressY == TA_Wrap ? AM_Wrap : (InOwner->AddressY == TA_Clamp ? AM_Clamp : AM_Mirror);
 
@@ -1533,7 +1536,16 @@ FVirtualTexture2DResource::FVirtualTexture2DResource(const UTexture2D* InOwner,
 
 	bSinglePhysicalSpace = InOwner->IsVirtualTexturedWithSinglePhysicalSpace();
 	bRequiresSinglePhysicalPool = InOwner->IsVirtualTexturedWithSinglePhysicalPool();
-
+	
+	if (InOwner->bUseVirtualTextureStreamingPriority)
+	{
+		VirtualTextureStreamingPriority = InOwner->VirtualTextureStreamingPriority;
+	}
+	else
+	{
+		VirtualTextureStreamingPriority = TextureLODSettings->GetTextureLODGroup(InOwner->LODGroup).VirtualTextureStreamingPriority;
+	}
+	
 	check(InVTData);
 	VTData = InVTData;
 
@@ -1615,8 +1627,7 @@ void FVirtualTexture2DResource::InitRHI(FRHICommandListBase& RHICmdList)
 		ProducerDesc.bIsLayerSRGB[LayerIndex] = bSRGB;
 	}
 	ProducerDesc.bRequiresSinglePhysicalPool = bRequiresSinglePhysicalPool;
-	// TODO [jonathan.bard] : For now, let's use the normal priority, but that could come from the texture asset and/or the texture group too : 
-	ProducerDesc.Priority = EVTProducerPriority::Normal;
+	ProducerDesc.Priority = VirtualTextureStreamingPriority;
 
 	FUploadingVirtualTexture* VirtualTexture = new FUploadingVirtualTexture(ProducerDesc.Name, VTData, FirstMipToUse);
 	ProducerHandle = GetRendererModule().RegisterVirtualTextureProducer(RHICmdList, ProducerDesc, VirtualTexture);
diff --git a/Engine/Source/Runtime/Engine/Private/TextureLODSettings.cpp b/Engine/Source/Runtime/Engine/Private/TextureLODSettings.cpp
--- a/Engine/Source/Runtime/Engine/Private/TextureLODSettings.cpp
+++ b/Engine/Source/Runtime/Engine/Private/TextureLODSettings.cpp
@@ -5,6 +5,7 @@
 #include "HAL/IConsoleManager.h"
 #include "ImageCoreUtils.h"
 #include "Interfaces/ITargetPlatform.h"
+#include "VirtualTextureEnum.h"
 
 #if ALLOW_OTHER_PLATFORM_CONFIG
 #include "UObject/UObjectIterator.h"
@@ -109,6 +110,8 @@ void FTextureLODGroup::SetupGroup()
 			Filter = ETextureSamplerFilter::AnisotropicLinear;
 		}
 	}
+
+	VirtualTextureStreamingPriority = EVTProducerPriority::Normal;
 }
 
 bool FTextureLODGroup::operator==(const FTextureLODGroup& Other) const
diff --git a/Engine/Source/Runtime/Engine/Public/TextureResource.h b/Engine/Source/Runtime/Engine/Public/TextureResource.h
--- a/Engine/Source/Runtime/Engine/Public/TextureResource.h
+++ b/Engine/Source/Runtime/Engine/Public/TextureResource.h
@@ -34,6 +34,7 @@ class FTexture3DResource;
 class FTexture2DArrayResource;
 class FVirtualTexture2DResource;
 struct IPooledRenderTarget;
+enum class EVTProducerPriority : uint8;
 
 /** Maximum number of slices in texture source art. */
 #define MAX_TEXTURE_SOURCE_SLICES 6
@@ -232,6 +233,7 @@ protected:
 	/** Cached owner settings */
 	bool bSinglePhysicalSpace = false;
 	bool bRequiresSinglePhysicalPool = false;
+	EVTProducerPriority VirtualTextureStreamingPriority;
 	/** Mip offset */
 	int32 FirstMipToUse = 0;
 	/** Built data owned by texture asset */
