Add new version of TQuadTree::GetElements that uses a callback instead of requiring an array to collect elements.

diff --git a/Engine/Source/Runtime/Engine/Public/GenericQuadTree.h b/Engine/Source/Runtime/Engine/Public/GenericQuadTree.h
--- a/Engine/Source/Runtime/Engine/Public/GenericQuadTree.h
+++ b/Engine/Source/Runtime/Engine/Public/GenericQuadTree.h
@@ -27,6 +27,11 @@ public:
 	template<typename ElementAllocatorType>
 	void GetElements(const FBox2D& Box, TArray<ElementType, ElementAllocatorType>& ElementsOut) const;
 
+	/** Given a 2D box, iterates over intersecting elements and execute the callback for each.
+	*   If Callback returns true, it will continue processing, if it returns false, iteration will end.*/
+	template<typename CallbackType>
+	inline void GetElements(const FBox2D& Box, CallbackType&& Callback) const;
+
 	/** Removes an object of type ElementType with an associated 2D box of size Box (log n). Does not cleanup tree*/
 	bool Remove(const ElementType& Instance, const FBox2D& Box);
 
@@ -81,6 +86,10 @@ private:
 	template<typename ElementAllocatorType>
 	void GetIntersectingElements(const FBox2D& Box, TArray<ElementType, ElementAllocatorType>& ElementsOut) const;
 
+	/** Given a list of nodes, execute Callback for any elements that intersect the box */
+	template<typename CallbackType>
+	bool GetIntersectingElements(const FBox2D& Box, CallbackType& Callback) const;
+
 	/** Given a list of nodes, remove the node that contains the given element */
 	bool RemoveNodeForElement(const ElementType& Element);
 
@@ -366,7 +375,44 @@ void TQuadTree<ElementType, NodeCapacity>::GetIntersectingElements(const FBox2D&
 			checkSlow(!ElementsOut.Contains(Node.Element));
 			ElementsOut.Add(Node.Element);
 		}
-	};
+	}
+}
+
+template <typename ElementType, int32 NodeCapacity>
+template<typename CallbackType>
+void TQuadTree<ElementType, NodeCapacity>::GetElements(const FBox2D& Box, CallbackType&& Callback) const
+{
+	TreeType* Quads[4];
+	const int32 NumQuads = GetQuads(Box, Quads);
+
+	// Always include any nodes contained in this quad
+	if (!GetIntersectingElements(Box, Callback))
+	{
+		return;
+	}
+
+	// As well as all relevant subtrees
+	for (int32 QuadIndex = 0; QuadIndex < NumQuads; QuadIndex++)
+	{
+		Quads[QuadIndex]->GetElements(Box, Callback);
+	}
+}
+
+template <typename ElementType, int32 NodeCapacity>
+template<typename CallbackType>
+bool TQuadTree<ElementType, NodeCapacity>::GetIntersectingElements(const FBox2D& Box, CallbackType& Callback) const
+{
+	for (const FNode& Node : Nodes)
+	{
+		// Only execute the callback if the boxes intersect
+		// Return if the Callback returns false, ending iteration
+		if (Box.Intersect(Node.Box) && !Callback(Node.Element))
+		{
+			return false;
+		}
+	}
+
+	return true;
 }
 
 template <typename ElementType, int32 NodeCapacity>
