Anim Assets: Fix crash when creating a preview scene for assets with a deleted compatible skeleton

diff --git a/Engine/Source/Runtime/Engine/Private/AnimationAsset.cpp b/Engine/Source/Runtime/Engine/Private/AnimationAsset.cpp
--- a/Engine/Source/Runtime/Engine/Private/AnimationAsset.cpp
+++ b/Engine/Source/Runtime/Engine/Private/AnimationAsset.cpp
@@ -421,7 +421,7 @@ USkeletalMesh* UAnimationAsset::GetPreviewMesh(bool bFindIfNotSet)
 #if WITH_EDITORONLY_DATA
 	USkeletalMesh* PreviewMesh = PreviewSkeletalMesh.LoadSynchronous();
 	// if somehow skeleton changes, just nullify it. 
-	if (PreviewMesh && !PreviewMesh->GetSkeleton()->IsCompatibleForEditor(Skeleton))
+	if (PreviewMesh && (PreviewMesh->GetSkeleton() == nullptr || !PreviewMesh->GetSkeleton()->IsCompatibleForEditor(Skeleton)))
 	{
 		PreviewMesh = nullptr;
 		SetPreviewMesh(nullptr);
