Anim Instance: Fix memory leak when GC-ing both mesh & anim instance in the same pass

Lazy-init of the proxy could re-occur depending on GC order. This prevents recreation of the proxy (which then would leak) in the case where the proxy has already been deleted.

diff --git a/Engine/Source/Runtime/Engine/Private/Animation/AnimInstance.cpp b/Engine/Source/Runtime/Engine/Private/Animation/AnimInstance.cpp
--- a/Engine/Source/Runtime/Engine/Private/Animation/AnimInstance.cpp
+++ b/Engine/Source/Runtime/Engine/Private/Animation/AnimInstance.cpp
@@ -333,7 +333,10 @@ void UAnimInstance::UninitializeAnimation()
 	
 	NativeUninitializeAnimation();
 
-	GetProxyOnGameThread<FAnimInstanceProxy>().Uninitialize(this);
+	if (AnimInstanceProxy != nullptr)
+	{
+		GetProxyOnGameThread<FAnimInstanceProxy>().Uninitialize(this);
+	}
 
 	StopAllMontages(0.f);
 
