Fix the DeviceToEmulate old system with hardcoded safezones

diff --git a/Engine/Source/Runtime/Engine/Private/LocalPlayer.cpp b/Engine/Source/Runtime/Engine/Private/LocalPlayer.cpp
--- a/Engine/Source/Runtime/Engine/Private/LocalPlayer.cpp
+++ b/Engine/Source/Runtime/Engine/Private/LocalPlayer.cpp
@@ -42,6 +42,7 @@
 
 #if WITH_EDITOR
 #include "Editor.h"
+#include "Settings/LevelEditorPlaySettings.h"
 #endif
 
 DEFINE_LOG_CATEGORY(LogPlayerManagement);
@@ -1188,25 +1189,29 @@ bool ULocalPlayer::GetProjectionData(FViewport* Viewport, FSceneViewProjectionDa
 		FPlane(0,	0,	0,	1));
 
 #if WITH_EDITOR
-	if (GEditor != nullptr && GEditor->IsFeatureLevelPreviewActive() && GEditor->GetAllowConstrainedAspectRatioInPreview())
+	ULevelEditorPlaySettings* PlaySettingsConfig = GetMutableDefault<ULevelEditorPlaySettings>();
+	if (PlaySettingsConfig != nullptr && PlaySettingsConfig->DeviceToEmulate.IsEmpty())
 	{
-		FSlateApplication::Get().OnConstrainedAspectRatioChanged.Broadcast(GEditor->PreviewPlatform.GetConstrainedAspectRatio());
-
-		FMargin Safezones(GEditor->PreviewPlatform.GetSafeZones());
-		if (Safezones != FSlateApplication::Get().GetCustomSafeZone())
+		if (GEditor != nullptr && GEditor->IsFeatureLevelPreviewActive() && GEditor->GetAllowConstrainedAspectRatioInPreview())
 		{
-			FSlateApplication::Get().SetCustomSafeZone(Safezones);
-			FSlateApplication::Get().OnDebugSafeZoneChanged.Broadcast(Safezones, true);
-		}
-	}
-	else
-	{
-		FSlateApplication::Get().OnConstrainedAspectRatioChanged.Broadcast(0.f);
+			FSlateApplication::Get().OnConstrainedAspectRatioChanged.Broadcast(GEditor->PreviewPlatform.GetConstrainedAspectRatio());
 
-		if (FMargin() != FSlateApplication::Get().GetCustomSafeZone())
+			FMargin Safezones(GEditor->PreviewPlatform.GetSafeZones());
+			if (Safezones != FSlateApplication::Get().GetCustomSafeZone())
+			{
+				FSlateApplication::Get().SetCustomSafeZone(Safezones);
+				FSlateApplication::Get().OnDebugSafeZoneChanged.Broadcast(Safezones, true);
+			}
+		}
+		else
 		{
-			FSlateApplication::Get().ResetCustomSafeZone();
-			FSlateApplication::Get().OnDebugSafeZoneChanged.Broadcast(FMargin(), true);
+			FSlateApplication::Get().OnConstrainedAspectRatioChanged.Broadcast(0.f);
+
+			if (FMargin() != FSlateApplication::Get().GetCustomSafeZone())
+			{
+				FSlateApplication::Get().ResetCustomSafeZone();
+				FSlateApplication::Get().OnDebugSafeZoneChanged.Broadcast(FMargin(), true);
+			}
 		}
 	}
 #endif
