FSpline: Get user attributes working again

diff --git a/Engine/Source/Runtime/Engine/Classes/Curves/Spline.h b/Engine/Source/Runtime/Engine/Classes/Curves/Spline.h
--- a/Engine/Source/Runtime/Engine/Classes/Curves/Spline.h
+++ b/Engine/Source/Runtime/Engine/Classes/Curves/Spline.h
@@ -69,7 +69,7 @@ struct FSpline
 	/* Attribute Interface */
 
 	// Non-Templated functions
-	bool SupportsAttributes() const { return false; /*return IsNew();*/ }
+	bool SupportsAttributes() const { return IsNew(); }
 	ENGINE_API bool HasAttributeChannel(FName AttributeName) const;
 	ENGINE_API TArray<FName> GetFloatPropertyChannels() const;
 	ENGINE_API TArray<FName> GetVectorPropertyChannels() const;
diff --git a/Engine/Source/Runtime/Engine/Private/Curves/Spline.cpp b/Engine/Source/Runtime/Engine/Private/Curves/Spline.cpp
--- a/Engine/Source/Runtime/Engine/Private/Curves/Spline.cpp
+++ b/Engine/Source/Runtime/Engine/Private/Curves/Spline.cpp
@@ -2,6 +2,7 @@
 
 #include "Curves/Spline.h"
 
+#include "BoxTypes.h"
 #include "Components/SplineComponent.h"	// for FSplineCurves, ideally removed
 #include "Misc/Base64.h"
 #include "Serialization/MemoryReader.h"
@@ -12,6 +13,8 @@
 
 #include UE_INLINE_GENERATED_CPP_BY_NAME(Spline)
 
+using UE::Geometry::FInterval1f;
+
 namespace UE::Spline
 {
 
@@ -137,7 +140,6 @@ public:
 	virtual ~FNewSpline() override = default;
     FNewSpline& operator=(const FNewSpline& Other);
 
-	bool operator==(const FNewSpline& Other) const;
 	bool operator!=(const FNewSpline& Other) const;
 
     virtual bool Serialize(FArchive& Ar) override;
@@ -222,7 +224,7 @@ private:
 	static inline const FName RotationAttrName = FName("Rotation");
 	static inline const FName ScaleAttrName = FName("Scale");
 	static inline const FName ParamToLengthAttrName = FName("ParamToLength");
-	static inline const FName LengthToParamAttrName = FName("LengthToParam");
+	static inline const FName NormalizedLengthToParamAttrName = FName("LengthToParam");
 
     /**
      * We keep these up to date with Points.
@@ -291,35 +293,55 @@ int32 FNewSpline::AddAttributeValue(FName Name, const AttrType& Value, float Par
 	{
 		ControlPointType ControlPoint = ControlPointType(Value);
 		int32 NumPoints = Channel->GetNumPoints();
+
+		float ParentSpaceParameter = ToInternalSplineSpace(Parameter);
+		FInterval1f MappedChildSpace = GetMappedChildSpace(Name);
 		
 		// Special case for empty spline
 		if (NumPoints == 0)
 		{
 			Channel->AppendPoint(ControlPoint);
+			SetAttributeChannelRange(Name, FInterval1f(ParentSpaceParameter, ParentSpaceParameter), EMappingRangeSpace::Parent);
 			return 0;
 		}
-		else if (NumPoints == 1)
+
+		if (NumPoints == 1)
 		{
-			if (Parameter > 0)
+			if (ParentSpaceParameter > MappedChildSpace.Min)
 			{
 				Channel->AppendPoint(ControlPoint);
 				Channel->Reparameterize(EParameterizationPolicy::Uniform);
+				SetAttributeChannelRange(Name, FInterval1f(MappedChildSpace.Min, ParentSpaceParameter), EMappingRangeSpace::Parent);
 				return 1;
 			}
 			else
 			{
 				Channel->PrependPoint(ControlPoint);
 				Channel->Reparameterize(EParameterizationPolicy::Uniform);
+				SetAttributeChannelRange(Name, FInterval1f(ParentSpaceParameter, MappedChildSpace.Min), EMappingRangeSpace::Parent);
 				return 0;
 			}
 		}
+
+		// append case
+		if (ParentSpaceParameter > MappedChildSpace.Max)
+		{
+			Channel->AppendPoint(ControlPoint);
+			return SetAttributeParameter<AttrType>(Name, Channel->GetNumPoints() - 1, Parameter);
+		}
+
+		// prepend case
+		if (ParentSpaceParameter < MappedChildSpace.Min)
+		{
+			Channel->PrependPoint(ControlPoint);
+			return SetAttributeParameter<AttrType>(Name, 0, Parameter);
+		}
 		
-		float InternalParameter = ToInternalSplineSpace(Parameter);
+		float ChildSpaceParameter = MapParameterToChildSpace(Name, ParentSpaceParameter);
 		float LocalT;
-		int32 SegmentIndex = GetSpline().FindSegmentIndex(InternalParameter, LocalT);
+		int32 SegmentIndex = Channel->FindSegmentIndex(ChildSpaceParameter, LocalT);
 		
 		int NewPointIndex = Channel->InsertPointAtSegmentParam(SegmentIndex, LocalT, ControlPoint);
-		Channel->Reparameterize(EParameterizationPolicy::Uniform);
 		return NewPointIndex;
 	}
 
@@ -344,6 +366,22 @@ void FNewSpline::RemoveAttributeValue(FName Name, int32 Index)
 
 	if (AttrSplineType* Channel = GetTypedAttributeChannel<AttrSplineType>(Name))
 	{
+		if (Channel->GetNumPoints() > 1)
+		{
+			if (Index == 0)
+			{
+				FInterval1f NewMappingRange = GetMappedChildSpace(Name);
+				NewMappingRange.Min = ToInternalSplineSpace(GetAttributeParameter<AttrType>(Name, Index + 1));
+				SetAttributeChannelRange(Name, NewMappingRange, EMappingRangeSpace::Parent);
+			}
+			else if (Index == Channel->GetNumPoints() - 1)
+			{
+				FInterval1f NewMappingRange = GetMappedChildSpace(Name);
+				NewMappingRange.Max = ToInternalSplineSpace(GetAttributeParameter<AttrType>(Name, Index - 1));
+				SetAttributeChannelRange(Name, NewMappingRange, EMappingRangeSpace::Parent);
+			}
+		}
+		
 		Channel->RemovePoint(Index);
 	}
 }
@@ -381,9 +419,131 @@ int32 FNewSpline::SetAttributeParameter(const FName& Name, int32 Index, float Pa
 	
 	if (AttrSplineType* Channel = GetTypedAttributeChannel<AttrSplineType>(Name))
 	{
+		float ParentSpaceParameter = ToInternalSplineSpace(Parameter);
+
+		// Helper to prevent code duplication.
+		// Uses current value of Index and ParentSpaceParameter at the time of calling to update the attribute parameter.
+		auto SetParentSpaceParameter = [this, &ParentSpaceParameter, Name, Channel, &Index]() -> int32
+		{
+			const float CurrentChildSpaceParameter = MapParameterToChildSpace(Name, ToInternalSplineSpace(GetAttributeParameter<AttrType>(Name, Index)));
+			const float DesiredChildSpaceParameter = MapParameterToChildSpace(Name, ParentSpaceParameter);
+
+			constexpr float MinStep = 2.f * UE_KINDA_SMALL_NUMBER;
+			float NewChildSpaceParameter = DesiredChildSpaceParameter > CurrentChildSpaceParameter
+				? FMath::Max(DesiredChildSpaceParameter, CurrentChildSpaceParameter + MinStep)
+				: FMath::Min(DesiredChildSpaceParameter, CurrentChildSpaceParameter - MinStep);
+			
+			return Channel->SetParameter(Index, NewChildSpaceParameter);
+		};
+
+		/**
+		 * Cases we handle below:
+		 * 1) Moving the only existing attribute.
+		 * 2) Moving an endpoint for a 2 point channel.
+		 * 3) Moving the first endpoint for a 3+ point channel.
+		 * 4) Moving the last endpoint for a 3+ point channel.
+		 * 5) Moving an internal point.
+		 */
 		
-		return Channel->SetParameter(Index, MapParameterToChildSpace(Name, ToInternalSplineSpace(Parameter)));
+		if (Channel->GetNumPoints() == 1)						// Case 1: Moving the only attribute
+		{
+			SetAttributeChannelRange(Name, FInterval1f(ParentSpaceParameter, ParentSpaceParameter), EMappingRangeSpace::Parent);
+			return Index;
+		}
+		else if (Channel->GetNumPoints() == 2)					// Case 2: Moving an end point while no internal points exist.
+		{
+			FInterval1f MappedRange = GetMappedChildSpace(Name);
+			float& RangeBound = Index == 0 ? MappedRange.Min : MappedRange.Max;
+			RangeBound = ParentSpaceParameter;
+
+			if (MappedRange.Min > MappedRange.Max)
+			{
+				// If range flips, swap end points and un-flip.
+				Channel->SetParameter(1, Channel->GetParameter(0) - 1.f);
+				SetAttributeChannelRange(Name, FInterval1f(MappedRange.Max, MappedRange.Min), EMappingRangeSpace::Parent);
+				Index = Index == 0 ? 1 : 0;
+			}
+			else
+			{
+				SetAttributeChannelRange(Name, MappedRange, EMappingRangeSpace::Parent);
+			}
+
+			return Index;
+		}
+		else if (Index == 0)									// Case 3: Moving the first point (which has exactly 1 neighbor)
+		{
+			TArray<float> InternalParameters;
+			for (int InternalIdx = 1; InternalIdx < Channel->GetNumPoints() - 1; ++InternalIdx)
+			{
+				InternalParameters.Add(MapParameterFromChildSpace(Name, Channel->GetParameter(InternalIdx)));
+			}
+
+			float ParentSpaceUpperBound = GetMappedChildSpace(Name).Max;
+			const float NeighborParentSpaceParameter = ToInternalSplineSpace(GetAttributeParameter<AttrType>(Name, Index + 1));
+			const bool bAttributesWillReorder = ParentSpaceParameter > NeighborParentSpaceParameter;
+
+			// If the end point is passing its neighbor, we need to actually reshuffle the points.
+			if (bAttributesWillReorder)
+			{
+				Index = SetParentSpaceParameter();
+				
+				// Update InternalParameters to properly re-parameterize the new internal point at Index - 1
+				InternalParameters[Index - 1] = ParentSpaceParameter;
+			}
+			
+			// Updates the lower bound of the mapping range.
+			const float ParentSpaceLowerBound = bAttributesWillReorder ? NeighborParentSpaceParameter : ParentSpaceParameter;
+			SetAttributeChannelRange(Name, FInterval1f(ParentSpaceLowerBound, ParentSpaceUpperBound), EMappingRangeSpace::Parent);
+
+			// Re-parameterize the internal points so that they do not move in parent space after the mapping range changed
+			for (int InternalIdx = 1; InternalIdx < Channel->GetNumPoints() - 1; ++InternalIdx)
+			{
+				Channel->SetParameter(InternalIdx, MapParameterToChildSpace(Name, InternalParameters[InternalIdx - 1]));
+			}
+			
+			return Index;
+		}
+		else if (Index == Channel->GetNumPoints() - 1)			// Case 4: Moving the last point (which has exactly 1 neighbor)
+		{
+			// Save internal attribute point parameters for later re-parameterization after changing the child space mapping.
+			TArray<float> InternalParameters;
+			for (int InternalIdx = 1; InternalIdx < Channel->GetNumPoints() - 1; ++InternalIdx)
+			{
+				InternalParameters.Add(MapParameterFromChildSpace(Name, Channel->GetParameter(InternalIdx)));
+			}
+
+			float ParentSpaceLowerBound = GetMappedChildSpace(Name).Min;
+			const float NeighborParentSpaceParameter = ToInternalSplineSpace(GetAttributeParameter<AttrType>(Name, Index - 1));
+			const bool bAttributesWillReorder = ParentSpaceParameter < NeighborParentSpaceParameter;
+			
+			// If the end point is passing its neighbor, we need to actually reshuffle the points.
+			if (bAttributesWillReorder)
+			{
+				Index = SetParentSpaceParameter();
+
+				// Update InternalParameters to properly re-parameterize the new internal point at Index - 1
+				InternalParameters[Index - 1] = ParentSpaceParameter;
+			}
+
+			// Update the child space mapping
+			float ParentSpaceUpperBound = bAttributesWillReorder ? NeighborParentSpaceParameter : ParentSpaceParameter; // Prevent over-shrinking if we reshuffle.
+			SetAttributeChannelRange(Name, FInterval1f(ParentSpaceLowerBound, ParentSpaceUpperBound), EMappingRangeSpace::Parent);
+			
+			// Re-parameterize the internal points so that they do not move in parent space after the mapping range changed
+			for (int InternalIdx = 1; InternalIdx < Channel->GetNumPoints() - 1; ++InternalIdx)
+			{
+				Channel->SetParameter(InternalIdx, MapParameterToChildSpace(Name, InternalParameters[InternalIdx - 1]));
+			}
+		
+			return Index;
+		}
+		else													// Case 5: Moving an internal point (which has exactly 2 neighbors)
+		{
+			Index = SetParentSpaceParameter();
+			return Index;
+		}
 	}
+	
 	return INDEX_NONE;
 }
 
@@ -861,7 +1021,7 @@ TArray<FName> FSpline::GetFloatPropertyChannels() const
 
 		// Don't report these internal attribute channels to the caller
 		ChannelNames.Remove(NewData->ParamToLengthAttrName);
-		ChannelNames.Remove(NewData->LengthToParamAttrName);
+		ChannelNames.Remove(NewData->NormalizedLengthToParamAttrName);
 	}
 	return ChannelNames;
 }
@@ -1365,7 +1525,7 @@ FNewSpline::FNewSpline()
 	CreateAttributeChannel<TTangentBezierSpline<FQuat>>(RotationAttrName);
 	CreateAttributeChannel<TTangentBezierSpline<FVector>>(ScaleAttrName);
 	CreateAttributeChannel<TTangentBezierSpline<float>>(ParamToLengthAttrName);
-	CreateAttributeChannel<TTangentBezierSpline<float>>(LengthToParamAttrName);
+	CreateAttributeChannel<TTangentBezierSpline<float>>(NormalizedLengthToParamAttrName);
     
     ResetToDefault();
 	
@@ -1377,7 +1537,7 @@ FNewSpline::FNewSpline(const FNewSpline& Other)
 	CreateAttributeChannel<TTangentBezierSpline<FQuat>>(RotationAttrName);
 	CreateAttributeChannel<TTangentBezierSpline<FVector>>(ScaleAttrName);
 	CreateAttributeChannel<TTangentBezierSpline<float>>(ParamToLengthAttrName);
-	CreateAttributeChannel<TTangentBezierSpline<float>>(LengthToParamAttrName);
+	CreateAttributeChannel<TTangentBezierSpline<float>>(NormalizedLengthToParamAttrName);
 	
     *this = Other;
 
@@ -1389,7 +1549,7 @@ FNewSpline::FNewSpline(const FLegacySpline& Other)
 	CreateAttributeChannel<TTangentBezierSpline<FQuat>>(RotationAttrName);
 	CreateAttributeChannel<TTangentBezierSpline<FVector>>(ScaleAttrName);
 	CreateAttributeChannel<TTangentBezierSpline<float>>(ParamToLengthAttrName);
-	CreateAttributeChannel<TTangentBezierSpline<float>>(LengthToParamAttrName);
+	CreateAttributeChannel<TTangentBezierSpline<float>>(NormalizedLengthToParamAttrName);
 	
 	TArray<FTangentBezierControlPoint> Points;
     for (int32 Idx = 0; Idx < Other.GetNumControlPoints(); Idx++)
@@ -1427,7 +1587,7 @@ FNewSpline::FNewSpline(const FSplineCurves& Other)
 	CreateAttributeChannel<TTangentBezierSpline<FQuat>>(RotationAttrName);
 	CreateAttributeChannel<TTangentBezierSpline<FVector>>(ScaleAttrName);
 	CreateAttributeChannel<TTangentBezierSpline<float>>(ParamToLengthAttrName);
-	CreateAttributeChannel<TTangentBezierSpline<float>>(LengthToParamAttrName);
+	CreateAttributeChannel<TTangentBezierSpline<float>>(NormalizedLengthToParamAttrName);
 	
 	TArray<FTangentBezierControlPoint> Points;
     for (int32 Idx = 0; Idx < Other.Position.Points.Num(); Idx++)
@@ -1479,11 +1639,6 @@ FNewSpline& FNewSpline::operator=(const FNewSpline& Other)
     return *this;
 }
 
-bool FNewSpline::operator==(const FNewSpline& Other) const
-{
-    return GetSpline() == Other.GetSpline() && GetSpline().GetTension() == Other.GetSpline().GetTension();
-}
-
 bool FNewSpline::operator!=(const FNewSpline& Other) const
 {
     return !(*this == Other);
@@ -1565,8 +1720,16 @@ float FNewSpline::GetSegmentLength(const int32 Index, const float Param, const F
 
 float FNewSpline::GetSplineLength() const
 {
+	float AccumulatedLength = 0.f;
+	for (int Idx = 0; Idx < GetSpline().GetNumSegments(); ++Idx)
+	{
+		AccumulatedLength += GetSegmentLength(Idx, 1.f);
+	}
+	return AccumulatedLength;
+	
 	// Evaluate the ParamToLength channel at the very end of the spline
-    return GetDistanceAtParameter(FromNormalizedSpace(1.0f));
+	//const float MaxParameter = static_cast<float>(GetSpline().GetNumSegments());
+    //return GetDistanceAtParameter(MaxParameter);
 }
 
 void FNewSpline::SetClosedLoop(bool bInClosedLoop)
@@ -1604,6 +1767,7 @@ void FNewSpline::ResetRotation()
     	AddAttributeValue<FQuat>(RotationAttrName, FQuat::Identity, static_cast<float>(Idx));
     }
 
+	SetAttributeChannelRange(RotationAttrName, FInterval1f(0.f, 1.f), EMappingRangeSpace::Normalized);
 	TTangentBezierSpline<FQuat>* RotSpline = GetTypedAttributeChannel<TTangentBezierSpline<FQuat>>(RotationAttrName);
 	RotSpline->SetKnotVector(GetSpline().GetKnotVector());
 }
@@ -1617,6 +1781,7 @@ void FNewSpline::ResetScale()
         AddAttributeValue<FVector>(ScaleAttrName, FVector::OneVector, static_cast<float>(Idx));
     }
 
+	SetAttributeChannelRange(ScaleAttrName, FInterval1f(0.f, 1.f), EMappingRangeSpace::Normalized);
 	TTangentBezierSpline<FVector>* ScaleSpline = GetTypedAttributeChannel<TTangentBezierSpline<FVector>>(ScaleAttrName);
 	ScaleSpline->SetKnotVector(GetSpline().GetKnotVector());
 }
@@ -1626,6 +1791,9 @@ void FNewSpline::AddPoint(const FSplinePoint& Point)
 	ValidateRotScale();
 	ON_SCOPE_EXIT
 	{
+		SetAttributeChannelRange(ScaleAttrName, FInterval1f(0.f, 1.f), EMappingRangeSpace::Normalized);
+		SetAttributeChannelRange(RotationAttrName, FInterval1f(0.f, 1.f), EMappingRangeSpace::Normalized);
+
 		ValidateRotScale();
 	};
 	
@@ -1691,6 +1859,9 @@ void FNewSpline::InsertPoint(const FSplinePoint& Point, int32 Index)
 	ValidateRotScale();
 	ON_SCOPE_EXIT
 	{
+		SetAttributeChannelRange(ScaleAttrName, FInterval1f(0.f, 1.f), EMappingRangeSpace::Normalized);
+		SetAttributeChannelRange(RotationAttrName, FInterval1f(0.f, 1.f), EMappingRangeSpace::Normalized);
+		
 		ValidateRotScale();
 	};
 	
@@ -1773,6 +1944,9 @@ void FNewSpline::RemovePoint(int32 Index)
 	ValidateRotScale();
 	ON_SCOPE_EXIT
 	{
+		SetAttributeChannelRange(ScaleAttrName, FInterval1f(0.f, 1.f), EMappingRangeSpace::Normalized);
+		SetAttributeChannelRange(RotationAttrName, FInterval1f(0.f, 1.f), EMappingRangeSpace::Normalized);
+		
 		ValidateRotScale();
 	};
 	
@@ -1937,23 +2111,19 @@ float FNewSpline::GetParameterAtIndex(int32 Index) const
     }
 
 	return FromInternalSplineSpace(GetSpline().GetParameter(Index));
-	
-    //return static_cast<float>(Index);
 }
 
 float FNewSpline::GetParameterAtDistance(float Distance) const
 {
-	const float NormalizedDistance = FMath::Clamp(Distance / FMath::Max(GetSplineLength(), UE_SMALL_NUMBER), 0.f, 1.f);
-	return EvaluateAttribute<float>(LengthToParamAttrName, NormalizedDistance);
+	return 0.f;
+	//const float NormalizedDistance = FMath::Clamp(Distance / FMath::Max(GetSplineLength(), UE_SMALL_NUMBER), 0.f, 1.f);
+	//return EvaluateAttribute<float>(NormalizedLengthToParamAttrName, ToInternalSplineSpace(NormalizedDistance));
 }
 
 float FNewSpline::GetDistanceAtParameter(float Parameter) const
 {
-	// a bit of a hack, we convert to normalized but then treat
-	// the resulting value as if it were in internal parent space.
-	// see the todo in UpdateReparamTable
-	Parameter = ToNormalizedSpace(Parameter);
-	return EvaluateAttribute<float>(ParamToLengthAttrName, Parameter);
+	return 0.f;
+	//return EvaluateAttribute<float>(ParamToLengthAttrName, ToInternalSplineSpace(Parameter));
 }
 
 float FNewSpline::FindNearest(const FVector& InLocation, float& OutSquaredDistance) const
@@ -2051,46 +2221,41 @@ float FNewSpline::ToInternalSplineSpace(float Parameter) const
 
 void FNewSpline::UpdateReparamTable()
 {
-    const int32 NumPoints = GetSpline().GetNumPoints();
-  	const int32 NumSegments = NumPoints - 1;
-    const int32 TotalSteps = NumSegments * ReparamStepsPerSegment + 1;
-	
-    // Start by clearing it
-    ClearAttributeChannel(ParamToLengthAttrName);
-    ClearAttributeChannel(LengthToParamAttrName);
-
-	// Todo: convert normalized parameters below into [0, NumSegments] space.
-	// This does not work when last tested, but should be corrected.
-	
-    // Here we get Param -> Length mapping.
-    float AccumulatedLength = 0.0f;
-    for (int32 SegmentIndex = 0; SegmentIndex < NumSegments; ++SegmentIndex)
-    {
-    	for (int32 Step = 0; Step < ReparamStepsPerSegment; ++Step)
-    	{
-    		// These are with respect to current segment.
-    		const float SegmentParam = static_cast<float>(Step) / ReparamStepsPerSegment;
-    		const float OffsetLength = GetSegmentLength(SegmentIndex, SegmentParam);
-
-    		// Map the above values to 'spline space' (i.e. [0, NumSegments]) by normalizing and converting.
-    		const float SplineParam = (static_cast<float>(SegmentIndex * ReparamStepsPerSegment + Step) / TotalSteps);
-    		const float CurrentLength = OffsetLength + AccumulatedLength;
-
-    		AddAttributeValue<float>(ParamToLengthAttrName, CurrentLength, SplineParam);
-    	}
-
-    	AccumulatedLength += GetSegmentLength(SegmentIndex, 1.0f);
-    }
-    AddAttributeValue<float>(ParamToLengthAttrName, AccumulatedLength, 1.0f);
-
-    // Here we get Length -> Param mapping
-    for (int32 Step = 0; Step <= NumSegments * ReparamStepsPerSegment + 1; ++Step)
-    {
-    	const float SplineParam = FMath::Clamp(static_cast<float>(Step) / TotalSteps, 0.0f, 1.0f);
-    	const float CurrentLength = EvaluateAttribute<float>(ParamToLengthAttrName, SplineParam);
-    	
-    	AddAttributeValue<float>(LengthToParamAttrName, FromNormalizedSpace(SplineParam), CurrentLength / AccumulatedLength);
-    }
+ //    const int32 NumPoints = GetSpline().GetNumPoints();
+ //  	const int32 NumSegments = NumPoints - 1;
+	//
+ //    // Start by clearing it
+ //    ClearAttributeChannel(ParamToLengthAttrName);
+ //    ClearAttributeChannel(NormalizedLengthToParamAttrName);
+	//
+ //    // Here we get Param -> Length mapping.
+ //    float AccumulatedLength = 0.0f;
+ //    for (int32 SegmentIndex = 0; SegmentIndex < NumSegments; ++SegmentIndex)
+ //    {
+ //    	for (int32 Step = 0; Step < ReparamStepsPerSegment; ++Step)
+ //    	{
+ //    		// These are with respect to current segment.
+ //    		const float SegmentParam = static_cast<float>(Step) / ReparamStepsPerSegment;
+ //    		const float OffsetLength = GetSegmentLength(SegmentIndex, SegmentParam);
+ //    
+ //    		// Map the above values to 'spline space' (i.e. [0, NumSegments]) by normalizing and converting.
+ //    		const float SplineParam = SegmentIndex + SegmentParam;
+ //    		const float CurrentLength = OffsetLength + AccumulatedLength;
+ //    		AddAttributeValue<float>(ParamToLengthAttrName, CurrentLength, SplineParam);
+ //    	}
+ //    
+ //    	AccumulatedLength += GetSegmentLength(SegmentIndex, 1.0f);
+ //    }
+ //    AddAttributeValue<float>(ParamToLengthAttrName, AccumulatedLength, static_cast<float>(NumSegments));
+	//
+ //    // Here we get Length -> Param mapping
+	// const int32 TotalSteps = NumSegments * ReparamStepsPerSegment + 1;
+ //    for (int32 Step = 0; Step <= TotalSteps; ++Step)
+ //    {
+ //    	const float SplineParam = FromNormalizedSpace(FMath::Clamp(static_cast<float>(Step) / TotalSteps, 0.0f, 1.0f));
+ //    	const float NormalizedLength = EvaluateAttribute<float>(ParamToLengthAttrName, ToInternalSplineSpace(SplineParam)) / AccumulatedLength;
+ //    	AddAttributeValue<float>(NormalizedLengthToParamAttrName, SplineParam, NormalizedLength);
+ //    }
 }
 
 void FNewSpline::RebuildLegacyCurves()
diff --git a/Engine/Source/Runtime/GeometryCore/Public/Splines/BSpline.h b/Engine/Source/Runtime/GeometryCore/Public/Splines/BSpline.h
--- a/Engine/Source/Runtime/GeometryCore/Public/Splines/BSpline.h
+++ b/Engine/Source/Runtime/GeometryCore/Public/Splines/BSpline.h
@@ -2,9 +2,11 @@
 
 #pragma once
 
+#include <cmath>
 #include "Spline.h"
 #include "InterpolationPolicies.h"
 #include "SplineMath.h"
+
 namespace UE
 {
 namespace Geometry
@@ -163,10 +165,10 @@ public:
 	}
     
 	// Parameterization methods
-	virtual bool SetParameter(int32 Index, float NewParameter)
+	virtual int32 SetParameter(int32 Index, float NewParameter)
 	{
 		UE_LOG(LogSpline, Warning, TEXT("SetParameter is disabled. Use Reparameterize() or a linear spline if you need manual control."));
-		return false;
+		return INDEX_NONE;
 	}
 	
 	// Returns parameter value for a control point index using Greville Abscissae
@@ -623,6 +625,44 @@ public:
 		}
 	}
 
+	/**
+	 * Replaces all instances of a knot.
+	 * 
+	 * @param OldParameter - Parameter value to search for
+	 * @param NewParameter - Value to replace instances of OldParameter
+	 */
+	void ReplaceKnot(float OldParameter, float NewParameter)
+	{
+		for (int32 i = 0; i < Knots.Num(); ++i)
+		{
+			if (FMath::IsNearlyEqual(Knots[i], OldParameter))
+			{
+				Knots[i] = NewParameter;
+			}
+		}
+	}
+
+	/**
+	 * Swaps all instances of two knots.
+	 * 
+	 * @param ParameterA - Values which should be replaced with ParameterB
+	 * @param ParameterB - Values which should be replaced with ParameterA
+	 */
+	void SwapKnots(float ParameterA, float ParameterB)
+	{
+		for (int32 i = 0; i < Knots.Num(); ++i)
+		{
+			if (FMath::IsNearlyEqual(Knots[i], ParameterA))
+			{
+				Knots[i] = ParameterB;
+			}
+			else if (FMath::IsNearlyEqual(Knots[i], ParameterB))
+			{
+				Knots[i] = ParameterA;
+			}
+		}
+	}
+
 	/**
 	 * Returns how many times a given knot value appears.
 	 * @param Parameter - Knot value to check
@@ -641,6 +681,84 @@ public:
 	    return Count;
 	}
 
+	// Given a parameter and an array of distinct knots (mult = 1),
+	// find the parameter closest to DesiredParameter which corresponds to no distinct knots.
+	static float GetNearestValidKnot(float DesiredParameter, const TArray<float>& Knots, float Epsilon = UE_KINDA_SMALL_NUMBER)
+	{
+		if (Knots.Num() == 0)
+		{
+			// If we have an empty knot vector, DesiredParameter is definitely valid.
+			return DesiredParameter;
+		}
+
+		if (Knots.Num() == 1)
+		{
+			// If we have 1 knot, we can special case to simplify logic in the loop below.
+			return DesiredParameter < Knots[0]
+				? FMath::Min(DesiredParameter, Knots[0] - Epsilon)
+				: FMath::Max(DesiredParameter, Knots[0] + Epsilon);
+		}
+
+		// A candidate is valid if it is not too close to any existing knots (including the knot which was used to generate the candidate)
+		auto IsCandidateValid = [&Knots, Epsilon](const float Candidate) -> bool
+		{
+			for (const float Knot : Knots)
+			{
+				if (FMath::Abs(Candidate - Knot) < Epsilon)
+				{
+					return false;
+				}
+			}
+
+			return true;
+		};
+		
+		// Candidates are closest valid parameters to existing knots which are too close
+		// to DesiredParameter to allow DesiredParameter to be considered a valid knot.
+		TArray<float> Candidates;
+		bool bDegenerate = false;
+		for (int Idx = 0; Idx < Knots.Num(); ++Idx)
+		{
+			const float Knot = Knots[Idx];
+
+			if (FMath::Abs(DesiredParameter - Knot) < Epsilon)
+			{
+				bDegenerate = true;
+				UE_LOG(LogSpline, Warning, TEXT("Parameter %f is too close to existing knot value %f."), DesiredParameter, Knot);
+
+				if (float Candidate = std::nextafter(Knot - Epsilon, -INFINITY); IsCandidateValid(Candidate))
+				{
+					Candidates.Add(Candidate);
+				}
+
+				if (float Candidate = std::nextafter(Knot + Epsilon, +INFINITY); IsCandidateValid(Candidate))
+				{
+					Candidates.Add(Candidate);
+				}
+			}
+		}
+
+		if (bDegenerate && Candidates.Num() == 0)
+		{
+			UE_LOG(LogSpline, Warning, TEXT("Parameter %f is invalid and no good alternatives were found, expect a degenerate segment."), DesiredParameter);
+		}
+		
+		float BestCandidate = DesiredParameter;
+		float BestCandidateDistanceToDesiredParameter = INFINITY;	// Set to infinity so that we take a candidate if one exists.
+		
+		for (const float Candidate : Candidates)
+		{
+			if (const float CandidateDistanceToDesiredParameter = FMath::Abs(DesiredParameter - Candidate);
+				CandidateDistanceToDesiredParameter < BestCandidateDistanceToDesiredParameter)
+			{
+				BestCandidate = Candidate;
+				BestCandidateDistanceToDesiredParameter = CandidateDistanceToDesiredParameter;
+			}
+		}
+
+		return BestCandidate;
+	}
+
 	void RemoveKnotsCompletely(float Parameter, int32 ExpectedKnotCount)
 	{
 		// Remove all occurrences of the knot without checking multiplicity
@@ -868,12 +986,6 @@ public:
 			Knots.Add(NewEndKnot);
 		}
 	}
-
-	FInterval1f GetParameterRange() const
-	{
-		return GetKnotRange();
-	}
-
 	
 protected:
 	
diff --git a/Engine/Source/Runtime/GeometryCore/Public/Splines/MultiSpline.h b/Engine/Source/Runtime/GeometryCore/Public/Splines/MultiSpline.h
--- a/Engine/Source/Runtime/GeometryCore/Public/Splines/MultiSpline.h
+++ b/Engine/Source/Runtime/GeometryCore/Public/Splines/MultiSpline.h
@@ -74,6 +74,13 @@ private:
     
     	// Move assignment
     	FChildSpline& operator=(FChildSpline&& Other) = default;
+
+    	friend FArchive& operator<<( FArchive& Ar, FChildSpline& Child )
+    	{
+    		Child.Serialize(Ar);
+    		return Ar;
+    	}
+    	
 		bool Serialize(FArchive& Ar)
 		{
 			Ar << MappingRangeSpace;
@@ -196,9 +203,38 @@ private:
 
 public:
     using SplineType = typename TSplineWrapper<SPLINETYPE>::SplineType;
+	typedef typename TSplineWrapper<SPLINETYPE>::ValueType ValueType;
 	
     SplineType& GetSpline() { return TSplineWrapper<SPLINETYPE>::InternalSpline; }
     const SplineType& GetSpline() const { return TSplineWrapper<SPLINETYPE>::InternalSpline; }
+
+	virtual bool operator==(const TMultiSpline<SplineType>& Other) const
+    {
+    	if (GetSpline() != Other.GetSpline())
+    	{
+    		return false;
+    	}
+
+    	TArray<FName> AttributeNames = GetAllAttributeChannelNames();
+    	TArray<FName> OtherAttributeNames = Other.GetAllAttributeChannelNames();
+
+    	if (AttributeNames.Num() != OtherAttributeNames.Num())
+    	{
+    		return false;
+    	}
+    	
+    	for (FName& Name : AttributeNames)
+    	{
+    		if (!OtherAttributeNames.Contains(Name))
+    		{
+    			return false;
+    		}
+
+    		// todo: compare attribute splines
+    	}
+    	
+	    return true;
+    }
 	
 	virtual void Clear() override
     {
@@ -217,13 +253,7 @@ public:
 			return false;
 		}
 
-		for (const FName& Name : GetAllAttributeChannelNames())
-		{
-			if (FChildSpline& Child = Children[Name]; !Child.Serialize(Ar))
-			{
-				return false;
-			}
-		}
+		Ar << Children;
 
 		return true;
 	}
@@ -433,7 +463,20 @@ public:
         return 0.f;
     }
 
+    /**
+     * Map a child spline's domain into the primary spline's domain.
+     * @param Name Name of the child spline
+     * @return The domain of the child spline in parent space.
+     */
+    FInterval1f GetMappedChildSpace(FName Name) const
+    {
+		if (const FChildSpline* Child = Children.Find(Name))
+		{
+			return Child->GetParentSpaceRange(GetSpline().GetParameterSpace());
+		}
     
+		return FInterval1f::Empty();
+    }
 
 	/**
      * Creates an attribute channel and returns the typed implementation
@@ -446,18 +489,18 @@ public:
     ImplType* CreateAttributeChannel(const FName& Name)
     {
         // Extract the value type from the implementation type
-        using ValueType = typename ImplType::ValueType;
+        using AttrType = typename ImplType::ValueType;
 
     	// Check if the type is already registered
     	uint32 TypeId = FSplineTypeRegistry::GetTypeId(ImplType::GetSplineTypeName(), 
-													  TSplineValueTypeTraits<ValueType>::Name);
+													  TSplineValueTypeTraits<AttrType>::Name);
     	// If not registered, register it manually
     	if (TypeId == 0)
     	{
     		// Get type information for registration
     		TypeId = ImplType::GetStaticTypeId();
     		FName ImplName = ImplType::GetSplineTypeName();
-    		FName ValueName = TSplineValueTypeTraits<ValueType>::Name;
+    		FName ValueName = TSplineValueTypeTraits<AttrType>::Name;
         
     		// Register with a factory function that creates this specific type
     		bool bSuccess = FSplineTypeRegistry::RegisterType(
@@ -476,7 +519,7 @@ public:
     	}
     	
         // Create the channel using the existing name-based method
-        CreateAttributeChannelInternal<ValueType>(Name, ImplType::GetSplineTypeName());
+        CreateAttributeChannelInternal<AttrType>(Name, ImplType::GetSplineTypeName());
     
         // Return the typed implementation
         return GetTypedAttributeChannel<ImplType>(Name);
diff --git a/Engine/Source/Runtime/GeometryCore/Public/Splines/PolyBezierSpline.h b/Engine/Source/Runtime/GeometryCore/Public/Splines/PolyBezierSpline.h
--- a/Engine/Source/Runtime/GeometryCore/Public/Splines/PolyBezierSpline.h
+++ b/Engine/Source/Runtime/GeometryCore/Public/Splines/PolyBezierSpline.h
@@ -1139,12 +1139,14 @@ private:
 		const ValueType& P3,
 		EParameterizationPolicy ParameterizationPolicy)
     {
+		constexpr float MinSegmentLength = UE_KINDA_SMALL_NUMBER;
+		
     	switch (ParameterizationPolicy)
     	{
     	case EParameterizationPolicy::ChordLength:
-    		return static_cast<float>(Math::Distance(P0, P3));
+    		return static_cast<float>(FMath::Max(Math::Distance(P0, P3), MinSegmentLength));
     	case EParameterizationPolicy::Centripetal:
-    		return static_cast<float>(Math::CentripetalDistance(P0, P3));
+    		return static_cast<float>(FMath::Max(Math::CentripetalDistance(P0, P3), MinSegmentLength));
     	case EParameterizationPolicy::Uniform:
     	default:
 			return 1.0f;
@@ -1337,106 +1339,119 @@ public:
         return Range.Interpolate(T);
 	}
 
-	virtual bool SetParameter(int32 Index, float NewParameter) override
+	virtual int32 SetParameter(int32 Index, float NewParameter) override
 	{
-		static constexpr int32 Degree = 3;
-
 		if (Index < 0 || Index >= Base::NumKeys() || Base::NumKeys() < 4)
 		{
-			return false;
+			return INDEX_NONE;
 		}
-
-		const int32 PointsPerSegment = 4;
+		
+		constexpr int32 PointsPerSegment = 4;
 		const int32 SegmentIndex = Index / PointsPerSegment;
 		const int32 LocalPointIndex = Index % PointsPerSegment;
 		int32 NumSegments = GetNumSegments();
+
+		// Prevent segment degeneration.
+		TArray<float> DistinctKnots = Base::GetDistinctKnotValues();
+		DistinctKnots.Remove(GetParameter(Index));
+		NewParameter = Base::GetNearestValidKnot(NewParameter, DistinctKnots);
+		
 		if (LocalPointIndex == 0) // P0
 		{
-			float OldParam = GetSegmentParameterRange(SegmentIndex).Min;
-			int32 Multiplicity = (Base::IsClampedEnds() && SegmentIndex == 0) ? (Degree + 1) : Degree;
-			//Base::RemoveKnotsCompletely(OldParam, GetExpectedNumKnots());
-			Base::Knots.Remove(OldParam);
-			Base::InsertKnot(NewParameter, Multiplicity);
+			int32 LeftSegmentIndex = SegmentIndex - 1;
+			int32 RightSegmentIndex = SegmentIndex;
+			
+			float OldParameter = GetSegmentParameterRange(SegmentIndex).Min;
+			Base::ReplaceKnot(OldParameter, NewParameter);
 
-			// Fix ordering
 			int32 Segment = SegmentIndex;
-			while (Segment > 0 && GetSegmentParameterRange(Segment).Min < GetSegmentParameterRange(Segment - 1).Min)
+			if (NewParameter < OldParameter)
 			{
-				SwapSegments(Segment, Segment - 1);
-				--Segment;
+				while (LeftSegmentIndex >= 0 && GetSegmentParameterRange(LeftSegmentIndex).Min > GetSegmentParameterRange(LeftSegmentIndex).Max)
+				{
+					FlipSegment(LeftSegmentIndex);
+					LeftSegmentIndex--;
+					RightSegmentIndex--;
+					Segment = RightSegmentIndex;
+				}
 			}
-			while (Segment < NumSegments - 1 && GetSegmentParameterRange(Segment).Min > GetSegmentParameterRange(Segment + 1).Min)
+			else if (NewParameter > OldParameter)
 			{
-				SwapSegments(Segment, Segment + 1);
-				++Segment;
+				while (RightSegmentIndex <= NumSegments - 1 && GetSegmentParameterRange(RightSegmentIndex).Min > GetSegmentParameterRange(RightSegmentIndex).Max)
+				{
+					FlipSegment(RightSegmentIndex);
+					LeftSegmentIndex++;
+					RightSegmentIndex++;
+					Segment = RightSegmentIndex;
+				}
 			}
-			return true;
+			
+			return Segment * PointsPerSegment + LocalPointIndex;
 		}
 		else if (LocalPointIndex == 3) // P3
 		{
-			float OldParam = GetSegmentParameterRange(SegmentIndex).Max;
-			int32 Multiplicity = (Base::IsClampedEnds() && SegmentIndex == GetNumSegments() - 1) ? (Degree + 1) : Degree;
-			// Base::RemoveKnotsCompletely(OldParam, GetExpectedNumKnots());
-			Base::Knots.Remove(OldParam);
-			Base::InsertKnot(NewParameter, Multiplicity);
+			int32 LeftSegmentIndex = SegmentIndex;
+			int32 RightSegmentIndex = SegmentIndex + 1;
+			
+			float OldParameter = GetSegmentParameterRange(SegmentIndex).Max;
+			Base::ReplaceKnot(OldParameter, NewParameter);
 
 			int32 Segment = SegmentIndex;
-            while (Segment > 0 && GetSegmentParameterRange(Segment).Max < GetSegmentParameterRange(Segment - 1).Max)
-            {
-                SwapSegments(Segment, Segment - 1);
-                --Segment;
-            }
-            while (Segment < NumSegments - 1 && GetSegmentParameterRange(Segment).Max > GetSegmentParameterRange(Segment + 1).Max)
-            {
-                SwapSegments(Segment, Segment + 1);
-                ++Segment;
-            }
-			return true;
+			if (NewParameter < OldParameter)
+			{
+				while (LeftSegmentIndex >= 0 && GetSegmentParameterRange(LeftSegmentIndex).Min > GetSegmentParameterRange(LeftSegmentIndex).Max)
+				{
+					FlipSegment(LeftSegmentIndex);
+					LeftSegmentIndex--;
+					RightSegmentIndex--;
+					Segment = LeftSegmentIndex;
+				}
+			}
+			else if (NewParameter > OldParameter)
+			{
+				while (RightSegmentIndex <= NumSegments - 1 && GetSegmentParameterRange(RightSegmentIndex).Min > GetSegmentParameterRange(RightSegmentIndex).Max)
+				{
+					FlipSegment(RightSegmentIndex);
+					LeftSegmentIndex++;
+					RightSegmentIndex++;
+					Segment = LeftSegmentIndex;
+				}
+			}
+			
+			return Segment * PointsPerSegment + LocalPointIndex;
 		}
 
-		return false;
+		return INDEX_NONE;
 	}
-	void SwapSegments(int32 SegmentA, int32 SegmentB)
+	
+	void FlipSegment(int32 Segment)
 	{
 		static constexpr int32 Degree = 3; // Cubic
 		static constexpr int32 ControlPointsPerSegment = 4;
 		static constexpr int32 KnotsPerSegment = Degree; // Not degree+1 because knot overlap between segments
-
-		if (SegmentA == SegmentB || SegmentA < 0 || SegmentB < 0)
-		{
-			return;
-		}
-
+		
 		const int32 NumSegments = GetNumSegments();
-		if (SegmentA >= NumSegments || SegmentB >= NumSegments)
+		if (Segment < 0 || Segment >= NumSegments)
 		{
 			return;
 		}
 
-		// Swap control points
-		int32 ControlIndexA = SegmentA * ControlPointsPerSegment;
-		int32 ControlIndexB = SegmentB * ControlPointsPerSegment;
-
-		for (int32 i = 0; i < ControlPointsPerSegment; ++i)
-		{
-			Base::Values.Swap(ControlIndexA + i, ControlIndexB + i);
-		}
-
-		// Swap repeated knots
-		// Important: due to clamping at endpoints, segments share some knots
-
-		// Degree + 1 knots define a segment (for Cubic: 4 knots)
-		int32 KnotStartA = Degree + SegmentA;
-		int32 KnotStartB = Degree + SegmentB;
+		// Swap P0 && P3
+		// Swap P1 && P2
+		
+		int32 P0 = Segment * ControlPointsPerSegment + 0;
+		int32 P1 = Segment * ControlPointsPerSegment + 1;
+		int32 P2 = Segment * ControlPointsPerSegment + 2;
+		int32 P3 = Segment * ControlPointsPerSegment + 3;
 
-		for (int32 i = 0; i < Degree; ++i)
-		{
-			if ((KnotStartA + i) < Base::Knots.Num() && (KnotStartB + i) < Base::Knots.Num())
-			{
-				Base::Knots.Swap(KnotStartA + i, KnotStartB + i);
-			}
-		}
+		Base::Values.Swap(P0, P3);
+		Base::Values.Swap(P1, P2);
+		
+		// Fix up the knot vector
+		FInterval1f SegmentParamRange = GetSegmentParameterRange(Segment);
+		Base::SwapKnots(SegmentParamRange.Min, SegmentParamRange.Max);
 	}
+	
 	virtual int32 FindIndexForParameter(float Parameter, float& OutLocalParam) const override
 	{
 		TArray<float> Knots = Base::GetDistinctKnotValues();
diff --git a/Engine/Source/Runtime/GeometryCore/Public/Splines/SplineInterfaces.h b/Engine/Source/Runtime/GeometryCore/Public/Splines/SplineInterfaces.h
--- a/Engine/Source/Runtime/GeometryCore/Public/Splines/SplineInterfaces.h
+++ b/Engine/Source/Runtime/GeometryCore/Public/Splines/SplineInterfaces.h
@@ -112,6 +112,11 @@ public:
      */
     ValueType Evaluate(float Parameter) const
     {
+    	if (GetParameterSpace().IsEmpty())
+    	{
+    		return ValueType();
+    	}
+    	
         return EvaluateImpl(HandleOutOfBounds(Parameter));
     }
     
diff --git a/Engine/Source/Runtime/GeometryCore/Public/Splines/TangentBezierSpline.h b/Engine/Source/Runtime/GeometryCore/Public/Splines/TangentBezierSpline.h
--- a/Engine/Source/Runtime/GeometryCore/Public/Splines/TangentBezierSpline.h
+++ b/Engine/Source/Runtime/GeometryCore/Public/Splines/TangentBezierSpline.h
@@ -1333,20 +1333,27 @@ public:
 		InternalSpline.Clear();
 		TangentModes.Empty();
 	}
+	
 	float GetParameter(int32 Index) const
 	{
 		const int32 ControlPointIndex = (Index == 0) ? 0 : ((Index - 1) * 4 + 3);
 		return InternalSpline.GetParameter(ControlPointIndex);
 	}
-	virtual bool SetParameter(int32 Index, float NewParameter)
+
+	int32 SetParameter(int32 Index, float NewParameter)
     {
 		const int32 ControlPointIndex = (Index == 0) ? 0 : ((Index - 1) * 4 + 3);
-        return InternalSpline.SetParameter(ControlPointIndex, NewParameter);
+		return (InternalSpline.SetParameter(ControlPointIndex, NewParameter) + 1) / 4;
     }
 
-	FInterval1f GetParameterRange() const
+	virtual FInterval1f GetParameterSpace() const override
 	{
-		return InternalSpline.GetParameterRange();
+		if (GetNumPoints() == 1)
+		{
+			return FInterval1f(0.f, 0.f);
+		}
+		
+		return InternalSpline.GetParameterSpace();
 	}
 	
 	FInterval1f GetSegmentParameterRange(int32 SegmentIndex) const
