Adding missing InitAction to fix runtime check.

diff --git a/Engine/Source/Runtime/Engine/Private/Texture2D.cpp b/Engine/Source/Runtime/Engine/Private/Texture2D.cpp
--- a/Engine/Source/Runtime/Engine/Private/Texture2D.cpp
+++ b/Engine/Source/Runtime/Engine/Private/Texture2D.cpp
@@ -1679,7 +1679,8 @@ void FVirtualTexture2DResource::InitializeEditorResources(IVirtualTexture* InVir
 			// Fallback to simple 1x1 texture of the fallback color.
 			const FRHITextureCreateDesc TextureDesc =
 				FRHITextureCreateDesc::Create2D(*Name, 1, 1, EPixelFormat::PF_B8G8R8A8)
-				.AddFlags(TexCreateFlags);
+				.AddFlags(TexCreateFlags)
+				.SetInitActionInitializer();
 
 			FRHICommandListImmediate& RHICmdList = FRHICommandListExecutor::GetImmediateCommandList();
 			FRHITextureInitializer Initializer = RHICmdList.CreateTextureInitializer(TextureDesc);
