Added missing support for Enumeration parameters in Material Function Instances.

#rnx

diff --git a/Engine/Source/Editor/MaterialEditor/Private/MaterialPropertyHelpers.cpp b/Engine/Source/Editor/MaterialEditor/Private/MaterialPropertyHelpers.cpp
--- a/Engine/Source/Editor/MaterialEditor/Private/MaterialPropertyHelpers.cpp
+++ b/Engine/Source/Editor/MaterialEditor/Private/MaterialPropertyHelpers.cpp
@@ -434,6 +434,7 @@ FReply FMaterialPropertyHelpers::OnClickedSaveNewFunctionInstance(class UMateria
 					ChildInstance->RuntimeVirtualTextureParameterValues = EditedInstance->RuntimeVirtualTextureParameterValues;
 					ChildInstance->SparseVolumeTextureParameterValues = EditedInstance->SparseVolumeTextureParameterValues;
 					ChildInstance->FontParameterValues = EditedInstance->FontParameterValues;
+					ChildInstance->EnumerationParameterValues = EditedInstance->EnumerationParameterValues;
 
 					const FStaticParameterSet& StaticParameters = EditedInstance->GetStaticParameters();
 					ChildInstance->StaticSwitchParameterValues = StaticParameters.StaticSwitchParameters;
@@ -528,6 +529,7 @@ FReply FMaterialPropertyHelpers::OnClickedSaveNewLayerInstance(class UMaterialFu
 					ChildInstance->RuntimeVirtualTextureParameterValues = EditedInstance->RuntimeVirtualTextureParameterValues;
 					ChildInstance->SparseVolumeTextureParameterValues = EditedInstance->SparseVolumeTextureParameterValues;
 					ChildInstance->FontParameterValues = EditedInstance->FontParameterValues;
+					ChildInstance->EnumerationParameterValues = EditedInstance->EnumerationParameterValues;
 
 					const FStaticParameterSet& StaticParameters = EditedInstance->GetStaticParameters();
 					ChildInstance->StaticSwitchParameterValues = StaticParameters.StaticSwitchParameters;
diff --git a/Engine/Source/Editor/UnrealEd/Private/PreviewMaterial.cpp b/Engine/Source/Editor/UnrealEd/Private/PreviewMaterial.cpp
--- a/Engine/Source/Editor/UnrealEd/Private/PreviewMaterial.cpp
+++ b/Engine/Source/Editor/UnrealEd/Private/PreviewMaterial.cpp
@@ -1024,8 +1024,8 @@ void UMaterialEditorInstanceConstant::ApplySourceFunctionChanges()
 		SourceFunction->RuntimeVirtualTextureParameterValues = SourceInstance->RuntimeVirtualTextureParameterValues;
 		SourceFunction->SparseVolumeTextureParameterValues = SourceInstance->SparseVolumeTextureParameterValues;
 		SourceFunction->FontParameterValues = SourceInstance->FontParameterValues;
+		SourceFunction->EnumerationParameterValues = SourceInstance->EnumerationParameterValues;
 
-		
 		SourceFunction->StaticSwitchParameterValues = SourceInstance->GetStaticParameters().StaticSwitchParameters;
 		const FStaticParameterSetEditorOnlyData& StaticParameters = SourceInstance->GetEditorOnlyStaticParameters();
 		SourceFunction->StaticComponentMaskParameterValues = StaticParameters.StaticComponentMaskParameters;
diff --git a/Engine/Source/Runtime/Engine/Private/Materials/MaterialExpressions.cpp b/Engine/Source/Runtime/Engine/Private/Materials/MaterialExpressions.cpp
--- a/Engine/Source/Runtime/Engine/Private/Materials/MaterialExpressions.cpp
+++ b/Engine/Source/Runtime/Engine/Private/Materials/MaterialExpressions.cpp
@@ -15244,6 +15244,17 @@ void UMaterialFunctionInstance::UpdateParameterSet()
 						}
 					}
 				}
+				else if (const UMaterialExpressionEnumerationParameter* EnumerationParameter = Cast<const UMaterialExpressionEnumerationParameter>(FunctionExpression))
+				{
+					for (FEnumerationParameterValue& EnumerationParameterValue : EnumerationParameterValues)
+					{
+						if (EnumerationParameterValue.ExpressionGUID == EnumerationParameter->ExpressionGUID)
+						{
+							EnumerationParameterValue.ParameterInfo.Name = EnumerationParameter->ParameterName;
+							break;
+						}
+					}
+				}
 			}
 		}
 	}
@@ -15260,6 +15271,7 @@ void UMaterialFunctionInstance::OverrideMaterialInstanceParameterValues(UMateria
 	Instance->RuntimeVirtualTextureParameterValues = RuntimeVirtualTextureParameterValues;
 	Instance->SparseVolumeTextureParameterValues = SparseVolumeTextureParameterValues;
 	Instance->FontParameterValues = FontParameterValues;
+	Instance->EnumerationParameterValues = EnumerationParameterValues;
 
 	// Static parameters
 	FStaticParameterSet StaticParametersOverride = Instance->GetStaticParameters();
@@ -15483,6 +15495,7 @@ bool UMaterialFunctionInstance::GetParameterOverrideValue(EMaterialParameterType
 	case EMaterialParameterType::Font: bResult = GameThread_GetParameterValue(FontParameterValues, ParameterInfo, OutResult); break;
 	case EMaterialParameterType::StaticSwitch: bResult = GameThread_GetParameterValue(StaticSwitchParameterValues, ParameterInfo, OutResult); break;
 	case EMaterialParameterType::StaticComponentMask: bResult = GameThread_GetParameterValue(StaticComponentMaskParameterValues, ParameterInfo, OutResult); break;
+	case EMaterialParameterType::Enumeration: bResult = GameThread_GetParameterValue(EnumerationParameterValues, ParameterInfo, OutResult); break;
 	default: checkNoEntry(); break;
 	}
 
diff --git a/Engine/Source/Runtime/Engine/Public/Materials/MaterialFunctionInstance.h b/Engine/Source/Runtime/Engine/Public/Materials/MaterialFunctionInstance.h
--- a/Engine/Source/Runtime/Engine/Public/Materials/MaterialFunctionInstance.h
+++ b/Engine/Source/Runtime/Engine/Public/Materials/MaterialFunctionInstance.h
@@ -69,6 +69,10 @@ class UMaterialFunctionInstance : public UMaterialFunctionInterface
 	UPROPERTY(EditAnywhere, Category = MaterialFunctionInstance)
 	TArray<struct FSparseVolumeTextureParameterValue> SparseVolumeTextureParameterValues;
 
+	/** Enumeration parameters. */
+	UPROPERTY(EditAnywhere, BlueprintReadOnly, Category = MaterialInstance, meta = (EditFixedOrder))
+	TArray<struct FEnumerationParameterValue> EnumerationParameterValues;
+
 #if WITH_EDITOR
 	ENGINE_API void UpdateParameterSet();
 	ENGINE_API void OverrideMaterialInstanceParameterValues(class UMaterialInstance* Instance);
