Allow component incremental pre-register / pre-unregister on listen server.

diff --git a/Engine/Source/Runtime/Engine/Private/Level.cpp b/Engine/Source/Runtime/Engine/Private/Level.cpp
--- a/Engine/Source/Runtime/Engine/Private/Level.cpp
+++ b/Engine/Source/Runtime/Engine/Private/Level.cpp
@@ -1812,10 +1812,10 @@ void ULevel::IncrementalUpdateComponents(int32 NumComponentsToUpdate, bool bReru
 	// The editor is never allowed to incrementally update components.  Make sure to pass in a value of zero for NumActorsToUpdate.
 	check(bFullyUpdateComponents || OwningWorld->IsGameWorld());
 
-	auto IsServer = [this]()
+	auto IsDedicatedServer = [this]()
 	{
 		const ENetMode NetMode = OwningWorld->GetNetMode();
-		return NetMode == NM_DedicatedServer || NetMode == NM_ListenServer;
+		return NetMode == NM_DedicatedServer;
 	};
 
 	FRegisterComponentContext LocalContext(OwningWorld);
@@ -1844,7 +1844,7 @@ void ULevel::IncrementalUpdateComponents(int32 NumComponentsToUpdate, bool bReru
 				// - world is inside a blocking level streaming
 				const bool bIsPreRegisterInitialComponentsAllowed = !bFullyUpdateComponents &&
 					GAllowIncrementalPreRegisterComponents &&
-					!IsServer() &&
+					!IsDedicatedServer() &&
 					FPhysScene::SupportsAsyncPhysicsStateCreation() &&
 					!OwningWorld->bRequestedBlockOnAsyncLoading && 
 					!OwningWorld->GetIsInBlockTillLevelStreamingCompleted();
@@ -2145,10 +2145,10 @@ bool ULevel::IncrementalUnregisterComponents(int32 NumComponentsToUnregister, FU
 		return true;
 	}
 
-	auto IsServer = [this]()
+	auto IsDedicatedServer = [this]()
 	{
 		const ENetMode NetMode = OwningWorld->GetNetMode();
-		return NetMode == NM_DedicatedServer || NetMode == NM_ListenServer;
+		return NetMode == NM_DedicatedServer;
 	};
 
 	auto IncrementalUnregisterComponentsInternal = [this, NumComponentsToUnregister]()
@@ -2201,7 +2201,7 @@ bool ULevel::IncrementalUnregisterComponents(int32 NumComponentsToUnregister, FU
 				const bool bIsPreUnregisterInitialComponentsAllowed = bIncrementalUnregisterComponents && 
 					(NumComponentsToUnregister != MAX_int32) &&
 					GAllowIncrementalPreUnregisterComponents &&
-					!IsServer() &&
+					!IsDedicatedServer() &&
 					FPhysScene::SupportsAsyncPhysicsStateDestruction() &&
 					!OwningWorld->bRequestedBlockOnAsyncLoading &&
 					!OwningWorld->GetIsInBlockTillLevelStreamingCompleted();
