FNewSpline: Fix several issues in SetAttributeParameter

1) When an internal point becomes an end point, we slightly reduce the child space. Everything is kept correct in parent space, but repeatedly performing this operation by jumping back and forth across the end will result in a child space which is very small which leads to float imprecision issues

2) The code for handling end points assumes they can only cross their immediate neighbor, but this is not accurate. It can cross any number of points. This primarily required changes to the reparameterization logic and the logic used to update the mapping range. The reparameterization logic was wrong in all cases where the final end point index changed by > 1 position, while the mapping range logic was only wrong in the case that the end point actually fully passed every attribute point, thereby becoming an end point on the other side.

3) Setting closed loop on the primary spline requires we set closed loop and reupdate the knot vector of the scale and rot channels in order to keep the point parameter mappings stable. Before revision 2 of this CL, we would ensure on setting closed loop. Now, we do not.

diff --git a/Engine/Source/Runtime/Engine/Private/Curves/Spline.cpp b/Engine/Source/Runtime/Engine/Private/Curves/Spline.cpp
--- a/Engine/Source/Runtime/Engine/Private/Curves/Spline.cpp
+++ b/Engine/Source/Runtime/Engine/Private/Curves/Spline.cpp
@@ -421,6 +421,22 @@ int32 FNewSpline::SetAttributeParameter(const FName& Name, int32 Index, float Pa
 	{
 		float ParentSpaceParameter = ToInternalSplineSpace(Parameter);
 
+		// Prevent collapse of the channel space and keep it starting at 0.
+		auto SanitizeChannelSpace = [Channel]()
+		{
+			const FInterval1f ChannelSpace = Channel->GetParameterSpace();
+			const float KnotOffset = -1.f * ChannelSpace.Min;
+			const float KnotScaleFactor = 1.f / (ChannelSpace.Max - ChannelSpace.Min);
+
+			TArray<float> Knots = Channel->GetKnotVector();
+			for (float& Knot : Knots)
+			{
+				Knot += KnotOffset;
+				Knot *= KnotScaleFactor;
+			}
+			Channel->SetKnotVector(Knots);
+		};
+		
 		// Helper to prevent code duplication.
 		// Uses current value of Index and ParentSpaceParameter at the time of calling to update the attribute parameter.
 		auto SetParentSpaceParameter = [this, &ParentSpaceParameter, Name, Channel, &Index]() -> int32
@@ -458,8 +474,8 @@ int32 FNewSpline::SetAttributeParameter(const FName& Name, int32 Index, float Pa
 
 			if (MappedRange.Min > MappedRange.Max)
 			{
-				// If range flips, swap end points and un-flip.
-				Channel->SetParameter(1, Channel->GetParameter(0) - 1.f);
+				// The mapping range will flip, swap end points and un-flip.
+				Channel->SetParameter(1, Channel->GetParameter(0) - 1.f);	// It doesn't actually matter what we do here, as long as we get the ordering & mapping range correct.
 				SetAttributeChannelRange(Name, FInterval1f(MappedRange.Max, MappedRange.Min), EMappingRangeSpace::Parent);
 				Index = Index == 0 ? 1 : 0;
 			}
@@ -486,12 +502,29 @@ int32 FNewSpline::SetAttributeParameter(const FName& Name, int32 Index, float Pa
 			if (bAttributesWillReorder)
 			{
 				Index = SetParentSpaceParameter();
+
+				// Shift values left up to the last invalidated index.
+				const int32 InvalidatedInternalParameterIdx = FMath::Clamp(Index - 1, 0, InternalParameters.Num() - 1);
+				for (int InternalParameterIdx = 1; InternalParameterIdx <= InvalidatedInternalParameterIdx; ++InternalParameterIdx)
+				{
+					InternalParameters[InternalParameterIdx - 1] = InternalParameters[InternalParameterIdx];
+				}
 				
-				// Update InternalParameters to properly re-parameterize the new internal point at Index - 1
-				InternalParameters[Index - 1] = ParentSpaceParameter;
+				// Do one of 2 things to the actual invalidated index:
+				if (Index == Channel->GetNumPoints() - 1)
+				{
+					// Old upper end point is now an internal point and needs to be re-parameterized.
+					InternalParameters[InvalidatedInternalParameterIdx] = ParentSpaceUpperBound;
+				}
+				else
+				{
+					// Old lower end point is now an internal point and needs to be re-parameterized.
+					InternalParameters[InvalidatedInternalParameterIdx] = ParentSpaceParameter;
+				}
 			}
 			
 			// Updates the lower bound of the mapping range.
+			ParentSpaceUpperBound = Index == Channel->GetNumPoints() - 1 ? ParentSpaceParameter : ParentSpaceUpperBound;
 			const float ParentSpaceLowerBound = bAttributesWillReorder ? NeighborParentSpaceParameter : ParentSpaceParameter;
 			SetAttributeChannelRange(Name, FInterval1f(ParentSpaceLowerBound, ParentSpaceUpperBound), EMappingRangeSpace::Parent);
 
@@ -501,6 +534,11 @@ int32 FNewSpline::SetAttributeParameter(const FName& Name, int32 Index, float Pa
 				Channel->SetParameter(InternalIdx, MapParameterToChildSpace(Name, InternalParameters[InternalIdx - 1]));
 			}
 			
+			if (bAttributesWillReorder)
+			{
+				SanitizeChannelSpace();
+			}
+			
 			return Index;
 		}
 		else if (Index == Channel->GetNumPoints() - 1)			// Case 4: Moving the last point (which has exactly 1 neighbor)
@@ -521,11 +559,28 @@ int32 FNewSpline::SetAttributeParameter(const FName& Name, int32 Index, float Pa
 			{
 				Index = SetParentSpaceParameter();
 
-				// Update InternalParameters to properly re-parameterize the new internal point at Index - 1
-				InternalParameters[Index - 1] = ParentSpaceParameter;
+				// Shift values right up to the last invalidated index.
+				const int32 InvalidatedInternalParameterIdx = FMath::Clamp(Index - 1, 0, InternalParameters.Num() - 1);
+				for (int InternalParameterIdx = InternalParameters.Num() - 2; InternalParameterIdx >= InvalidatedInternalParameterIdx; --InternalParameterIdx)
+				{
+					InternalParameters[InternalParameterIdx + 1] = InternalParameters[InternalParameterIdx];
+				}
+				
+				// Do one of 2 things to the actual invalidated index:
+				if (Index == 0)
+				{
+					// Old upper end point is now an internal point and needs to be re-parameterized.
+					InternalParameters[InvalidatedInternalParameterIdx] = ParentSpaceLowerBound;
+				}
+				else
+				{
+					// Old lower end point is now an internal point and needs to be re-parameterized.
+					InternalParameters[InvalidatedInternalParameterIdx] = ParentSpaceParameter;
+				}
 			}
 
 			// Update the child space mapping
+			ParentSpaceLowerBound = Index == 0 ? ParentSpaceParameter : ParentSpaceLowerBound;
 			float ParentSpaceUpperBound = bAttributesWillReorder ? NeighborParentSpaceParameter : ParentSpaceParameter; // Prevent over-shrinking if we reshuffle.
 			SetAttributeChannelRange(Name, FInterval1f(ParentSpaceLowerBound, ParentSpaceUpperBound), EMappingRangeSpace::Parent);
 			
@@ -534,12 +589,32 @@ int32 FNewSpline::SetAttributeParameter(const FName& Name, int32 Index, float Pa
 			{
 				Channel->SetParameter(InternalIdx, MapParameterToChildSpace(Name, InternalParameters[InternalIdx - 1]));
 			}
+
+			if (bAttributesWillReorder)
+			{
+				SanitizeChannelSpace();
+			}
 		
 			return Index;
 		}
 		else													// Case 5: Moving an internal point (which has exactly 2 neighbors)
 		{
+			FInterval1f MappedRange = GetMappedChildSpace(Name);
 			Index = SetParentSpaceParameter();
+			
+			if (Index == 0)
+			{
+				// The internal point is now the lower end point
+				MappedRange.Min = ParentSpaceParameter;
+			}
+			else if (Index == Channel->GetNumPoints() - 1)
+			{
+				// The internal point is now the upper end point
+				MappedRange.Max = ParentSpaceParameter;
+			}
+			
+			SetAttributeChannelRange(Name, MappedRange, EMappingRangeSpace::Parent);
+			
 			return Index;
 		}
 	}
@@ -1735,6 +1810,15 @@ float FNewSpline::GetSplineLength() const
 void FNewSpline::SetClosedLoop(bool bInClosedLoop)
 {
     GetSpline().SetClosedLoop(bInClosedLoop);
+	
+	TTangentBezierSpline<FVector>* ScaleChild = GetTypedAttributeChannel<TTangentBezierSpline<FVector>>(ScaleAttrName);
+	ScaleChild->SetClosedLoop(bInClosedLoop);
+	ScaleChild->SetKnotVector(GetSpline().GetKnotVector());
+		
+	TTangentBezierSpline<FQuat>* RotChild = GetTypedAttributeChannel<TTangentBezierSpline<FQuat>>(RotationAttrName);
+	RotChild->SetClosedLoop(bInClosedLoop);
+	RotChild->SetKnotVector(GetSpline().GetKnotVector());
+	
     UpdateSpline();
 }
 
diff --git a/Engine/Source/Runtime/GeometryCore/Public/BoxTypes.h b/Engine/Source/Runtime/GeometryCore/Public/BoxTypes.h
--- a/Engine/Source/Runtime/GeometryCore/Public/BoxTypes.h
+++ b/Engine/Source/Runtime/GeometryCore/Public/BoxTypes.h
@@ -173,7 +173,18 @@ struct TInterval1
 		{
 			return 1;
 		}
-		else if (Min == Max)
+		else
+		{
+			return GetUnclampedT(Value);
+		}
+	}
+
+	/**
+	 * Convert Value into T Value.
+	 */
+	RealType GetUnclampedT(RealType Value) const
+	{
+		if (Min == Max)
 		{
 			return 0.5;
 		}
diff --git a/Engine/Source/Runtime/GeometryCore/Public/Splines/MultiSpline.h b/Engine/Source/Runtime/GeometryCore/Public/Splines/MultiSpline.h
--- a/Engine/Source/Runtime/GeometryCore/Public/Splines/MultiSpline.h
+++ b/Engine/Source/Runtime/GeometryCore/Public/Splines/MultiSpline.h
@@ -114,45 +114,45 @@ private:
 		    return true;
 		}
     	
-        float MapParameterToChildSpace(float ParentParameter, const FInterval1f& ParentSpace) const
-        {
-            if (!Spline)
-            {
-                return 0.f;
-            }
+		float MapParameterToChildSpace(float ParentParameter, const FInterval1f& ParentSpace) const
+		{
+			if (!Spline)
+			{
+				return 0.f;
+			}
 
 			if (FInterval1f ChildSpace = Spline->GetParameterSpace(); !ChildSpace.IsEmpty())
 			{
 				FInterval1f MappedParentRange = GetParentSpaceRange(ParentSpace);
-    
-				float T = MappedParentRange.GetT(ParentParameter);
+	
+				float T = MappedParentRange.GetUnclampedT(ParentParameter);
 				return ChildSpace.Interpolate(T);
 			}
 			else
 			{
 				return 0.f;
 			}
-        }
-    
-        float MapParameterFromChildSpace(float ChildParameter, const FInterval1f& ParentSpace) const
-        {
-            if (!Spline)
-            {
-                return 0.f;
-            }
+		}
+	
+		float MapParameterFromChildSpace(float ChildParameter, const FInterval1f& ParentSpace) const
+		{
+			if (!Spline)
+			{
+				return 0.f;
+			}
 
 			if (FInterval1f ChildSpace = Spline->GetParameterSpace(); !ChildSpace.IsEmpty())
 			{
 				FInterval1f MappedParentRange = GetParentSpaceRange(ParentSpace);
-    
-				float T = ChildSpace.GetT(ChildParameter);
+	
+				float T = ChildSpace.GetUnclampedT(ChildParameter);
 				return MappedParentRange.Interpolate(T);
 			}
 			else
 			{
 				return 0.f;
 			}
-        }
+		}
 
         FInterval1f GetParentSpaceRange(const FInterval1f& ParentSpace) const
         {
