Fix Instanced Structs Property Migration

diff --git a/Engine/Source/Runtime/Engine/Private/ObjectEditorUtils.cpp b/Engine/Source/Runtime/Engine/Private/ObjectEditorUtils.cpp
--- a/Engine/Source/Runtime/Engine/Private/ObjectEditorUtils.cpp
+++ b/Engine/Source/Runtime/Engine/Private/ObjectEditorUtils.cpp
@@ -213,16 +213,17 @@ namespace FObjectEditorUtils
 		if (FStructProperty* const DestStructProperty = CastField<FStructProperty>(InDestinationProperty))
 		{
 			FStructProperty* const SrcStructProperty = CastField<FStructProperty>(InSourceProperty);
-
-			// Ensure that the target struct is initialized before copying fields from the source.
-			DestStructProperty->InitializeValue_InContainer(InTargetPtr);
-
+			
+			// Now go through each property and make sure we take care of things like nested instanced objects
 			const int32 PropertyArrayDim = DestStructProperty->ArrayDim;
 			for (int32 ArrayIndex = 0; ArrayIndex < PropertyArrayDim; ArrayIndex++)
 			{
 				const void* const SourcePtr = SrcStructProperty->ContainerPtrToValuePtr<void>(InSourcePtr, ArrayIndex);
 				void* const TargetPtr = DestStructProperty->ContainerPtrToValuePtr<void>(InTargetPtr, ArrayIndex);
 
+				// Copy complete element value
+				InSourceProperty->CopyCompleteValue(TargetPtr, SourcePtr);
+
 				for (TFieldIterator<FProperty> It(SrcStructProperty->Struct); It; ++It)
 				{
 					FProperty* const InnerProperty = *It;
