Add a Map Check warning when static lights are present in the scene but disabled due to Lumen.

diff --git a/Engine/Source/Runtime/Core/Private/Misc/MapErrors.cpp b/Engine/Source/Runtime/Core/Private/Misc/MapErrors.cpp
--- a/Engine/Source/Runtime/Core/Private/Misc/MapErrors.cpp
+++ b/Engine/Source/Runtime/Core/Private/Misc/MapErrors.cpp
@@ -71,6 +71,7 @@ FLazyName FMapErrors::LODActorNoActorFound(TEXT("LODActorNoActor"));
 FLazyName FMapErrors::InvalidVirtualTextureUsage(TEXT("InvalidVirtualTextureUsage"));
 FLazyName FMapErrors::PrimitiveComponentHasInvalidTranslucentShadowSetting(TEXT("PrimitiveComponentHasInvalidTranslucentShadowSetting"));
 FLazyName FMapErrors::UseFauxOrthoViewPosDeprecation_Warning(TEXT("UseFauxOrthoViewPosDeprecation_Warning"));
+FLazyName FMapErrors::StaticLightWithLumen_Warning(TEXT("StaticLightWithLumen_Warning"));
 
 namespace
 {
diff --git a/Engine/Source/Runtime/Core/Public/Misc/MapErrors.h b/Engine/Source/Runtime/Core/Public/Misc/MapErrors.h
--- a/Engine/Source/Runtime/Core/Public/Misc/MapErrors.h
+++ b/Engine/Source/Runtime/Core/Public/Misc/MapErrors.h
@@ -158,6 +158,9 @@ struct FMapErrors
 	/** bUseFauxOrthoViewPos has been deprecated. Property will be lost on re-save. */
 	static CORE_API FLazyName UseFauxOrthoViewPosDeprecation_Warning;
 
+	/** {LightComponent} is using static lighting while Lumen is enabled. Static lights will fail to render. */
+	static CORE_API FLazyName StaticLightWithLumen_Warning;
+
 	/** Navigation */
 
 	/** Paths need to be rebuilt: **TODO** */
diff --git a/Engine/Source/Runtime/Engine/Classes/Components/LightComponent.h b/Engine/Source/Runtime/Engine/Classes/Components/LightComponent.h
--- a/Engine/Source/Runtime/Engine/Classes/Components/LightComponent.h
+++ b/Engine/Source/Runtime/Engine/Classes/Components/LightComponent.h
@@ -456,6 +456,7 @@ public:
 	ENGINE_API virtual bool CanEditChange(const FProperty* InProperty) const override;
 	ENGINE_API virtual void PostEditChangeProperty(FPropertyChangedEvent& PropertyChangedEvent) override;
 	ENGINE_API virtual void UpdateLightSpriteTexture() override;
+	ENGINE_API virtual void CheckForErrors() override;
 #endif // WITH_EDITOR
 	ENGINE_API virtual void BeginDestroy() override;
 	ENGINE_API virtual bool IsReadyForFinishDestroy() override;
diff --git a/Engine/Source/Runtime/Engine/Private/Components/LightComponent.cpp b/Engine/Source/Runtime/Engine/Private/Components/LightComponent.cpp
--- a/Engine/Source/Runtime/Engine/Private/Components/LightComponent.cpp
+++ b/Engine/Source/Runtime/Engine/Private/Components/LightComponent.cpp
@@ -24,6 +24,7 @@
 #include "PrimitiveDrawingUtils.h"
 #include "SceneView.h"
 #include "Logging/MessageLog.h"
+#include "Misc/MapErrors.h"
 #include "Misc/UObjectToken.h"
 #include "Components/PointLightComponent.h"
 #include "Components/DirectionalLightComponent.h"
@@ -34,6 +35,7 @@
 #include "UObject/UnrealType.h"
 #include "UObject/Package.h"
 #include "WorldPartition/ActorInstanceGuids.h"
+#include "GameFramework/WorldSettings.h"
 
 #include UE_INLINE_GENERATED_CPP_BY_NAME(LightComponent)
 
@@ -41,6 +43,8 @@
 #include "Rendering/StaticLightingSystemInterface.h"
 #endif
 
+#define LOCTEXT_NAMESPACE "LightComponent"
+
 #if WITH_EDITOR
 static const TCHAR* GLightSpriteAssetName = TEXT("/Engine/EditorResources/LightIcons/S_LightError.S_LightError");
 #endif
@@ -855,6 +859,30 @@ void ULightComponent::UpdateLightSpriteTexture()
 	}
 }
 
+void ULightComponent::CheckForErrors()
+{
+	Super::CheckForErrors();
+
+	const bool bUsingStaticLighting = HasStaticLighting() &&
+		!GetWorld()->GetWorldSettings()->bForceNoPrecomputedLighting &&
+		IsStaticLightingAllowed();
+	
+	static const auto CVarDynamicGIMethod = IConsoleManager::Get().FindTConsoleVariableDataInt(TEXT("r.DynamicGlobalIlluminationMethod"));
+	static const auto CVarLumenAllowed = IConsoleManager::Get().FindTConsoleVariableDataInt(TEXT("r.Lumen.DiffuseIndirect.Allow"));
+
+	const bool bLumenEnabled = CVarDynamicGIMethod->GetValueOnGameThread() == 1 &&
+		CVarLumenAllowed->GetValueOnGameThread() &&
+		DoesPlatformSupportLumenGI(GetScene()->GetShaderPlatform()) &&
+		DoesProjectSupportLumenGI(GetScene()->GetShaderPlatform());
+
+	if (bUsingStaticLighting && bLumenEnabled)
+	{
+		FMessageLog("MapCheck").Warning()
+				->AddToken(FUObjectToken::Create(this))
+				->AddToken(FTextToken::Create(LOCTEXT( "MapCheck_Message_StaticLightWithLumen", "is using static lighting while Lumen is enabled. Static lights will fail to render." )))
+				->AddToken(FMapErrorToken::Create(FMapErrors::StaticLightWithLumen_Warning));
+	}
+}
 #endif // WITH_EDITOR
 
 void ULightComponent::BeginDestroy()
@@ -1760,3 +1788,5 @@ static FAutoConsoleCommand ToggleLightCmd(
 	ECVF_Cheat
 	);
 
+
+#undef LOCTEXT_NAMESPACE
