FSpline: Add a CVar that controls the functionality of ValidateRotScale

The validation code is useful but it can throw false positives when comparing actual and expected parameters due to float imprecision when dealing with large parameter spaces. Until the validation is more robust, we will disable it by default.

diff --git a/Engine/Source/Runtime/Engine/Private/Curves/Spline.cpp b/Engine/Source/Runtime/Engine/Private/Curves/Spline.cpp
--- a/Engine/Source/Runtime/Engine/Private/Curves/Spline.cpp
+++ b/Engine/Source/Runtime/Engine/Private/Curves/Spline.cpp
@@ -64,6 +64,13 @@ bool GEncodeAsHex = true;
 // 	GEncodeAsHex,
 // 	TEXT("If true, splines are imported/exported using hexadecimal encoding. If false, the data is encoded in base 64.")
 // );
+
+bool GValidateRotScale = false;
+FAutoConsoleVariableRef CVarValidateRotScale(
+	TEXT("Spline.ValidateRotScale"),
+	GValidateRotScale,
+	TEXT("True if we should validate rotation and scale attributes when structurally modifying the spline.")
+);
 	
 }
 
@@ -2438,6 +2445,11 @@ int32 FNewSpline::ConvertInternalParameterToNearestPointIndex(float Parameter) c
 
 void FNewSpline::ValidateRotScale() const
 {
+	if (!UE::Spline::GValidateRotScale)
+	{
+		return;
+	}
+	
 	// Validate child existence
 	TTangentBezierSpline<FQuat>* RotSpline = GetTypedAttributeChannel<TTangentBezierSpline<FQuat>>(RotationAttrName);
 	TTangentBezierSpline<FVector>* ScaleSpline = GetTypedAttributeChannel<TTangentBezierSpline<FVector>>(ScaleAttrName);
