Improve ensure messages in SetupAttachment.

diff --git a/Engine/Source/Runtime/Engine/Private/Components/SceneComponent.cpp b/Engine/Source/Runtime/Engine/Private/Components/SceneComponent.cpp
--- a/Engine/Source/Runtime/Engine/Private/Components/SceneComponent.cpp
+++ b/Engine/Source/Runtime/Engine/Private/Components/SceneComponent.cpp
@@ -2209,13 +2209,13 @@ void USceneComponent::SetupAttachment(class USceneComponent* InParent, FName InS
 {
 	if (InParent != AttachParent || InSocketName != AttachSocketName)
 	{
-		if (ensureMsgf(!bRegistered, TEXT("SetupAttachment should only be used to initialize AttachParent and AttachSocketName for a future AttachToComponent. Once a component is registered you must use AttachToComponent. Owner [%s], InParent [%s], InSocketName [%s]"), *GetPathNameSafe(GetOwner()), *GetNameSafe(InParent), *InSocketName.ToString()))
+		if (ensureMsgf(!bRegistered, TEXT("SetupAttachment should only be used to initialize AttachParent and AttachSocketName for a future AttachToComponent. Once a component is registered you must use AttachToComponent. Component [%s] Owner [%s], InParent [%s], InSocketName [%s]"), *GetFullName(), *GetPathNameSafe(GetOwner()), *GetNameSafe(InParent), *InSocketName.ToString()))
 		{
-			if (ensureMsgf(InParent != this, TEXT("Cannot attach a component to itself.")))
+			if (ensureMsgf(InParent != this, TEXT("Cannot attach component %s to itself."), *GetFullName()))
 			{
-				if (ensureMsgf(InParent == nullptr || !InParent->IsAttachedTo(this), TEXT("Setting up attachment would create a cycle.")))
+				if (ensureMsgf(InParent == nullptr || !InParent->IsAttachedTo(this), TEXT("Attaching %s to %s would create a cycle."), *GetFullName(), *InParent->GetFullName()))
 				{
-					if (ensureMsgf(AttachParent == nullptr || !AttachParent->AttachChildren.Contains(this), TEXT("SetupAttachment cannot be used once a component has already had AttachTo used to connect it to a parent.")))
+					if (ensureMsgf(AttachParent == nullptr || !AttachParent->AttachChildren.Contains(this), TEXT("SetupAttachment cannot be used once a component has already had AttachTo used to connect it to a parent. %s is already attached to %s"), *GetFullName(), *AttachParent->GetFullName()))
 					{
 						SetAttachParent(InParent);
 						SetAttachSocketName(InSocketName);
