- Fix crash in gpu sort test due to buffer refactor and modifying cvars on the RT

#rnx

diff --git a/Engine/Source/Runtime/Engine/Private/GPUSort.cpp b/Engine/Source/Runtime/Engine/Private/GPUSort.cpp
--- a/Engine/Source/Runtime/Engine/Private/GPUSort.cpp
+++ b/Engine/Source/Runtime/Engine/Private/GPUSort.cpp
@@ -22,8 +22,10 @@ DEFINE_LOG_CATEGORY_STATIC(LogGPUSort, Log, All);
 	Global settings.
 ------------------------------------------------------------------------------*/
 
-static TAutoConsoleVariable<int32> CVarDebugOffsets(TEXT("GPUSort.DebugOffsets"),0,TEXT("Debug GPU sort offsets."));
-static TAutoConsoleVariable<int32> CVarDebugSort(TEXT("GPUSort.DebugSort"),0,TEXT("Debug GPU sorting."));
+static bool GbSortDebugOffsets = false;
+static bool GbSortDebug = false;
+static FAutoConsoleVariableRef CVarDebugOffsets(TEXT("GPUSort.DebugOffsets"), GbSortDebugOffsets, TEXT("Debug GPU sort offsets."));
+static FAutoConsoleVariableRef CVarDebugSort(TEXT("GPUSort.DebugSort"), GbSortDebug, TEXT("Debug GPU sorting."));
 
 #define GPUSORT_BITCOUNT 32
 #define RADIX_BITS 4
@@ -514,8 +516,8 @@ int32 SortGPUBuffers(FRHICommandList& RHICmdList, FGPUSortBuffers SortBuffers, i
 {
 	FRadixSortParameters SortParameters;
 	FRadixSortUniformBufferRef SortUniformBufferRef;
-	const bool bDebugOffsets = CVarDebugOffsets.GetValueOnRenderThread() != 0;
-	const bool bDebugSort = CVarDebugSort.GetValueOnRenderThread() != 0;
+	const bool bDebugOffsets = GbSortDebugOffsets;
+	const bool bDebugSort = GbSortDebug;
 
 	SCOPED_DRAW_EVENTF(RHICmdList, SortGPU, TEXT("Sort(%d)"), Count);
 
@@ -723,8 +725,8 @@ static bool RunGPUSortTest(FRHICommandListImmediate& RHICmdList, int32 TestSize,
 	int32 ResultBufferIndex;
 	int32 IncorrectKeyIndex = 0;
 	const int32 BufferSize = TestSize * sizeof(uint32);
-	const bool bDebugOffsets = CVarDebugOffsets.GetValueOnRenderThread() != 0;
-	const bool bDebugSort = CVarDebugSort.GetValueOnRenderThread() != 0;
+	const bool bDebugOffsets = GbSortDebugOffsets;
+	const bool bDebugSort = GbSortDebug;
 
 	// Generate the test keys.
 	Keys.Reserve(TestSize);
@@ -739,6 +741,7 @@ static bool RunGPUSortTest(FRHICommandListImmediate& RHICmdList, int32 TestSize,
 	RefSortedKeys.Sort();
 
 	// Upload initial keys and values to the GPU.
+	for ( int32 i=0; i < UE_ARRAY_COUNT(KeysBufferRHI); ++i )
 	{
 		const FRHIBufferCreateDesc KeysCreateDesc =
 			FRHIBufferCreateDesc::CreateVertex(TEXT("KeysBuffer"), BufferSize)
@@ -748,10 +751,10 @@ static bool RunGPUSortTest(FRHICommandListImmediate& RHICmdList, int32 TestSize,
 
 		FRHIBufferInitializer Initializer = RHICmdList.CreateBufferInitializer(KeysCreateDesc);
 		Initializer.WriteData(Keys.GetData(), BufferSize);
-		KeysBufferRHI[0] = Initializer.Finalize();
-		KeysBufferRHI[1] = RHICmdList.CreateBuffer(KeysCreateDesc);
+		KeysBufferRHI[i] = Initializer.Finalize();
 	}
 
+	for (int32 i = 0; i < UE_ARRAY_COUNT(ValuesBufferRHI); ++i)
 	{
 		const FRHIBufferCreateDesc ValuesCreateDesc =
 			FRHIBufferCreateDesc::CreateVertex(TEXT("ValuesBuffer"), BufferSize)
@@ -761,8 +764,7 @@ static bool RunGPUSortTest(FRHICommandListImmediate& RHICmdList, int32 TestSize,
 
 		FRHIBufferInitializer Initializer = RHICmdList.CreateBufferInitializer(ValuesCreateDesc);
 		Initializer.WriteData(Keys.GetData(), BufferSize);
-		ValuesBufferRHI[0] = Initializer.Finalize();
-		ValuesBufferRHI[1] = RHICmdList.CreateBuffer(ValuesCreateDesc);
+		ValuesBufferRHI[i] = Initializer.Finalize();
 	}
 
 	// Allocate GPU resources.
@@ -866,15 +868,9 @@ static bool RunGPUSortTest(FRHICommandListImmediate& RHICmdList, int32 TestSize,
  */
 static void RunGPUSortTestWithDebug(FRHICommandListImmediate& RHICmdList, int32 TestSize, ERHIFeatureLevel::Type FeatureLevel)
 {
-	static IConsoleVariable* IVarDebugOffsets = IConsoleManager::Get().FindConsoleVariable(TEXT("GPUSort.DebugOffsets"));
-	static IConsoleVariable* IVarDebugSort = IConsoleManager::Get().FindConsoleVariable(TEXT("GPUSort.DebugSort"));
-	const bool bWasDebuggingOffsets = CVarDebugOffsets.GetValueOnRenderThread() != 0;
-	const bool bWasDebuggingSort = CVarDebugSort.GetValueOnRenderThread() != 0;
-	IVarDebugOffsets->Set(1, ECVF_SetByCode);
-	IVarDebugSort->Set(1, ECVF_SetByCode);
+	TGuardValue<bool> SortDebugOffsets(GbSortDebugOffsets, true);
+	TGuardValue<bool> SortDebug(GbSortDebug, true);
 	RunGPUSortTest(RHICmdList, TestSize, FeatureLevel);
-	IVarDebugOffsets->Set(bWasDebuggingOffsets ? 1 : 0, ECVF_SetByCode);
-	IVarDebugSort->Set(bWasDebuggingSort ? 1 : 0, ECVF_SetByCode);
 }
 
 /**
