World Partition: fix crash when reloading a standalone level instance level with data layers when LevelStreaming.ShouldReuseUnloadedButStillAroundLevels is set.

diff --git a/Engine/Source/Runtime/Engine/Private/WorldPartition/DataLayer/WorldDataLayers.cpp b/Engine/Source/Runtime/Engine/Private/WorldPartition/DataLayer/WorldDataLayers.cpp
--- a/Engine/Source/Runtime/Engine/Private/WorldPartition/DataLayer/WorldDataLayers.cpp
+++ b/Engine/Source/Runtime/Engine/Private/WorldPartition/DataLayer/WorldDataLayers.cpp
@@ -104,6 +104,14 @@ void AWorldDataLayers::PostRegisterAllComponents()
 	}
 }
 
+void AWorldDataLayers::PostUnregisterAllComponents()
+{
+	Super::PostUnregisterAllComponents();
+
+	ResetDataLayerRuntimeStates();
+	InitializeDataLayerRuntimeStates();
+}
+
 #define AWORLDDATALAYERS_UPDATE_REPLICATED_DATALAYERS(ReplicatedArray, SourceArray) \
 	MARK_PROPERTY_DIRTY_FROM_NAME(AWorldDataLayers, ReplicatedArray, this); \
 	ReplicatedArray = SourceArray;
diff --git a/Engine/Source/Runtime/Engine/Public/WorldPartition/DataLayer/WorldDataLayers.h b/Engine/Source/Runtime/Engine/Public/WorldPartition/DataLayer/WorldDataLayers.h
--- a/Engine/Source/Runtime/Engine/Public/WorldPartition/DataLayer/WorldDataLayers.h
+++ b/Engine/Source/Runtime/Engine/Public/WorldPartition/DataLayer/WorldDataLayers.h
@@ -94,6 +94,7 @@ public:
 	//~ Begin AActor Interface
 	ENGINE_API virtual void RewindForReplay() override;
 	ENGINE_API virtual void PostRegisterAllComponents() override;
+	ENGINE_API virtual void PostUnregisterAllComponents() override;
 #if WITH_EDITOR
 	ENGINE_API virtual void PreEditUndo() override;
 	ENGINE_API virtual void PostEditUndo() override;
