SplineComponent: Make Spline transient for 5.6

We just don't want this thing serializing at all. We leave it around to minimize the risk of breaking the component for 5.6 but it should remain non-functional and never save to disk

diff --git a/Engine/Source/Runtime/Engine/Classes/Components/SplineComponent.h b/Engine/Source/Runtime/Engine/Classes/Components/SplineComponent.h
--- a/Engine/Source/Runtime/Engine/Classes/Components/SplineComponent.h
+++ b/Engine/Source/Runtime/Engine/Classes/Components/SplineComponent.h
@@ -211,7 +211,7 @@ class USplineComponent : public UPrimitiveComponent
 protected:
 
 	// todo: replicate, but I believe we need to make the struct able to replicate in the first place.
-	UPROPERTY(EditAnywhere, Category=Points)
+	UPROPERTY(Transient, DuplicateTransient, EditAnywhere, Category=Points)
 	FSpline Spline;
 	
 public:
