Add QueueWorkFunction and SendQueuedWork to the TaskSyncManager which can be used to queue a batch of function pointers to execute as a simple tick task.
Change it so work registered with RegisterTickGroupWorkHandle will activate a task on the game thread as needed as well as at the end of the tick group.
Combined, this lets you efficiently schedule batch execution of game thread functions from any number of worker thread functions.
Add CreateManualTickTask and FWorkQueueTickFunction utilities which can be useful for other types of async work.
Add a.AnimNext.GameThreadEventBatchMode to enable batching. The default of 1 uses a per-module queue, while 2 uses this new global queue if available.

diff --git a/Engine/Plugins/Experimental/AnimNext/Source/AnimNext/Internal/Module/ModuleTickFunction.h b/Engine/Plugins/Experimental/AnimNext/Source/AnimNext/Internal/Module/ModuleTickFunction.h
--- a/Engine/Plugins/Experimental/AnimNext/Source/AnimNext/Internal/Module/ModuleTickFunction.h
+++ b/Engine/Plugins/Experimental/AnimNext/Source/AnimNext/Internal/Module/ModuleTickFunction.h
@@ -5,6 +5,7 @@
 #include "CoreMinimal.h"
 #include "Engine/EngineBaseTypes.h"
 #include "Containers/SpscQueue.h"
+#include "TaskSyncManager.h"
 
 struct FAnimNextModuleInstance;
 
@@ -49,6 +50,9 @@ struct FModuleEventTickFunction : public FTickFunction
 	void AddSubsequent(UObject* InObject, FTickFunction& InTickFunction);
 	void RemoveSubsequent(UObject* InObject, FTickFunction& InTickFunction);
 	void RemoveAllExternalSubsequents();
+	void InitializeBatchedWork(class UWorld* World);
+
+	static void RunTaskOnGameThread(TUniqueFunction<void(void)>&& InFunction);
 
 	FAnimNextModuleInstance* ModuleInstance = nullptr;
 
@@ -56,6 +60,12 @@ struct FModuleEventTickFunction : public FTickFunction
 	TSpscQueue<TUniqueFunction<void(const FModuleTaskContext&)>> PreExecuteTasks;
 	TSpscQueue<TUniqueFunction<void(const FModuleTaskContext&)>> PostExecuteTasks;
 
+	// Tasks queued to execute on the game thread after the current tick completes
+	TArray<TUniqueFunction<void()>> PostExecuteGameThreadTaskList;
+
+	// Handle to use for submitting batched work to the task sync manager
+	UE::Tick::FActiveSyncWorkHandle BatchedWorkHandle;
+
 	// Multicast delegate called before this event
 	FOnPreModuleEvent OnPreModuleEvent;
 
diff --git a/Engine/Plugins/Experimental/AnimNext/Source/AnimNext/Private/Module/AnimNextModuleInstance.cpp b/Engine/Plugins/Experimental/AnimNext/Source/AnimNext/Private/Module/AnimNextModuleInstance.cpp
--- a/Engine/Plugins/Experimental/AnimNext/Source/AnimNext/Private/Module/AnimNextModuleInstance.cpp
+++ b/Engine/Plugins/Experimental/AnimNext/Source/AnimNext/Private/Module/AnimNextModuleInstance.cpp
@@ -16,6 +16,7 @@
 #include "Module/ModuleGuard.h"
 #include "Module/RigUnit_AnimNextModuleEvents.h"
 #include "Module/ProxyVariablesContext.h"
+#include "Module/ModuleTickFunction.h"
 #include "UObject/UObjectIterator.h"
 #include "Variables/IAnimNextVariableProxyHost.h"
 #include "RewindDebugger/AnimNextTrace.h"
@@ -236,6 +237,12 @@ void FAnimNextModuleInstance::Initialize()
 				}
 			}
 
+			if (ModuleEvent.bUserEvent && !ModuleEvent.bIsGameThreadTask)
+			{
+				// Hook up with Task Sync Manager if possible
+				TickFunction.InitializeBatchedWork(World);
+			}
+
 			if (ModuleEvent.bUserEvent && !bFoundFirstUserEvent)
 			{
 				TickFunction.bFirstUserEvent = true;
@@ -639,14 +646,7 @@ void FAnimNextModuleInstance::ShowDebugDrawing(bool bInShowDebugDrawing)
 
 void FAnimNextModuleInstance::RunTaskOnGameThread(TUniqueFunction<void(void)>&& InFunction)
 {
-	if(IsInGameThread())
-	{
-		InFunction();
-	}
-	else
-	{
-		FFunctionGraphTask::CreateAndDispatchWhenReady(MoveTemp(InFunction), TStatId(), nullptr, ENamedThreads::GameThread);
-	}
+	UE::AnimNext::FModuleEventTickFunction::RunTaskOnGameThread(MoveTemp(InFunction));
 }
 
 UE::AnimNext::FModuleEventTickFunction* FAnimNextModuleInstance::FindTickFunctionByName(FName InEventName)
diff --git a/Engine/Plugins/Experimental/AnimNext/Source/AnimNext/Private/Module/ModuleTickFunction.cpp b/Engine/Plugins/Experimental/AnimNext/Source/AnimNext/Private/Module/ModuleTickFunction.cpp
--- a/Engine/Plugins/Experimental/AnimNext/Source/AnimNext/Private/Module/ModuleTickFunction.cpp
+++ b/Engine/Plugins/Experimental/AnimNext/Source/AnimNext/Private/Module/ModuleTickFunction.cpp
@@ -18,6 +18,21 @@
 
 namespace UE::AnimNext
 {
+	namespace Private
+	{
+		// When we run a module tick event, this is the currently executing tick function
+		static thread_local FModuleEventTickFunction* CurrentTickFunction = nullptr;
+
+		static constexpr int32 NoBatching = 0;
+		static constexpr int32 ModuleBatching = 1;
+		static constexpr int32 TaskSyncBatching = 2;
+
+		static int32 GameThreadEventBatchMode = ModuleBatching;
+		static FAutoConsoleVariableRef CVarGameThreadEventBatchMode(
+			TEXT("a.AnimNext.GameThreadEventBatchMode"),
+			GameThreadEventBatchMode,
+			TEXT("Batching mode for game thread events, 0 = no batching, 1 = simple event queue, 2 = task sync manager if available"));
+	}
 
 void FModuleEventTickFunction::Initialize(float InDeltaTime)
 {
@@ -82,7 +97,7 @@ void FModuleEventTickFunction::EndTick(float DeltaTime)
 		)
 		{
 			// Queue task to disable ourselves
-			FAnimNextModuleInstance::RunTaskOnGameThread([ModuleHandle = ModuleInstance->Handle, WeakObject = TWeakObjectPtr<UObject>(ModuleInstance->Object)]()
+			RunTaskOnGameThread([ModuleHandle = ModuleInstance->Handle, WeakObject = TWeakObjectPtr<UObject>(ModuleInstance->Object)]()
 			{
 				check(IsInGameThread());
 				if (UObject* Object = WeakObject.Get())
@@ -96,7 +111,7 @@ void FModuleEventTickFunction::EndTick(float DeltaTime)
 
 	if (!MainThreadActionEventList.IsEmpty())
 	{
-		FAnimNextModuleInstance::RunTaskOnGameThread([MainThreadActionEventList = MoveTemp(MainThreadActionEventList)]()
+		RunTaskOnGameThread([MainThreadActionEventList = MoveTemp(MainThreadActionEventList)]()
 			{
 				SCOPED_NAMED_EVENT(AnimNext_Module_EndTick_GameThread, FColor::Orange);
 				check(IsInGameThread());
@@ -116,7 +131,61 @@ void FModuleEventTickFunction::EndTick(float DeltaTime)
 
 void FModuleEventTickFunction::ExecuteTick(float DeltaTime, ELevelTick TickType, ENamedThreads::Type CurrentThread, const FGraphEventRef& MyCompletionGraphEvent)
 {
-	Run(DeltaTime);
+	check(PostExecuteGameThreadTaskList.IsEmpty());	// Should always be empty since we move it to the game thread task
+
+	if (BatchedWorkHandle.IsValid() && GetActualEndTickGroup() != TickGroup)
+	{
+		// This spans multiple tick groups so cannot use the tick group queue
+		BatchedWorkHandle.Reset();
+	}
+
+	{
+		check(Private::CurrentTickFunction == nullptr);
+		Private::CurrentTickFunction = this;
+
+		Run(DeltaTime);
+
+		Private::CurrentTickFunction = nullptr;
+	}
+
+	// If we have a task sync queue, flush it and extend this task
+	if (BatchedWorkHandle.IsValid())
+	{
+		FGraphEventRef ThisEvent = GetCompletionHandle();
+		BatchedWorkHandle.SendQueuedWork(&ThisEvent);
+	}
+
+	// Dispatch any post-execute game thread tasks we might have queued up
+	if (!PostExecuteGameThreadTaskList.IsEmpty())
+	{
+		auto PostExecuteGameThreadTask = [PostExecuteGameThreadTaskList = MoveTemp(PostExecuteGameThreadTaskList)]()
+			{
+				SCOPED_NAMED_EVENT(AnimNext_Module_PostEventTick, FColor::Orange);
+				check(IsInGameThread());
+
+				for (auto& Function : PostExecuteGameThreadTaskList)
+				{
+					Function();
+				}
+			};
+
+		if (IsInGameThread())
+		{
+			// We are already on the main thread, we've batched up all our work, execute it in bulk now
+			PostExecuteGameThreadTask();
+		}
+		else
+		{
+			// Dispatch our main thread work as a single task
+			FGraphEventRef PostExecuteCompletionEvent = FFunctionGraphTask::CreateAndDispatchWhenReady(MoveTemp(PostExecuteGameThreadTask), TStatId(), nullptr, ENamedThreads::GameThread);
+
+			if (IsCompletionHandleValid())
+			{
+				// Update our completion event to make sure anyone that depends on us also depends on any game thread tasks we generate
+				GetCompletionHandle()->DontCompleteUntil(PostExecuteCompletionEvent);
+			}
+		}
+	}
 }
 
 void FModuleEventTickFunction::Run(float InDeltaTime)
@@ -220,6 +289,16 @@ void FModuleEventTickFunction::RemoveAllExternalSubsequents()
 	ExternalSubsequents.Reset();
 }
 
+void FModuleEventTickFunction::InitializeBatchedWork(UWorld* World)
+{
+	UE::Tick::FTaskSyncManager* SyncManager = UE::Tick::FTaskSyncManager::Get();
+
+	if (Private::GameThreadEventBatchMode == Private::TaskSyncBatching && SyncManager)
+	{
+		SyncManager->RegisterTickGroupWorkHandle(World, TickGroup, BatchedWorkHandle);
+	}
+}
+
 FString FModuleEventTickFunction::DiagnosticMessage()
 {
 	TStringBuilder<256> Builder;
@@ -228,4 +307,38 @@ FString FModuleEventTickFunction::DiagnosticMessage()
 	return Builder.ToString();
 }
 
+void FModuleEventTickFunction::RunTaskOnGameThread(TUniqueFunction<void(void)>&& InFunction)
+{
+	if (!InFunction)
+	{
+		return;
+	}
+
+	if (FModuleEventTickFunction* TickFunction = Private::CurrentTickFunction)
+	{
+		if (TickFunction->BatchedWorkHandle.IsValid())
+		{
+			// Use task sync manager if it was set up
+			TickFunction->BatchedWorkHandle.QueueWorkFunction(MoveTemp(InFunction));
+			return;
+		}
+		else if (Private::GameThreadEventBatchMode != Private::NoBatching)
+		{
+			// Use normal batching if enabled
+			TickFunction->PostExecuteGameThreadTaskList.Add(MoveTemp(InFunction));
+			return;
+		}
+	}
+
+	// Fallback to old behavior if no tick function or mode is 0
+	if (IsInGameThread())
+	{
+		InFunction();
+	}
+	else
+	{
+		FFunctionGraphTask::CreateAndDispatchWhenReady(MoveTemp(InFunction), TStatId(), nullptr, ENamedThreads::GameThread);
+	}
+}
+
 }
diff --git a/Engine/Source/Runtime/Engine/Private/TaskSyncManager.cpp b/Engine/Source/Runtime/Engine/Private/TaskSyncManager.cpp
--- a/Engine/Source/Runtime/Engine/Private/TaskSyncManager.cpp
+++ b/Engine/Source/Runtime/Engine/Private/TaskSyncManager.cpp
@@ -112,6 +112,11 @@ bool FActiveSyncWorkHandle::HasRequestedWork() const
 	return bWorkRequested;
 }
 
+bool FActiveSyncWorkHandle::HasWorkQueue() const
+{
+	return bHasWorkQueue;
+}
+
 bool FActiveSyncWorkHandle::AbandonWork()
 {
 	if (ensure(SyncPoint.IsValid()))
@@ -148,8 +153,131 @@ bool FActiveSyncWorkHandle::RequestWork(FTickFunction* WorkFunction, ESyncWorkRe
 	return false;
 }
 
+
+bool FActiveSyncWorkHandle::QueueWorkFunction(FWorkFunction&& WorkFunction)
+{
+	if (ensure(SyncPoint.IsValid()))
+	{
+		return SyncPoint->QueueWorkFunction(*this, MoveTemp(WorkFunction));
+	}
+	return false;
+}
+
+bool FActiveSyncWorkHandle::SendQueuedWork(FGraphEventRef* TaskToExtend)
+{
+	if (ensure(SyncPoint.IsValid()))
+	{
+		return SyncPoint->SendQueuedWork(*this, TaskToExtend);
+	}
+	return false;
+}
+
+FWorkQueueTickFunction::FWorkQueueTickFunction()
+{
+	bCanEverTick = true;
+	bDispatchManually = true;
+	bStartWithTickEnabled = true;
+
+	bSetAsOpen = true;
+	bClearAfterExecute = true;
+	bOpenAfterExecute = true;
+	WorkExecutionIndex = INDEX_NONE;
+
+	static FName QueueWorkTickFunction("QueueWorkTickFunction");
+	DebugName = QueueWorkTickFunction;
+}
+
+FWorkQueueTickFunction::~FWorkQueueTickFunction()
+{
+	check(!IsWorkExecuting());
+}
+
+void FWorkQueueTickFunction::SetQueueOpen(bool bShouldBeOpen)
+{
+	bSetAsOpen = bShouldBeOpen;
+}
+
+bool FWorkQueueTickFunction::AddWork(FWorkFunction&& Work)
+{
+	check(!IsWorkExecuting());
+	
+	if (IsQueueOpen())
+	{
+		QueuedWork.Add(MoveTemp(Work));
+		return true;
+	}
+	return false;
+}
+
+void FWorkQueueTickFunction::ClearWork()
+{
+	check(!IsWorkExecuting());
+	QueuedWork.Reset();
+}
+
+void FWorkQueueTickFunction::ExecuteWork()
+{
+	TRACE_CPUPROFILER_EVENT_SCOPE(WorkQueueTickFunction_ExecuteWork);
+	check(!IsWorkExecuting());
+
+	for (WorkExecutionIndex = 0; WorkExecutionIndex < QueuedWork.Num(); WorkExecutionIndex++)
+	{
+		QueuedWork[WorkExecutionIndex]();
+	}
+
+	WorkExecutionIndex = INDEX_NONE;
+
+	if (bClearAfterExecute)
+	{
+		ClearWork();
+	}
+
+	if (bOpenAfterExecute)
+	{
+		SetQueueOpen(true);
+	}
+}
+
+void FWorkQueueTickFunction::SetDebugName(FName InDebugName, const FString& InDetailString)
+{
+	DebugName = InDebugName;
+	DebugDetailString = InDetailString;
+}
+
+void FWorkQueueTickFunction::ExecuteTick(float DeltaTime, ELevelTick TickType, ENamedThreads::Type CurrentThread, const FGraphEventRef& MyCompletionGraphEvent)
+{
+	ExecuteWork();
+}
+
+FString FWorkQueueTickFunction::DiagnosticMessage()
+{
+	return FString::Printf(TEXT("%s:%s"), *DebugName.ToString(), *DebugDetailString);
+}
+
+FName FWorkQueueTickFunction::DiagnosticContext(bool bDetailed)
+{
+	if (bDetailed && !DebugDetailString.IsEmpty())
+	{
+		return FName(*DiagnosticMessage());
+	}
+
+	return DebugName;
+}
+
+void FWorkQueueTickFunction::SetClearAfterExecute(bool bShouldClear)
+{
+	check(!IsWorkExecuting());
+	bClearAfterExecute = bShouldClear;
+}
+
+void FWorkQueueTickFunction::SetOpenAfterExecute(bool bShouldClear)
+{
+	check(!IsWorkExecuting());
+	bClearAfterExecute = bShouldClear;
+}
+
+
 /** Task to execute a tick function manually outside of it's normal tick group */
-// TODO move into TickTaskManager?
 class FManualTickFunctionTask
 {
 	FTickFunction* Target;
@@ -219,7 +347,7 @@ void FActiveSyncPoint::ResetWorkForFrame()
 			CurrentWork.bWorkReserved = true;
 		}
 
-		if (!CurrentWork.bWorkRequested && CurrentWork.TickFunction)
+		if (!CurrentWork.bWorkRequested && !CurrentWork.bIsWorkQueueFunction && CurrentWork.TickFunction)
 		{
 			// This could delete a wrapper tick function
 			CurrentWork.SetTickFunction(nullptr, false);
@@ -252,7 +380,7 @@ bool FActiveSyncPoint::HandleFirstTickGroup()
 	return false;
 }
 
-bool FActiveSyncPoint::IsReadyToProcessWork()
+bool FActiveSyncPoint::IsReadyToProcessWork() const
 {
 	if (FrameStatus < ESyncPointStatus::DispatchWaitingForWork || FrameStatus == ESyncPointStatus::ExecutionComplete)
 	{
@@ -260,11 +388,23 @@ bool FActiveSyncPoint::IsReadyToProcessWork()
 		return false;
 	}
 
-	// See if we should dispatch a process task
+	if (SyncPointDescription.ActivationRules == ESyncPointActivationRules::ActivateForAnyWork)
+	{
+		// Dispatch if we have any valid requests
+		for (const FActiveSyncWork& CurrentWork : ActiveWork)
+		{
+			if (CurrentWork.bWorkRequested)
+			{
+				return true;
+			}
+		}
+		return false;
+	}
+
 	bool bHasRequest = false;
 	bool bHasReservation = false;
 
-	for (FActiveSyncWork& CurrentWork : ActiveWork)
+	for (const FActiveSyncWork& CurrentWork : ActiveWork)
 	{
 		if (CurrentWork.bWorkRequested)
 		{
@@ -276,10 +416,20 @@ bool FActiveSyncPoint::IsReadyToProcessWork()
 		}
 	}
 
-	if (bHasRequest)
+	if (bHasRequest && !bHasReservation)
 	{
 		// Needs to be dispatched if it has no reservations or will start for any work
-		return !bHasReservation || (SyncPointDescription.ActivationRules == ESyncPointActivationRules::ActivateForAnyWork);
+		return true;
+	}
+	return false;
+}
+
+bool FActiveSyncPoint::IsTooLateToAddWork(bool bWorkReserved) const
+{
+	if (FrameStatus == ESyncPointStatus::ExecutionComplete || (FrameStatus >= ESyncPointStatus::Executing && !bWorkReserved))
+	{
+		// Never too late for tick group work
+		return !IsTickGroupWork();
 	}
 	return false;
 }
@@ -342,15 +492,21 @@ bool FActiveSyncPoint::GetWorkToExecute(float DeltaTime, ELevelTick TickType, EN
 
 		// Create our redo event but don't dispatch it yet
 		// This might get dispatched immediately after leaving the lock here
-		FBaseGraphTask* Task = TGraphTask<FManualTickFunctionTask>::CreateTask(nullptr, CurrentThread).ConstructAndHold(static_cast<FTickFunction*>(this), CurrentThread, DeltaTime, TickType);
-		ReactivationEvent = Task->GetCompletionEvent();
-		MyCompletionGraphEvent->DontCompleteUntil(ReactivationEvent);
+		ReactivationEvent = FTaskSyncManager::Get()->CreateManualTickTask(WorldContextId, this, DeltaTime, TickType);
 
 		return false;
 	}
 	else
 	{
-		FrameStatus = ESyncPointStatus::ExecutionComplete;
+		if (IsTickGroupWork())
+		{
+			// Tick group work can restart at any time
+			FrameStatus = ESyncPointStatus::WaitingForManualTask;
+		}
+		else
+		{
+			FrameStatus = ESyncPointStatus::ExecutionComplete;
+		}
 		return false;
 	}
 }
@@ -399,7 +555,7 @@ void FActiveSyncPoint::ExecuteTick(float DeltaTime, ELevelTick TickType, ENamedT
 	}
 
 	// TODO add pre/post batch function calls
-
+	TRACE_CPUPROFILER_EVENT_SCOPE(ActiveSyncPoint_ExecuteTick);
 	while (GetWorkToExecute(DeltaTime, TickType, CurrentThread, MyCompletionGraphEvent))
 	{
 		for (int32 i = 0; i < ExecutingWork.Num(); i++)
@@ -453,6 +609,7 @@ FActiveSyncWorkHandle FActiveSyncPoint::RegisterWorkHandle()
 	ReturnHandle.WorkIndex = WorkIndex;
 	ReturnHandle.bWorkReserved = false;
 	ReturnHandle.bWorkRequested = false;
+	ReturnHandle.bHasWorkQueue = false;
 
 	return MoveTemp(ReturnHandle);
 }
@@ -468,6 +625,11 @@ bool FActiveSyncPoint::ReserveFutureWork(FActiveSyncWorkHandle& Handle, ESyncWor
 
 		if (ensure(CurrentWork.bHasActiveHandle))
 		{
+			if (CurrentWork.bIsWorkQueueFunction)
+			{
+				return false;
+			}
+
 			CurrentWork.bReserveEveryFrame = (Repeat == ESyncWorkRepetition::EveryFrame);
 
 			if (Repeat == ESyncWorkRepetition::Never)
@@ -505,6 +667,11 @@ bool FActiveSyncPoint::RequestWork(FActiveSyncWorkHandle& Handle, FTickFunction*
 
 		if (ensure(CurrentWork.bHasActiveHandle))
 		{
+			if (CurrentWork.bIsWorkQueueFunction)
+			{
+				return false;
+			}
+
 			if (Repeat == ESyncWorkRepetition::Never)
 			{
 				Handle.bWorkRequested = false;
@@ -516,7 +683,7 @@ bool FActiveSyncPoint::RequestWork(FActiveSyncWorkHandle& Handle, FTickFunction*
 				Handle.bWorkRequested = true;
 
 				// If the task is completely done or in the middle of processing and we didn't reserve work
-				if (FrameStatus == ESyncPointStatus::ExecutionComplete || (FrameStatus >= ESyncPointStatus::Executing && !CurrentWork.bWorkReserved))
+				if (IsTooLateToAddWork(CurrentWork.bWorkReserved))
 				{
 					FTaskSyncManager* Manager = FTaskSyncManager::Get();
 					
@@ -545,6 +712,101 @@ bool FActiveSyncPoint::RequestWork(FActiveSyncWorkHandle& Handle, FTickFunction*
 	return false;
 }
 
+bool FActiveSyncPoint::QueueWorkFunction(FActiveSyncWorkHandle& Handle, FWorkFunction&& WorkFunction)
+{
+	UE::TScopeLock WorkScope(WorkLock);
+	int32 WorkIndex = Handle.GetIndex();
+
+	if (ensure(Handle.IsValid() && ActiveWork.IsValidIndex(WorkIndex)))
+	{
+		FActiveSyncWork& CurrentWork = ActiveWork[WorkIndex];
+
+		if (ensure(CurrentWork.bHasActiveHandle))
+		{
+			if (CurrentWork.bWorkRequested || CurrentWork.bRequestEveryFrame || CurrentWork.bReserveEveryFrame)
+			{
+				// Can't convert a work request into queued work
+				return false;
+			}
+
+			// Allocate a work queue if needed
+			if (!CurrentWork.bIsWorkQueueFunction)
+			{
+				CurrentWork.SetTickFunction(new FWorkQueueTickFunction(), true);
+				CurrentWork.bIsWorkQueueFunction = true;
+			}
+
+			Handle.bWorkReserved = true;
+			Handle.bHasWorkQueue = true;
+
+			CurrentWork.bWorkReserved = true;
+			check(CurrentWork.TickFunction && CurrentWork.bIsWorkQueueFunction);
+
+			FWorkQueueTickFunction* QueueFunction = (FWorkQueueTickFunction*)CurrentWork.TickFunction;
+
+			return QueueFunction->AddWork(MoveTemp(WorkFunction));
+		}
+	}
+	return false;
+}
+
+bool FActiveSyncPoint::SendQueuedWork(FActiveSyncWorkHandle& Handle, FGraphEventRef* TaskToExtend)
+{
+	UE::TScopeLock WorkScope(WorkLock);
+	int32 WorkIndex = Handle.GetIndex();
+
+	if (ensure(Handle.IsValid() && ActiveWork.IsValidIndex(WorkIndex)))
+	{
+		FActiveSyncWork& CurrentWork = ActiveWork[WorkIndex];
+
+		if (ensure(CurrentWork.bHasActiveHandle))
+		{
+			if (!CurrentWork.bIsWorkQueueFunction || CurrentWork.bWorkRequested)
+			{
+				// Can't convert a work request into queued work
+				return false;
+			}
+			
+			check(CurrentWork.TickFunction && CurrentWork.bIsWorkQueueFunction);
+			FWorkQueueTickFunction* QueueFunction = (FWorkQueueTickFunction*)CurrentWork.TickFunction;
+
+			if (QueueFunction->HasQueuedWork())
+			{
+				// Actually dispatch
+				Handle.bWorkRequested = true;
+				
+				// If the task is completely done or in the middle of processing and we didn't reserve work
+				if (IsTooLateToAddWork(CurrentWork.bWorkReserved))
+				{
+					FTaskSyncManager* Manager = FTaskSyncManager::Get();
+
+					// Can pass through the function as this handle will keep it alive or cancel it
+					CurrentWork.bLateWorkRequested = Manager->HandleLateWorkRequest(this, WorkIndex, QueueFunction);
+				}
+
+				CurrentWork.bWorkReserved = false;
+				CurrentWork.bWorkRequested = !CurrentWork.bLateWorkRequested;
+
+				if (IsReadyToProcessWork())
+				{
+					DispatchWorkTask(TaskToExtend);
+				}
+			}
+			else
+			{
+				CurrentWork.bWorkReserved = false;
+				CurrentWork.bWorkRequested = false;
+
+				Handle.bWorkRequested = false;
+				Handle.bWorkReserved = false;
+			}
+
+			return true;
+		}
+	}
+	return false;
+}
+
 bool FActiveSyncPoint::AbandonWork(FActiveSyncWorkHandle& Handle)
 {
 	UE::TScopeLock WorkScope(WorkLock);
@@ -627,8 +889,9 @@ uint32 FActiveSyncPoint::AllocateActiveWork()
 	return ActiveWork.AddDefaulted();
 }
 
-void FActiveSyncPoint::DispatchWorkTask()
+void FActiveSyncPoint::DispatchWorkTask(FGraphEventRef* TaskToExtend)
 {
+	FGraphEventRef ThisTask;
 	if (FrameStatus < ESyncPointStatus::Dispatched)
 	{
 		ensure(DispatchManually());
@@ -637,16 +900,40 @@ void FActiveSyncPoint::DispatchWorkTask()
 	{
 		check(ReactivationEvent.IsValid());
 		// Dispatch our reactivation event to let the task activate
-		ReactivationEvent->Unlock();
+		ThisTask = MoveTemp(ReactivationEvent);
 		ReactivationEvent.SafeRelease();
-		FrameStatus = ESyncPointStatus::Dispatched; // TODO do we need a different state for dispatched a second time?
+
+		FrameStatus = ESyncPointStatus::Dispatched;
+		ThisTask->Unlock();
+	}
+	else if (FrameStatus == ESyncPointStatus::WaitingForManualTask)
+	{
+		// Spawn a manual task right here
+		check(!ReactivationEvent.IsValid());
+
+		ThisTask = FTaskSyncManager::Get()->CreateManualTickTask(WorldContextId, this);
+		
+		FrameStatus = ESyncPointStatus::Dispatched;
+		ThisTask->Unlock();
 	}
 	else
 	{
-		// If it gets here the work will get handled by the already dispatched task
+		// If it gets here the work will get handled by the already dispatched task, but we can extend the passed in task
 		check(FrameStatus == ESyncPointStatus::Dispatched || FrameStatus == ESyncPointStatus::Executing);
 	}
-	
+
+	if (TaskToExtend)
+	{
+		if (!ThisTask && IsCompletionHandleValid())
+		{
+			ThisTask = GetCompletionHandle();
+		}
+
+		if (ThisTask)
+		{
+			(*TaskToExtend)->DontCompleteUntil(MoveTemp(ThisTask));
+		}
+	}
 }
 
 FTaskSyncManager* FTaskSyncManager::Get()
@@ -912,7 +1199,16 @@ void FTaskSyncManager::InitializeBatchForFrame(FBatchData& BatchData, ULevel* Pe
 		}
 	}
 
-	// TODO execute one of the TickGroupWork functions? But should it be at the top or bottom
+	for (int32 Group = 0; Group < BatchData.TickGroupWork.Num(); Group++)
+	{
+		// Reset the tick group now
+		FActiveSyncPoint* TickGroupWork = BatchData.TickGroupWork[Group].Get();
+		if (TickGroupWork)
+		{
+			TickGroupWork->FrameStatus = ESyncPointStatus::TaskNotCreated;
+			TickGroupWork->ResetWorkForFrame();
+		}
+	}
 }
 
 FActiveSyncPoint* FTaskSyncManager::GetOrCreateSyncPoint(FBatchData& BatchData, FRegisteredSyncPointData& SyncData)
@@ -984,9 +1280,9 @@ FActiveSyncPoint* FTaskSyncManager::GetOrCreateSyncPoint(FBatchData& BatchData,
 FActiveSyncPoint* FTaskSyncManager::GetOrCreateTickGroupWork(FBatchData& BatchData, ETickingGroup TickGroup)
 {
 	check(TickGroup >= 0 && TickGroup < ETickingGroup::TG_MAX);
-	if (TickGroup > BatchData.TickGroupWork.Num())
+	if (TickGroup >= BatchData.TickGroupWork.Num())
 	{
-		BatchData.TickGroupWork.SetNum(TickGroup);
+		BatchData.TickGroupWork.SetNum(TickGroup + 1);
 	}
 
 	FActiveSyncPoint* FoundPoint = BatchData.TickGroupWork[TickGroup].Get();
@@ -1000,7 +1296,18 @@ FActiveSyncPoint* FTaskSyncManager::GetOrCreateTickGroupWork(FBatchData& BatchDa
 	NewTickFunction->EndTickGroup = TickGroup;
 	NewTickFunction->WorldContextId = BatchData.WorldContext;
 	NewTickFunction->SyncPointDescription.EventType = ESyncPointEventType::GameThreadTask;
-	NewTickFunction->SyncPointDescription.ActivationRules = ESyncPointActivationRules::AlwaysActivate;
+	NewTickFunction->SyncPointDescription.ActivationRules = ESyncPointActivationRules::ActivateForAnyWork;
+
+	if (IsTicking() && TickGroup == CurrentTickGroup)
+	{
+		// Can be used immediately
+		NewTickFunction->FrameStatus = ESyncPointStatus::WaitingForManualTask;
+	}
+	else
+	{
+		NewTickFunction->FrameStatus = ESyncPointStatus::TaskNotCreated;
+	}
+	
 
 	check(NewTickFunction->IsTickGroupWork());
 	
@@ -1030,6 +1337,16 @@ void FTaskSyncManager::StartTickGroup(const UWorld* InWorld, ETickingGroup Group
 					TicksToManualDispatch.Add(&SyncPoint);
 				}
 			}
+
+			if (BatchData.TickGroupWork.IsValidIndex(Group))
+			{
+				FActiveSyncPoint* TickGroupWork = BatchData.TickGroupWork[Group].Get();
+				if (TickGroupWork)
+				{
+					// Mark as ready for work
+					TickGroupWork->FrameStatus = ESyncPointStatus::WaitingForManualTask;
+				}
+			}
 		}
 	}
 }
@@ -1046,8 +1363,9 @@ void FTaskSyncManager::EndTickGroup(const UWorld* InWorld, ETickingGroup Group)
 			FActiveSyncPoint* TickGroupWork = BatchData.TickGroupWork[Group].Get();
 			if (TickGroupWork)
 			{
-				// TODO We also may want to run this earlier in the tick group, such as after the first stalled request
+				// Clear out the work queue if anything is left over
 				TickGroupWork->ExecuteFromGameThread(CurrentDeltaTime, CurrentTickType);
+				TickGroupWork->FrameStatus = ESyncPointStatus::ExecutionComplete;
 			}
 		}
 	}
@@ -1433,4 +1751,22 @@ void FTaskSyncManager::CancelTemporaryWorkRequest(FActiveSyncPoint* RequestedSyn
 			}
 		}
 	}
+}
+
+FGraphEventRef FTaskSyncManager::CreateManualTickTask(FWorldContextId WorldContext, FTickFunction* TickFunction, float DeltaTime, ELevelTick TickType)
+{
+	check(TickFunction);
+
+	ENamedThreads::Type TargetThread = ENamedThreads::SetTaskPriority(TickFunction->bRunOnAnyThread ? ENamedThreads::AnyThread : ENamedThreads::GameThread,
+		TickFunction->bHighPriority ? ENamedThreads::HighTaskPriority : ENamedThreads::NormalTaskPriority);
+
+	// Get the world's delta time if needed
+	if (DeltaTime < 0.0f && IsTicking())
+	{
+		DeltaTime = CurrentDeltaTime;
+		TickType = CurrentTickType;
+	}
+
+	FBaseGraphTask* Task = TGraphTask<FManualTickFunctionTask>::CreateTask(nullptr, ENamedThreads::AnyThread).ConstructAndHold(TickFunction, TargetThread, DeltaTime, TickType);
+	return Task->GetCompletionEvent();
 }
\ No newline at end of file
diff --git a/Engine/Source/Runtime/Engine/Public/TaskSyncManager.h b/Engine/Source/Runtime/Engine/Public/TaskSyncManager.h
--- a/Engine/Source/Runtime/Engine/Public/TaskSyncManager.h
+++ b/Engine/Source/Runtime/Engine/Public/TaskSyncManager.h
@@ -121,6 +121,9 @@ public:
 namespace UE::Tick
 {
 
+/** Simple lambda type */
+typedef TUniqueFunction<void()> FWorkFunction;
+
 /** Enum describing the success/failure of a task sync operation */
 enum class ESyncOperationResult : uint8
 {
@@ -180,6 +183,8 @@ enum class ESyncPointStatus : uint8
 	Executing,
 	/** A follow up task has been dispatched to waiting for more work before executing again */
 	WaitingForMoreWork,
+	/** Not executing, but can start a manual task if needed */
+	WaitingForManualTask,
 	/** Completely done executing for the frame */
 	ExecutionComplete,
 };
@@ -427,6 +432,9 @@ struct FActiveSyncWorkHandle final
 	/** True if this handle has been used to request work. This will still be true until work is abandoned */
 	ENGINE_API bool HasRequestedWork() const;
 
+	/** True if this handle is associated with a work queue that can be added to */
+	ENGINE_API bool HasWorkQueue() const;
+
 	/** 
 	 * Request to reserve work using this handle, which must be filled to complete the frame's tasks. 
 	 * If EveryFrame is passed in, it will reserve work at the start of every frame that must be requested or abandoned.
@@ -439,6 +447,12 @@ struct FActiveSyncWorkHandle final
 	 */
 	ENGINE_API bool RequestWork(FTickFunction* FunctionToExecute, ESyncWorkRepetition Repeat);
 
+	/** Adds an arbitrary work function to the queue of work to execute later, this will fail if work has already been requested or reserved */
+	ENGINE_API bool QueueWorkFunction(FWorkFunction&& WorkFunction);
+
+	/** Sends any previously queued work to be executed as part of the sync point. Will call DontCompleteUntil on passed in TaskToExtend if valid */
+	ENGINE_API bool SendQueuedWork(FGraphEventRef* TaskToExtend = nullptr);
+
 	/** 
 	 * Abandon any requested or reserved work and return to how it was at initial registration.
 	 * This clears HasReserved/HasRequested, but does not invalidate the handle.
@@ -465,6 +479,9 @@ private:
 	/** True if this may represent active requested work */
 	uint32 bWorkRequested : 1 = 0;
 
+	/** True if this represents a function work queue */
+	uint32 bHasWorkQueue : 1 = 0;
+
 	/** Get index as signed int */
 	inline int32 GetIndex() const
 	{
@@ -478,12 +495,81 @@ private:
 		WorkIndex = 0;
 		bWorkReserved = false;
 		bWorkRequested = false;
+		bHasWorkQueue = false;
 	}
 
 	friend struct FActiveSyncPoint;
 	friend class FTaskSyncManager;
 };
 
+/**
+ * A Tick function that can be used to execute a set of arbitrary work. This can be used as a normal tick function or passed to the work request functions.
+ * This is NOT threadsafe on it's own, it needs to be protected by something like the locks in SendQueuedWork to use from multiple threads.
+ */
+struct FWorkQueueTickFunction : public FTickFunction
+{
+	ENGINE_API FWorkQueueTickFunction();
+	ENGINE_API virtual ~FWorkQueueTickFunction();
+
+	/** Return true if this is currently executing, if so no work can be added to the queue */
+	inline bool IsWorkExecuting() const
+	{
+		return WorkExecutionIndex != INDEX_NONE;
+	}
+
+	/** Return true if any work has been queued */
+	inline bool HasQueuedWork() const
+	{
+		return IsWorkExecuting() || QueuedWork.Num() > 0;
+	}
+
+	/** Return true if work can be added to the queue */
+	inline bool IsQueueOpen() const
+	{
+		return !IsWorkExecuting() && bSetAsOpen;
+	}
+
+	/** Explicitly close or open the queue, can be used to finalize before execution */
+	ENGINE_API void SetQueueOpen(bool bShouldBeOpen);
+
+	/** Try to add a new function to the queue */
+	ENGINE_API bool AddWork(FWorkFunction&& Work);
+
+	/** Clear the queue, not allowed if execution has started */
+	ENGINE_API void ClearWork();
+
+	/** Actually execute the work, also called by execute tick if this scheduled as a normal tick */
+	ENGINE_API void ExecuteWork();
+
+	/** Sets the debug name used by DiagnosticContext */
+	ENGINE_API void SetDebugName(FName InDebugName, const FString& InDetailString = FString());
+
+	/** Sets if it should clear the queue after every execution (default behavior) */
+	ENGINE_API void SetClearAfterExecute(bool bShouldClear);
+
+	/** Sets if it should reopen a closed queue after every execution (default behavior) */
+	ENGINE_API void SetOpenAfterExecute(bool bShouldOpen);
+
+protected:
+
+	// FTickFunction interface
+	ENGINE_API void ExecuteTick(float DeltaTime, ELevelTick TickType, ENamedThreads::Type CurrentThread, const FGraphEventRef& MyCompletionGraphEvent) override;
+	ENGINE_API FString DiagnosticMessage() override;
+	ENGINE_API FName DiagnosticContext(bool bDetailed) override;
+
+	bool bSetAsOpen;
+	bool bClearAfterExecute;
+	bool bOpenAfterExecute;
+	int32 WorkExecutionIndex;
+
+	FName DebugName;
+	FString DebugDetailString;
+
+	// Simple queue of work to execute, could be changed to a different type in the future
+	TArray<FWorkFunction> QueuedWork;
+};
+
+
 /** Struct defining an active sync point inside a specific world context, implemented as a tick function that may not be registered */
 struct FActiveSyncPoint final : public TSharedFromThis<FActiveSyncPoint, ESPMode::ThreadSafe>, private FTickFunction
 {
@@ -513,6 +599,12 @@ struct FActiveSyncPoint final : public TSharedFromThis<FActiveSyncPoint, ESPMode
 	/** Creates a new item of work bound to a specific function */
 	ENGINE_API bool RequestWork(FActiveSyncWorkHandle& Handle, FTickFunction* WorkFunction, ESyncWorkRepetition Repeat);
 
+	/** Adds an arbitrary work function to the queue of work to execute later, this will fail if work has already been requested or reserved */
+	ENGINE_API bool QueueWorkFunction(FActiveSyncWorkHandle& Handle, FWorkFunction&& WorkFunction);
+
+	/** Sends any previously queued work to be executed as part of the sync point. Will call DontCompleteUntil on passed in TaskToExtend if valid */
+	ENGINE_API bool SendQueuedWork(FActiveSyncWorkHandle& Handle, FGraphEventRef* TaskToExtend = nullptr);
+
 	/** Abandons requested or reserved work, which could trigger other work to start */
 	ENGINE_API bool AbandonWork(FActiveSyncWorkHandle& Handle);
 
@@ -555,6 +647,9 @@ private:
 				/** If true, a late work request was created for this handle */
 				uint32 bLateWorkRequested : 1;
 
+				/** If true, this points to a FWorkQueueTickFunction that can be added to and sent */
+				uint32 bIsWorkQueueFunction : 1;
+
 				/** If true, the tick function allocated by this work and should be freed when reset */
 				uint32 bAllocatedTickFunction : 1;
 			};
@@ -644,7 +739,10 @@ private:
 	bool HandleFirstTickGroup();
 
 	/** Returns true if this sync point is ready to trigger due to work */
-	bool IsReadyToProcessWork();
+	bool IsReadyToProcessWork() const;
+
+	/** Returns true if we're too late in execution to schedule work directly */
+	bool IsTooLateToAddWork(bool bWorkReserved) const;
 
 	/** Handle copying of work into / out of ExecutingWork array, returns true if there is anything to do */
 	bool GetWorkToExecute(float DeltaTime, ELevelTick TickType, ENamedThreads::Type CurrentThread, const FGraphEventRef& MyCompletionGraphEvent);
@@ -656,7 +754,7 @@ private:
 	uint32 AllocateActiveWork();
 
 	/** Dispatches the task to handle work, this should be called outside the work lock */
-	void DispatchWorkTask();
+	void DispatchWorkTask(FGraphEventRef* TaskToExtend = nullptr);
 
 	/** Registered sync point that created this */
 	FSyncPointDescription SyncPointDescription;
@@ -743,6 +841,14 @@ public:
 	 */
 	ENGINE_API FTaskSyncResult RegisterTickGroupWorkHandle(FWorldContextId WorldContext, ETickingGroup TickGroup, FActiveSyncWorkHandle& OutWorkHandle);
 
+	/**
+	 * Creates (but does not dispatch) a simple task that will execute a tick function. 
+	 * The caller is responsible for dispatching the task and must keep the tick function valid until the task completes.
+	 * This is useful for cases where you want to spawn a tick halfway through a tick group
+	 * If default delta time is passed, it will try to use the currently ticking world's info
+	 */
+	ENGINE_API FGraphEventRef CreateManualTickTask(FWorldContextId WorldContext, FTickFunction* TickFunction, float DeltaTime = -1.0f, ELevelTick TickType = LEVELTICK_All);
+
 
 	// Game thread-only functions, in general these will not affect the current frame
 
