[ISKMC] Primitive bounds in world space rather than local space.

#rnx

diff --git a/Engine/Plugins/PCG/Source/PCG/Private/Components/PCGProceduralISMComponent.cpp b/Engine/Plugins/PCG/Source/PCG/Private/Components/PCGProceduralISMComponent.cpp
--- a/Engine/Plugins/PCG/Source/PCG/Private/Components/PCGProceduralISMComponent.cpp
+++ b/Engine/Plugins/PCG/Source/PCG/Private/Components/PCGProceduralISMComponent.cpp
@@ -414,7 +414,7 @@ FMatrix UPCGProceduralISMComponent::GetRenderMatrix() const
 
 FBoxSphereBounds UPCGProceduralISMComponent::CalcBounds(const FTransform& BoundTransform) const
 {
-	return WorldBounds.InverseTransformBy(GetComponentTransform()).TransformBy(BoundTransform);
+	return WorldBounds.InverseTransformBy(GetComponentTransform().Inverse() * BoundTransform);
 }
 
 #if WITH_EDITOR
diff --git a/Engine/Plugins/PCG/Source/PCG/Private/Compute/DataInterfaces/Elements/PCGSkinnedMeshSpawnerDataInterface.cpp b/Engine/Plugins/PCG/Source/PCG/Private/Compute/DataInterfaces/Elements/PCGSkinnedMeshSpawnerDataInterface.cpp
--- a/Engine/Plugins/PCG/Source/PCG/Private/Compute/DataInterfaces/Elements/PCGSkinnedMeshSpawnerDataInterface.cpp
+++ b/Engine/Plugins/PCG/Source/PCG/Private/Compute/DataInterfaces/Elements/PCGSkinnedMeshSpawnerDataInterface.cpp
@@ -483,8 +483,6 @@ void UPCGSkinnedMeshSpawnerDataProvider::CreatePrimitiveDescriptors(FPCGContext*
 	}
 
 	
-	const FBox LocalBounds = SourceComponent->GetGridBounds().InverseTransformBy(SourceComponent->GetOwner()->GetActorTransform());
-
 	if (Settings->InstanceDataPackerParameters && Settings->InstanceDataPackerParameters->GetAttributeNames(/*OutNames=*/nullptr))
 	{
 		TArray<FName> AttributeNames;
@@ -587,7 +585,7 @@ void UPCGSkinnedMeshSpawnerDataProvider::CreatePrimitiveDescriptors(FPCGContext*
 			
 			Descriptor = Settings->MeshSelectorParameters->TemplateDescriptor;
 			Descriptor.bIsInstanceDataGPUOnly = true;
-			Descriptor.PrimitiveLocalBoundsOverride = LocalBounds;
+			Descriptor.PrimitiveBoundsOverride = SourceComponent->GetGridBounds();
 			Descriptor.NumInstancesGPUOnly = StringKeyAndInstanceCount.Value;
 			Descriptor.NumCustomDataFloatsGPUOnly = CustomFloatCount;
 			Descriptor.SkinnedAsset = FSoftObjectPath(MeshPathString);
diff --git a/Engine/Source/Runtime/Engine/Classes/Animation/AnimBank.h b/Engine/Source/Runtime/Engine/Classes/Animation/AnimBank.h
--- a/Engine/Source/Runtime/Engine/Classes/Animation/AnimBank.h
+++ b/Engine/Source/Runtime/Engine/Classes/Animation/AnimBank.h
@@ -396,7 +396,7 @@ public:
 	bool bAffectDistanceFieldLighting = true;
 
 	UPROPERTY()
-	FBox PrimitiveLocalBoundsOverride;
+	FBox PrimitiveBoundsOverride;
 
 	UPROPERTY()
 	bool bIsInstanceDataGPUOnly = false;
diff --git a/Engine/Source/Runtime/Engine/Classes/Components/InstancedSkinnedMeshComponent.h b/Engine/Source/Runtime/Engine/Classes/Components/InstancedSkinnedMeshComponent.h
--- a/Engine/Source/Runtime/Engine/Classes/Components/InstancedSkinnedMeshComponent.h
+++ b/Engine/Source/Runtime/Engine/Classes/Components/InstancedSkinnedMeshComponent.h
@@ -97,7 +97,7 @@ protected:
 	TArray<struct FAnimBankItem> AnimBankItems;
 
 	UPROPERTY()
-	FBox PrimitiveLocalBoundsOverride;
+	FBox PrimitiveBoundsOverride;
 
 	UPROPERTY()
 	bool bIsInstanceDataGPUOnly = false;
@@ -161,8 +161,8 @@ public:
 
 	bool UsesGPUOnlyInstances() const { return bIsInstanceDataGPUOnly; }
 
-	FBox GetPrimitiveLocalBoundsOverride() const { return PrimitiveLocalBoundsOverride; }
-	void SetPrimitiveLocalBoundsOverride(const FBox& InBounds) { PrimitiveLocalBoundsOverride = InBounds; }
+	FBox GetPrimitiveBoundsOverride() const { return PrimitiveBoundsOverride; }
+	void SetPrimitiveBoundsOverride(const FBox& InBounds) { PrimitiveBoundsOverride = InBounds; }
 
 	/** Add an instance to this component. Transform is given in local space of this component unless bWorldSpace is set. */
 	UFUNCTION(BlueprintCallable, Category = "Components|InstancedSkinnedMesh")
@@ -306,7 +306,7 @@ public:
 	FInstancedSkinnedMeshComponentInstanceData(const UInstancedSkinnedMeshComponent* InComponent)
 		: FSceneComponentInstanceData(InComponent)
 		, SkinnedAsset(InComponent->GetSkinnedAsset())
-		, PrimitiveLocalBoundsOverride(InComponent->GetPrimitiveLocalBoundsOverride())
+		, PrimitiveBoundsOverride(InComponent->GetPrimitiveBoundsOverride())
 		, bIsInstanceDataGPUOnly(InComponent->UsesGPUOnlyInstances())
 		, NumInstancesGPUOnly(InComponent->GetInstanceCountGPUOnly())
 	{
@@ -346,7 +346,7 @@ public:
 	bool bHasPerInstanceHitProxies = false;
 
 	UPROPERTY()
-	FBox PrimitiveLocalBoundsOverride;
+	FBox PrimitiveBoundsOverride;
 
 	UPROPERTY()
 	bool bIsInstanceDataGPUOnly = false;
diff --git a/Engine/Source/Runtime/Engine/Private/Animation/AnimBank.cpp b/Engine/Source/Runtime/Engine/Private/Animation/AnimBank.cpp
--- a/Engine/Source/Runtime/Engine/Private/Animation/AnimBank.cpp
+++ b/Engine/Source/Runtime/Engine/Private/Animation/AnimBank.cpp
@@ -1407,7 +1407,7 @@ void FSkinnedMeshComponentDescriptorBase::InitFrom(const UInstancedSkinnedMeshCo
 	bVisibleInRayTracing = Template->bVisibleInRayTracing;
 	bAffectDynamicIndirectLighting = Template->bAffectDynamicIndirectLighting;
 	bAffectDistanceFieldLighting = Template->bAffectDistanceFieldLighting;
-	PrimitiveLocalBoundsOverride = Template->GetPrimitiveLocalBoundsOverride();
+	PrimitiveBoundsOverride = Template->GetPrimitiveBoundsOverride();
 #if WITH_EDITOR
 	HLODBatchingPolicy = Template->HLODBatchingPolicy;
 	bIncludeInHLOD = Template->bEnableAutoLODGeneration;
@@ -1430,7 +1430,7 @@ void FSkinnedMeshComponentDescriptorBase::InitComponent(UInstancedSkinnedMeshCom
 	Component->bVisibleInRayTracing = bVisibleInRayTracing;
 	Component->bAffectDynamicIndirectLighting = bAffectDynamicIndirectLighting;
 	Component->bAffectDistanceFieldLighting = bAffectDistanceFieldLighting;
-	Component->SetPrimitiveLocalBoundsOverride(PrimitiveLocalBoundsOverride);
+	Component->SetPrimitiveBoundsOverride(PrimitiveBoundsOverride);
 #if WITH_EDITOR
 	Component->HLODBatchingPolicy = HLODBatchingPolicy;
 	Component->bEnableAutoLODGeneration = bIncludeInHLOD;
@@ -1463,7 +1463,7 @@ bool FSkinnedMeshComponentDescriptorBase::operator==(const FSkinnedMeshComponent
 #endif 
 		bCastShadow == Other.bCastShadow &&
 		bIsInstanceDataGPUOnly == Other.bIsInstanceDataGPUOnly &&
-		PrimitiveLocalBoundsOverride == Other.PrimitiveLocalBoundsOverride &&
+		PrimitiveBoundsOverride == Other.PrimitiveBoundsOverride &&
 		NumInstancesGPUOnly == Other.NumInstancesGPUOnly &&
 		NumCustomDataFloatsGPUOnly == Other.NumCustomDataFloatsGPUOnly &&
 		bVisibleInRayTracing == Other.bVisibleInRayTracing &&
diff --git a/Engine/Source/Runtime/Engine/Private/Components/InstancedSkinnedMeshComponent.cpp b/Engine/Source/Runtime/Engine/Private/Components/InstancedSkinnedMeshComponent.cpp
--- a/Engine/Source/Runtime/Engine/Private/Components/InstancedSkinnedMeshComponent.cpp
+++ b/Engine/Source/Runtime/Engine/Private/Components/InstancedSkinnedMeshComponent.cpp
@@ -1431,7 +1431,7 @@ void UInstancedSkinnedMeshComponent::ApplyComponentInstanceData(struct FInstance
 
 	SelectedInstances = InstancedMeshData->SelectedInstances;
 	bHasPerInstanceHitProxies = InstancedMeshData->bHasPerInstanceHitProxies;
-	PrimitiveLocalBoundsOverride = InstancedMeshData->PrimitiveLocalBoundsOverride;
+	PrimitiveBoundsOverride = InstancedMeshData->PrimitiveBoundsOverride;
 	bIsInstanceDataGPUOnly = InstancedMeshData->bIsInstanceDataGPUOnly;
 	NumInstancesGPUOnly = InstancedMeshData->NumInstancesGPUOnly;
 #endif
@@ -1440,9 +1440,9 @@ void UInstancedSkinnedMeshComponent::ApplyComponentInstanceData(struct FInstance
 FBoxSphereBounds UInstancedSkinnedMeshComponent::CalcBounds(const FTransform& LocalToWorld) const
 {
 	//SCOPE_CYCLE_COUNTER(STAT_CalcSkelMeshBounds);
-	if (PrimitiveLocalBoundsOverride.IsValid)
+	if (PrimitiveBoundsOverride.IsValid)
 	{
-		return PrimitiveLocalBoundsOverride.TransformBy(LocalToWorld);
+		return PrimitiveBoundsOverride.InverseTransformBy(GetComponentTransform().Inverse() * LocalToWorld);
 	}
 	else
 	{
