Fix for DSTM server crash when deserialzing a raw pointer that is on a remote server.

diff --git a/Engine/Source/Runtime/Engine/Classes/GameFramework/CharacterMovementReplication.h b/Engine/Source/Runtime/Engine/Classes/GameFramework/CharacterMovementReplication.h
--- a/Engine/Source/Runtime/Engine/Classes/GameFramework/CharacterMovementReplication.h
+++ b/Engine/Source/Runtime/Engine/Classes/GameFramework/CharacterMovementReplication.h
@@ -280,7 +280,12 @@ public:
 	FVector NewVel; // Note: if bBaseRelativeVelocity is set, this is a relative velocity to the Movement base.
 	FRotator NewRot;
 	FVector GravityDirection;
+
+#if UE_WITH_REMOTE_OBJECT_HANDLE
+	TObjectPtr<UPrimitiveComponent> NewBase;
+#else
 	UPrimitiveComponent* NewBase;
+#endif
 	FName NewBaseBoneName;
 	bool bAckGoodMove;
 	bool bBaseRelativePosition;
diff --git a/Engine/Source/Runtime/Engine/Private/Components/CharacterMovementComponent.cpp b/Engine/Source/Runtime/Engine/Private/Components/CharacterMovementComponent.cpp
--- a/Engine/Source/Runtime/Engine/Private/Components/CharacterMovementComponent.cpp
+++ b/Engine/Source/Runtime/Engine/Private/Components/CharacterMovementComponent.cpp
@@ -10735,7 +10735,11 @@ bool FCharacterMoveResponseDataContainer::Serialize(UCharacterMovementComponent&
 			ClientAdjustment.NewRot = FRotator::ZeroRotator;
 		}
 
+#if UE_WITH_REMOTE_OBJECT_HANDLE
+		SerializeOptionalValue<TObjectPtr<UPrimitiveComponent>>(bIsSaving, Ar, ClientAdjustment.NewBase, nullptr);
+#else
 		SerializeOptionalValue<UPrimitiveComponent*>(bIsSaving, Ar, ClientAdjustment.NewBase, nullptr);
+#endif
 		SerializeOptionalValue<FName>(bIsSaving, Ar, ClientAdjustment.NewBaseBoneName, NAME_None);
 		SerializeOptionalValue<uint8>(bIsSaving, Ar, ClientAdjustment.MovementMode, MOVE_Walking);
 		Ar.SerializeBits(&ClientAdjustment.bBaseRelativePosition, 1);
