- When transacting simply clear out the cached current platform data ptr, expecting it to be set once any triggered async compression finishes or during ::PreSave
- Added missing WRITE_SCOPE and swapped out WRITE with READ scope (as it's actually serializing into the archive)
- Moved ScopeTransaction for applying multiple modifiers from AnimModifier API to high level caller

diff --git a/Engine/Source/Editor/AnimationModifiers/Private/AnimationModifier.cpp b/Engine/Source/Editor/AnimationModifiers/Private/AnimationModifier.cpp
--- a/Engine/Source/Editor/AnimationModifiers/Private/AnimationModifier.cpp
+++ b/Engine/Source/Editor/AnimationModifiers/Private/AnimationModifier.cpp
@@ -48,10 +48,7 @@ void UAnimationModifier::ApplyToAnimationSequence(UAnimSequence* AnimSequence) c
 
 	checkf(AnimSequence, TEXT("Invalid Animation Sequence supplied"));
 
-	//TODO: We want to use UE_COMPRESSED_DATA_READ_SCOPE instead of a wait. But the scope is not working since the animation modifier will modify the anim sequence
-	//and will deadlock when closing the animation controller bracket.
-	//UE_COMPRESSED_DATA_READ_SCOPE(AnimSequence);
-	//Make sure Anim sequence compression is done before starting a new one.
+	//Make sure Anim sequence compression is done
 	AnimSequence->WaitOnExistingCompression();
 
 	USkeleton* Skeleton = AnimSequence->GetSkeleton();
@@ -488,14 +485,16 @@ void UAnimationModifier::ApplyToAll(TSubclassOf<UAnimationModifier> ModifierSubC
 	{
 		// Make sure all packages (in this case UAnimSequences) are loaded to ensure the TObjectIterator has any instances to iterate over
 		LoadModifierReferencers(ModifierSubClass);
-		
-		const FScopedTransaction Transaction(LOCTEXT("UndoAction_ApplyModifiers", "Applying Animation Modifier to Animation Sequence(s)"));		
+
 		for (UAnimationModifiersAssetUserData* AssetUserData : TObjectRange<UAnimationModifiersAssetUserData>{})
 		{
 			UAnimSequence* AnimSequence = Cast<UAnimSequence>(AssetUserData->GetOuter());
 			USkeleton* Skeleton = Cast<USkeleton>(AssetUserData->GetOuter());
 			ensure(AnimSequence || Skeleton);
 
+			IAnimationDataController& Controller = AnimSequence->GetController();
+			IAnimationDataController::FScopedBracket ScopedBracket(Controller, LOCTEXT("ApplyingAnimationModifier", "Applying Animation Modifiers"));
+			
 			for (UAnimationModifier* Modifier : AssetUserData->AnimationModifierInstances)
 			{
 				if (Modifier->GetClass() == ModifierClass && IsValidChecked(Modifier))
diff --git a/Engine/Source/Editor/AnimationModifiers/Private/AnimationModifierDetailCustomization.cpp b/Engine/Source/Editor/AnimationModifiers/Private/AnimationModifierDetailCustomization.cpp
--- a/Engine/Source/Editor/AnimationModifiers/Private/AnimationModifierDetailCustomization.cpp
+++ b/Engine/Source/Editor/AnimationModifiers/Private/AnimationModifierDetailCustomization.cpp
@@ -7,6 +7,7 @@
 #include "DetailCategoryBuilder.h"
 #include "DetailLayoutBuilder.h"
 #include "DetailWidgetRow.h"
+#include "ScopedTransaction.h"
 #include "Internationalization/Internationalization.h"
 #include "Misc/AssertionMacros.h"
 #include "SlotBase.h"
@@ -86,7 +87,8 @@ FReply FAnimationModifierDetailCustomization::OnUpdateRevisionButtonClicked()
 FReply FAnimationModifierDetailCustomization::OnApplyButtonClicked(bool bForceApply)
 {
 	if (ModifierInstance)
-	{
+	{		
+		const FScopedTransaction Transaction(LOCTEXT("UndoAction_ApplyModifiers", "Applying Animation Modifier to Animation Sequence(s)"));		
 		UAnimationModifier::ApplyToAll(ModifierInstance->GetClass(), bForceApply);
 	}
 	return FReply::Handled();
diff --git a/Engine/Source/Runtime/Engine/Private/Animation/AnimSequence.cpp b/Engine/Source/Runtime/Engine/Private/Animation/AnimSequence.cpp
--- a/Engine/Source/Runtime/Engine/Private/Animation/AnimSequence.cpp
+++ b/Engine/Source/Runtime/Engine/Private/Animation/AnimSequence.cpp
@@ -756,87 +756,95 @@ void UAnimSequence::Serialize(FArchive& Ar)
 #if WITH_EDITOR
 				if (!GetOutermost()->HasAnyPackageFlags(PKG_Cooked))
 				{
-					const bool bIsLoading = Ar.IsLoading();
-					const bool bIsSaving = Ar.IsSaving();
-					
-					TUniquePtr<UE::TSharedLock<UE::FSharedRecursiveMutex>> SharedLock;
-
-					if (bIsLoading)
-					{
-						SharedCompressedDataMutex.Lock();
-					}
-					else
+					if (!bIsTransacting)
 					{
-						SharedLock = MakeUnique<UE::TSharedLock<UE::FSharedRecursiveMutex>>(SharedCompressedDataMutex);
-					}
+						const bool bIsLoading = Ar.IsLoading();
+						const bool bIsSaving = Ar.IsSaving();
 
-					if (bCookingTargetNeedsCompressedData)
-					{
-						const ITargetPlatform* Platform = Ar.CookingTarget();
-						checkf(Platform, TEXT("Invalid cooking target platform"));
-						// During cook, serialize out singular (target platform) compressed data 
-						FCompressedAnimSequence& PlatformCompressedData = GetPlatformCompressedData(Platform);
-						check(PlatformCompressedData.IsValid(this, true));
-						SerializeCompressedData(Ar, PlatformCompressedData);
-					}
-					else if (bIsDuplicating || bIsTransacting || bIsCountingMemory)
-					{
-						// During duplication/transacting etc. serialize all contained compressed data
-						if (bIsSaving)
-						{
-							uint32 NumPlatforms = DataByPlatformKeyHash.Num();
-							Ar << NumPlatforms;
+						TUniquePtr<UE::TSharedLock<UE::FSharedRecursiveMutex>> SharedLock;
 
-							FIoHash CurrentPlatformHash = FIoHash::Zero;
-							for (TPair<FIoHash, TUniquePtr<FCompressedAnimSequence>>& CompressedDataPair : DataByPlatformKeyHash)
+					    if (bIsLoading)
+					    {
+						    SharedCompressedDataMutex.Lock();
+					    }
+					    else
+					    {
+						    SharedLock = MakeUnique<UE::TSharedLock<UE::FSharedRecursiveMutex>>(SharedCompressedDataMutex);
+					    }
+
+						if (bCookingTargetNeedsCompressedData)
+						{
+							const ITargetPlatform* Platform = Ar.CookingTarget();
+							checkf(Platform, TEXT("Invalid cooking target platform"));
+							// During cook, serialize out singular (target platform) compressed data 
+							FCompressedAnimSequence& PlatformCompressedData = GetPlatformCompressedData(Platform);
+							check(PlatformCompressedData.IsValid(this, true));
+							SerializeCompressedData(Ar, PlatformCompressedData);
+						}
+						else if (bIsDuplicating || bIsCountingMemory)
+						{					
+							// During duplication/transacting etc. serialize all contained compressed data
+							if (bIsSaving)
 							{
-								Ar << CompressedDataPair.Key;
-								SerializeCompressedData(Ar, *CompressedDataPair.Value.Get());
+								uint32 NumPlatforms = DataByPlatformKeyHash.Num();
+								Ar << NumPlatforms;
 
-								if (CurrentPlatformData && CurrentPlatformData == CompressedDataPair.Value.Get())
+								FIoHash CurrentPlatformHash = FIoHash::Zero;
+								for (TPair<FIoHash, TUniquePtr<FCompressedAnimSequence>>& CompressedDataPair : DataByPlatformKeyHash)
 								{
-									CurrentPlatformHash = CompressedDataPair.Key;
+									Ar << CompressedDataPair.Key;
+									SerializeCompressedData(Ar, *CompressedDataPair.Value.Get());
+
+									if (CurrentPlatformData && CurrentPlatformData == CompressedDataPair.Value.Get())
+									{
+										CurrentPlatformHash = CompressedDataPair.Key;
+									}
 								}
-							}
 							
-							// Store off the expected hash for CurrentPlatformData
-							Ar << CurrentPlatformHash;
-						}
-						else
-						{
-							uint32 NumPlatforms = 0;
-							Ar << NumPlatforms;
-							for (uint32 Index = 0; Index < NumPlatforms; ++Index)
-							{
-								FIoHash PlatformHash;
-								Ar << PlatformHash;
-								FCompressedAnimSequence& TargetData = *DataByPlatformKeyHash.Emplace(PlatformHash, MakeUnique<FCompressedAnimSequence>()).Get();
-								SerializeCompressedData(Ar, TargetData);
+								// Store off the expected hash for CurrentPlatformData
+								Ar << CurrentPlatformHash;
 							}
-
-							CurrentPlatformData = nullptr;
-							
-							// Try and restore short-circuit ptr for CurrentPlatformData according to its previously serialized hash
-							FIoHash CurrentPlatformHash = FIoHash::Zero;
-							Ar << CurrentPlatformHash;
-							if (CurrentPlatformHash != FIoHash::Zero)
+							else
 							{
-								if (TUniquePtr<FCompressedAnimSequence>* Entry = DataByPlatformKeyHash.Find(CurrentPlatformHash))
+								uint32 NumPlatforms = 0;
+								Ar << NumPlatforms;
+								for (uint32 Index = 0; Index < NumPlatforms; ++Index)
 								{
-									CurrentPlatformData = Entry->Get();
+									FIoHash PlatformHash;
+									Ar << PlatformHash;
+									FCompressedAnimSequence& TargetData = *DataByPlatformKeyHash.Emplace(PlatformHash, MakeUnique<FCompressedAnimSequence>()).Get();
+									SerializeCompressedData(Ar, TargetData);
 								}
-							}	
+
+								CurrentPlatformData = nullptr;
+							
+								// Try and restore short-circuit ptr for CurrentPlatformData according to its previously serialized hash
+								FIoHash CurrentPlatformHash = FIoHash::Zero;
+								Ar << CurrentPlatformHash;
+								if (CurrentPlatformHash != FIoHash::Zero)
+								{
+									if (TUniquePtr<FCompressedAnimSequence>* Entry = DataByPlatformKeyHash.Find(CurrentPlatformHash))
+									{
+										CurrentPlatformData = Entry->Get();
+									}
+								}	
+							}
 						}
-					}
 
-					if (bIsLoading)
-					{
-						SharedCompressedDataMutex.Unlock();
-					}
-					else
-					{
-						SharedLock.Reset();
-					}
+					    if (bIsLoading)
+					    {
+						    SharedCompressedDataMutex.Unlock();
+					    }
+					    else
+					    {
+						    SharedLock.Reset();
+					    }
+				    }
+				    else
+				    {
+						UE_COMPRESSED_DATA_WRITE_SCOPE(this);
+					    CurrentPlatformData = nullptr;
+				    }
 				}
 				else
 #endif
@@ -1102,6 +1110,17 @@ void UAnimSequence::PreSave(FObjectPreSaveContext ObjectSaveContext)
 	{
 		ObjectSaveContext.RegisterDeterminismHelper(new FAnimSequenceDeterminismHelper(this));
 	}
+
+	// Updated current platform cache data ptr during non-cook editor runtime
+	if (!ObjectSaveContext.IsCooking())
+	{
+		// Find compressed data for current platform and cache its ptr
+		if(HasCompressedDataForPlatform(nullptr))
+		{
+			UE_COMPRESSED_DATA_WRITE_SCOPE(this);
+			CurrentPlatformData = &GetPlatformCompressedData(nullptr);
+		}
+	}
 #endif
 
 	Super::PreSave(ObjectSaveContext);
diff --git a/Engine/Source/Runtime/Engine/Private/Animation/AnimationCompressionDerivedData.cpp b/Engine/Source/Runtime/Engine/Private/Animation/AnimationCompressionDerivedData.cpp
--- a/Engine/Source/Runtime/Engine/Private/Animation/AnimationCompressionDerivedData.cpp
+++ b/Engine/Source/Runtime/Engine/Private/Animation/AnimationCompressionDerivedData.cpp
@@ -227,8 +227,10 @@ bool FAnimationSequenceAsyncCacheTask::BuildData() const
 	const bool bCompressionSuccessful = bBoneCompressionOk && bCurveCompressionOk;
 	const FString CompressionName = DataToCompress.BoneCompressionSettings->GetFullName();
 
-	if (bCompressionSuccessful && !bIsCanceled)
+	const UAnimSequence* AnimSequence = WeakAnimSequence.Get();
+	if (bCompressionSuccessful && !bIsCanceled && AnimSequence)
 	{
+		UE_COMPRESSED_DATA_WRITE_SCOPE(AnimSequence);		
 		OutData.CompressedByteStream = MoveTemp(CompressionResult.CompressedByteStream);
 		OutData.CompressedDataStructure = MoveTemp(CompressionResult.AnimData);
 		OutData.BoneCompressionCodec = CompressionResult.Codec;
@@ -248,43 +250,43 @@ bool FAnimationSequenceAsyncCacheTask::BuildData() const
 void FAnimationSequenceAsyncCacheTask::LaunchCompressionTask(const UE::FSharedString& Name, const UE::DerivedData::FCacheKey& Key)
 {
 	Owner.LaunchTask(TEXT("AnimationSequenceCompression"), [this, Name, Key]
-		{
-			COOK_STAT(auto Timer = AnimSequenceCookStats::UsageStats.TimeSyncWork());
+	{
+		COOK_STAT(auto Timer = AnimSequenceCookStats::UsageStats.TimeSyncWork());
 
-			// Release execution resource as soon as the task is done
-			ON_SCOPE_EXIT{ ExecutionResource = nullptr; };
+		// Release execution resource as soon as the task is done
+		ON_SCOPE_EXIT{ ExecutionResource = nullptr; };
+
+		if (!BuildData())
+		{
+			return;
+		}
 
-			if (!BuildData())
+		if (const UAnimSequence* AnimSequence = WeakAnimSequence.Get())
+		{
+			if (!CompressedData->IsValid(AnimSequence, true))
 			{
-				return;
+				UE_LOG(LogAnimationCompression, Warning, TEXT("Generated invalid compressed animation data for %s"), *CompressibleAnimPtr->FullName);
 			}
-
-			if (const UAnimSequence* AnimSequence = WeakAnimSequence.Get())
+			else
 			{
-				if (!CompressedData->IsValid(AnimSequence, true))
-				{
-					UE_LOG(LogAnimationCompression, Warning, TEXT("Generated invalid compressed animation data for %s"), *CompressibleAnimPtr->FullName);
-				}
-				else
-				{
-					UE_COMPRESSED_DATA_WRITE_SCOPE(AnimSequence);					
-					TArray64<uint8> RecordData;
-					FMemoryWriter64 Ar(RecordData, /*bIsPersistent*/ true);
-					CompressedData->SerializeCompressedData(Ar, true, const_cast<UAnimSequence*>(AnimSequence), nullptr, CompressibleAnimPtr->BoneCompressionSettings, CompressibleAnimPtr->CurveCompressionSettings);
-					UE_LOG(LogAnimationCompression, Display, TEXT("Storing compressed animation data for %s, at %s/%s"), *Name, *FString(Key.Bucket.ToString()), *LexToString(Key.Hash));
-					UE::DerivedData::GetCache().PutValue({ {Name, Key, UE::DerivedData::FValue::Compress(MakeSharedBufferFromArray(MoveTemp(RecordData)))} }, Owner);
-
-					COOK_STAT(Timer.AddMiss(int64(Ar.Tell())));
-				}
+				UE_COMPRESSED_DATA_READ_SCOPE(AnimSequence);
+				TArray64<uint8> RecordData;
+				FMemoryWriter64 Ar(RecordData, /*bIsPersistent*/ true);
+				CompressedData->SerializeCompressedData(Ar, true, const_cast<UAnimSequence*>(AnimSequence), nullptr, CompressibleAnimPtr->BoneCompressionSettings, CompressibleAnimPtr->CurveCompressionSettings);
+				UE_LOG(LogAnimationCompression, Display, TEXT("Storing compressed animation data for %s, at %s/%s"), *Name, *FString(Key.Bucket.ToString()), *LexToString(Key.Hash));
+				UE::DerivedData::GetCache().PutValue({ {Name, Key, UE::DerivedData::FValue::Compress(MakeSharedBufferFromArray(MoveTemp(RecordData)))} }, Owner);
+
+				COOK_STAT(Timer.AddMiss(int64(Ar.Tell())));
+			}
 
-				if (Compression::FAnimationCompressionMemorySummaryScope::ShouldStoreCompressionResults())
-				{
-					const double CompressionEndTime = FPlatformTime::Seconds();
-					const double CompressionTime = CompressionEndTime - CompressionStartTime;
-					Compression::FAnimationCompressionMemorySummaryScope::CompressionResultSummary().GatherPostCompressionStats(*CompressedData, CompressibleAnimPtr->BoneData, AnimSequence->GetFName(), CompressionTime, true);
-				}
+			if (Compression::FAnimationCompressionMemorySummaryScope::ShouldStoreCompressionResults())
+			{
+				const double CompressionEndTime = FPlatformTime::Seconds();
+				const double CompressionTime = CompressionEndTime - CompressionStartTime;
+				Compression::FAnimationCompressionMemorySummaryScope::CompressionResultSummary().GatherPostCompressionStats(*CompressedData, CompressibleAnimPtr->BoneData, AnimSequence->GetFName(), CompressionTime, true);
 			}
-		});
+		}
+	});
 }
 
 void FAnimationSequenceAsyncCacheTask::CalculateRequiredMemoryEstimate()
