AnimSequence compressed data residency: keep bi-directional info for references as the DDC hash might change, due to underlying asset data changes, between requesting residency and releasing (keeping data corresponding to hash at request resident).

diff --git a/Engine/Plugins/Animation/PoseSearch/Source/Runtime/Private/PoseSearchDerivedData.cpp b/Engine/Plugins/Animation/PoseSearch/Source/Runtime/Private/PoseSearchDerivedData.cpp
--- a/Engine/Plugins/Animation/PoseSearch/Source/Runtime/Private/PoseSearchDerivedData.cpp
+++ b/Engine/Plugins/Animation/PoseSearch/Source/Runtime/Private/PoseSearchDerivedData.cpp
@@ -1449,15 +1449,16 @@ void FPoseSearchDatabaseAsyncCacheTask::ClearAnimSequenceResidency()
 	}
 	
 	const ITargetPlatform* TargetPlatform = GetTargetPlatformManager()->GetRunningTargetPlatform();
+	const uint32 DatabaseHash = GetTypeHash(Database.Get());
 	for (TWeakObjectPtr<const UObject>& Dependency : DatabaseDependencies)
 	{
 		if (Dependency.IsValid())
 		{
 			if (UAnimSequence* AnimSequence = Cast<UAnimSequence>(const_cast<UObject*>(Dependency.Get())))
 			{
-				if (AnimSequence->HasResidency(TargetPlatform))
+				if (AnimSequence->HasResidency(DatabaseHash))
 				{
-					AnimSequence->ReleaseResidency(TargetPlatform, GetTypeHash(Database.Get()));
+					AnimSequence->ReleaseResidency(TargetPlatform, DatabaseHash);
 				}
 			}
 		}
diff --git a/Engine/Source/Runtime/Engine/Classes/Animation/AnimSequence.h b/Engine/Source/Runtime/Engine/Classes/Animation/AnimSequence.h
--- a/Engine/Source/Runtime/Engine/Classes/Animation/AnimSequence.h
+++ b/Engine/Source/Runtime/Engine/Classes/Animation/AnimSequence.h
@@ -817,8 +817,10 @@ protected:
 	// Whether or not compressed data should be cleared in case of residency releasing (cook-time behaviour only)
 	TAtomic<bool> bShouldClearCompressedData = false;
 	
-	// Multi-map keeping track of required (target-platform specific) compressed data, stored by hash and provided identifier from API usage	
-	TMultiMap<FIoHash, uint32> ResidencyReferencerHashes;
+	// Bi-directional multi-maps keeping track of required (target-platform specific) compressed data, stored by hash and provided identifier from API usage
+	TMultiMap<uint32, FIoHash> ResidencyReferencerHashes;
+	TMultiMap<FIoHash, uint32> PlatformHashToReferencers;
+
 	mutable FRWLock ResidencyLock;
 
 	ENGINE_API bool TryCancelAsyncTasks();
@@ -876,6 +878,7 @@ public:
 	
 	// Whether or not compressed animation data has been requested to stay resident until released 
 	ENGINE_API bool HasResidency(const ITargetPlatform* InPlatform) const;
+	ENGINE_API bool HasResidency(uint32 InReferencerHash) const;
 	// Fetches and keeps resident the compressed animation data for the provided target platform (until released)
 	ENGINE_API void RequestResidency(const ITargetPlatform* InPlatform, uint32 InReferencerHash);
 	// Releases previously requested residency of compressed animation data
diff --git a/Engine/Source/Runtime/Engine/Private/Animation/AnimSequence.cpp b/Engine/Source/Runtime/Engine/Private/Animation/AnimSequence.cpp
--- a/Engine/Source/Runtime/Engine/Private/Animation/AnimSequence.cpp
+++ b/Engine/Source/Runtime/Engine/Private/Animation/AnimSequence.cpp
@@ -5320,14 +5320,20 @@ void UAnimSequence::ClearCompressionData(const FIoHash& InKeyHash)
 bool UAnimSequence::RequiresResidency(const FIoHash& InKeyHash) const
 {
 	FReadScopeLock ScopeLock(ResidencyLock);
-	return ResidencyReferencerHashes.Contains(InKeyHash);
+	return PlatformHashToReferencers.Contains(InKeyHash);
 }
 
 bool UAnimSequence::HasResidency(const ITargetPlatform* InPlatform) const
 {
 	const FIoHash KeyHash = GetDerivedDataKeyHash(InPlatform);
 	FReadScopeLock ScopeLock(ResidencyLock);
-	return ResidencyReferencerHashes.Contains(KeyHash);
+	return PlatformHashToReferencers.Contains(KeyHash);
+}
+
+bool UAnimSequence::HasResidency(uint32 InReferencerHash) const
+{
+	FReadScopeLock ScopeLock(ResidencyLock);
+	return ResidencyReferencerHashes.Contains(InReferencerHash);
 }
 
 void UAnimSequence::RequestResidency(const ITargetPlatform* InPlatform, uint32 InReferencerHash)
@@ -5338,7 +5344,7 @@ void UAnimSequence::RequestResidency(const ITargetPlatform* InPlatform, uint32 I
 	TArray<uint32> ReferencerHashes;
 	{
 		FReadScopeLock ScopeLock(ResidencyLock);
-		ResidencyReferencerHashes.MultiFind(KeyHash, ReferencerHashes);
+		PlatformHashToReferencers.MultiFind(KeyHash, ReferencerHashes);
 	}
 	
 	// First request for residency
@@ -5351,28 +5357,34 @@ void UAnimSequence::RequestResidency(const ITargetPlatform* InPlatform, uint32 I
 	if (!ReferencerHashes.Contains(InReferencerHash))
 	{
 		FWriteScopeLock ScopeLock(ResidencyLock);
-		ResidencyReferencerHashes.Add(KeyHash, InReferencerHash);
+		// Store referencer hash alongside platform has at point of request (as data changes during compression can impact platform hash)
+		ResidencyReferencerHashes.Add(InReferencerHash, KeyHash);
+		PlatformHashToReferencers.Add(KeyHash, InReferencerHash);
 	}
 }
 
 void UAnimSequence::ReleaseResidency(const ITargetPlatform* InPlatform, uint32 InReferencerHash)
 {
-	const FIoHash KeyHash = GetDerivedDataKeyHash(InPlatform);
+	FIoHash KeyHash = FIoHash::Zero;
 	TArray<uint32> ReferencerHashes;
 	{
 		FWriteScopeLock ScopeLock(ResidencyLock);
-	
-		check(ResidencyReferencerHashes.Remove(KeyHash, InReferencerHash) == 1);
+		KeyHash = ResidencyReferencerHashes.FindChecked(InReferencerHash);
 
+		// Remove entry from both maps
+		check(ResidencyReferencerHashes.Remove(InReferencerHash) == 1);
+		check(PlatformHashToReferencers.Remove(KeyHash, InReferencerHash) == 1);
+		
 		// Check if there are any references left
-		ResidencyReferencerHashes.MultiFind(KeyHash, ReferencerHashes);
+		PlatformHashToReferencers.MultiFind(KeyHash, ReferencerHashes);
+
 	
 		UE_CLOG(false, LogAnimationCompression, Display, TEXT("ReleaseResidency platform data [KeyHash]: %s %s %x %s"), *InPlatform->DisplayName().ToString(), *LexToString(KeyHash), InReferencerHash, *GetPathName());
 	}
 	
-	if (ReferencerHashes.Num() == 0)
+	if (ReferencerHashes.Num() == 0 && !KeyHash.IsZero())
 	{
-		// If previously this anim sequence was marked as never to be cooked again clear out any compressed data without residency references
+		// If previously this anim sequence was marked as never to be cooked again clear out compressed data without residency references
 		if(bShouldClearCompressedData)
 		{
 			ClearCompressionData(KeyHash);
