Crash fix for AnimationSequence save (followup to 42370935)

diff --git a/Engine/Source/Runtime/Engine/Private/Animation/AnimSequence.cpp b/Engine/Source/Runtime/Engine/Private/Animation/AnimSequence.cpp
--- a/Engine/Source/Runtime/Engine/Private/Animation/AnimSequence.cpp
+++ b/Engine/Source/Runtime/Engine/Private/Animation/AnimSequence.cpp
@@ -1115,10 +1115,11 @@ void UAnimSequence::PreSave(FObjectPreSaveContext ObjectSaveContext)
 	if (!ObjectSaveContext.IsCooking())
 	{
 		// Find compressed data for current platform and cache its ptr
-		if(HasCompressedDataForPlatform(nullptr))
+		const ITargetPlatform* RunningPlatform = GetTargetPlatformManagerRef().GetRunningTargetPlatform();
+		if(HasCompressedDataForPlatform(RunningPlatform))
 		{
 			UE_COMPRESSED_DATA_WRITE_SCOPE(this);
-			CurrentPlatformData = &GetPlatformCompressedData(nullptr);
+			CurrentPlatformData = &GetPlatformCompressedData(RunningPlatform);
 		}
 	}
 #endif
