Spline: Fix horde warning from unchecked pointer dereference.

diff --git a/Engine/Source/Runtime/Engine/Private/Curves/Spline.cpp b/Engine/Source/Runtime/Engine/Private/Curves/Spline.cpp
--- a/Engine/Source/Runtime/Engine/Private/Curves/Spline.cpp
+++ b/Engine/Source/Runtime/Engine/Private/Curves/Spline.cpp
@@ -1189,12 +1189,13 @@ bool FSpline::CreateAttributeChannel(FName AttributeName) const
 {
 	if (NewData && SupportsAttributes())
 	{
-		TTangentBezierSpline<AttrType>* NewAttrChannel = NewData->CreateAttributeChannel<TTangentBezierSpline<AttrType>>(AttributeName);
-		
-		NewAttrChannel->SetPreInfinityMode(EOutOfBoundsHandlingMode::Constant);
-		NewAttrChannel->SetPostInfinityMode(EOutOfBoundsHandlingMode::Constant);
+		if (TTangentBezierSpline<AttrType>* NewAttrChannel = NewData->CreateAttributeChannel<TTangentBezierSpline<AttrType>>(AttributeName); NewAttrChannel != nullptr)
+		{
+			NewAttrChannel->SetPreInfinityMode(EOutOfBoundsHandlingMode::Constant);
+			NewAttrChannel->SetPostInfinityMode(EOutOfBoundsHandlingMode::Constant);
 		
-		return NewAttrChannel != nullptr;
+			return true;
+		}
 	}
 	
 	return false;
