Set r.VRS.DisableForCustomPrimitiveData to false by default (pending full removal), as the issue it was designed to avoid has been addressed.
As of CL 38376645, Nanite Software VRS properly samples at the centroid of its cluster, so custom primitive data is sampled correctly.

diff --git a/Engine/Source/Runtime/Engine/Private/Materials/MaterialShared.cpp b/Engine/Source/Runtime/Engine/Private/Materials/MaterialShared.cpp
--- a/Engine/Source/Runtime/Engine/Private/Materials/MaterialShared.cpp
+++ b/Engine/Source/Runtime/Engine/Private/Materials/MaterialShared.cpp
@@ -107,8 +107,9 @@ static TAutoConsoleVariable<bool> CVarUsingUseNewMaterialTranslatorPrototype(
 
 static TAutoConsoleVariable<bool> CVarDisableVRSForCustomPrimitiveData(
 	TEXT("r.VRS.DisableForCustomPrimitiveData"),
-	true,
-	TEXT("Controls whether or not to disable VRS for materials using Custom Primitive Data."),
+	false,
+	TEXT("Controls whether or not to disable VRS for materials using Custom Primitive Data.\n")
+	TEXT("This was a workaround for Nanite not using centroid sampling for VRS clusters, and should no longer be necessary in 5.6."),
 	ECVF_RenderThreadSafe | ECVF_ReadOnly);
 
 bool IsUsingNewMaterialTranslatorPrototype()
