World Partition: make sure to mark non-editoronly objects referenced by editoronly objects or through editoronly properties from editoronly  or non-editoronly objects as transient so they don't get included in cooked builds.
This works differently in PIE, as we use the instancing context to skip loading these objects.

diff --git a/Engine/Source/Runtime/Engine/Private/WorldPartition/WorldPartitionLevelHelper.cpp b/Engine/Source/Runtime/Engine/Private/WorldPartition/WorldPartitionLevelHelper.cpp
--- a/Engine/Source/Runtime/Engine/Private/WorldPartition/WorldPartitionLevelHelper.cpp
+++ b/Engine/Source/Runtime/Engine/Private/WorldPartition/WorldPartitionLevelHelper.cpp
@@ -276,6 +276,12 @@ void FWorldPartitionLevelHelper::MoveExternalActorsToLevel(const TArray<FWorldPa
 			}
 			check(Actor->GetPackage() == LevelPackage);
 
+			if (PackageObjectMapping.bIsEditorOnly)
+			{
+				Actor->SetFlags(RF_Transient);
+				UE_LOG(LogWorldPartition, Log, TEXT("Marked actor %s transient as it was referenced by an editor-only context"), *Actor->GetPathName());
+			}
+
 			// Process objects found in the source actor package
 			if (bIsActoPackageExternal)
 			{
@@ -312,6 +318,12 @@ void FWorldPartitionLevelHelper::MoveExternalActorsToLevel(const TArray<FWorldPa
 							// Move objects in the destination level package
 							Object->Rename(nullptr, LevelPackage, REN_NonTransactional | REN_DoNotDirty | REN_DontCreateRedirectors);
 						}
+
+						if (PackageObjectMapping.bIsEditorOnly)
+						{
+							Object->SetFlags(RF_Transient);
+							UE_LOG(LogWorldPartition, Log, TEXT("Marked actor object %s transient as it was referenced by an editor-only context"), *Object->GetPathName());
+						}
 					}
 				}
 
