[Static Mesh/Spline Mesh] Add source mesh bounds to static mesh LOD render data and use LOD 0's source bounds when fitting meshes to the spline with spline meshes.

This prevents potential gaps caused by the bounds being pushed out when generating simplified LODs for the meshes.

Also, add a CVar to revert to prior behavior for the sake of any projects that might have been depending on this behavior.

diff --git a/Engine/Source/Developer/MeshBuilder/Private/StaticMeshBuilder.cpp b/Engine/Source/Developer/MeshBuilder/Private/StaticMeshBuilder.cpp
--- a/Engine/Source/Developer/MeshBuilder/Private/StaticMeshBuilder.cpp
+++ b/Engine/Source/Developer/MeshBuilder/Private/StaticMeshBuilder.cpp
@@ -270,7 +270,10 @@ static bool PrepareNaniteStaticMeshBuild(
 	return true;
 }
 
-static bool InitNaniteBuildInput(FStaticMeshNaniteBuildContext& Context, Nanite::IBuilderModule::FInputMeshData& OutData)
+static bool InitNaniteBuildInput(
+	FStaticMeshNaniteBuildContext& Context,
+	Nanite::IBuilderModule::FInputMeshData& OutData,
+	FBoxSphereBounds& OutVertexBounds)
 {
 	FMeshDescription MeshDescription; 
 	if (!Context.SourceModel->CloneMeshDescription(MeshDescription))
@@ -325,7 +328,6 @@ static bool InitNaniteBuildInput(FStaticMeshNaniteBuildContext& Context, Nanite:
 	TArray<int32> WedgeMap;
 
 	// Build the vertex and index buffer
-	FBoxSphereBounds VertexBounds;
 	UE::Private::StaticMeshBuilder::BuildVertexBuffer(
 		Context.StaticMesh,
 		MeshDescription,
@@ -336,7 +338,7 @@ static bool InitNaniteBuildInput(FStaticMeshNaniteBuildContext& Context, Nanite:
 		OutData.Vertices,
 		MeshDescriptionHelper.GetOverlappingCorners(),
 		RemapVerts,
-		VertexBounds,
+		OutVertexBounds,
 		bNeedTangents,
 		bNeedWedgeMap);
 
@@ -357,8 +359,8 @@ static bool InitNaniteBuildInput(FStaticMeshNaniteBuildContext& Context, Nanite:
 		OutData.Sections[SectionIndex].MaterialIndex = Context.StaticMesh->GetSectionInfoMap().Get(0, SectionIndex).MaterialIndex;
 	}
 
-	OutData.VertexBounds.Min = FVector3f(VertexBounds.Origin - VertexBounds.BoxExtent);
-	OutData.VertexBounds.Max = FVector3f(VertexBounds.Origin + VertexBounds.BoxExtent);
+	OutData.VertexBounds.Min = FVector3f(OutVertexBounds.Origin - OutVertexBounds.BoxExtent);
+	OutData.VertexBounds.Max = FVector3f(OutVertexBounds.Origin + OutVertexBounds.BoxExtent);
 
 	TVertexInstanceAttributesRef<FVector2f const> VertexInstanceUVs = MeshDescription.VertexInstanceAttributes().GetAttributesRef<FVector2f>(MeshAttribute::VertexInstance::TextureCoordinate);
 	OutData.NumTexCoords = VertexInstanceUVs.IsValid() ? VertexInstanceUVs.GetNumChannels() : 0;
@@ -386,7 +388,8 @@ static Nanite::FAssemblyPartResourceRef BuildNaniteAssemblyPart(FStaticMeshNanit
 	}
 
 	Nanite::IBuilderModule::FInputMeshData InputMeshData;
-	if (!InitNaniteBuildInput(ChildContext, InputMeshData))
+	FBoxSphereBounds VertexBounds;
+	if (!InitNaniteBuildInput(ChildContext, InputMeshData, VertexBounds))
 	{
 		return nullptr;
 	}
@@ -619,7 +622,7 @@ static bool BuildNanite(
 
 	// Build new vertex buffers
 	Nanite::IBuilderModule::FInputMeshData InputMeshData;
-	if (!InitNaniteBuildInput(Context, InputMeshData))
+	if (!InitNaniteBuildInput(Context, InputMeshData, LOD0Resources.SourceMeshBounds))
 	{
 		return false;
 	}
@@ -1135,8 +1138,6 @@ bool FStaticMeshBuilder::Build(FStaticMeshRenderData& StaticMeshRenderData, cons
 		TArray<TArray<uint32> > PerSectionIndices;
 		PerSectionIndices.AddDefaulted(MeshDescriptions[LodIndex].PolygonGroups().Num());
 
-		FBoxSphereBounds LODBounds;
-
 		if (UE::Tasks::FCancellationTokenScope::IsCurrentWorkCanceled())
 		{
 			return false;
@@ -1153,12 +1154,12 @@ bool FStaticMeshBuilder::Build(FStaticMeshRenderData& StaticMeshRenderData, cons
 			BuildVertexData,
 			MeshDescriptionHelper.GetOverlappingCorners(),
 			RemapVerts,
-			LODBounds,
+			StaticMeshLOD.SourceMeshBounds,
 			true /* bNeedTangents */,
 			true /* bNeedWedgeMap */
 		);
 
-		MeshBoundsBuilder += LODBounds;
+		MeshBoundsBuilder += StaticMeshLOD.SourceMeshBounds;
 
 		TVertexInstanceAttributesRef<FVector2f> VertexInstanceUVs = MeshDescriptions[LodIndex].VertexInstanceAttributes().GetAttributesRef<FVector2f>(MeshAttribute::VertexInstance::TextureCoordinate);
 		const uint32 NumTextureCoord = VertexInstanceUVs.IsValid() ? VertexInstanceUVs.GetNumChannels() : 0;
diff --git a/Engine/Source/Runtime/Core/Private/UObject/UE5ReleaseStreamObjectVersion.cpp b/Engine/Source/Runtime/Core/Private/UObject/UE5ReleaseStreamObjectVersion.cpp
--- a/Engine/Source/Runtime/Core/Private/UObject/UE5ReleaseStreamObjectVersion.cpp
+++ b/Engine/Source/Runtime/Core/Private/UObject/UE5ReleaseStreamObjectVersion.cpp
@@ -14,7 +14,7 @@ TMap<FGuid, FGuid> FUE5ReleaseStreamObjectVersion::GetSystemGuids()
 	SystemGuids.Add(DevGuids.NIAGARASHADERMAP_DERIVEDDATA_VER, FGuid("7BCE96ABAA6E46D689306FE673DED4DE"));
 	SystemGuids.Add(DevGuids.Niagara_LatestScriptCompileVersion, FGuid("5CA8BABDBA114F41A3303EBE474C35DE"));
 	SystemGuids.Add(DevGuids.SkeletalMeshDerivedDataVersion, FGuid("BA91E71F642E4813A1B74B80D0448928"));
-	SystemGuids.Add(DevGuids.STATICMESH_DERIVEDDATA_VER, FGuid("7A556175518D458AA8786ED69AB8DDE7"));
+	SystemGuids.Add(DevGuids.STATICMESH_DERIVEDDATA_VER, FGuid("14B86848953D46989C69A8B195AD7F6C"));
 
 	return SystemGuids;
 }
diff --git a/Engine/Source/Runtime/Engine/Private/Components/SplineMeshComponent.cpp b/Engine/Source/Runtime/Engine/Private/Components/SplineMeshComponent.cpp
--- a/Engine/Source/Runtime/Engine/Private/Components/SplineMeshComponent.cpp
+++ b/Engine/Source/Runtime/Engine/Private/Components/SplineMeshComponent.cpp
@@ -538,7 +538,7 @@ void USplineMeshComponent::UpdateMesh_Concurrent()
 
 FSplineMeshShaderParams USplineMeshComponent::CalculateShaderParams() const
 {
-	return FSplineMeshSceneProxyDesc(this).CalculateShaderParams(GetStaticMesh());
+	return FSplineMeshSceneProxyDesc(this).CalculateShaderParams();
 }
 
 void USplineMeshComponent::UpdateRenderStateAndCollision()
@@ -755,17 +755,17 @@ void USplineMeshComponent::UpdateBounds()
 
 float USplineMeshComponent::ComputeRatioAlongSpline(const float DistanceAlong) const 
 {
-	return FSplineMeshSceneProxyDesc(this).ComputeRatioAlongSpline(GetStaticMesh(), DistanceAlong);
+	return FSplineMeshSceneProxyDesc(this).ComputeRatioAlongSpline(DistanceAlong);
 }
 
 void USplineMeshComponent::ComputeVisualMeshSplineTRange(float& MinT, float& MaxT) const
 {
-	return FSplineMeshSceneProxyDesc(this).ComputeVisualMeshSplineTRange(GetStaticMesh(), MinT, MaxT);
+	return FSplineMeshSceneProxyDesc(this).ComputeVisualMeshSplineTRange(MinT, MaxT);
 }
 
 FBox USplineMeshComponent::ComputeDistortedBounds(const FTransform& InLocalToWorld, const FBoxSphereBounds& InMeshBounds, const FBoxSphereBounds* InBoundsToDistort) const
 {
-	return FSplineMeshSceneProxyDesc(this).ComputeDistortedBounds(GetStaticMesh(), InLocalToWorld, InMeshBounds, InBoundsToDistort);
+	return FSplineMeshSceneProxyDesc(this).ComputeDistortedBounds(InLocalToWorld, InMeshBounds, InBoundsToDistort);
 }
 
 FTransform USplineMeshComponent::GetSocketTransform(FName InSocketName, ERelativeTransformSpace TransformSpace) const
@@ -805,7 +805,7 @@ FTransform USplineMeshComponent::GetSocketTransform(FName InSocketName, ERelativ
 
 FTransform USplineMeshComponent::CalcSliceTransform(const float DistanceAlong) const
 {
-	return FSplineMeshSceneProxyDesc(this).CalcSliceTransform(GetStaticMesh(), DistanceAlong);
+	return FSplineMeshSceneProxyDesc(this).CalcSliceTransform(DistanceAlong);
 }
 
 FTransform USplineMeshComponent::CalcSliceTransformAtSplineOffset(const float Alpha, const float MinT, const float MaxT) const
diff --git a/Engine/Source/Runtime/Engine/Private/SplineMeshSceneProxy.cpp b/Engine/Source/Runtime/Engine/Private/SplineMeshSceneProxy.cpp
--- a/Engine/Source/Runtime/Engine/Private/SplineMeshSceneProxy.cpp
+++ b/Engine/Source/Runtime/Engine/Private/SplineMeshSceneProxy.cpp
@@ -81,7 +81,7 @@ FSplineMeshSceneProxy::FSplineMeshSceneProxy(const FStaticMeshSceneProxyDesc& In
 	}
 
 	// Copy spline params from component
-	SplineParams = InSplineDesc.CalculateShaderParams(InMeshDesc.GetStaticMesh());
+	SplineParams = InSplineDesc.CalculateShaderParams();
 
 	// If we're using GPU Scene, we place the spline mesh parameters in the instance data buffer, with the
 	// exception of mobile platforms that are unable to pull this data from the structured buffer in the VS
@@ -346,7 +346,7 @@ FNaniteSplineMeshSceneProxy::FNaniteSplineMeshSceneProxy(const Nanite::FMaterial
 	}
 
 	// Copy spline params from component
-	SplineParams = InSplineDesc.CalculateShaderParams(InMeshDesc.GetStaticMesh());
+	SplineParams = InSplineDesc.CalculateShaderParams();
 	SplineMeshInstanceData.Setup(SplineParams);
 	SetupInstanceSceneDataBuffers(&SplineMeshInstanceData);
 
diff --git a/Engine/Source/Runtime/Engine/Private/SplineMeshSceneProxyDesc.cpp b/Engine/Source/Runtime/Engine/Private/SplineMeshSceneProxyDesc.cpp
--- a/Engine/Source/Runtime/Engine/Private/SplineMeshSceneProxyDesc.cpp
+++ b/Engine/Source/Runtime/Engine/Private/SplineMeshSceneProxyDesc.cpp
@@ -5,6 +5,14 @@
 #include "SplineMeshSceneProxy.h"
 #include "StaticMeshResources.h"
 #include "Components/SplineMeshComponent.h"
+#include "HAL/ConsoleManager.h"
+
+static TAutoConsoleVariable<bool> CVarSplineMeshFitToSourceMeshBounds(
+	TEXT("r.SplineMesh.FitToSourceMeshBounds"),
+	true,
+	TEXT("When true, will use the bounds of the LOD 0 source mesh to fit the mesh to the spline, as opposed to the collective ")
+	TEXT("mesh bounds of all LODs. This prevents gaps that might occur due bounds being expanded by lower LODs or bounds extensions.")
+);
 
 static FVector3f SplineEvalPos(const FVector3f& StartPos, const FVector3f& StartTangent, const FVector3f& EndPos, const FVector3f& EndTangent, float A)
 {
@@ -66,9 +74,25 @@ void FSplineMeshSceneProxyDesc::InitializeFrom(const USplineMeshComponent* InCom
 	SplineBoundaryMax = InComponent->SplineBoundaryMax;
 	bSmoothInterpRollScale = InComponent->bSmoothInterpRollScale;
 	ForwardAxis = InComponent->ForwardAxis;
+
+	if (const UStaticMesh* StaticMesh = InComponent->GetStaticMesh())
+	{
+		SourceMeshBounds = StaticMesh->GetBounds(); // legacy behavior
+
+		if (CVarSplineMeshFitToSourceMeshBounds.GetValueOnAnyThread())
+		{
+			if (const FStaticMeshRenderData* RenderData = StaticMesh->GetRenderData())
+			{
+				if (!RenderData->LODResources.IsEmpty())
+				{
+					SourceMeshBounds = RenderData->LODResources[0].SourceMeshBounds;
+				}
+			}
+		}
+	}
 }
 
-FSplineMeshShaderParams FSplineMeshSceneProxyDesc::CalculateShaderParams(UStaticMesh* StaticMesh) const
+FSplineMeshShaderParams FSplineMeshSceneProxyDesc::CalculateShaderParams() const
 {
 	FSplineMeshShaderParams Output;
 
@@ -97,13 +121,12 @@ FSplineMeshShaderParams FSplineMeshSceneProxyDesc::CalculateShaderParams(UStatic
 	Output.MeshZScale = 1.0f;
 	Output.MeshZOffset = 0.0f;
 
-	if (StaticMesh)
+	if (SourceMeshBounds.SphereRadius > 0.0f)
 	{
-		const FBoxSphereBounds StaticMeshBounds = StaticMesh->GetBounds();
-		const float BoundsXYRadius = FVector3f(StaticMeshBounds.BoxExtent).Dot((Output.MeshX + Output.MeshY).GetUnsafeNormal());
+		const float BoundsXYRadius = FVector3f(SourceMeshBounds.BoxExtent).Dot((Output.MeshX + Output.MeshY).GetUnsafeNormal());
 
-		const float MeshMinZ = UE::SplineMesh::RealToFloatChecked(USplineMeshComponent::GetAxisValueRef(StaticMeshBounds.Origin - StaticMeshBounds.BoxExtent, ForwardAxis));
-		const float MeshZLen = UE::SplineMesh::RealToFloatChecked(2 * USplineMeshComponent::GetAxisValueRef(StaticMeshBounds.BoxExtent, ForwardAxis));
+		const float MeshMinZ = UE::SplineMesh::RealToFloatChecked(USplineMeshComponent::GetAxisValueRef(SourceMeshBounds.Origin - SourceMeshBounds.BoxExtent, ForwardAxis));
+		const float MeshZLen = UE::SplineMesh::RealToFloatChecked(2 * USplineMeshComponent::GetAxisValueRef(SourceMeshBounds.BoxExtent, ForwardAxis));
 		const float InvMeshZLen = (MeshZLen <= 0.0f) ? 1.0f : 1.0f / MeshZLen;
 		constexpr float MeshTexelLen = float(SPLINE_MESH_TEXEL_WIDTH - 1);
 
@@ -232,11 +255,11 @@ void FSplineMeshSceneProxyDesc::InitRayTracingProxyVertexFactory(UStaticMesh* Me
 	InitSplineVertexFactory_Internal((*RenderData->RayTracingProxy->LODVertexFactories)[InLODIndex], *RenderData->RayTracingProxy->LODs[InLODIndex].VertexBuffers, Mesh->GetLightMapCoordinateIndex(), FeatureLevel, InOverrideColorVertexBuffer);
 }
 
-FBox FSplineMeshSceneProxyDesc::ComputeDistortedBounds(UStaticMesh* StaticMesh, const FTransform& InLocalToWorld, const FBoxSphereBounds& InMeshBounds, const FBoxSphereBounds* InBoundsToDistort) const
+FBox FSplineMeshSceneProxyDesc::ComputeDistortedBounds(const FTransform& InLocalToWorld, const FBoxSphereBounds& InMeshBounds, const FBoxSphereBounds* InBoundsToDistort) const
 {
 	float MinT = 0.0f;
 	float MaxT = 1.0f;
-	ComputeVisualMeshSplineTRange(StaticMesh, MinT, MaxT);
+	ComputeVisualMeshSplineTRange(MinT, MaxT);
 	const FBoxSphereBounds& BoundsToDistort = InBoundsToDistort ? *InBoundsToDistort : InMeshBounds;
 
 	auto ForwardAxisType = static_cast<ESplineMeshAxis::Type>(ForwardAxis);
@@ -254,7 +277,7 @@ FBox FSplineMeshSceneProxyDesc::ComputeDistortedBounds(UStaticMesh* StaticMesh,
 	const FVector B(-6 * SplineParams.StartPos - 4 * SplineParams.StartTangent - 2 * SplineParams.EndTangent + 6 * SplineParams.EndPos);
 	const FVector C(SplineParams.StartTangent);
 
-	auto AppendAxisExtrema = [this, &StaticMesh,&BoundingBox, &FlattenedBounds, MinT, MaxT](const double Discriminant, const double A, const double B)
+	auto AppendAxisExtrema = [this, &BoundingBox, &FlattenedBounds, MinT, MaxT](const double Discriminant, const double A, const double B)
 		{
 			// Negative discriminant means no solution; A == 0 implies coincident start/end points
 			if (Discriminant > 0 && !FMath::IsNearlyZero(A))
@@ -292,13 +315,13 @@ FBox FSplineMeshSceneProxyDesc::ComputeDistortedBounds(UStaticMesh* StaticMesh,
 		const double BoundsMin = USplineMeshComponent::GetAxisValueRef(BoundsToDistort.Origin - BoundsToDistort.BoxExtent, ForwardAxisType);
 		const double BoundsMax = USplineMeshComponent::GetAxisValueRef(BoundsToDistort.Origin + BoundsToDistort.BoxExtent, ForwardAxisType);
 
-		float Alpha = ComputeRatioAlongSpline(StaticMesh, UE::SplineMesh::RealToFloatChecked(BoundsMin));
+		float Alpha = ComputeRatioAlongSpline(UE::SplineMesh::RealToFloatChecked(BoundsMin));
 		if (Alpha < MinT)
 		{
 			BoundingBox += FlattenedBounds.TransformBy(CalcSliceTransformAtSplineOffset(Alpha, MinT, MaxT));
 		}
 
-		Alpha = ComputeRatioAlongSpline(StaticMesh, UE::SplineMesh::RealToFloatChecked(BoundsMax));
+		Alpha = ComputeRatioAlongSpline(UE::SplineMesh::RealToFloatChecked(BoundsMax));
 		if (Alpha > MaxT)
 		{
 			BoundingBox += FlattenedBounds.TransformBy(CalcSliceTransformAtSplineOffset(Alpha, MinT, MaxT));
@@ -308,14 +331,14 @@ FBox FSplineMeshSceneProxyDesc::ComputeDistortedBounds(UStaticMesh* StaticMesh,
 	return BoundingBox.TransformBy(InLocalToWorld);
 }
 
-FTransform FSplineMeshSceneProxyDesc::CalcSliceTransform(UStaticMesh* StaticMesh, const float DistanceAlong) const
+FTransform FSplineMeshSceneProxyDesc::CalcSliceTransform(const float DistanceAlong) const
 {
-	const float Alpha = ComputeRatioAlongSpline(StaticMesh, DistanceAlong);
+	const float Alpha = ComputeRatioAlongSpline(DistanceAlong);
 	
 
 	float MinT = 0.f;
 	float MaxT = 1.f;
-	ComputeVisualMeshSplineTRange(StaticMesh, MinT, MaxT);
+	ComputeVisualMeshSplineTRange(MinT, MaxT);
 
 	return CalcSliceTransformAtSplineOffset(Alpha, MinT, MaxT);
 }
@@ -411,7 +434,7 @@ FTransform FSplineMeshSceneProxyDesc::CalcSliceTransformAtSplineOffset(const flo
 	return SliceTransform;
 }
 
-float FSplineMeshSceneProxyDesc::ComputeRatioAlongSpline(UStaticMesh* StaticMesh, float DistanceAlong) const
+float FSplineMeshSceneProxyDesc::ComputeRatioAlongSpline(float DistanceAlong) const
 {
 	// Find how far 'along' mesh (or custom boundaries) we are
 	float Alpha = 0.f;
@@ -421,11 +444,10 @@ float FSplineMeshSceneProxyDesc::ComputeRatioAlongSpline(UStaticMesh* StaticMesh
 	{
 		Alpha = (DistanceAlong - SplineBoundaryMin) / (SplineBoundaryMax - SplineBoundaryMin);
 	}
-	else if (StaticMesh)
+	else if (SourceMeshBounds.SphereRadius > 0.0f)
 	{
-		const FBoxSphereBounds StaticMeshBounds = StaticMesh->GetBounds();
-		const double MeshMinZ = USplineMeshComponent::GetAxisValueRef(StaticMeshBounds.Origin, ForwardAxis) - USplineMeshComponent::GetAxisValueRef(StaticMeshBounds.BoxExtent, ForwardAxis);
-		const double MeshRangeZ = 2 * USplineMeshComponent::GetAxisValueRef(StaticMeshBounds.BoxExtent, ForwardAxis);
+		const double MeshMinZ = USplineMeshComponent::GetAxisValueRef(SourceMeshBounds.Origin, ForwardAxis) - USplineMeshComponent::GetAxisValueRef(SourceMeshBounds.BoxExtent, ForwardAxis);
+		const double MeshRangeZ = 2 * USplineMeshComponent::GetAxisValueRef(SourceMeshBounds.BoxExtent, ForwardAxis);
 		if (MeshRangeZ > UE_SMALL_NUMBER)
 		{
 			Alpha = UE::SplineMesh::RealToFloatChecked((DistanceAlong - MeshMinZ) / MeshRangeZ);
@@ -434,18 +456,17 @@ float FSplineMeshSceneProxyDesc::ComputeRatioAlongSpline(UStaticMesh* StaticMesh
 	return Alpha;
 }
 
-void FSplineMeshSceneProxyDesc::ComputeVisualMeshSplineTRange(UStaticMesh* StaticMesh, float& MinT, float& MaxT) const
+void FSplineMeshSceneProxyDesc::ComputeVisualMeshSplineTRange(float& MinT, float& MaxT) const
 {
 	MinT = 0.0;
 	MaxT = 1.0;	
 	const bool bHasCustomBoundary = !FMath::IsNearlyEqual(SplineBoundaryMin, SplineBoundaryMax);
 	if (bHasCustomBoundary)
 	{
-		const FBoxSphereBounds& MeshBounds = StaticMesh->GetBounds();
 		// If there's a custom boundary, alter the min/max of the spline we need to evaluate
 		ESplineMeshAxis::Type ForwardAxisType = static_cast<ESplineMeshAxis::Type>(ForwardAxis);
-		const float BoundsMin = UE::SplineMesh::RealToFloatChecked(USplineMeshComponent::GetAxisValueRef(MeshBounds.Origin - MeshBounds.BoxExtent, ForwardAxisType));
-		const float BoundsMax = UE::SplineMesh::RealToFloatChecked(USplineMeshComponent::GetAxisValueRef(MeshBounds.Origin + MeshBounds.BoxExtent, ForwardAxisType));
+		const float BoundsMin = UE::SplineMesh::RealToFloatChecked(USplineMeshComponent::GetAxisValueRef(SourceMeshBounds.Origin - SourceMeshBounds.BoxExtent, ForwardAxisType));
+		const float BoundsMax = UE::SplineMesh::RealToFloatChecked(USplineMeshComponent::GetAxisValueRef(SourceMeshBounds.Origin + SourceMeshBounds.BoxExtent, ForwardAxisType));
 		const float BoundsMinT = (BoundsMin - SplineBoundaryMin) / (SplineBoundaryMax - SplineBoundaryMin);
 		const float BoundsMaxT = (BoundsMax - SplineBoundaryMin) / (SplineBoundaryMax - SplineBoundaryMin);
 
diff --git a/Engine/Source/Runtime/Engine/Private/StaticMesh.cpp b/Engine/Source/Runtime/Engine/Private/StaticMesh.cpp
--- a/Engine/Source/Runtime/Engine/Private/StaticMesh.cpp
+++ b/Engine/Source/Runtime/Engine/Private/StaticMesh.cpp
@@ -803,6 +803,7 @@ void FStaticMeshLODResources::Serialize(FArchive& Ar, UObject* Owner, int32 Inde
 	FStripDataFlags StripFlags(Ar, ClassDataStripFlags);
 
 	Ar << Sections;
+	Ar << SourceMeshBounds;
 	Ar << MaxDeviation;
 
 #if WITH_EDITORONLY_DATA
diff --git a/Engine/Source/Runtime/Engine/Public/SplineMeshSceneProxyDesc.h b/Engine/Source/Runtime/Engine/Public/SplineMeshSceneProxyDesc.h
--- a/Engine/Source/Runtime/Engine/Public/SplineMeshSceneProxyDesc.h
+++ b/Engine/Source/Runtime/Engine/Public/SplineMeshSceneProxyDesc.h
@@ -19,14 +19,15 @@ struct FSplineMeshSceneProxyDesc
 	float SplineBoundaryMax = 0.0f;
 	TEnumAsByte<ESplineMeshAxis::Type> ForwardAxis = ESplineMeshAxis::X;
 	bool bSmoothInterpRollScale = false;
+	FBoxSphereBounds SourceMeshBounds { ForceInit };
 
 	FSplineMeshParams GetSplineParams() const { return SplineParams; }
 
-	FSplineMeshShaderParams CalculateShaderParams(UStaticMesh* StaticMesh) const;
-	float ComputeRatioAlongSpline(UStaticMesh* StaticMesh, float DistanceAlong) const;
-	void ComputeVisualMeshSplineTRange(UStaticMesh* StaticMesh, float& MinT, float& MaxT) const;
-	ENGINE_API FBox ComputeDistortedBounds(UStaticMesh* StaticMesh, const FTransform& InLocalToWorld, const FBoxSphereBounds& InMeshBounds, const FBoxSphereBounds* InBoundsToDistort = nullptr) const;
-	FTransform CalcSliceTransform(UStaticMesh* StaticMesh, const float DistanceAlong) const;
+	FSplineMeshShaderParams CalculateShaderParams() const;
+	float ComputeRatioAlongSpline(float DistanceAlong) const;
+	void ComputeVisualMeshSplineTRange(float& MinT, float& MaxT) const;
+	ENGINE_API FBox ComputeDistortedBounds(const FTransform& InLocalToWorld, const FBoxSphereBounds& InMeshBounds, const FBoxSphereBounds* InBoundsToDistort = nullptr) const;
+	FTransform CalcSliceTransform(const float DistanceAlong) const;
 	FTransform CalcSliceTransformAtSplineOffset(const float Alpha, const float MinT, const float MaxT) const;
 	
 	static void InitVertexFactory(UStaticMesh* Mesh, const ERHIFeatureLevel::Type FeatureLevel, int32 InLODIndex, FColorVertexBuffer* InOverrideColorVertexBuffer);
diff --git a/Engine/Source/Runtime/Engine/Public/StaticMeshResources.h b/Engine/Source/Runtime/Engine/Public/StaticMeshResources.h
--- a/Engine/Source/Runtime/Engine/Public/StaticMeshResources.h
+++ b/Engine/Source/Runtime/Engine/Public/StaticMeshResources.h
@@ -438,6 +438,9 @@ public:
 	/** Geometry for ray tracing. */
 	FRayTracingGeometry* RayTracingGeometry = nullptr;
 
+	/** The bounds of the source mesh data for this LOD */
+	FBoxSphereBounds SourceMeshBounds { ForceInit };
+
 	/** The maximum distance by which this LOD deviates from the base from which it was generated. */
 	float MaxDeviation;
 
