[FSimpleStreamableAssetManager] Objects will no longer rely on Asset being valid during unregistration
Ensure debug pointers are compiled out in shipping configuration

diff --git a/Engine/Source/Runtime/Engine/Private/Streaming/SimpleStreamableAssetManager.cpp b/Engine/Source/Runtime/Engine/Private/Streaming/SimpleStreamableAssetManager.cpp
--- a/Engine/Source/Runtime/Engine/Private/Streaming/SimpleStreamableAssetManager.cpp
+++ b/Engine/Source/Runtime/Engine/Private/Streaming/SimpleStreamableAssetManager.cpp
@@ -48,14 +48,39 @@ void FSimpleStreamableAssetManager::Shutdown()
 
 void FSimpleStreamableAssetManager::Process()
 {
-	check(Instance != nullptr);
-	Instance->UpdateTask_Async();
+	if (IsEnabled())
+	{
+		check(Instance != nullptr);
+		FScopedLock ScopedLock(GetCriticalSection(), IsEnabled());
+		Instance->UpdateTask_Async();
+	}
+}
+
+void FSimpleStreamableAssetManager::UnregisterAsset(UStreamableRenderAsset* InAsset)
+{
+	if (IsEnabled())
+	{
+		check(Instance != nullptr);
+		Instance->RemovedAssetsRecords.Push(FAssetRecord{
+			.AssetRegistrationIndex = InAsset->SimpleStreamableAssetManagerIndex
+#if SIMPLE_STREAMABLE_ASSET_MANAGER_ALLOW_DEBUG_POINTER
+			, .StreamableRenderAsset_ForDebug = InAsset
+#endif
+		});
+		InAsset->SimpleStreamableAssetManagerIndex = INDEX_NONE;
+	}
 }
 
 uint32 FSimpleStreamableAssetManager::GetAllocatedSize()
 {
-	check(Instance != nullptr);
-	return Instance->GetAllocatedSize_Impl();
+	constexpr uint32 StaticSize = sizeof(Instance) + sizeof(GUseSimpleStreamableAssetManager) + sizeof(CVarUseSimpleStreamableAssetManager);
+	if (IsEnabled())
+	{
+		check(Instance != nullptr);
+		FScopedLock SimpleScopedLock(GetCriticalSection(), IsEnabled());
+		return StaticSize + Instance->GetAllocatedSize_Impl();
+	}
+	return StaticSize;
 }
 
 void FSimpleStreamableAssetManager::GetAssetReferenceBounds(const UStreamableRenderAsset* Asset, TArray<FBox>& AssetBoxes)
@@ -121,7 +146,6 @@ void FSimpleStreamableAssetManager::GetAssetReferenceBounds_Impl(const UStreamab
 
 	if (AssetIndex != INDEX_NONE)
 	{
-		check(AssetIndex != INDEX_NONE);
 		const TSet<FAssetBoundElement>& AssetElements = AssetIndexToBounds4Index[AssetIndex];
 
 		AssetBoxes.Reserve(AssetElements.Num());
@@ -141,7 +165,7 @@ uint32 FSimpleStreamableAssetManager::GetAllocatedSize_Impl() const
 	uint32 AllocatedSize = sizeof(FSimpleStreamableAssetManager);
 
 	AllocatedSize += ObjectUsedIndices.GetAllocatedSize();
-	for (const TSet<UStreamableRenderAsset*>& Assets : ObjectRegistrationIndexToAssetProperty)
+	for (const TSet<FAssetRecord>& Assets : ObjectRegistrationIndexToAssetProperty)
 	{
 		AllocatedSize += Assets.GetAllocatedSize();
 	}
@@ -155,9 +179,9 @@ uint32 FSimpleStreamableAssetManager::GetAllocatedSize_Impl() const
 
 	AllocatedSize += AssetUsedIndices.GetAllocatedSize();
 
-	for (const TSet<FAssetBoundElement>& AssetBondElements : AssetIndexToBounds4Index)
+	for (const TSet<FAssetBoundElement>& AssetBoundElements : AssetIndexToBounds4Index)
 	{
-		AllocatedSize += AssetBondElements.GetAllocatedSize();
+		AllocatedSize += AssetBoundElements.GetAllocatedSize();
 	}
 
 	AllocatedSize += BoundsViewInfos.GetAllocatedSize();
@@ -176,6 +200,7 @@ void FSimpleStreamableAssetManager::UpdateTask_Async()
 	TArray<TLocklessGrowingStorage<FUpdate>::FStorageShard*> Pending_UpdateRecords{};
 	TArray<TLocklessGrowingStorage<FRegister>::FStorageShard*> Pending_RegisterRecords{};
 	TArray<TLocklessGrowingStorage<FUnregister>::FStorageShard*> Pending_UnregisterRecords{};
+	TArray<TLocklessGrowingStorage<FAssetRecord>::FStorageShard*> Pending_RemovedAssetRecords{};
 
 	{
 		/** It is important to extract Shards in this order
@@ -187,6 +212,7 @@ void FSimpleStreamableAssetManager::UpdateTask_Async()
 		Pending_UpdateRecords = UpdateRecords.ExtractShards();
 		Pending_RegisterRecords = RegisterRecords.ExtractShards();
 		Pending_UnregisterRecords = UnregisterRecords.ExtractShards();
+		Pending_RemovedAssetRecords = RemovedAssetsRecords.ExtractShards();
 	}
 	
 	/** Assign a indice to each proxy **/
@@ -221,8 +247,6 @@ void FSimpleStreamableAssetManager::UpdateTask_Async()
 			}
 		}
 	}
-
-	
 	
 	{
 		SCOPED_NAMED_EVENT(FSimpleStreamableAssetManager_Process, FColor::Silver);
@@ -263,6 +287,26 @@ void FSimpleStreamableAssetManager::UpdateTask_Async()
 				}
 			}
 		}
+
+		{
+			SCOPED_NAMED_EVENT(FSimpleStreamableAssetManager_RemoveAssets, FColor::Silver);
+			for (TLocklessGrowingStorage<FAssetRecord>::FStorageShard* Shard : Pending_RemovedAssetRecords)
+			{
+				TArrayView<FAssetRecord> View = Shard->GetData();
+				for (const FAssetRecord& Record : View)
+				{
+					const int32 AssetIndex = Record.AssetRegistrationIndex;
+					if (AssetIndex != INDEX_NONE)
+					{
+						AssetUsedIndices[AssetIndex] = false;
+						--UsedAssetIndices;
+						FreeAssetIndexHint = FMath::Min(FreeAssetIndexHint, AssetIndex);
+
+						check(AssetIndexToBounds4Index[AssetIndex].Num() == 0);
+					}
+				}
+			}
+		}
 	}
 	
 	if (LastRegisteredObjectCount != RegisteredObjectCount)
@@ -311,6 +355,11 @@ void FSimpleStreamableAssetManager::UpdateTask_Async()
 			delete Shard;
 		}
 		TRACE_COUNTER_SET(RemovedObjects, Removed);
+
+		for (TLocklessGrowingStorage<FAssetRecord>::FStorageShard* Shard : Pending_RemovedAssetRecords)
+		{
+			delete Shard;
+		}
 	}
 }
 
@@ -495,7 +544,7 @@ FBoxSphereBounds FSimpleStreamableAssetManager::GetBounds(int32 BoundsIndex) con
 void FSimpleStreamableAssetManager::AddRenderAssetElements(const TArrayView<FStreamingRenderAssetPrimitiveInfo>& RenderAssetInstanceInfos, int32 BoundsIndex, int32 ObjectRegistrationIndex, bool bForceMipStreaming)
 {
 	const int32 BoundsOffset = ObjectRegistrationIndexToBounds4Index[ObjectRegistrationIndex].Add(BoundsIndex);
-	TSet<UStreamableRenderAsset*>& ObjectAssets = ObjectRegistrationIndexToAssetProperty[ObjectRegistrationIndex];
+	TSet<FAssetRecord>& ObjectAssets = ObjectRegistrationIndexToAssetProperty[ObjectRegistrationIndex];
 	ObjectAssets.Reserve(RenderAssetInstanceInfos.Num());
 	for (const FStreamingRenderAssetPrimitiveInfo& AssetInfo : RenderAssetInstanceInfos)
 	{
@@ -504,9 +553,7 @@ void FSimpleStreamableAssetManager::AddRenderAssetElements(const TArrayView<FStr
 		{
 			continue;
 		}
-			
-		ObjectAssets.Add(RenderAsset);			
-
+		
 		const int32 AssetIndex = [&]
 		{
 			if (RenderAsset->SimpleStreamableAssetManagerIndex != INDEX_NONE)
@@ -532,6 +579,12 @@ void FSimpleStreamableAssetManager::AddRenderAssetElements(const TArrayView<FStr
 
 		check(AssetIndex != INDEX_NONE);
 
+		ObjectAssets.Add(FAssetRecord{ .AssetRegistrationIndex = AssetIndex
+#if SIMPLE_STREAMABLE_ASSET_MANAGER_ALLOW_DEBUG_POINTER
+			, .StreamableRenderAsset_ForDebug = RenderAsset
+#endif
+			});			
+
 		if (AssetIndex >= AssetIndexToBounds4Index.Num())
 		{
 			check(AssetIndex == AssetIndexToBounds4Index.Num());
@@ -555,26 +608,20 @@ void FSimpleStreamableAssetManager::RemoveRenderAssetElements(int32 ObjectRegist
 	TArray<int32>& ObjectBoundsIndices = ObjectRegistrationIndexToBounds4Index[ObjectRegistrationIndex];
 	const int32 BoundCount = ObjectBoundsIndices.Num();
 	
-	TSet<UStreamableRenderAsset*>& ObjectAssets = ObjectRegistrationIndexToAssetProperty[ObjectRegistrationIndex];
+	TSet<FAssetRecord>& ObjectAssets = ObjectRegistrationIndexToAssetProperty[ObjectRegistrationIndex];
 	
-	for (UStreamableRenderAsset* Asset : ObjectAssets)
+	for (const FAssetRecord& Asset : ObjectAssets)
 	{
-		const int32 AssetIndex = Asset->SimpleStreamableAssetManagerIndex;
+		const int32& AssetIndex = Asset.AssetRegistrationIndex;
 		check(AssetIndex != INDEX_NONE);
 		TSet<FAssetBoundElement>& AssetBondElements = AssetIndexToBounds4Index[AssetIndex];
-		
-		for(int32 BoundOffset = 0; BoundOffset < BoundCount; ++BoundOffset)
-		{
-			AssetBondElements.Remove(FAssetBoundElement{.ObjectRegistrationIndex = static_cast<uint32>(ObjectRegistrationIndex), .ElementOffset = static_cast<uint32>(BoundOffset)});
-		}
-		
-		// If no references stop tracking 
-		if (AssetBondElements.Num() == 0) 
+
+		if (AssetBondElements.Num() > 0)
 		{
-			AssetUsedIndices[AssetIndex] = false;
-			--UsedAssetIndices;
-			FreeAssetIndexHint = FMath::Min(FreeAssetIndexHint, AssetIndex);
-			Asset->SimpleStreamableAssetManagerIndex = INDEX_NONE;
+			for(int32 BoundOffset = 0; BoundOffset < BoundCount; ++BoundOffset)
+			{
+				AssetBondElements.Remove(FAssetBoundElement{.ObjectRegistrationIndex = static_cast<uint32>(ObjectRegistrationIndex), .ElementOffset = static_cast<uint32>(BoundOffset)});
+			}
 		}
 	}
 	ObjectAssets.Reset();
@@ -599,11 +646,11 @@ void FSimpleStreamableAssetManager::GetRenderAssetScreenSize_Impl(
 	if (AssetType != EStreamableRenderAssetType::Texture || MaxSize_VisibleOnly < MaxAssetSize || LogPrefix)
 	{
 		const int32 AssetIndex = InAsset->SimpleStreamableAssetManagerIndex;
+		// Asset might be registered, but it might not be yet used by any proxy
 		if (AssetIndex != INDEX_NONE)
 		{
 			const TSet<FAssetBoundElement>& AssetBoundElements = AssetIndexToBounds4Index[AssetIndex];
-			check(AssetBoundElements.Num() >= 1);
-			
+						
 			for (const FAssetBoundElement& AssetBoundElement : AssetBoundElements)
 			{
 				const int32 AssetBoundsIndex =  AssetBoundElement.ElementIndex; 
diff --git a/Engine/Source/Runtime/Engine/Private/Streaming/StreamingManagerTexture.cpp b/Engine/Source/Runtime/Engine/Private/Streaming/StreamingManagerTexture.cpp
--- a/Engine/Source/Runtime/Engine/Private/Streaming/StreamingManagerTexture.cpp
+++ b/Engine/Source/Runtime/Engine/Private/Streaming/StreamingManagerTexture.cpp
@@ -1030,6 +1030,9 @@ void FRenderAssetStreamingManager::AddStreamingRenderAsset(UStreamableRenderAsse
  */
 void FRenderAssetStreamingManager::RemoveStreamingRenderAsset( UStreamableRenderAsset* RenderAsset )
 {
+	// FSimpleStreamableAssetManager::UnregisterAsset can be called outside critical Section
+	FSimpleStreamableAssetManager::UnregisterAsset(RenderAsset);
+	
 	FScopeLock ScopeLock(&CriticalSection);
 	LLM_SCOPE(ELLMTag::StreamingManager);
 #if LLM_ALLOW_ASSETS_TAGS
@@ -2727,12 +2730,8 @@ bool FRenderAssetStreamingManager::HandleStreamingManagerMemoryCommand( const TC
 			MemSize += LevelManager->GetAllocatedSize();
 		}
 	}
-
-	if (FSimpleStreamableAssetManager::IsEnabled())
-	{
-		FSimpleStreamableAssetManager::FScopedLock SimpleScopedLock(FSimpleStreamableAssetManager::GetCriticalSection(), FSimpleStreamableAssetManager::IsEnabled());
-		MemSize += FSimpleStreamableAssetManager::GetAllocatedSize();
-	}
+	
+	MemSize += FSimpleStreamableAssetManager::GetAllocatedSize();
 
 	Ar.Logf(TEXT("StreamingManagerTexture: %.2f KB used"), MemSize / 1024.0f);
 
diff --git a/Engine/Source/Runtime/Engine/Private/Streaming/TextureInstanceTask.cpp b/Engine/Source/Runtime/Engine/Private/Streaming/TextureInstanceTask.cpp
--- a/Engine/Source/Runtime/Engine/Private/Streaming/TextureInstanceTask.cpp
+++ b/Engine/Source/Runtime/Engine/Private/Streaming/TextureInstanceTask.cpp
@@ -179,11 +179,9 @@ void FDoWorkTask::DoWork()
 {
 	LLM_SCOPE(ELLMTag::Textures);
 	TRACE_CPUPROFILER_EVENT_SCOPE(RenderAssetInstanceTask);
-
-	if(FSimpleStreamableAssetManager::IsEnabled())
+	
 	{
 		TRACE_CPUPROFILER_EVENT_SCOPE(FSimpleStreamableAssetManager);
-		FSimpleStreamableAssetManager::FScopedLock ScopedLock(FSimpleStreamableAssetManager::GetCriticalSection(), FSimpleStreamableAssetManager::IsEnabled());
 		FSimpleStreamableAssetManager::Process();
 	}
 
diff --git a/Engine/Source/Runtime/Engine/Public/Streaming/SimpleStreamableAssetManager.h b/Engine/Source/Runtime/Engine/Public/Streaming/SimpleStreamableAssetManager.h
--- a/Engine/Source/Runtime/Engine/Public/Streaming/SimpleStreamableAssetManager.h
+++ b/Engine/Source/Runtime/Engine/Public/Streaming/SimpleStreamableAssetManager.h
@@ -14,6 +14,8 @@ struct FStreamingViewInfoExtra;
 struct FRenderAssetStreamingSettings;
 struct FBounds4;
 
+#define SIMPLE_STREAMABLE_ASSET_MANAGER_ALLOW_DEBUG_POINTER !(UE_BUILD_SHIPPING)
+
 template<typename T, uint32 N>
 struct TLocklessStaticStorage
 {
@@ -149,14 +151,16 @@ public:
 	
 	struct FUnregister
 	{
+		TSharedPtr<int32, ESPMode::ThreadSafe> ObjectRegistrationIndex{};
+		
+#if SIMPLE_STREAMABLE_ASSET_MANAGER_ALLOW_DEBUG_POINTER
 		union
 		{
 			UPTRINT ObjectKey = 0;
 			const FPrimitiveSceneProxy* SceneProxy_ForDebug;
 		};
-		TSharedPtr<int32, ESPMode::ThreadSafe> ObjectRegistrationIndex{};
-
-
+#endif
+		
 		FUnregister() = default;
 		FUnregister(const FUnregister&) = default;
 		FUnregister& operator=(const FUnregister&) = default;
@@ -165,21 +169,24 @@ public:
 
 		template<typename TObject>
 		FUnregister(const TObject* Object)
-			: ObjectKey(reinterpret_cast<UPTRINT>(Object))
-			, ObjectRegistrationIndex(Object->SimpleStreamableAssetManagerIndex)
+			: ObjectRegistrationIndex(Object->SimpleStreamableAssetManagerIndex)
+#if SIMPLE_STREAMABLE_ASSET_MANAGER_ALLOW_DEBUG_POINTER
+			, ObjectKey(reinterpret_cast<UPTRINT>(Object))
+#endif
 		{}
 	};
 
 	struct FUpdate
 	{
+		TSharedPtr<int32, ESPMode::ThreadSafe> ObjectRegistrationIndex{};
+		FBoxSphereBounds ObjectBounds{};
+#if SIMPLE_STREAMABLE_ASSET_MANAGER_ALLOW_DEBUG_POINTER
 		union
 		{
 			UPTRINT ObjectKey = 0;
 			const FPrimitiveSceneProxy* SceneProxy_ForDebug;
 		};
-
-		TSharedPtr<int32, ESPMode::ThreadSafe> ObjectRegistrationIndex{};
-		FBoxSphereBounds ObjectBounds{};
+#endif
 		float MinDistance = 0.0f;
 		float MaxDistance = FLT_MAX;
 		float LastRenderedTime = -1000.0f;
@@ -196,9 +203,11 @@ public:
 			const TSharedPtr<int32, ESPMode::ThreadSafe>& InObjectRegistrationIndex,
 			const FBoxSphereBounds& InBounds,
 			const float InMinDistance, const float InMaxDistance, const float InLastRenderedTime, bool InForceMipStreaming)
-				: ObjectKey(reinterpret_cast<UPTRINT>(InObject))
-				, ObjectRegistrationIndex(InObjectRegistrationIndex)
+				: ObjectRegistrationIndex(InObjectRegistrationIndex)
 				, ObjectBounds(InBounds)
+#if SIMPLE_STREAMABLE_ASSET_MANAGER_ALLOW_DEBUG_POINTER
+				, ObjectKey(reinterpret_cast<UPTRINT>(InObject))
+#endif
 				, MinDistance(InMinDistance)
 				, MaxDistance(InMaxDistance)
 				, LastRenderedTime(InLastRenderedTime)
@@ -240,6 +249,15 @@ public:
 		
 		friend bool operator==(const FAssetBoundElement& A, const FAssetBoundElement& B) { return A.Raw == B.Raw; }
 	};
+
+	struct FAssetRecord
+	{
+		int32 AssetRegistrationIndex = 0;
+#if SIMPLE_STREAMABLE_ASSET_MANAGER_ALLOW_DEBUG_POINTER
+		const UStreamableRenderAsset* StreamableRenderAsset_ForDebug = nullptr;
+#endif
+		friend bool operator==(const FAssetRecord& A, const FAssetRecord& B) { return A.AssetRegistrationIndex == B.AssetRegistrationIndex; }
+	};
 private:
 	
 	
@@ -247,7 +265,8 @@ private:
 	ENGINE_API static int32 GUseSimpleStreamableAssetManager;
 	static FSimpleStreamableAssetManager* Instance;
 	FCriticalSection CriticalSection;
-	
+
+	TLocklessGrowingStorage<FAssetRecord> RemovedAssetsRecords;
 	TLocklessGrowingStorage<FRegister> RegisterRecords;
 	TLocklessGrowingStorage<FUnregister> UnregisterRecords;
 	TLocklessGrowingStorage<FUpdate> UpdateRecords;
@@ -258,7 +277,7 @@ private:
 
 	int32 FreeObjectIndexHint = 0;
 	TBitArray<> ObjectUsedIndices;
-	TArray<TSet<UStreamableRenderAsset*>> ObjectRegistrationIndexToAssetProperty;
+	TArray<TSet<FAssetRecord>> ObjectRegistrationIndexToAssetProperty;
 	TArray<TArray<int32>> ObjectRegistrationIndexToBounds4Index;
 
 	void RegisterRecord(FRegister& Record);
@@ -306,8 +325,6 @@ private:
 		const FUpdate& Record,
 		float& MinDistanceSq, float& MinRangeSq, float& MaxRangeSq);
 
-	
-
 	void UpdateTask_Async();
 	
 	void GetAssetReferenceBounds_Impl(const UStreamableRenderAsset* Asset, TArray<FBox>& AssetBoxes) const;
@@ -323,6 +340,8 @@ public:
 	static void Init();
 	static void Shutdown();
 	static void Process();
+	
+	static void UnregisterAsset(UStreamableRenderAsset* InAsset);
 
 	static uint32 GetAllocatedSize();
 
@@ -372,4 +391,5 @@ FSimpleStreamableAssetManager::FRegister::FRegister(const TObject* InObject, con
 	}
 }
 
-inline int32 GetTypeHash(const FSimpleStreamableAssetManager::FAssetBoundElement& Object) { return Object.Raw; }
\ No newline at end of file
+inline int32 GetTypeHash(const FSimpleStreamableAssetManager::FAssetBoundElement& Object) { return Object.Raw; }
+inline int32 GetTypeHash(const FSimpleStreamableAssetManager::FAssetRecord& Object) { return Object.AssetRegistrationIndex; }
\ No newline at end of file
