[FSimpleStreamableAssetManager] In Editor after building texture resources will be released first and proxy rebuilt later - making check if all objects unregistred fails
Additionaly, Set do not release memory on elements  removed we need to empty it, replacing TSet with TArray to reduce the memory footprint and iteration time

diff --git a/Engine/Source/Runtime/Engine/Private/Streaming/SimpleStreamableAssetManager.cpp b/Engine/Source/Runtime/Engine/Private/Streaming/SimpleStreamableAssetManager.cpp
--- a/Engine/Source/Runtime/Engine/Private/Streaming/SimpleStreamableAssetManager.cpp
+++ b/Engine/Source/Runtime/Engine/Private/Streaming/SimpleStreamableAssetManager.cpp
@@ -165,7 +165,7 @@ uint32 FSimpleStreamableAssetManager::GetAllocatedSize_Impl() const
 	uint32 AllocatedSize = sizeof(FSimpleStreamableAssetManager);
 
 	AllocatedSize += ObjectUsedIndices.GetAllocatedSize();
-	for (const TSet<FAssetRecord>& Assets : ObjectRegistrationIndexToAssetProperty)
+	for (const TArray<FAssetRecord>& Assets : ObjectRegistrationIndexToAssetProperty)
 	{
 		AllocatedSize += Assets.GetAllocatedSize();
 	}
@@ -301,8 +301,8 @@ void FSimpleStreamableAssetManager::UpdateTask_Async()
 						AssetUsedIndices[AssetIndex] = false;
 						--UsedAssetIndices;
 						FreeAssetIndexHint = FMath::Min(FreeAssetIndexHint, AssetIndex);
-
-						check(AssetIndexToBounds4Index[AssetIndex].Num() == 0);
+						
+						AssetIndexToBounds4Index[AssetIndex].Empty();
 					}
 				}
 			}
@@ -544,7 +544,7 @@ FBoxSphereBounds FSimpleStreamableAssetManager::GetBounds(int32 BoundsIndex) con
 void FSimpleStreamableAssetManager::AddRenderAssetElements(const TArrayView<FStreamingRenderAssetPrimitiveInfo>& RenderAssetInstanceInfos, int32 BoundsIndex, int32 ObjectRegistrationIndex, bool bForceMipStreaming)
 {
 	const int32 BoundsOffset = ObjectRegistrationIndexToBounds4Index[ObjectRegistrationIndex].Add(BoundsIndex);
-	TSet<FAssetRecord>& ObjectAssets = ObjectRegistrationIndexToAssetProperty[ObjectRegistrationIndex];
+	TSet<FAssetRecord> ObjectAssets;
 	ObjectAssets.Reserve(RenderAssetInstanceInfos.Num());
 	for (const FStreamingRenderAssetPrimitiveInfo& AssetInfo : RenderAssetInstanceInfos)
 	{
@@ -596,6 +596,8 @@ void FSimpleStreamableAssetManager::AddRenderAssetElements(const TArrayView<FStr
 
 		AssetIndexToBounds4Index[AssetIndex].Add(FAssetBoundElement{.ElementIndex = uint32(BoundsIndex), .ObjectRegistrationIndex = uint32(ObjectRegistrationIndex), .ElementOffset =  uint32(BoundsOffset), .bForceLOD = bForceMipStreaming, .TexelFactor =  AssetInfo.TexelFactor});
 	}
+
+	ObjectRegistrationIndexToAssetProperty[ObjectRegistrationIndex] = ObjectAssets.Array();
 }
 
 void FSimpleStreamableAssetManager::RemoveRenderAssetElements(int32 ObjectRegistrationIndex)
@@ -608,7 +610,7 @@ void FSimpleStreamableAssetManager::RemoveRenderAssetElements(int32 ObjectRegist
 	TArray<int32>& ObjectBoundsIndices = ObjectRegistrationIndexToBounds4Index[ObjectRegistrationIndex];
 	const int32 BoundCount = ObjectBoundsIndices.Num();
 	
-	TSet<FAssetRecord>& ObjectAssets = ObjectRegistrationIndexToAssetProperty[ObjectRegistrationIndex];
+	TArray<FAssetRecord>& ObjectAssets = ObjectRegistrationIndexToAssetProperty[ObjectRegistrationIndex];
 	
 	for (const FAssetRecord& Asset : ObjectAssets)
 	{
@@ -655,11 +657,11 @@ void FSimpleStreamableAssetManager::GetRenderAssetScreenSize_Impl(
 			{
 				const int32 AssetBoundsIndex =  AssetBoundElement.ElementIndex; 
 				
-				const FBoundsViewInfo& BoundsVieWInfo = BoundsViewInfos[AssetBoundsIndex];
+				const FBoundsViewInfo& BoundsViewInfo = BoundsViewInfos[AssetBoundsIndex];
 			
 				const float TexelFactor = AssetBoundElement.TexelFactor;
 				const bool bForcedLODs = AssetBoundElement.bForceLOD;
-				FRenderAssetInstanceAsyncView::ProcessElement(AssetType, BoundsVieWInfo, AssetType != EStreamableRenderAssetType::Texture ? TexelFactor : TexelFactor * BoundsVieWInfo.ComponentScale, bForcedLODs, MaxSize, MaxSize_VisibleOnly, MaxNumForcedLODs);
+				FRenderAssetInstanceAsyncView::ProcessElement(AssetType, BoundsViewInfo, AssetType != EStreamableRenderAssetType::Texture ? TexelFactor : TexelFactor * BoundsViewInfo.ComponentScale, bForcedLODs, MaxSize, MaxSize_VisibleOnly, MaxNumForcedLODs);
 
 				if (MaxSize_VisibleOnly >= MaxAssetSize || MaxNumForcedLODs >= MaxAllowedMip)
 				{
diff --git a/Engine/Source/Runtime/Engine/Public/Streaming/SimpleStreamableAssetManager.h b/Engine/Source/Runtime/Engine/Public/Streaming/SimpleStreamableAssetManager.h
--- a/Engine/Source/Runtime/Engine/Public/Streaming/SimpleStreamableAssetManager.h
+++ b/Engine/Source/Runtime/Engine/Public/Streaming/SimpleStreamableAssetManager.h
@@ -277,7 +277,7 @@ private:
 
 	int32 FreeObjectIndexHint = 0;
 	TBitArray<> ObjectUsedIndices;
-	TArray<TSet<FAssetRecord>> ObjectRegistrationIndexToAssetProperty;
+	TArray<TArray<FAssetRecord>> ObjectRegistrationIndexToAssetProperty;
 	TArray<TArray<int32>> ObjectRegistrationIndexToBounds4Index;
 
 	void RegisterRecord(FRegister& Record);
