Fix SimpleStreamable compile errors

diff --git a/Engine/Source/Runtime/Engine/Private/Streaming/SimpleStreamableAssetManager.cpp b/Engine/Source/Runtime/Engine/Private/Streaming/SimpleStreamableAssetManager.cpp
--- a/Engine/Source/Runtime/Engine/Private/Streaming/SimpleStreamableAssetManager.cpp
+++ b/Engine/Source/Runtime/Engine/Private/Streaming/SimpleStreamableAssetManager.cpp
@@ -9,8 +9,8 @@
 
 #include "ProfilingDebugging/CountersTrace.h"
 
-#include "Streaming\TextureInstanceView.h"
-#include "Streaming\TextureInstanceView.inl"
+#include "Streaming/TextureInstanceView.h"
+#include "Streaming/TextureInstanceView.inl"
 
 TRACE_DECLARE_INT_COUNTER(RegisteredObjects, TEXT("StreamableAssets/SimpleStreamableAssetManager/RegisteredObjects"));
 TRACE_DECLARE_INT_COUNTER(RegisteredAssets, TEXT("StreamableAssets/SimpleStreamableAssetManager/RegisteredAssets"));
