Add a NetDriver option to reduce spurious ack packets with no other data.

diff --git a/Engine/Source/Runtime/Engine/Classes/Engine/NetDriver.h b/Engine/Source/Runtime/Engine/Classes/Engine/NetDriver.h
--- a/Engine/Source/Runtime/Engine/Classes/Engine/NetDriver.h
+++ b/Engine/Source/Runtime/Engine/Classes/Engine/NetDriver.h
@@ -2354,6 +2354,15 @@ public:
 	/** Enumerate all connections to find one that matches a particular player id. */
 	TObjectPtr<UNetConnection> FindConnection(struct FUniqueNetIdRepl& PlayerId);
 
+	/**
+	 * Used by UNetConnection after receiving a packet to determine whether to send an ack packet on the next
+	 * connection flush regardless of whether there's any other data to send. Set via the config property bSendImmediateAcks.
+	 */
+	bool IsSendingImmediateAcks() const
+	{
+		return bSendImmediateAcks;
+	}
+
 protected:
 	
 	/** Stream of random numbers to be used by this instance of UNetDriver */
@@ -2504,4 +2513,13 @@ private:
 
 	/** Server support for NMT_JoinNoPawn and NMT_JoinNoPawnSplit. */
 	bool bSupportNoPawnConnection = false;
+
+	/**
+	 * If true, after receiving a packet, send an ack packet on the next connection flush regardless of whether
+	 * there's any other data to send. This is the default.
+	 * If false, the ack won't be sent until other data is ready or the KeepAliveTime is reached. This can help
+	 * reduce the number of individual packets sent.
+	 */
+	UPROPERTY(Config)
+	bool bSendImmediateAcks = true;
 };
diff --git a/Engine/Source/Runtime/Engine/Private/NetConnection.cpp b/Engine/Source/Runtime/Engine/Private/NetConnection.cpp
--- a/Engine/Source/Runtime/Engine/Private/NetConnection.cpp
+++ b/Engine/Source/Runtime/Engine/Private/NetConnection.cpp
@@ -3329,9 +3329,12 @@ void UNetConnection::ReceivedPacket( FBitReader& Reader, bool bIsReinjectedPacke
 			++Driver->OutTotalAcks;
 		}
 
-		// We do want to let the other side know about the ack, so even if there are no other outgoing data when we tick the connection we will send an ackpacket.
-		TimeSensitive = 1;
-		++HasDirtyAcks;
+		if (Driver->IsSendingImmediateAcks())
+		{
+			// If enabled, we want to let the other side know about the ack, so even if there are no other outgoing data when we tick the connection we will send an ackpacket.
+			TimeSensitive = 1;
+			++HasDirtyAcks;
+		}
 	}
 
 	// Flush trace content collector
