SplineComponent: Remove some attribute channel functions which are not working yet

We just want to deal with float attributes while the attribute channel API is in a state of constant change. These are actually not even working right now.

diff --git a/Engine/Source/Runtime/Engine/Classes/Components/SplineComponent.h b/Engine/Source/Runtime/Engine/Classes/Components/SplineComponent.h
--- a/Engine/Source/Runtime/Engine/Classes/Components/SplineComponent.h
+++ b/Engine/Source/Runtime/Engine/Classes/Components/SplineComponent.h
@@ -880,25 +880,12 @@ public:
 	UFUNCTION(BlueprintCallable, Category = "Spline|Metadata")
 	ENGINE_API TArray<FName> GetFloatPropertyChannels() const;
 
-	/**
-	 * @return An array of all vector attribute channel names.
-	 */
-	UFUNCTION(BlueprintCallable, Category = "Spline|Metadata")
-	ENGINE_API TArray<FName> GetVectorPropertyChannels() const;
-
 	/**
 	 * Create a float attribute channel.
 	 * @param PropertyName The desired name of the new attribute channel.
 	 */
 	UFUNCTION(BlueprintCallable, Category = "Spline|Metadata")
 	ENGINE_API void CreateFloatPropertyChannel(FName PropertyName);
-
-	/**
-	 * Create a vector attribute channel.
-	 * @param PropertyName The desired name of the new attribute channel.
-	 */
-	UFUNCTION(BlueprintCallable, Category = "Spline|Metadata")
-	ENGINE_API void CreateVectorPropertyChannel(FName PropertyName);
 	
 	/**
 	 * Get a float attribute value by parameter.
@@ -909,15 +896,6 @@ public:
 	UFUNCTION(BlueprintCallable, Category = "Spline|Metadata")
 	ENGINE_API float GetFloatPropertyAtSplineInputKey(float InKey, FName PropertyName) const;
 
-	/**
-	 * Get a vector attribute value by parameter.
-	 * @param InKey The parameter to use when evaluating the specified attribute channel.
-	 * @param PropertyName The name of the attribute channel to evaluate.
-	 * @return The attribute value.
-	 */
-	UFUNCTION(BlueprintCallable, Category = "Spline|Metadata")
-	ENGINE_API FVector GetVectorPropertyAtSplineInputKey(float InKey, FName PropertyName) const;
-
 	/**
 	 * Add a float attribute value by parameter.
 	 * @param InKey The desired attribute parameter.
@@ -927,18 +905,6 @@ public:
 	 */
 	UFUNCTION(BlueprintCallable, Category = "Spline|Metadata")
 	ENGINE_API int32 SetFloatPropertyAtSplineInputKey(float InKey, float Value, FName PropertyName) const;
-
-	/**
-	 * Add a vector attribute value by parameter.
-	 * @param InKey The desired attribute parameter.
-	 * @param Value The desired attribute value.
-	 * @param PropertyName The name of the attribute channel that will contain the new attribute value.
-	 * @return The index of the new attribute.
-	 */
-	UFUNCTION(BlueprintCallable, Category = "Spline|Metadata")
-	ENGINE_API int32 SetVectorPropertyAtSplineInputKey(float InKey, const FVector& Value, FName PropertyName) const;
-	
-	// todo: vector variants of below functions
 	
 	/**
 	 * Get a float attribute parameter by index.
diff --git a/Engine/Source/Runtime/Engine/Private/Components/SplineComponent.cpp b/Engine/Source/Runtime/Engine/Private/Components/SplineComponent.cpp
--- a/Engine/Source/Runtime/Engine/Private/Components/SplineComponent.cpp
+++ b/Engine/Source/Runtime/Engine/Private/Components/SplineComponent.cpp
@@ -667,16 +667,6 @@ TArray<FName> USplineComponent::GetFloatPropertyChannels() const
 	return TArray<FName>();
 }
 
-TArray<FName> USplineComponent::GetVectorPropertyChannels() const
-{
-	if (UE::SplineComponent::GUseSplineCurves == false && Spline.SupportsAttributes())
-	{
-		return Spline.GetVectorPropertyChannels();
-	}
-
-	return TArray<FName>();
-}
-
 template<class T>
 bool CreatePropertyChannel(const USplineMetadata* Metadata, const FSpline& Spline, FName PropertyName)
 {
@@ -792,31 +782,16 @@ void USplineComponent::CreateFloatPropertyChannel(FName PropertyName)
 	CreatePropertyChannel<float>(GetSplinePointsMetadata(), Spline, PropertyName);
 }
 
-void USplineComponent::CreateVectorPropertyChannel(FName PropertyName)
-{
-	CreatePropertyChannel<FVector>(GetSplinePointsMetadata(), Spline, PropertyName);
-}
-
 float USplineComponent::GetFloatPropertyAtSplineInputKey(float InKey, FName PropertyName) const
 {
 	return GetPropertyAtSplineInputKey<float>(GetSplinePointsMetadata(), Spline, InKey, PropertyName);
 }
 
-FVector USplineComponent::GetVectorPropertyAtSplineInputKey(float InKey, FName PropertyName) const
-{
-	return GetPropertyAtSplineInputKey<FVector>(GetSplinePointsMetadata(), Spline, InKey, PropertyName);
-}
-
 int32 USplineComponent::SetFloatPropertyAtSplineInputKey(float InKey, float Value, FName PropertyName) const
 {
 	return SetPropertyAtSplineInputKey<float>(Spline, InKey, Value, PropertyName);
 }
 
-int32 USplineComponent::SetVectorPropertyAtSplineInputKey(float InKey, const FVector& Value, FName PropertyName) const
-{
-	return SetPropertyAtSplineInputKey<FVector>(Spline, InKey, Value, PropertyName);
-}
-
 float USplineComponent::GetFloatPropertyInputKeyAtIndex(int32 Index, FName PropertyName) const
 {
 	return GetInputKeyAtIndex(Spline, Index, PropertyName);
