SplineComponent: Make Spline non-transient

diff --git a/Engine/Source/Runtime/Engine/Classes/Components/SplineComponent.h b/Engine/Source/Runtime/Engine/Classes/Components/SplineComponent.h
--- a/Engine/Source/Runtime/Engine/Classes/Components/SplineComponent.h
+++ b/Engine/Source/Runtime/Engine/Classes/Components/SplineComponent.h
@@ -211,7 +211,7 @@ class USplineComponent : public UPrimitiveComponent
 protected:
 
 	// todo: replicate, but I believe we need to make the struct able to replicate in the first place.
-	UPROPERTY(Transient, DuplicateTransient, EditAnywhere, Category=Points)
+	UPROPERTY(EditAnywhere, Category=Points)
 	FSpline Spline;
 	
 public:
