Changing VT functions to take FRHICommandListBase to match the type passed to InitRHI

diff --git a/Engine/Plugins/Experimental/GPULightmass/Source/GPULightmass/Private/LightmapPreviewVirtualTexture.cpp b/Engine/Plugins/Experimental/GPULightmass/Source/GPULightmass/Private/LightmapPreviewVirtualTexture.cpp
--- a/Engine/Plugins/Experimental/GPULightmass/Source/GPULightmass/Private/LightmapPreviewVirtualTexture.cpp
+++ b/Engine/Plugins/Experimental/GPULightmass/Source/GPULightmass/Private/LightmapPreviewVirtualTexture.cpp
@@ -61,7 +61,7 @@ FLightmapPreviewVirtualTexture::FLightmapPreviewVirtualTexture(FLightmapRenderSt
 }
 
 FVTRequestPageResult FLightmapPreviewVirtualTexture::RequestPageData(
-	FRHICommandList& RHICmdList,
+	FRHICommandListBase& RHICmdList,
 	const FVirtualTextureProducerHandle& InProducerHandle,
 	uint8 LayerMask,
 	uint8 vLevel,
@@ -72,7 +72,7 @@ FVTRequestPageResult FLightmapPreviewVirtualTexture::RequestPageData(
 }
 
 IVirtualTextureFinalizer* FLightmapPreviewVirtualTexture::ProducePageData(
-	FRHICommandList& RHICmdList,
+	FRHICommandListBase& RHICmdList,
 	ERHIFeatureLevel::Type FeatureLevel,
 	EVTProducePageFlags Flags,
 	const FVirtualTextureProducerHandle& InProducerHandle,
diff --git a/Engine/Plugins/Experimental/GPULightmass/Source/GPULightmass/Private/LightmapPreviewVirtualTexture.h b/Engine/Plugins/Experimental/GPULightmass/Source/GPULightmass/Private/LightmapPreviewVirtualTexture.h
--- a/Engine/Plugins/Experimental/GPULightmass/Source/GPULightmass/Private/LightmapPreviewVirtualTexture.h
+++ b/Engine/Plugins/Experimental/GPULightmass/Source/GPULightmass/Private/LightmapPreviewVirtualTexture.h
@@ -20,7 +20,7 @@ public:
 	}
 
 	virtual FVTRequestPageResult RequestPageData(
-		FRHICommandList& RHICmdList,
+		FRHICommandListBase& RHICmdList,
 		const FVirtualTextureProducerHandle& ProducerHandle,
 		uint8 LayerMask,
 		uint8 vLevel,
@@ -29,7 +29,7 @@ public:
 	) override;
 
 	virtual IVirtualTextureFinalizer* ProducePageData(
-		FRHICommandList& RHICmdList,
+		FRHICommandListBase& RHICmdList,
 		ERHIFeatureLevel::Type FeatureLevel,
 		EVTProducePageFlags Flags,
 		const FVirtualTextureProducerHandle& InProducerHandle,
diff --git a/Engine/Source/Runtime/Engine/Private/Texture2D.cpp b/Engine/Source/Runtime/Engine/Private/Texture2D.cpp
--- a/Engine/Source/Runtime/Engine/Private/Texture2D.cpp
+++ b/Engine/Source/Runtime/Engine/Private/Texture2D.cpp
@@ -1634,7 +1634,7 @@ void FVirtualTexture2DResource::InitRHI(FRHICommandListBase& RHICmdList)
 
 	// Only create the miptails mini-texture in-editor.
 #if WITH_EDITOR
-	InitializeEditorResources(VirtualTexture);
+	InitializeEditorResources(RHICmdList, VirtualTexture);
 #endif
 
 	if (TextureRHI.IsValid())
@@ -1644,7 +1644,7 @@ void FVirtualTexture2DResource::InitRHI(FRHICommandListBase& RHICmdList)
 }
 
 #if WITH_EDITOR
-void FVirtualTexture2DResource::InitializeEditorResources(IVirtualTexture* InVirtualTexture)
+void FVirtualTexture2DResource::InitializeEditorResources(FRHICommandListBase& RHICmdList, IVirtualTexture* InVirtualTexture)
 {
 	// Create a texture resource for use when we try to access the VT as a regular texture.
 	// This may happen for debugging or when we are in a preview mode that doesn't support VT.
@@ -1682,7 +1682,6 @@ void FVirtualTexture2DResource::InitializeEditorResources(IVirtualTexture* InVir
 				.AddFlags(TexCreateFlags)
 				.SetInitActionInitializer();
 
-			FRHICommandListImmediate& RHICmdList = FRHICommandListExecutor::GetImmediateCommandList();
 			FRHITextureInitializer Initializer = RHICmdList.CreateTextureInitializer(TextureDesc);
 
 			{
@@ -1694,8 +1693,7 @@ void FVirtualTexture2DResource::InitializeEditorResources(IVirtualTexture* InVir
 		}
 		else
 		{
-			FRHICommandListImmediate& RHICommandList = FRHICommandListExecutor::GetImmediateCommandList();
-			FRDGBuilder GraphBuilder(RHICommandList);
+			FRDGBuilder GraphBuilder(RHICmdList.GetAsImmediate());
 
 			// We want to strip borders when compositing tiles since we're just laying out tiles in a regular texture.
 			// But if we have block compressed formats with border not divisible by 4 then doing this will lead to an unaligned copy. We keep the small unwanted borders in that case.
@@ -1711,8 +1709,7 @@ void FVirtualTexture2DResource::InitializeEditorResources(IVirtualTexture* InVir
 			
 			// First try reusing the current referenced texture.
 			// We do that because swapping the referenced texture can be expensive if we are doing a recreate of all textures.
-			FRHITexture* ReferencedTexture = TextureReferenceRHI->GetReferencedTexture();
-			if (ReferencedTexture != nullptr)
+			if (FRHITexture* ReferencedTexture = TextureReferenceRHI->GetReferencedTexture())
 			{
 				FRHITextureDesc const& ExistingDesc = ReferencedTexture->GetDesc();
 				if (ExistingDesc.Format == PixelFormat && ExistingDesc.GetSize() == FIntVector(MipWidthInTiles * TileSizeInPixels, MipHeightInTiles * TileSizeInPixels, 1))
@@ -1741,7 +1738,7 @@ void FVirtualTexture2DResource::InitializeEditorResources(IVirtualTexture* InVir
 				for (uint32 TileX = 0u; TileX < MipWidthInTiles; ++TileX)
 				{
 					const uint32 vAddress = FMath::MortonCode2(TileX) | (FMath::MortonCode2(TileY) << 1);
-					const FVTRequestPageResult RequestResult = InVirtualTexture->RequestPageData(RHICommandList, ProducerHandle, LayerMask, MipLevel, vAddress, EVTRequestPagePriority::High);
+					const FVTRequestPageResult RequestResult = InVirtualTexture->RequestPageData(RHICmdList, ProducerHandle, LayerMask, MipLevel, vAddress, EVTRequestPagePriority::High);
 
 					// High priority request should never be Saturated.
 					// It's possible for status to be Invalid, if requesting data from a mip level that doesn't exist for the given producer (when using sparse UDIMs).
@@ -1753,7 +1750,7 @@ void FVirtualTexture2DResource::InitializeEditorResources(IVirtualTexture* InVir
 						TargetLayer.pPageLocation = FIntVector(TileX, TileY, 0);
 
 						IVirtualTextureFinalizer* Finalizer = InVirtualTexture->ProducePageData(
-							RHICommandList,
+							RHICmdList,
 							GMaxRHIFeatureLevel,
 							bCopyUnwantedBordersForAlignment ? EVTProducePageFlags::None : EVTProducePageFlags::SkipPageBorders,
 							ProducerHandle, LayerMask, MipLevel, vAddress,
diff --git a/Engine/Source/Runtime/Engine/Private/VT/RuntimeVirtualTexture.cpp b/Engine/Source/Runtime/Engine/Private/VT/RuntimeVirtualTexture.cpp
--- a/Engine/Source/Runtime/Engine/Private/VT/RuntimeVirtualTexture.cpp
+++ b/Engine/Source/Runtime/Engine/Private/VT/RuntimeVirtualTexture.cpp
@@ -55,7 +55,7 @@ namespace
 		}
 
 		virtual FVTRequestPageResult RequestPageData(
-			FRHICommandList& RHICmdList,
+			FRHICommandListBase& RHICmdList,
 			const FVirtualTextureProducerHandle& ProducerHandle,
 			uint8 LayerMask,
 			uint8 vLevel,
@@ -67,7 +67,7 @@ namespace
 		}
 
 		virtual IVirtualTextureFinalizer* ProducePageData(
-			FRHICommandList& RHICmdList,
+			FRHICommandListBase& RHICmdList,
 			ERHIFeatureLevel::Type FeatureLevel,
 			EVTProducePageFlags Flags,
 			const FVirtualTextureProducerHandle& ProducerHandle,
diff --git a/Engine/Source/Runtime/Engine/Private/VT/UploadingVirtualTexture.cpp b/Engine/Source/Runtime/Engine/Private/VT/UploadingVirtualTexture.cpp
--- a/Engine/Source/Runtime/Engine/Private/VT/UploadingVirtualTexture.cpp
+++ b/Engine/Source/Runtime/Engine/Private/VT/UploadingVirtualTexture.cpp
@@ -124,7 +124,7 @@ uint32 FUploadingVirtualTexture::GetLocalMipBias(uint8 vLevel, uint32 vAddress)
 	return Current_vLevel - vLevel;
 }
 
-FVTRequestPageResult FUploadingVirtualTexture::RequestPageData(FRHICommandList& RHICmdList, const FVirtualTextureProducerHandle& ProducerHandle, uint8 LayerMask, uint8 vLevel, uint64 vAddress, EVTRequestPagePriority Priority)
+FVTRequestPageResult FUploadingVirtualTexture::RequestPageData(FRHICommandListBase& RHICmdList, const FVirtualTextureProducerHandle& ProducerHandle, uint8 LayerMask, uint8 vLevel, uint64 vAddress, EVTRequestPagePriority Priority)
 {
 	vLevel += FirstMipOffset;
 
@@ -132,7 +132,7 @@ FVTRequestPageResult FUploadingVirtualTexture::RequestPageData(FRHICommandList&
 	return StreamingManager->RequestTile(RHICmdList, this, ProducerHandle, LayerMask, vLevel, (uint32)vAddress, Priority);
 }
 
-IVirtualTextureFinalizer* FUploadingVirtualTexture::ProducePageData(FRHICommandList& RHICmdList,
+IVirtualTextureFinalizer* FUploadingVirtualTexture::ProducePageData(FRHICommandListBase& RHICmdList,
 	ERHIFeatureLevel::Type FeatureLevel,
 	EVTProducePageFlags Flags,
 	const FVirtualTextureProducerHandle& ProducerHandle, uint8 LayerMask, uint8 vLevel, uint64 vAddress,
diff --git a/Engine/Source/Runtime/Engine/Private/VT/UploadingVirtualTexture.h b/Engine/Source/Runtime/Engine/Private/VT/UploadingVirtualTexture.h
--- a/Engine/Source/Runtime/Engine/Private/VT/UploadingVirtualTexture.h
+++ b/Engine/Source/Runtime/Engine/Private/VT/UploadingVirtualTexture.h
@@ -79,8 +79,8 @@ public:
 	// IVirtualTexture interface
 	virtual uint32 GetLocalMipBias(uint8 vLevel, uint32 vAddress) const override;
 	virtual bool IsPageStreamed(uint8 vLevel, uint32 vAddress) const override { return true; }
-	virtual FVTRequestPageResult RequestPageData(FRHICommandList& RHICmdList, const FVirtualTextureProducerHandle& ProducerHandle, uint8 LayerMask, uint8 vLevel, uint64 vAddress, EVTRequestPagePriority Priority) override;
-	virtual IVirtualTextureFinalizer* ProducePageData(FRHICommandList& RHICmdList,
+	virtual FVTRequestPageResult RequestPageData(FRHICommandListBase& RHICmdList, const FVirtualTextureProducerHandle& ProducerHandle, uint8 LayerMask, uint8 vLevel, uint64 vAddress, EVTRequestPagePriority Priority) override;
+	virtual IVirtualTextureFinalizer* ProducePageData(FRHICommandListBase& RHICmdList,
 		ERHIFeatureLevel::Type FeatureLevel,
 		EVTProducePageFlags Flags,
 		const FVirtualTextureProducerHandle& ProducerHandle, uint8 LayerMask, uint8 vLevel, uint64 vAddress,
diff --git a/Engine/Source/Runtime/Engine/Private/VT/VirtualTextureAdapter.cpp b/Engine/Source/Runtime/Engine/Private/VT/VirtualTextureAdapter.cpp
--- a/Engine/Source/Runtime/Engine/Private/VT/VirtualTextureAdapter.cpp
+++ b/Engine/Source/Runtime/Engine/Private/VT/VirtualTextureAdapter.cpp
@@ -420,7 +420,7 @@ public:
 	}
 
 	virtual FVTRequestPageResult RequestPageData(
-		FRHICommandList& RHICmdList,
+		FRHICommandListBase& RHICmdList,
 		const FVirtualTextureProducerHandle& ProducerHandle,
 		uint8 LayerMask,
 		uint8 vLevel,
@@ -433,7 +433,7 @@ public:
 	}
 
 	virtual IVirtualTextureFinalizer* ProducePageData(
-		FRHICommandList& RHICmdList,
+		FRHICommandListBase& RHICmdList,
 		ERHIFeatureLevel::Type FeatureLevel,
 		EVTProducePageFlags Flags,
 		const FVirtualTextureProducerHandle& ProducerHandle,
diff --git a/Engine/Source/Runtime/Engine/Private/VT/VirtualTextureChunkManager.cpp b/Engine/Source/Runtime/Engine/Private/VT/VirtualTextureChunkManager.cpp
--- a/Engine/Source/Runtime/Engine/Private/VT/VirtualTextureChunkManager.cpp
+++ b/Engine/Source/Runtime/Engine/Private/VT/VirtualTextureChunkManager.cpp
@@ -61,7 +61,7 @@ void FVirtualTextureChunkStreamingManager::CancelForcedResources()
 {
 }
 
-FVTRequestPageResult FVirtualTextureChunkStreamingManager::RequestTile(FRHICommandList& RHICmdList, FUploadingVirtualTexture* VTexture, const FVirtualTextureProducerHandle& ProducerHandle, uint8 LayerMask, uint8 vLevel, uint32 vAddress, EVTRequestPagePriority Priority)
+FVTRequestPageResult FVirtualTextureChunkStreamingManager::RequestTile(FRHICommandListBase& RHICmdList, FUploadingVirtualTexture* VTexture, const FVirtualTextureProducerHandle& ProducerHandle, uint8 LayerMask, uint8 vLevel, uint32 vAddress, EVTRequestPagePriority Priority)
 {
 	SCOPE_CYCLE_COUNTER(STAT_VTP_RequestTile);
 
@@ -142,7 +142,7 @@ FVTRequestPageResult FVirtualTextureChunkStreamingManager::RequestTile(FRHIComma
 	return FVTRequestPageResult(EVTRequestPageStatus::Pending, TranscodeHandle.PackedData);
 }
 
-IVirtualTextureFinalizer* FVirtualTextureChunkStreamingManager::ProduceTile(FRHICommandList& RHICmdList, uint32 SkipBorderSize, uint8 NumLayers, uint8 LayerMask, uint64 RequestHandle, const FVTProduceTargetLayer* TargetLayers)
+IVirtualTextureFinalizer* FVirtualTextureChunkStreamingManager::ProduceTile(FRHICommandListBase& RHICmdList, uint32 SkipBorderSize, uint8 NumLayers, uint8 LayerMask, uint64 RequestHandle, const FVTProduceTargetLayer* TargetLayers)
 {
 	SCOPE_CYCLE_COUNTER(STAT_VTP_ProduceTile);
 
diff --git a/Engine/Source/Runtime/Engine/Private/VT/VirtualTextureChunkManager.h b/Engine/Source/Runtime/Engine/Private/VT/VirtualTextureChunkManager.h
--- a/Engine/Source/Runtime/Engine/Private/VT/VirtualTextureChunkManager.h
+++ b/Engine/Source/Runtime/Engine/Private/VT/VirtualTextureChunkManager.h
@@ -42,8 +42,8 @@ public:
 	virtual void NotifyLevelOffset(class ULevel* Level, const FVector& Offset) override {}
 	// End IStreamingManager interface
 
-	FVTRequestPageResult RequestTile(FRHICommandList& RHICmdList, FUploadingVirtualTexture* VTexture, const FVirtualTextureProducerHandle& ProducerHandle, uint8 LayerMask, uint8 vLevel, uint32 vAddress, EVTRequestPagePriority Priority);
-	IVirtualTextureFinalizer* ProduceTile(FRHICommandList& RHICmdList, uint32 SkipBorderSize, uint8 NumLayers, uint8 LayerMask, uint64 RequestHandle, const FVTProduceTargetLayer* TargetLayers);
+	FVTRequestPageResult RequestTile(FRHICommandListBase& RHICmdList, FUploadingVirtualTexture* VTexture, const FVirtualTextureProducerHandle& ProducerHandle, uint8 LayerMask, uint8 vLevel, uint32 vAddress, EVTRequestPagePriority Priority);
+	IVirtualTextureFinalizer* ProduceTile(FRHICommandListBase& RHICmdList, uint32 SkipBorderSize, uint8 NumLayers, uint8 LayerMask, uint64 RequestHandle, const FVTProduceTargetLayer* TargetLayers);
 	void GatherProducePageDataTasks(FVirtualTextureProducerHandle const& ProducerHandle, FGraphEventArray& InOutTasks) const;
 	void GatherProducePageDataTasks(uint64 RequestHandle, FGraphEventArray& InOutTasks) const;
 
diff --git a/Engine/Source/Runtime/Engine/Private/VT/VirtualTextureLevelRedirector.cpp b/Engine/Source/Runtime/Engine/Private/VT/VirtualTextureLevelRedirector.cpp
--- a/Engine/Source/Runtime/Engine/Private/VT/VirtualTextureLevelRedirector.cpp
+++ b/Engine/Source/Runtime/Engine/Private/VT/VirtualTextureLevelRedirector.cpp
@@ -22,7 +22,7 @@ bool FVirtualTextureLevelRedirector::IsPageStreamed(uint8 vLevel, uint32 vAddres
 }
 
 FVTRequestPageResult FVirtualTextureLevelRedirector::RequestPageData(
-	FRHICommandList& RHICmdList,
+	FRHICommandListBase& RHICmdList,
 	const FVirtualTextureProducerHandle& ProducerHandle,
 	uint8 LayerMask,
 	uint8 vLevel,
@@ -35,7 +35,7 @@ FVTRequestPageResult FVirtualTextureLevelRedirector::RequestPageData(
 }
 
 IVirtualTextureFinalizer* FVirtualTextureLevelRedirector::ProducePageData(
-	FRHICommandList& RHICmdList,
+	FRHICommandListBase& RHICmdList,
 	ERHIFeatureLevel::Type FeatureLevel,
 	EVTProducePageFlags Flags,
 	const FVirtualTextureProducerHandle& ProducerHandle,
diff --git a/Engine/Source/Runtime/Engine/Private/VT/VirtualTextureLevelRedirector.h b/Engine/Source/Runtime/Engine/Private/VT/VirtualTextureLevelRedirector.h
--- a/Engine/Source/Runtime/Engine/Private/VT/VirtualTextureLevelRedirector.h
+++ b/Engine/Source/Runtime/Engine/Private/VT/VirtualTextureLevelRedirector.h
@@ -19,7 +19,7 @@ public:
 	virtual bool IsPageStreamed(uint8 vLevel, uint32 vAddress) const override;
 
 	virtual FVTRequestPageResult RequestPageData(
-		FRHICommandList& RHICmdList,
+		FRHICommandListBase& RHICmdList,
 		const FVirtualTextureProducerHandle& ProducerHandle,
 		uint8 LayerMask,
 		uint8 vLevel,
@@ -28,7 +28,7 @@ public:
 	) override;
 
 	virtual IVirtualTextureFinalizer* ProducePageData(
-		FRHICommandList& RHICmdList,
+		FRHICommandListBase& RHICmdList,
 		ERHIFeatureLevel::Type FeatureLevel,
 		EVTProducePageFlags Flags,
 		const FVirtualTextureProducerHandle& ProducerHandle,
diff --git a/Engine/Source/Runtime/Engine/Private/VT/VirtualTextureTranscodeCache.cpp b/Engine/Source/Runtime/Engine/Private/VT/VirtualTextureTranscodeCache.cpp
--- a/Engine/Source/Runtime/Engine/Private/VT/VirtualTextureTranscodeCache.cpp
+++ b/Engine/Source/Runtime/Engine/Private/VT/VirtualTextureTranscodeCache.cpp
@@ -346,7 +346,7 @@ const FVTUploadTileHandle* FVirtualTextureTranscodeCache::AcquireTaskResult(FVTT
 }
 
 FVTTranscodeTileHandle FVirtualTextureTranscodeCache::SubmitTask(
-	FRHICommandList& RHICmdList,
+	FRHICommandListBase& RHICmdList,
 	FVirtualTextureUploadCache& InUploadCache,
 	const FVTTranscodeKey& InKey,
 	const FVirtualTextureProducerHandle& InProducerHandle,
@@ -451,7 +451,7 @@ FVTTranscodeTileHandle FVirtualTextureTranscodeCache::SubmitTask(
 	return FVTTranscodeTileHandle(TaskIndex, TaskEntry.Magic);
 }
 
-void FVirtualTextureTranscodeCache::RetireOldTasks(FRHICommandList& RHICmdList, FVirtualTextureUploadCache& InUploadCache)
+void FVirtualTextureTranscodeCache::RetireOldTasks(FRHICommandListBase& RHICmdList, FVirtualTextureUploadCache& InUploadCache)
 {
 	const uint32 CurrentFrame = GFrameNumberRenderThread;
 
diff --git a/Engine/Source/Runtime/Engine/Private/VT/VirtualTextureTranscodeCache.h b/Engine/Source/Runtime/Engine/Private/VT/VirtualTextureTranscodeCache.h
--- a/Engine/Source/Runtime/Engine/Private/VT/VirtualTextureTranscodeCache.h
+++ b/Engine/Source/Runtime/Engine/Private/VT/VirtualTextureTranscodeCache.h
@@ -74,7 +74,7 @@ public:
 	FVTTranscodeTileHandleAndStatus FindTask(const FVTTranscodeKey& InKey) const;
 
 	FVTTranscodeTileHandle SubmitTask(
-		FRHICommandList& RHICmdList,
+		FRHICommandListBase& RHICmdList,
 		FVirtualTextureUploadCache& InUploadCache,
 		const FVTTranscodeKey& InKey,
 		const FVirtualTextureProducerHandle& InProducerHandle,
@@ -88,7 +88,7 @@ public:
 	FGraphEventRef GetTaskEvent(FVTTranscodeTileHandle InHandle) const;
 	void GatherProducePageDataTasks(FVirtualTextureProducerHandle const& ProducerHandle, FGraphEventArray& InOutTasks) const;
 
-	void RetireOldTasks(FRHICommandList& RHICmdList, FVirtualTextureUploadCache& InUploadCache);
+	void RetireOldTasks(FRHICommandListBase& RHICmdList, FVirtualTextureUploadCache& InUploadCache);
 
 private:
 	enum ListType
diff --git a/Engine/Source/Runtime/Engine/Private/VT/VirtualTextureUploadCache.cpp b/Engine/Source/Runtime/Engine/Private/VT/VirtualTextureUploadCache.cpp
--- a/Engine/Source/Runtime/Engine/Private/VT/VirtualTextureUploadCache.cpp
+++ b/Engine/Source/Runtime/Engine/Private/VT/VirtualTextureUploadCache.cpp
@@ -39,7 +39,7 @@ static TAutoConsoleVariable<int32> CVarMaxUploadRequests(
 	ECVF_RenderThreadSafe);
 
 
-uint32 FVTUploadTileAllocator::Allocate(FRHICommandList& RHICmdList, EVTUploadType InUploadType, EPixelFormat InFormat, uint32 InTileSize)
+uint32 FVTUploadTileAllocator::Allocate(FRHICommandListBase& RHICmdList, EVTUploadType InUploadType, EPixelFormat InFormat, uint32 InTileSize)
 {
 	// Find matching FormatBuffer.
 	const FPixelFormatInfo& FormatInfo = GPixelFormats[InFormat];
@@ -110,7 +110,7 @@ uint32 FVTUploadTileAllocator::Allocate(FRHICommandList& RHICmdList, EVTUploadTy
 	return Handle.PackedValue;
 }
 
-void FVTUploadTileAllocator::Free(FRHICommandList& RHICmdList, uint32 InHandle)
+void FVTUploadTileAllocator::Free(FRHICommandListBase& RHICmdList, uint32 InHandle)
 {
 	FHandle Handle;
 	Handle.PackedValue = InHandle;
@@ -165,7 +165,7 @@ FVTUploadTileAllocator::FStagingBuffer::~FStagingBuffer()
 	Release(nullptr);
 }
 
-void FVTUploadTileAllocator::FStagingBuffer::Init(FRHICommandList& RHICmdList, EVTUploadType InUploadType, uint32 InBufferStrideBytes, uint32 InTileSizeBytes)
+void FVTUploadTileAllocator::FStagingBuffer::Init(FRHICommandListBase& RHICmdList, EVTUploadType InUploadType, uint32 InBufferStrideBytes, uint32 InTileSizeBytes)
 {
 	TileSize = InTileSizeBytes;
 	TileSizeAligned = Align(InTileSizeBytes, 128u);
@@ -206,7 +206,7 @@ void FVTUploadTileAllocator::FStagingBuffer::Init(FRHICommandList& RHICmdList, E
 	}
 }
 
-void FVTUploadTileAllocator::FStagingBuffer::Release(FRHICommandList* RHICmdList)
+void FVTUploadTileAllocator::FStagingBuffer::Release(FRHICommandListBase* RHICmdList)
 {
 	const uint32 BufferSize = TileSizeAligned * NumTiles;
 
@@ -535,7 +535,7 @@ void FVirtualTextureUploadCache::ReleaseRHI()
 	Pools.Empty();
 }
 
-FVTUploadTileHandle FVirtualTextureUploadCache::PrepareTileForUpload(FRHICommandList& RHICmdList, FVTUploadTileBuffer& OutBuffer, EPixelFormat InFormat, uint32 InTileSize)
+FVTUploadTileHandle FVirtualTextureUploadCache::PrepareTileForUpload(FRHICommandListBase& RHICmdList, FVTUploadTileBuffer& OutBuffer, EPixelFormat InFormat, uint32 InTileSize)
 {
 	SCOPE_CYCLE_COUNTER(STAT_VTP_StageTile)
 
@@ -553,7 +553,7 @@ FVTUploadTileHandle FVirtualTextureUploadCache::PrepareTileForUpload(FRHICommand
 	return FVTUploadTileHandle(Index);
 }
 
-void FVirtualTextureUploadCache::SubmitTile(FRHICommandList& RHICmdList, const FVTUploadTileHandle& InHandle, const FVTProduceTargetLayer& Target, int InSkipBorderSize)
+void FVirtualTextureUploadCache::SubmitTile(FRHICommandListBase& RHICmdList, const FVTUploadTileHandle& InHandle, const FVTProduceTargetLayer& Target, int InSkipBorderSize)
 {
 	// Get entry and remove from pending uploads.
 	check(PendingUpload.IsValidIndex(InHandle.Index));
@@ -574,7 +574,7 @@ void FVirtualTextureUploadCache::SubmitTile(FRHICommandList& RHICmdList, const F
 	PoolEntry.PendingSubmit.Emplace(Entry);
 }
 
-void FVirtualTextureUploadCache::CancelTile(FRHICommandList& RHICmdList, const FVTUploadTileHandle& InHandle)
+void FVirtualTextureUploadCache::CancelTile(FRHICommandListBase& RHICmdList, const FVTUploadTileHandle& InHandle)
 {
 	check(PendingUpload.IsValidIndex(InHandle.Index));
 	FTileEntry& Entry = PendingUpload[InHandle.Index];
@@ -582,7 +582,7 @@ void FVirtualTextureUploadCache::CancelTile(FRHICommandList& RHICmdList, const F
 	PendingUpload.RemoveAt(InHandle.Index);
 }
 
-void FVirtualTextureUploadCache::UpdateFreeList(FRHICommandList& RHICmdList, bool bForceFreeAll)
+void FVirtualTextureUploadCache::UpdateFreeList(FRHICommandListBase& RHICmdList, bool bForceFreeAll)
 {
 	// Iterate tiles pending release and free them once we can be sure they have been used.
 	const uint32 CurrentFrame = GFrameNumberRenderThread;
diff --git a/Engine/Source/Runtime/Engine/Private/VT/VirtualTextureUploadCache.h b/Engine/Source/Runtime/Engine/Private/VT/VirtualTextureUploadCache.h
--- a/Engine/Source/Runtime/Engine/Private/VT/VirtualTextureUploadCache.h
+++ b/Engine/Source/Runtime/Engine/Private/VT/VirtualTextureUploadCache.h
@@ -65,9 +65,9 @@ class FVTUploadTileAllocator
 {
 public:
 	/** Allocate a tile. Sometimes does an allocation of the backing CPU/GPU block of memory. */
-	uint32 Allocate(FRHICommandList& RHICmdList, EVTUploadType InUploadType, EPixelFormat InFormat, uint32 InTileSize);
+	uint32 Allocate(FRHICommandListBase& RHICmdList, EVTUploadType InUploadType, EPixelFormat InFormat, uint32 InTileSize);
 	/** Free a tile. Sometimes does a free of the backing CPU/GPU block of memory. */
-	void Free(FRHICommandList& RHICmdList, uint32 InHandle);
+	void Free(FRHICommandListBase& RHICmdList, uint32 InHandle);
 
 	/** Get upload buffer description from handle. */
 	FVTUploadTileBuffer GetBufferFromHandle(uint32 InHandle) const;
@@ -99,11 +99,12 @@ private:
 	{
 		~FStagingBuffer();
 
-		void Init(FRHICommandList& RHICmdList, EVTUploadType InUploadType, uint32 InBufferStrideBytes, uint32 InTileSizeBytes);
-		void Release(FRHICommandList* RHICmdList);
+		void Init(FRHICommandListBase& RHICmdList, EVTUploadType InUploadType, uint32 InBufferStrideBytes, uint32 InTileSizeBytes);
+		void Release(FRHICommandListBase* RHICmdList);
 
 		/** GPU buffer if used on platform. */
 		TRefCountPtr<FRHIBuffer> RHIBuffer;
+
 		/** Memory pointer to locked GPU buffer if used on platform, or to allocated CPU heap memory if not. */
 		void* Memory = nullptr;
 		uint32 TileSize = 0u;
@@ -152,21 +153,21 @@ public:
 	//~ End FRenderResource Interface.
 
 	/** Get a staging upload buffer for streaming texture data into. */
-	FVTUploadTileHandle PrepareTileForUpload(FRHICommandList& RHICmdList, FVTUploadTileBuffer& OutBuffer, EPixelFormat InFormat, uint32 InTileSize);
+	FVTUploadTileHandle PrepareTileForUpload(FRHICommandListBase& RHICmdList, FVTUploadTileBuffer& OutBuffer, EPixelFormat InFormat, uint32 InTileSize);
 	/** 
 	 * Mark streamed upload data ready for upload to to the physical virtual texture.
 	 * Depending on the platform the upload might happen here, or be deferred to the Finalize() call.
 	 */
-	void SubmitTile(FRHICommandList& RHICmdList, const FVTUploadTileHandle& InHandle, const FVTProduceTargetLayer& Target, int InSkipBorderSize);
+	void SubmitTile(FRHICommandListBase& RHICmdList, const FVTUploadTileHandle& InHandle, const FVTProduceTargetLayer& Target, int InSkipBorderSize);
 	/** Cancel a tile that was already in flight. */
-	void CancelTile(FRHICommandList& RHICmdList, const FVTUploadTileHandle& InHandle);
+	void CancelTile(FRHICommandListBase& RHICmdList, const FVTUploadTileHandle& InHandle);
 
 	//~ Begin IVirtualTextureFinalizer Interface.
 	virtual void Finalize(FRDGBuilder& GraphBuilder) override;
 	//~ End IVirtualTextureFinalizer Interface.
 
 	/** Call on a tick to recycle submitted staging buffers. */
-	void UpdateFreeList(FRHICommandList& RHICmdList, bool bForceFreeAll = false);
+	void UpdateFreeList(FRHICommandListBase& RHICmdList, bool bForceFreeAll = false);
 
 	/** Returns true if underlying allocator within the budget set by r.VT.MaxUploadMemory.  */
 	uint32 IsInMemoryBudget() const;
diff --git a/Engine/Source/Runtime/Engine/Public/TextureResource.h b/Engine/Source/Runtime/Engine/Public/TextureResource.h
--- a/Engine/Source/Runtime/Engine/Public/TextureResource.h
+++ b/Engine/Source/Runtime/Engine/Public/TextureResource.h
@@ -187,7 +187,7 @@ public:
 	virtual const FVirtualTexture2DResource* GetVirtualTexture2DResource() const { return this; }
 
 #if WITH_EDITOR
-	ENGINE_API virtual void InitializeEditorResources(class IVirtualTexture* InVirtualTexture);
+	ENGINE_API virtual void InitializeEditorResources(FRHICommandListBase& RHICmdList, IVirtualTexture* InVirtualTexture);
 #endif
 
 	ENGINE_API virtual uint32 GetSizeX() const override;
diff --git a/Engine/Source/Runtime/RenderCore/Public/VirtualTexturing.h b/Engine/Source/Runtime/RenderCore/Public/VirtualTexturing.h
--- a/Engine/Source/Runtime/RenderCore/Public/VirtualTexturing.h
+++ b/Engine/Source/Runtime/RenderCore/Public/VirtualTexturing.h
@@ -19,8 +19,7 @@
 #include "UObject/NameTypes.h"
 
 class FRDGBuilder;
-class FRHICommandList;
-class FRHICommandListImmediate;
+class FRHICommandListBase;
 class FRHIShaderResourceView;
 class FRHITexture;
 class FRHIUnorderedAccessView;
@@ -315,7 +314,7 @@ public:
 	* @param Priority Priority of the request, used to drive async IO/task priority needed to generate data for request
 	* @return FVTRequestPageResult describing the availability of the request
 	*/
-	virtual FVTRequestPageResult RequestPageData(FRHICommandList& RHICmdList, const FVirtualTextureProducerHandle& ProducerHandle, uint8 LayerMask, uint8 vLevel, uint64 vAddress, EVTRequestPagePriority Priority) = 0;
+	virtual FVTRequestPageResult RequestPageData(FRHICommandListBase& RHICmdList, const FVirtualTextureProducerHandle& ProducerHandle, uint8 LayerMask, uint8 vLevel, uint64 vAddress, EVTRequestPagePriority Priority) = 0;
 
 	/**
 	* Upload page data to the cache, data must have been previously requested, and reported either 'Available' or 'Pending'
@@ -333,7 +332,8 @@ public:
 	* @param TargetLayers Array of 'FVTProduceTargetLayer' structs, gives location where each layer should write data
 	* @return a 'IVirtualTextureFinalizer' which must be finalized to complete the operation
 	*/
-	virtual IVirtualTextureFinalizer* ProducePageData(FRHICommandList& RHICmdList,
+	virtual IVirtualTextureFinalizer* ProducePageData(
+		FRHICommandListBase& RHICmdList,
 		ERHIFeatureLevel::Type FeatureLevel,
 		EVTProducePageFlags Flags,
 		const FVirtualTextureProducerHandle& ProducerHandle, uint8 LayerMask, uint8 vLevel, uint64 vAddress,
diff --git a/Engine/Source/Runtime/Renderer/Private/MaterialCache/MaterialCacheVirtualTexture.cpp b/Engine/Source/Runtime/Renderer/Private/MaterialCache/MaterialCacheVirtualTexture.cpp
--- a/Engine/Source/Runtime/Renderer/Private/MaterialCache/MaterialCacheVirtualTexture.cpp
+++ b/Engine/Source/Runtime/Renderer/Private/MaterialCache/MaterialCacheVirtualTexture.cpp
@@ -49,7 +49,7 @@ public: /** IVirtualTexture */
 	}
 
 	virtual FVTRequestPageResult RequestPageData(
-		FRHICommandList& RHICmdList,
+		FRHICommandListBase& RHICmdList,
 		const FVirtualTextureProducerHandle& ProducerHandle,
 		uint8 LayerMask,
 		uint8 vLevel,
@@ -88,7 +88,7 @@ public: /** IVirtualTexture */
 	}
 
 	virtual IVirtualTextureFinalizer* ProducePageData(
-		FRHICommandList& RHICmdList,
+		FRHICommandListBase& RHICmdList,
 		ERHIFeatureLevel::Type FeatureLevel,
 		EVTProducePageFlags Flags,
 		const FVirtualTextureProducerHandle& ProducerHandle,
diff --git a/Engine/Source/Runtime/Renderer/Private/VT/AdaptiveVirtualTexture.cpp b/Engine/Source/Runtime/Renderer/Private/VT/AdaptiveVirtualTexture.cpp
--- a/Engine/Source/Runtime/Renderer/Private/VT/AdaptiveVirtualTexture.cpp
+++ b/Engine/Source/Runtime/Renderer/Private/VT/AdaptiveVirtualTexture.cpp
@@ -82,7 +82,7 @@ public:
 	}
 
 	virtual FVTRequestPageResult RequestPageData(
-		FRHICommandList& RHICmdList,
+		FRHICommandListBase& RHICmdList,
 		const FVirtualTextureProducerHandle& ProducerHandle,
 		uint8 LayerMask,
 		uint8 vLevel,
@@ -99,7 +99,7 @@ public:
 	}
 
 	virtual IVirtualTextureFinalizer* ProducePageData(
-		FRHICommandList& RHICmdList,
+		FRHICommandListBase& RHICmdList,
 		ERHIFeatureLevel::Type FeatureLevel,
 		EVTProducePageFlags Flags,
 		const FVirtualTextureProducerHandle& ProducerHandle,
diff --git a/Engine/Source/Runtime/Renderer/Private/VT/RuntimeVirtualTextureProducer.cpp b/Engine/Source/Runtime/Renderer/Private/VT/RuntimeVirtualTextureProducer.cpp
--- a/Engine/Source/Runtime/Renderer/Private/VT/RuntimeVirtualTextureProducer.cpp
+++ b/Engine/Source/Runtime/Renderer/Private/VT/RuntimeVirtualTextureProducer.cpp
@@ -156,7 +156,7 @@ FRuntimeVirtualTextureProducer::FRuntimeVirtualTextureProducer(
 }
 
 FVTRequestPageResult FRuntimeVirtualTextureProducer::RequestPageData(
-	FRHICommandList& RHICmdList,
+	FRHICommandListBase& RHICmdList,
 	const FVirtualTextureProducerHandle& ProducerHandle,
 	uint8 LayerMask,
 	uint8 vLevel,
@@ -174,7 +174,7 @@ FVTRequestPageResult FRuntimeVirtualTextureProducer::RequestPageData(
 }
 
 IVirtualTextureFinalizer* FRuntimeVirtualTextureProducer::ProducePageData(
-	FRHICommandList& RHICmdList,
+	FRHICommandListBase& RHICmdList,
 	ERHIFeatureLevel::Type FeatureLevel,
 	EVTProducePageFlags Flags,
 	const FVirtualTextureProducerHandle& ProducerHandle,
diff --git a/Engine/Source/Runtime/Renderer/Private/VT/RuntimeVirtualTextureProducer.h b/Engine/Source/Runtime/Renderer/Private/VT/RuntimeVirtualTextureProducer.h
--- a/Engine/Source/Runtime/Renderer/Private/VT/RuntimeVirtualTextureProducer.h
+++ b/Engine/Source/Runtime/Renderer/Private/VT/RuntimeVirtualTextureProducer.h
@@ -103,7 +103,7 @@ public:
 	}
 
 	virtual FVTRequestPageResult RequestPageData(
-		FRHICommandList& RHICmdList,
+		FRHICommandListBase& RHICmdList,
 		const FVirtualTextureProducerHandle& ProducerHandle,
 		uint8 LayerMask,
 		uint8 vLevel,
@@ -112,7 +112,7 @@ public:
 	) override;
 
 	virtual IVirtualTextureFinalizer* ProducePageData(
-		FRHICommandList& RHICmdList,
+		FRHICommandListBase& RHICmdList,
 		ERHIFeatureLevel::Type FeatureLevel,
 		EVTProducePageFlags Flags,
 		const FVirtualTextureProducerHandle& ProducerHandle,
