Post Process Volume:  previous CL 42396889 sort logic was buggy.  Fixed.

#rnx

diff --git a/Engine/Source/Runtime/Engine/Private/World.cpp b/Engine/Source/Runtime/Engine/Private/World.cpp
--- a/Engine/Source/Runtime/Engine/Private/World.cpp
+++ b/Engine/Source/Runtime/Engine/Private/World.cpp
@@ -10088,6 +10088,10 @@ void UWorld::InsertPostProcessVolume(IInterface_PostProcessVolume* InVolume)
 		{
 			break;
 		}
+		else if (TargetProperties.Priority > CurrentProperties.Priority)
+		{
+			continue;
+		}
 		if (TargetProperties.bIsUnbound && !CurrentProperties.bIsUnbound)
 		{
 			break;
