Fixed skeletal gpu object crash on render resource destruction
- CL 41894456 "attempted" to solve this, but unfortunately checked the render resource list index on the wrong thread

diff --git a/Engine/Source/Runtime/Engine/Private/SkeletalRenderGPUSkin.cpp b/Engine/Source/Runtime/Engine/Private/SkeletalRenderGPUSkin.cpp
--- a/Engine/Source/Runtime/Engine/Private/SkeletalRenderGPUSkin.cpp
+++ b/Engine/Source/Runtime/Engine/Private/SkeletalRenderGPUSkin.cpp
@@ -2111,8 +2111,7 @@ void UpdateVertexFactory(
 	{
 		FGPUBaseSkinVertexFactory* VertexFactory = FactoryPtr.Get();
 
-		// Do not recreate the factory if it's been released, given the loose scheduling this may result in dangling factories
-		if (VertexFactory != nullptr && VertexFactory->IsInitialized())
+		if (VertexFactory != nullptr)
 		{
 			// Setup the update data for enqueue
 			FDynamicUpdateVertexFactoryData VertexUpdateData(VertexFactory, InVertexBuffers);
@@ -2121,6 +2120,12 @@ void UpdateVertexFactory(
 			ENQUEUE_RENDER_COMMAND(UpdateGPUSkinVertexFactory)(UE::RenderCommandPipe::SkeletalMesh,
 				[VertexUpdateData](FRHICommandList& RHICmdList)
 			{
+				// Do not recreate the factory if it's been released, given the loose scheduling this may result in dangling factories
+				if (!VertexUpdateData.VertexFactory->IsInitialized())
+				{
+					return;
+				}
+
 				// Use the cloth data type for both variants since the base version will just ignore the cloth parts.
 				FGPUSkinAPEXClothDataType Data;
 				InitGPUSkinVertexFactoryComponents(&Data, VertexUpdateData.VertexBuffers, VertexUpdateData.VertexFactory);
