Iris
* Fix crash when ReinitPollFrequency is called with the NetEntityFactory enabled

#rnx

diff --git a/Engine/Source/Runtime/Engine/Private/Net/Iris/ReplicationSystem/EngineReplicationBridge.cpp b/Engine/Source/Runtime/Engine/Private/Net/Iris/ReplicationSystem/EngineReplicationBridge.cpp
--- a/Engine/Source/Runtime/Engine/Private/Net/Iris/ReplicationSystem/EngineReplicationBridge.cpp
+++ b/Engine/Source/Runtime/Engine/Private/Net/Iris/ReplicationSystem/EngineReplicationBridge.cpp
@@ -714,10 +714,14 @@ bool UEngineReplicationBridge::CanCreateDestructionInfo() const
 
 float UEngineReplicationBridge::GetPollFrequencyOfRootObject(const UObject* ReplicatedObject) const
 {
-	const AActor* ReplicatedActor = CastChecked<AActor>(ReplicatedObject);
-	float PollFrequency = ReplicatedActor->GetNetUpdateFrequency();
-	GetClassPollFrequency(ReplicatedActor->GetClass(), PollFrequency);
-	return PollFrequency;
+	if (const AActor* ReplicatedActor = Cast<AActor>(ReplicatedObject))
+	{
+		float PollFrequency = ReplicatedActor->GetNetUpdateFrequency();
+		GetClassPollFrequency(ReplicatedActor->GetClass(), PollFrequency);
+		return PollFrequency;
+	}
+	//TODO: Ask NetObjectFactory for this value instead
+	return 100.0f;
 }
 
 bool UEngineReplicationBridge::IsClassReplicatedByDefault(const UClass* Class) const
