[SimpleStreamableAssetManager] Made CVar s.StreamableAssets.UseSimpleStreamableAssetManager readonly since it can't change dynamically

diff --git a/Engine/Source/Runtime/Engine/Private/Streaming/SimpleStreamableAssetManager.cpp b/Engine/Source/Runtime/Engine/Private/Streaming/SimpleStreamableAssetManager.cpp
--- a/Engine/Source/Runtime/Engine/Private/Streaming/SimpleStreamableAssetManager.cpp
+++ b/Engine/Source/Runtime/Engine/Private/Streaming/SimpleStreamableAssetManager.cpp
@@ -32,7 +32,7 @@ FSimpleStreamableAssetManager::GUseSimpleStreamableAssetManager,
 TEXT("Whether to use FSimpleStreamableAssetManager.\n")
 TEXT("If 0 (current default), StreamingAsset works with LevelStreamingManager by collecting UPrimitiveComponents mostly operating on GT.\n")
 TEXT("If 1, The FSimpleStreamableAssetManager is Enabled and works by integrating with SceneProxy that is responsible for feeding the system."),
-ECVF_SetByGameSetting
+ECVF_SetByGameSetting | ECVF_ReadOnly
 );
 
 
