Texture.cpp : fix previous attempt at NeverStream change checking CanEverRender for cooker

change the NeverStream logic so the bool reflects user desire, not the pow2 geometry of the texture
there is still a constraint that nonpow2 textures can't stream
detect that in IsPossibleToStream() at runtime

diff --git a/Engine/Source/Runtime/Engine/Private/Texture.cpp b/Engine/Source/Runtime/Engine/Private/Texture.cpp
--- a/Engine/Source/Runtime/Engine/Private/Texture.cpp
+++ b/Engine/Source/Runtime/Engine/Private/Texture.cpp
@@ -318,13 +318,15 @@ void UTexture::ReleaseResource()
 		CachedSRRState.Clear();
 
 		// Free the resource.
-		ENQUEUE_RENDER_COMMAND(DeleteResource)([this, ToDelete = PrivateResource](FRHICommandListImmediate& RHICmdList)
+		FTextureResource* ToDelete = PrivateResource;
+		PrivateResource = nullptr;
+
+		ENQUEUE_RENDER_COMMAND(DeleteResource)([this, ToDelete](FRHICommandListImmediate& RHICmdList)
 		{
 			PrivateResourceRenderThread = nullptr;
 			ToDelete->ReleaseResource();
 			delete ToDelete;
 		});
-		PrivateResource = nullptr;
 	}
 }
 
@@ -692,23 +694,11 @@ void UTexture::ValidateSettingsAfterImportOrEdit(bool * pRequiresNotifyMaterials
 			// longlat cube always generates pow2 output
 			bIsPowerOfTwo = true;
 		}
+		// ?? bug ?? bIsPowerOfTwo does not consider ResizeDuringBuildX
 
-		// Downscale can violate IsPow2, but it only acts when NoMipMaps, so it's moot
-
-		if ( ! bIsPowerOfTwo )
-		{
-			// streaming only supports power of 2 mips
-			// due to failure to compensate for the GPU row pitch
-			// it only works for mips that naturally have the required 256 pitch
-			// so mip levels >= 256 and power of 2 only
-			// (this used to be in Texture2D.cpp)
-			// see WarnRequiresTightPackedMip
-			// there are other issues with streaming non-pow2
-			//  all streamable levels must be valid textures, so block-of-4 alignment for BCN
-			//  is easiest to guarantee if the source is pow2
-			NeverStream = true;	
-		}
-	
+		// note: NeverStream used to be turned on here if !bIsPowerOfTwo
+		//	now we leave it off to reflect artist intent, and rely on it being set at runtime by IsPossibleToStream
+		
 		const int32 LargerSourceDimension = FMath::Max( Source.GetSizeX() , Source.GetSizeY() );
 		bool bLargeTextureMustBeVT = LargerSourceDimension > GetMaximumDimensionOfNonVT();
 
@@ -1817,11 +1807,53 @@ const TArray<UAssetUserData*>* UTexture::GetAssetUserDataArray() const
 
 bool UTexture::IsPossibleToStream() const
 {
+	// user flags requesting non-streaming :
 	if ( NeverStream || LODGroup == TEXTUREGROUP_UI )
 	{
 		return false;
 	}
 
+	// we can only stream pow2 textures at the moment, check for that :
+
+	if ( FApp::CanEverRender() ) // don't do this branch during Cook
+	{
+		// previously NeverStream was always set in the Editor for nonpow2 textures
+		// now it is only set if user wants it
+		// but nonpow2 textures must still not stream, so detect it here on the cooked data :
+		FTexturePlatformData** PtrPlatformData = const_cast<UTexture *>(this)->GetRunningPlatformData();
+		if ( PtrPlatformData && *PtrPlatformData )
+		{
+			FTexturePlatformData * PlatformData = *PtrPlatformData;
+
+			bool bIsAsyncWorkComplete = true;
+			#if WITH_EDITOR
+			bIsAsyncWorkComplete = PlatformData->IsAsyncWorkComplete() && ! IsDefaultTexture();
+			#endif
+
+			// PlatformData->Mips.Num() == 0 during Cook
+			if( bIsAsyncWorkComplete && PlatformData->Mips.Num() > 0 )
+			{
+				if ( PlatformData->Mips.Num() == 1 )
+				{
+					return false;
+				}
+
+				uint32 SizeX = PlatformData->SizeX;
+				uint32 SizeY = PlatformData->SizeY;
+				uint32 SizeZ = PlatformData->GetNumSlices(); // == 6 for cubes
+				if ( !FMath::IsPowerOfTwo(SizeX) || !FMath::IsPowerOfTwo(SizeY) )
+				{
+					return false;
+				}
+				if ( GetTextureClass() == ETextureClass::Volume && !FMath::IsPowerOfTwo(SizeZ) )
+				{
+					// Arrays can stream with no constraint on SizeZ
+					return false;
+				}
+			}
+		}
+	}
+
 	#if WITH_EDITORONLY_DATA
 	if (!GetOutermost()->HasAnyPackageFlags(PKG_FilterEditorOnly))
 	{
@@ -1837,14 +1869,9 @@ bool UTexture::IsPossibleToStream() const
 	
 	if ( Source.IsValid() )
 	{
-		// should have set NeverStream for nonpow2
-		// note: this is not the case for all old textures
-		// ValidateSettingsAfterImportOrEdit makes sure this is true
-		//check( Source.IsPowerOfTwo() || PowerOfTwoMode != ETexturePowerOfTwoSetting::None )
-		
-		// duplicate the checks done for NeverStream :
+		// turn off streaming for non-pow2 sources :
 
-		// IsPowerOfTwo only checks XY :
+		// AreAllBlocksPowerOfTwo only checks XY :
 		bool bIsPowerOfTwo = Source.AreAllBlocksPowerOfTwo();
 		if ( ! FMath::IsPowerOfTwo(Source.GetVolumeSizeZ()) )
 		{
@@ -1861,7 +1888,6 @@ bool UTexture::IsPossibleToStream() const
 		}
 		if ( ! bIsPowerOfTwo )
 		{
-			// NeverStream should have been set
 			return false;
 		}
 	}
