FNewSpline: Fix attribute append/prepend

diff --git a/Engine/Source/Runtime/Engine/Private/Curves/Spline.cpp b/Engine/Source/Runtime/Engine/Private/Curves/Spline.cpp
--- a/Engine/Source/Runtime/Engine/Private/Curves/Spline.cpp
+++ b/Engine/Source/Runtime/Engine/Private/Curves/Spline.cpp
@@ -333,15 +333,29 @@ int32 FNewSpline::AddAttributeValue(FName Name, const AttrType& Value, float Par
 		// append case
 		if (ParentSpaceParameter > MappedChildSpace.Max)
 		{
+			// It is important to compute ChildSpaceParameter before the AppendPoint call, otherwise it will not be correct.
+			float ChildSpaceParameter = MapParameterToChildSpace(Name, ParentSpaceParameter);
 			Channel->AppendPoint(ControlPoint);
-			return SetAttributeParameter<AttrType>(Name, Channel->GetNumPoints() - 1, Parameter);
+
+			// By growing the child space and the mapped parent range proportionally, we keep the internal points stable in parent space.
+			Channel->SetParameter(Channel->GetNumPoints() - 1, ChildSpaceParameter);
+			SetAttributeChannelRange(Name, FInterval1f(MappedChildSpace.Min, ParentSpaceParameter), EMappingRangeSpace::Parent);
+			
+			return Channel->GetNumPoints() - 1;
 		}
 
 		// prepend case
 		if (ParentSpaceParameter < MappedChildSpace.Min)
 		{
+			// It is important to compute ChildSpaceParameter before the AppendPoint call, otherwise it will not be correct.
+			float ChildSpaceParameter = MapParameterToChildSpace(Name, ParentSpaceParameter);
 			Channel->PrependPoint(ControlPoint);
-			return SetAttributeParameter<AttrType>(Name, 0, Parameter);
+
+			// By growing the child space and the mapped parent range proportionally, we keep the internal points stable in parent space.
+			Channel->SetParameter(0, ChildSpaceParameter);
+			SetAttributeChannelRange(Name, FInterval1f(ParentSpaceParameter, MappedChildSpace.Max), EMappingRangeSpace::Parent);
+			
+			return 0;
 		}
 		
 		float ChildSpaceParameter = MapParameterToChildSpace(Name, ParentSpaceParameter);
