Fix SoundModule spam when plugin is disabled.

diff --git a/Engine/Source/Runtime/Engine/Private/SoundModulationDestination.cpp b/Engine/Source/Runtime/Engine/Private/SoundModulationDestination.cpp
--- a/Engine/Source/Runtime/Engine/Private/SoundModulationDestination.cpp
+++ b/Engine/Source/Runtime/Engine/Private/SoundModulationDestination.cpp
@@ -4,6 +4,7 @@
 
 #include "Algo/AnyOf.h"
 #include "AudioDevice.h"
+#include "AudioPluginUtilities.h"
 #include "DSP/FloatArrayMath.h"
 
 #include UE_INLINE_GENERATED_CPP_BY_NAME(SoundModulationDestination)
@@ -40,6 +41,15 @@ PRAGMA_ENABLE_DEPRECATION_WARNINGS
 }
 #endif // WITH_EDITORONLY_DATA
 
+namespace SoundModulationDestinationPrivate
+{
+	static bool IsPluginEnabled()
+	{
+		static bool bEnabled = AudioPluginUtilities::GetDesiredModulationPlugin() != nullptr;
+		return bEnabled;
+	}
+}
+
 
 FSoundModulationDefaultRoutingSettings::FSoundModulationDefaultRoutingSettings()
 	: FSoundModulationDefaultSettings()
@@ -137,18 +147,7 @@ namespace Audio
 		}
 		else
 		{
-			if (FAudioDeviceManager* AudioDeviceManager = FAudioDeviceManager::Get())
-			{
-				if (const FAudioDevice* AudioDevice = AudioDeviceManager->GetAudioDeviceRaw(InDeviceId))
-				{
-					if (AudioDevice->IsModulationPluginEnabled())
-					{
-						// Even with Modulation enabled, we could not find the parameter, so likely user error
-						UE_LOG(LogAudio, Warning, TEXT("Modulation Parameter %s not found while Initializing Modulation Destination. Using default parameter instead."), *InParameterName.ToString());
-					}
-				}
-			}
-			
+			UE_CLOG(SoundModulationDestinationPrivate::IsPluginEnabled(),LogAudio, Warning, TEXT("Modulation Parameter %s not found while Initializing Modulation Destination. Using default parameter instead."), *InParameterName.ToString());
 			DestinationData->Parameter = Audio::GetDefaultModulationParameter();
 		}
 
