Added a new property ForwardAxis to the USkeletalMesh to indicate the forward direction axis
Updated the material's boolean to indicate the material is applied to a skeletal mesh

#rnx

diff --git a/Engine/Source/Editor/UnrealEd/Private/Fbx/FbxSkeletalMeshImport.cpp b/Engine/Source/Editor/UnrealEd/Private/Fbx/FbxSkeletalMeshImport.cpp
--- a/Engine/Source/Editor/UnrealEd/Private/Fbx/FbxSkeletalMeshImport.cpp
+++ b/Engine/Source/Editor/UnrealEd/Private/Fbx/FbxSkeletalMeshImport.cpp
@@ -20,6 +20,7 @@
 #include "UObject/Object.h"
 #include "Misc/PackageName.h"
 #include "Animation/Skeleton.h"
+#include "Materials/Material.h"
 #include "Engine/SkeletalMesh.h"
 #include "Engine/SkeletalMeshLODSettings.h"
 #include "Engine/SkinnedAssetCommon.h"
@@ -1441,6 +1442,15 @@ bool UnFbx::FFbxImporter::FillSkeletalMeshImportData(TArray<FbxNode*>& NodeArray
 
 	CleanUpUnusedMaterials(*SkelMeshImportDataPtr);
 
+	for (SkeletalMeshImportData::FMaterial& MaterialInfo : SkelMeshImportDataPtr->Materials)
+	{
+		if (UMaterial* Material = Cast<UMaterial>(MaterialInfo.Material))
+		{
+			Material->bUsedWithSkeletalMesh = true;
+		}
+	}
+
+
 	if (LastImportedMaterialNames.Num() > 0)
 	{
 		SetMaterialOrderByName(*SkelMeshImportDataPtr, LastImportedMaterialNames);
@@ -2358,6 +2368,9 @@ void UnFbx::FFbxImporter::UpdateSkeletalMeshImportData(USkeletalMesh *SkeletalMe
 				}
 			}
 		}
+
+		// Note: For 36.00, the forward axis is either Y or X. A more complete solution will be implemented later.
+		SkeletalMesh->SetForwardAxis(ImportData->bConvertScene ? (ImportData->bForceFrontXAxis ? EAxis::X : EAxis::Y) : EAxis::Y);
 	}
 }
 
diff --git a/Engine/Source/Runtime/Engine/Classes/Engine/SkeletalMesh.h b/Engine/Source/Runtime/Engine/Classes/Engine/SkeletalMesh.h
--- a/Engine/Source/Runtime/Engine/Classes/Engine/SkeletalMesh.h
+++ b/Engine/Source/Runtime/Engine/Classes/Engine/SkeletalMesh.h
@@ -3008,6 +3008,14 @@ protected:
 	UPROPERTY(EditAnywhere, BlueprintReadOnly, Category=Rendering)
 	float OverlayMaterialMaxDrawDistance;
 
+	/**
+	 * Axis that the skeletal mesh is facing. Default is the Y axis.
+	 * The facing axis represents the MeshDescription's orientation and is set on import.
+	 * Therefore, if this property is modified elsewhere, the associated MeshDescription should be appropriately modified to reflect the new orientation.
+	 */
+	UPROPERTY(VisibleAnywhere, Category = Mesh, BlueprintGetter = GetForwardAxis)
+	TEnumAsByte<EAxis::Type> ForwardAxis = EAxis::Y;
+
 public:
 	/** Set the default mesh deformer used by this mesh. A mesh deformer is used to deform the skeletal mesh at runtime.
 	 *  The default deformer can be overridden by the USkinnedMeshComponent object. */
@@ -3087,6 +3095,19 @@ public:
 		PRAGMA_ENABLE_DEPRECATION_WARNINGS
 	}
 
+	/** Get the forward axis used by this mesh */
+	UFUNCTION(BlueprintCallable, Category = Mesh)
+	TEnumAsByte<EAxis::Type> GetForwardAxis() const
+	{
+		return ForwardAxis;
+	}
+
+	/** Set the forward axis used by this mesh */
+	void SetForwardAxis(TEnumAsByte<EAxis::Type> InForwardAxis)
+	{
+		ForwardAxis = InForwardAxis;
+	}
+
 private:
 	/**
 	 * Initial step for the Post Load process - Can't be done in parallel. USkinnedAsset Interface.
