Revert the LOD skipping optimization for FSkinWeightProfileManager

#rnx

diff --git a/Engine/Source/Runtime/Engine/Private/Animation/SkinWeightProfileManager.cpp b/Engine/Source/Runtime/Engine/Private/Animation/SkinWeightProfileManager.cpp
--- a/Engine/Source/Runtime/Engine/Private/Animation/SkinWeightProfileManager.cpp
+++ b/Engine/Source/Runtime/Engine/Private/Animation/SkinWeightProfileManager.cpp
@@ -42,14 +42,6 @@ static TAutoConsoleVariable<int32> CVarForceSkinWeightProfileManagerAsyncTaskOnG
 	ECVF_Default | ECVF_RenderThreadSafe
 );
 
-static int32 GSkipLODWithoutDataOrStreamingRequest = 1;
-static FAutoConsoleVariableRef CVarSkipLODWithoutDataOrStreamingRequest(
-	TEXT("SkinWeightProfileManager.SkipLODWithoutDataOrStreamingRequest"),
-	GSkipLODWithoutDataOrStreamingRequest,
-	TEXT(""),
-	ECVF_Default | ECVF_RenderThreadSafe
-);
-
 FSkinWeightProfileManager::FSkinWeightProfileManager(UWorld* InWorld) : LastGamethreadProfileIndex(INDEX_NONE), WeakWorld(nullptr)
 {
 	if (InWorld)
@@ -227,7 +219,7 @@ void FSkinWeightProfileManager::DoTick(float DeltaTime, ENamedThreads::Type Curr
 		{
 			QUICK_SCOPE_CYCLE_COUNTER(STAT_SkinWeightProfileManager_TickMesh);
 
-			FSetProfileRequest& Request = PendingSetProfileRequests[RequestIndex];
+			const FSetProfileRequest& Request = PendingSetProfileRequests[RequestIndex];
 			
 			// Remove if canceled or objects have gone stale 
 			bool bRemove = Request.IdentifyingObject.IsStale() || Request.WeakSkeletalMesh.IsStale() ||
@@ -244,21 +236,12 @@ void FSkinWeightProfileManager::DoTick(float DeltaTime, ENamedThreads::Type Curr
 			if (!bRemove)
 			{	
 				bool bAllBuffersReady = true;
-				for (int32 LODProfileArrayIndex = 0; LODProfileArrayIndex < Request.LODIndices.Num(); ++LODProfileArrayIndex)
+				for (const int32 LODIndex : Request.LODIndices)
 				{
-					int32 LODRenderDataArrayIndex = Request.LODIndices[LODProfileArrayIndex];
-					const FSkinWeightProfilesData& SkinWeightProfilesData = RenderData->LODRenderData[LODRenderDataArrayIndex].SkinWeightProfilesData;
-					if (!SkinWeightProfilesData.ContainsOverrideBuffer(Request.ProfileStack))
+					if (!RenderData->LODRenderData[LODIndex].SkinWeightProfilesData.ContainsOverrideBuffer(Request.ProfileStack))
 					{
-						// Check if this LOD is loaded or scheduled to be streamed in, if not then we can ignore it and remove from Request.LODIndices
-						if (!ShouldIgnoreLOD(SkinWeightProfilesData, Request.WeakSkeletalMesh))
-						{
-							bAllBuffersReady = false;
-							break;
-						}
-
-						Request.LODIndices.RemoveAtSwap(LODProfileArrayIndex);
-						--LODProfileArrayIndex;
+						bAllBuffersReady = false;
+						break;
 					}
 				}
 
@@ -342,31 +325,6 @@ void FSkinWeightProfileManager::CleanupRequest(const FSetProfileRequest& Request
 	
 }
 
-bool FSkinWeightProfileManager::ShouldIgnoreLOD(const FSkinWeightProfilesData& SkinWeightProfilesData, const TWeakObjectPtr<USkeletalMesh>& WeakSkeletalMesh) const
-{
-	if (!GSkipLODWithoutDataOrStreamingRequest)
-	{
-		return false;
-	}
-
-	// Are LOD's vertices loaded? 
-	FSkinWeightVertexBuffer* BaseBuffer = SkinWeightProfilesData.BaseBuffer;
-	if (BaseBuffer && BaseBuffer->GetDataVertexBuffer()->VertexBufferRHI->GetSize())
-	{
-		return false;
-	}
-
-	// Are they scheduled to be streamed in?
-	const USkeletalMesh* SkeletalMesh = WeakSkeletalMesh.Get();
-	if (SkeletalMesh && SkeletalMesh->HasPendingInitOrStreaming())
-	{
-		return false;
-	}
-
-	// This mesh LOD isn't loaded and not scheduled for streaming, we can ignore it
-	return true;
-}
-
 bool FSkinWeightProfileManager::IsTickableWhenPaused() const
 {
 	return false;
diff --git a/Engine/Source/Runtime/Engine/Public/Animation/SkinWeightProfileManager.h b/Engine/Source/Runtime/Engine/Public/Animation/SkinWeightProfileManager.h
--- a/Engine/Source/Runtime/Engine/Public/Animation/SkinWeightProfileManager.h
+++ b/Engine/Source/Runtime/Engine/Public/Animation/SkinWeightProfileManager.h
@@ -117,7 +117,6 @@ public:
 	void DoTick(float DeltaTime, ENamedThreads::Type CurrentThread, const FGraphEventRef& MyCompletionGraphEvent);
 protected:	
 	void CleanupRequest(const FSetProfileRequest& Request);
-	bool ShouldIgnoreLOD(const FSkinWeightProfilesData& SkinWeightProfilesData, const TWeakObjectPtr<USkeletalMesh>& WeakSkeletalMesh) const;
 protected:
 	TArray<FSetProfileRequest, TInlineAllocator<4>> CanceledRequest;
 	TArray<FSetProfileRequest> PendingSetProfileRequests;
