SplineComponent: Fix crash in indexing into an empty legacy spline introduced by recent refactor.

diff --git a/Engine/Source/Runtime/Engine/Private/Components/SplineComponent.cpp b/Engine/Source/Runtime/Engine/Private/Components/SplineComponent.cpp
--- a/Engine/Source/Runtime/Engine/Private/Components/SplineComponent.cpp
+++ b/Engine/Source/Runtime/Engine/Private/Components/SplineComponent.cpp
@@ -1327,6 +1327,11 @@ int32 USplineComponent::GetNumberOfSplineSegments() const
 
 float USplineComponent::GetInputKeyValueAtSplinePoint(int32 PointIndex) const
 {
+	if (GetNumberOfSplinePoints() == 0)
+	{
+		return 0.f;
+	}
+
 	PointIndex = GetClampedIndex(PointIndex);
 	
 	return UE::SplineComponent::GUseSplineCurves
@@ -1349,6 +1354,11 @@ FSplinePoint USplineComponent::GetSplinePointAt(int32 PointIndex, ESplineCoordin
 
 FVector USplineComponent::GetLocationAtSplinePoint(int32 PointIndex, ESplineCoordinateSpace::Type CoordinateSpace) const
 {
+	if (GetNumberOfSplinePoints() == 0)
+	{
+		return FVector::ZeroVector; // from legacy DummyPointPosition
+	}
+
 	PointIndex = GetClampedIndex(PointIndex);
 	
 	const FVector Location = UE::SplineComponent::GUseSplineCurves
@@ -1370,6 +1380,11 @@ FVector USplineComponent::GetTangentAtSplinePoint(int32 PointIndex, ESplineCoord
 
 FVector USplineComponent::GetArriveTangentAtSplinePoint(int32 PointIndex, ESplineCoordinateSpace::Type CoordinateSpace) const
 {
+	if (GetNumberOfSplinePoints() == 0)
+	{
+		return FVector::ForwardVector; // from legacy DummyPointPosition
+	}
+
 	PointIndex = GetClampedIndex(PointIndex);
 
 	const FVector Tangent = UE::SplineComponent::GUseSplineCurves
@@ -1381,6 +1396,11 @@ FVector USplineComponent::GetArriveTangentAtSplinePoint(int32 PointIndex, ESplin
 
 FVector USplineComponent::GetLeaveTangentAtSplinePoint(int32 PointIndex, ESplineCoordinateSpace::Type CoordinateSpace) const
 {
+	if (GetNumberOfSplinePoints() == 0)
+	{
+		return FVector::ForwardVector; // from legacy DummyPointPosition
+	}
+
 	PointIndex = GetClampedIndex(PointIndex);
 
 	const FVector Tangent = UE::SplineComponent::GUseSplineCurves
@@ -1417,6 +1437,11 @@ float USplineComponent::GetRollAtSplinePoint(int32 PointIndex, ESplineCoordinate
 
 FVector USplineComponent::GetScaleAtSplinePoint(int32 PointIndex) const
 {
+	if (GetNumberOfSplinePoints() == 0)
+	{
+		return FVector::OneVector; // from legacy DummyPointScale
+	}
+
 	PointIndex = GetClampedIndex(PointIndex);
 
 	return UE::SplineComponent::GUseSplineCurves
