[SimpleStreamableAssetManager] Fixed FStaticMeshSceneProxy::GetStreamableRenderAssetInfo computing an invalid TexelFactor because it wasn't using the correct primitive bounds. Also fixed uninitialized FPrimitiveSceneProxy's Bounds.

diff --git a/Engine/Source/Runtime/Engine/Private/StaticMeshSceneProxy.cpp b/Engine/Source/Runtime/Engine/Private/StaticMeshSceneProxy.cpp
--- a/Engine/Source/Runtime/Engine/Private/StaticMeshSceneProxy.cpp
+++ b/Engine/Source/Runtime/Engine/Private/StaticMeshSceneProxy.cpp
@@ -600,9 +600,8 @@ void FStaticMeshSceneProxy::GetStreamableRenderAssetInfo(const FBoxSphereBounds&
 	
 	if (StaticMesh->RenderResourceSupportsStreaming() && (StaticMesh->GetRenderAssetType() == EStreamableRenderAssetType::StaticMesh))
 	{
-		const FBoxSphereBounds& ProxyBounds = GetBounds();
-		const float TexelFactor = ForcedLodModel > 0 ? - (StaticMesh->GetRenderData()->LODResources.Num() - ForcedLodModel + 1) : ProxyBounds.SphereRadius * 2.0f;
-		OutStreamableRenderAssets.Add(FStreamingRenderAssetPrimitiveInfo(const_cast<UStaticMesh*>(StaticMesh), ProxyBounds, TexelFactor, PackedRelativeBox_Identity, true, false));
+		const float TexelFactor = ForcedLodModel > 0 ? - (StaticMesh->GetRenderData()->LODResources.Num() - ForcedLodModel + 1) : InPrimitiveBounds.SphereRadius * 2.0f;
+		OutStreamableRenderAssets.Add(FStreamingRenderAssetPrimitiveInfo(const_cast<UStaticMesh*>(StaticMesh), InPrimitiveBounds, TexelFactor, PackedRelativeBox_Identity, true, false));
 	}
 }
 
diff --git a/Engine/Source/Runtime/Engine/Public/PrimitiveSceneProxy.h b/Engine/Source/Runtime/Engine/Public/PrimitiveSceneProxy.h
--- a/Engine/Source/Runtime/Engine/Public/PrimitiveSceneProxy.h
+++ b/Engine/Source/Runtime/Engine/Public/PrimitiveSceneProxy.h
@@ -1618,7 +1618,7 @@ private:
 	FMatrix LocalToWorld;
 
 	/** The primitive's bounds. */
-	FBoxSphereBounds Bounds;
+	FBoxSphereBounds Bounds = FBoxSphereBounds(ForceInit);
 
 	/** The primitive's local space bounds. */
 	FBoxSphereBounds LocalBounds = FBoxSphereBounds(ForceInit);
