Make Soundclass parents Transient (Resubmit, no submixes)
[Submitting on behalf of Alfaroh.CorneyIII]

#rnx

diff --git a/Engine/Source/Runtime/Engine/Classes/Sound/SoundClass.h b/Engine/Source/Runtime/Engine/Classes/Sound/SoundClass.h
--- a/Engine/Source/Runtime/Engine/Classes/Sound/SoundClass.h
+++ b/Engine/Source/Runtime/Engine/Classes/Sound/SoundClass.h
@@ -221,7 +221,7 @@ public:
 	UPROPERTY(EditAnywhere, BlueprintReadOnly, Category = General)
 	TArray<FPassiveSoundMixModifier> PassiveSoundMixModifiers;
 
-	UPROPERTY(BlueprintReadOnly, Category = General)
+	UPROPERTY(Transient, BlueprintReadOnly, Category = General)
 	TObjectPtr<USoundClass> ParentClass;
 
 	/** Array of user data stored with the asset */
