Fixes to AnimBank DDC, AnimBank resource mem tracking, some misc fixes, as well as some experimental tagging

diff --git a/Engine/Source/Runtime/Engine/Classes/Animation/AnimBank.h b/Engine/Source/Runtime/Engine/Classes/Animation/AnimBank.h
--- a/Engine/Source/Runtime/Engine/Classes/Animation/AnimBank.h
+++ b/Engine/Source/Runtime/Engine/Classes/Animation/AnimBank.h
@@ -198,6 +198,8 @@ public:
 	virtual bool IsReadyForFinishDestroy() override;
 	virtual bool NeedsLoadForTargetPlatform(const ITargetPlatform* TargetPlatform) const override;
 
+	virtual void GetResourceSizeEx(FResourceSizeEx& CumulativeResourceSize) override;
+
 	void InitResources();
 	void ReleaseResources();
 
diff --git a/Engine/Source/Runtime/Engine/Classes/Components/InstancedSkinnedMeshComponent.h b/Engine/Source/Runtime/Engine/Classes/Components/InstancedSkinnedMeshComponent.h
--- a/Engine/Source/Runtime/Engine/Classes/Components/InstancedSkinnedMeshComponent.h
+++ b/Engine/Source/Runtime/Engine/Classes/Components/InstancedSkinnedMeshComponent.h
@@ -51,6 +51,7 @@ struct FSkinnedMeshInstanceData
 };
 
 UCLASS(ClassGroup=Rendering, hidecategories=(Object,Activation,Collision,"Components|Activation",Physics), editinlinenew, meta=(BlueprintSpawnableComponent), Blueprintable, MinimalAPI)
+UE_EXPERIMENTAL(5.6, "This class is currently extremely experimental and should not be used at this time.")
 class UInstancedSkinnedMeshComponent : public USkinnedMeshComponent
 {
 	GENERATED_UCLASS_BODY()
diff --git a/Engine/Source/Runtime/Engine/Private/Animation/AnimBank.cpp b/Engine/Source/Runtime/Engine/Private/Animation/AnimBank.cpp
--- a/Engine/Source/Runtime/Engine/Private/Animation/AnimBank.cpp
+++ b/Engine/Source/Runtime/Engine/Private/Animation/AnimBank.cpp
@@ -526,7 +526,7 @@ bool FAnimBankBuildAsyncCacheTask::BuildData(const UE::FSharedString& Name, cons
 		UAnimSequence::FScopedCompressedAnimSequence CompressedAnimSequence = Sequence->GetCompressedData(TargetPlatform);
 		const FCompressedAnimSequence& PlatformCompressedData = CompressedAnimSequence.Get();
 		// We should always have compressed data at this point
-		if (PlatformCompressedData.IsBoneDataValid(Sequence))
+		if (!PlatformCompressedData.IsValid(Sequence) || PlatformCompressedData.BoneCompressionCodec == nullptr)
 		{
 			UE_LOG(LogAnimBank, Error, TEXT("Animation bank referenced sequence is missing compressed data!"));
 			return false;
@@ -782,6 +782,26 @@ bool UAnimBank::NeedsLoadForTargetPlatform(const ITargetPlatform* TargetPlatform
 	return DoesTargetPlatformSupportNanite(TargetPlatform);
 }
 
+void UAnimBank::GetResourceSizeEx(FResourceSizeEx& CumulativeResourceSize)
+{
+	Super::GetResourceSizeEx(CumulativeResourceSize);
+
+	CumulativeResourceSize.AddDedicatedSystemMemoryBytes(Data.Mapping.MeshGlobalRefPose.GetAllocatedSize());
+	CumulativeResourceSize.AddDedicatedSystemMemoryBytes(Data.Mapping.AnimGlobalRefPose.GetAllocatedSize());
+	CumulativeResourceSize.AddDedicatedSystemMemoryBytes(Data.Mapping.MeshToAnimIndexMap.GetAllocatedSize());
+	CumulativeResourceSize.AddDedicatedSystemMemoryBytes(Data.Mapping.RetargetingTable.GetAllocatedSize());
+	CumulativeResourceSize.AddDedicatedSystemMemoryBytes(Data.Mapping.PositionKeys.GetAllocatedSize());
+	CumulativeResourceSize.AddDedicatedSystemMemoryBytes(Data.Mapping.RotationKeys.GetAllocatedSize());
+
+	for (const FAnimBankEntry& Entry : Data.Entries)
+	{
+		CumulativeResourceSize.AddDedicatedSystemMemoryBytes(sizeof(Entry));
+		CumulativeResourceSize.AddDedicatedSystemMemoryBytes(Entry.PositionKeys.GetAllocatedSize());
+		CumulativeResourceSize.AddDedicatedSystemMemoryBytes(Entry.RotationKeys.GetAllocatedSize());
+		CumulativeResourceSize.AddDedicatedSystemMemoryBytes(Entry.ScalingKeys.GetAllocatedSize());
+	}
+}
+
 void UAnimBank::InitResources()
 {
 	if (!FApp::CanEverRender())
@@ -1018,7 +1038,7 @@ FIoHash UAnimBank::CreateDerivedDataKeyHash(const ITargetPlatform* TargetPlatfor
 	}
 
 	FMemoryHasherBlake3 Writer;
-	FGuid AnimBankVersionGuid(0x51842261, 0xDBAF4D8C, 0xB13ABA3C, 0x50EC8691);
+	FGuid AnimBankVersionGuid(0x51842261, 0xDBAF4D8C, 0xB13ABA3C, 0x40EC7666);
 	Writer << AnimBankVersionGuid;
 
 	if (IsValid(Asset))
diff --git a/Engine/Source/Runtime/Engine/Private/Components/InstancedSkinnedMeshComponent.cpp b/Engine/Source/Runtime/Engine/Private/Components/InstancedSkinnedMeshComponent.cpp
--- a/Engine/Source/Runtime/Engine/Private/Components/InstancedSkinnedMeshComponent.cpp
+++ b/Engine/Source/Runtime/Engine/Private/Components/InstancedSkinnedMeshComponent.cpp
@@ -34,6 +34,10 @@ static TAutoConsoleVariable<int32> CVarInstancedSkinnedMeshesForceRefPose(
 	TEXT("r.InstancedSkinnedMeshes.ForceRefPose"),
 	0,
 	TEXT("Whether to force ref pose for instanced skinned meshes"),
+	FConsoleVariableDelegate::CreateLambda([](IConsoleVariable* InVariable)
+	{
+		FGlobalComponentRecreateRenderStateContext Context;
+	}),
 	ECVF_RenderThreadSafe
 );
 
diff --git a/Engine/Source/Runtime/Renderer/Private/Skinning/AnimBankTransformProvider.cpp b/Engine/Source/Runtime/Renderer/Private/Skinning/AnimBankTransformProvider.cpp
--- a/Engine/Source/Runtime/Renderer/Private/Skinning/AnimBankTransformProvider.cpp
+++ b/Engine/Source/Runtime/Renderer/Private/Skinning/AnimBankTransformProvider.cpp
@@ -703,24 +703,26 @@ void FAnimBankTransformProvider::ProvideGPUBankTransforms(FSkinningTransformProv
 
 	FRDGBuilder& GraphBuilder = Context.GraphBuilder;
 	const FAnimBankGPUData BankData = BuildAnimBankGPUData(BankRecordMap, GraphBuilder);
-	if (BankData.bValid)
+	if (!BankData.bValid)
 	{
-		// Evaluate animation banks
-		{
-			FAnimBankEvaluateCS::FParameters* PassParameters = GraphBuilder.AllocParameters<FAnimBankEvaluateCS::FParameters>();
-			PassParameters->TransformBuffer = GraphBuilder.CreateUAV(BankData.TransformBuffer);
-			PassParameters->BankBuffer = GraphBuilder.CreateSRV(BankData.BankDataBuffer);
-			PassParameters->HeaderBuffer = GraphBuilder.CreateSRV(BankData.BoneBlockBuffer);
+		return;
+	}
 
-			auto ComputeShader = GetGlobalShaderMap(GMaxRHIFeatureLevel)->GetShader<FAnimBankEvaluateCS>();
-			FComputeShaderUtils::AddPass(
-				GraphBuilder,
-				RDG_EVENT_NAME("AnimBankEvaluate"),
-				ComputeShader,
-				PassParameters,
-				FIntVector(BankData.BlockCount, 1, 1)
-			);
-		}
+	// Evaluate animation banks
+	{
+		FAnimBankEvaluateCS::FParameters* PassParameters = GraphBuilder.AllocParameters<FAnimBankEvaluateCS::FParameters>();
+		PassParameters->TransformBuffer = GraphBuilder.CreateUAV(BankData.TransformBuffer);
+		PassParameters->BankBuffer = GraphBuilder.CreateSRV(BankData.BankDataBuffer);
+		PassParameters->HeaderBuffer = GraphBuilder.CreateSRV(BankData.BoneBlockBuffer);
+
+		auto ComputeShader = GetGlobalShaderMap(GMaxRHIFeatureLevel)->GetShader<FAnimBankEvaluateCS>();
+		FComputeShaderUtils::AddPass(
+			GraphBuilder,
+			RDG_EVENT_NAME("AnimBankEvaluate"),
+			ComputeShader,
+			PassParameters,
+			FIntVector(BankData.BlockCount, 1, 1)
+		);
 	}
 
 	// Scatter animation bank results
@@ -736,6 +738,10 @@ void FAnimBankTransformProvider::ProvideCPUBankTransforms(FSkinningTransformProv
 
 	FRDGBuilder& GraphBuilder = Context.GraphBuilder;
 	const FAnimBankCPUData BankData = BuildAnimBankCPUData(BankRecordMap, GraphBuilder);
+	if (!BankData.bValid)
+	{
+		return;
+	}
 
 	// Scatter animation bank results
 	ScatterAnimation(Context, MakeConstArrayView(BankData.IdToOffsetMapping), BankData.TransformBuffer);
@@ -751,4 +757,4 @@ const FSkinningTransformProvider::FProviderId& GetAnimBankProviderId()
 	{
 		return AnimBankCPUProviderId;
 	}
-}
\ No newline at end of file
+}
