FSpline: Reenable reparam table

In addition to turning this functionality on, we:

1) Only generate reparam info when necessary,
2) Only create default attribute channels when necessary (not when copying another spline),
2) Minimized lots of unnecessary FSpline construction/conversion when caching and restoring during BP construction script runs.

diff --git a/Engine/Source/Runtime/Engine/Classes/Components/SplineComponent.h b/Engine/Source/Runtime/Engine/Classes/Components/SplineComponent.h
--- a/Engine/Source/Runtime/Engine/Classes/Components/SplineComponent.h
+++ b/Engine/Source/Runtime/Engine/Classes/Components/SplineComponent.h
@@ -1065,6 +1065,11 @@ public:
 		: FSceneComponentInstanceData(SourceComponent)
 		, bSplineHasBeenEdited(false)
 	{}
+	explicit FSplineComponentInstanceData(const USplineComponent* SourceComponent, const FSpline& SourceSpline)
+		: FSceneComponentInstanceData(SourceComponent)
+		, bSplineHasBeenEdited(false)
+		, Spline(SourceSpline)
+	{}
 	virtual ~FSplineComponentInstanceData() = default;
 
 	virtual bool ContainsData() const override
diff --git a/Engine/Source/Runtime/Engine/Classes/Curves/Spline.h b/Engine/Source/Runtime/Engine/Classes/Curves/Spline.h
--- a/Engine/Source/Runtime/Engine/Classes/Curves/Spline.h
+++ b/Engine/Source/Runtime/Engine/Classes/Curves/Spline.h
@@ -24,8 +24,8 @@ struct FSpline
 	ENGINE_API FSpline();
 	ENGINE_API FSpline(const FSpline& Other);
 	ENGINE_API FSpline& operator=(const FSpline& Other);
+	ENGINE_API FSpline& operator=(const FSplineCurves& Other);
 	ENGINE_API ~FSpline();
-	ENGINE_API static FSpline FromSplineCurves(const FSplineCurves& InSpline);
 	
 	/* Control Point Index Interface */
 	
diff --git a/Engine/Source/Runtime/Engine/Private/Components/SplineComponent.cpp b/Engine/Source/Runtime/Engine/Private/Components/SplineComponent.cpp
--- a/Engine/Source/Runtime/Engine/Private/Components/SplineComponent.cpp
+++ b/Engine/Source/Runtime/Engine/Private/Components/SplineComponent.cpp
@@ -2538,15 +2538,22 @@ TStructOnScope<FActorComponentInstanceData> USplineComponent::GetComponentInstan
 	}
 	else
 	{
-		TStructOnScope<FActorComponentInstanceData> InstanceData = MakeStructOnScope<FActorComponentInstanceData, FSplineComponentInstanceData>(this);
-		FSplineComponentInstanceData* SplineComponentInstanceData = InstanceData.Cast<FSplineComponentInstanceData>();
-	
+		TStructOnScope<FActorComponentInstanceData> InstanceData;
+		
 		if (bSplineHasBeenEdited)
 		{
-			SplineComponentInstanceData->Spline = Spline;
+			InstanceData = MakeStructOnScope<FActorComponentInstanceData, FSplineComponentInstanceData>(this, Spline);
+			FSplineComponentInstanceData* SplineComponentInstanceData = InstanceData.Cast<FSplineComponentInstanceData>();
+			
 			SplineComponentInstanceData->bClosedLoop = bClosedLoop;
+			SplineComponentInstanceData->bSplineHasBeenEdited = bSplineHasBeenEdited;
+		}
+		else
+		{
+			InstanceData = MakeStructOnScope<FActorComponentInstanceData, FSplineComponentInstanceData>(this);
+			FSplineComponentInstanceData* SplineComponentInstanceData = InstanceData.Cast<FSplineComponentInstanceData>();
+			SplineComponentInstanceData->bSplineHasBeenEdited = bSplineHasBeenEdited;
 		}
-		SplineComponentInstanceData->bSplineHasBeenEdited = bSplineHasBeenEdited;
 
 		return InstanceData;
 	}
@@ -2568,7 +2575,7 @@ void USplineComponent::SetSpline(const FSpline& InSpline)
 void USplineComponent::SetSpline(const FSplineCurves& InSplineCurves)
 {
 	WarninglessSplineCurves() = InSplineCurves;
-	Spline = FSpline::FromSplineCurves(InSplineCurves);
+	Spline = InSplineCurves;
 }
 
 int32 USplineComponent::GetVersion() const
diff --git a/Engine/Source/Runtime/Engine/Private/Curves/Spline.cpp b/Engine/Source/Runtime/Engine/Private/Curves/Spline.cpp
--- a/Engine/Source/Runtime/Engine/Private/Curves/Spline.cpp
+++ b/Engine/Source/Runtime/Engine/Private/Curves/Spline.cpp
@@ -157,7 +157,7 @@ public:
     }
 
 	float GetSegmentLength(const int32 Index, const float Param, const FVector& Scale3D = FVector(1.0f)) const;
-	float GetSplineLength() const;
+	float GetSplineLength();
 	
     virtual void SetClosedLoop(bool bInClosedLoop) override;
 
@@ -191,8 +191,8 @@ public:
 	EInterpCurveMode GetSplinePointType(int32 Index) const;
 
 	float GetParameterAtIndex(int32 Index) const;
-	float GetParameterAtDistance(float Distance) const;
-	float GetDistanceAtParameter(float Parameter) const;
+	float GetParameterAtDistance(float Distance);
+	float GetDistanceAtParameter(float Parameter);
 
     virtual float FindNearest(const FVector& InLocation, float& OutSquaredDistance) const override;
     float FindNearestOnSegment(const FVector& InLocation, int32 SegmentIndex, float& OutSquaredDistance) const;
@@ -260,6 +260,9 @@ private:
 	
 	/** Convert from [0, NumSegments] to internal spline space */
 	float ToInternalSplineSpace(float Parameter) const;
+
+	bool bReparamTableDirty = true;
+	void MarkReparamTableDirty() { bReparamTableDirty = true; }
 	
 	/** Clears and repopulates reparameterization attribute channels. */
 	void UpdateReparamTable();
@@ -811,34 +814,32 @@ FSpline& FSpline::operator=(const FSpline& Other)
 	return *this;
 }
 
-FSpline::~FSpline()
+FSpline& FSpline::operator=(const FSplineCurves& Other)
 {
-#if WITH_EDITOR
-	if (OnSplineImplementationChangedHandle.IsValid())
+	switch (CurrentImplementation)
 	{
-		UE::Spline::GetOnSplineImplementationChanged().Remove(OnSplineImplementationChangedHandle);
-		OnSplineImplementationChangedHandle.Reset();
+	default:
+		break;
+	case 1:
+		*LegacyData = FLegacySpline(Other);
+		break;
+	case 2:
+		*NewData = FNewSpline(Other);
+		break;
 	}
-#endif
+	
+	return *this;
 }
 
-FSpline FSpline::FromSplineCurves(const FSplineCurves& InSpline)
+FSpline::~FSpline()
 {
-	FSpline NewSpline;
-	
 #if WITH_EDITOR
-	NewSpline.PreviousImplementation = UE::Spline::GImplementation;
-#endif
-	NewSpline.CurrentImplementation = UE::Spline::GImplementation;
-
-	switch (NewSpline.CurrentImplementation)
+	if (OnSplineImplementationChangedHandle.IsValid())
 	{
-	default: break;		// Intentionally doing nothing
-	case 1: NewSpline.LegacyData = MakeShared<FLegacySpline>(InSpline); break;
-	case 2: NewSpline.NewData = MakeShared<FNewSpline>(InSpline); break;
+		UE::Spline::GetOnSplineImplementationChanged().Remove(OnSplineImplementationChangedHandle);
+		OnSplineImplementationChangedHandle.Reset();
 	}
-	
-	return NewSpline;
+#endif
 }
 
 void FSpline::AddPoint(const FSplinePoint& InPoint)
@@ -1636,11 +1637,6 @@ FNewSpline::FNewSpline()
 
 FNewSpline::FNewSpline(const FNewSpline& Other)
 {
-	CreateAttributeChannel<TTangentBezierSpline<FQuat>>(RotationAttrName);
-	CreateAttributeChannel<TTangentBezierSpline<FVector>>(ScaleAttrName);
-	CreateAttributeChannel<TTangentBezierSpline<float>>(ParamToLengthAttrName);
-	CreateAttributeChannel<TTangentBezierSpline<float>>(NormalizedLengthToParamAttrName);
-	
     *this = Other;
 
 	ValidateRotScale();
@@ -1733,9 +1729,10 @@ FNewSpline& FNewSpline::operator=(const FNewSpline& Other)
     if (this != &Other)
     {
 		TMultiSpline::operator=(Other);
+		PositionCurve = Other.PositionCurve;
+		RotationCurve = Other.RotationCurve;
+		ScaleCurve = Other.ScaleCurve;
 		ReparamStepsPerSegment = Other.ReparamStepsPerSegment;
-
-		UpdateSpline();
     }
 	
     return *this;
@@ -1749,12 +1746,12 @@ bool FNewSpline::operator!=(const FNewSpline& Other) const
 bool FNewSpline::Serialize(FArchive& Ar)
 {
     TMultiSpline::Serialize(Ar);
-	
     Ar << ReparamStepsPerSegment;
-	
+
+	// Do not serialize legacy curves, only generate if loading.
 	if (Ar.IsLoading())
 	{
-		UpdateSpline();
+		RebuildLegacyCurves();
 	}
 	
     return true;
@@ -1820,18 +1817,11 @@ float FNewSpline::GetSegmentLength(const int32 Index, const float Param, const F
 	return Length;
 }
 
-float FNewSpline::GetSplineLength() const
-{
-	float AccumulatedLength = 0.f;
-	for (int Idx = 0; Idx < GetSpline().GetNumSegments(); ++Idx)
-	{
-		AccumulatedLength += GetSegmentLength(Idx, 1.f);
-	}
-	return AccumulatedLength;
-	
+float FNewSpline::GetSplineLength()
+{	
 	// Evaluate the ParamToLength channel at the very end of the spline
-	//const float MaxParameter = static_cast<float>(GetSpline().GetNumSegments());
-    //return GetDistanceAtParameter(MaxParameter);
+	const float MaxParameter = static_cast<float>(GetSpline().GetNumSegments());
+    return GetDistanceAtParameter(MaxParameter);
 }
 
 void FNewSpline::SetClosedLoop(bool bInClosedLoop)
@@ -2224,17 +2214,19 @@ float FNewSpline::GetParameterAtIndex(int32 Index) const
 	return FromInternalSplineSpace(GetSpline().GetParameter(Index));
 }
 
-float FNewSpline::GetParameterAtDistance(float Distance) const
+float FNewSpline::GetParameterAtDistance(float Distance)
 {
-	return 0.f;
-	//const float NormalizedDistance = FMath::Clamp(Distance / FMath::Max(GetSplineLength(), UE_SMALL_NUMBER), 0.f, 1.f);
-	//return EvaluateAttribute<float>(NormalizedLengthToParamAttrName, ToInternalSplineSpace(NormalizedDistance));
+	UpdateReparamTable();
+	
+	const float NormalizedDistance = FMath::Clamp(Distance / FMath::Max(GetSplineLength(), UE_SMALL_NUMBER), 0.f, 1.f);
+	return EvaluateAttribute<float>(NormalizedLengthToParamAttrName, ToInternalSplineSpace(NormalizedDistance));
 }
 
-float FNewSpline::GetDistanceAtParameter(float Parameter) const
+float FNewSpline::GetDistanceAtParameter(float Parameter)
 {
-	return 0.f;
-	//return EvaluateAttribute<float>(ParamToLengthAttrName, ToInternalSplineSpace(Parameter));
+	UpdateReparamTable();
+
+	return EvaluateAttribute<float>(ParamToLengthAttrName, ToInternalSplineSpace(Parameter));
 }
 
 float FNewSpline::FindNearest(const FVector& InLocation, float& OutSquaredDistance) const
@@ -2290,7 +2282,7 @@ FVector FNewSpline::EvaluateScale(float Parameter) const
 void FNewSpline::UpdateSpline()
 {
     GetSpline().UpdateTangents();			// updates the tangents on our version of the points
-    UpdateReparamTable();		// updates both reparam tables
+    MarkReparamTableDirty();
     RebuildLegacyCurves();		// updates legacy curves based on points, never evaluates internal spline.
 }
 
@@ -2332,56 +2324,60 @@ float FNewSpline::ToInternalSplineSpace(float Parameter) const
 
 void FNewSpline::UpdateReparamTable()
 {
- //    const int32 NumPoints = GetSpline().GetNumPoints();
- //  	const int32 NumSegments = NumPoints - 1;
-	//
- //    // Start by clearing it
- //    ClearAttributeChannel(ParamToLengthAttrName);
- //    ClearAttributeChannel(NormalizedLengthToParamAttrName);
-	//
- //    // Here we get Param -> Length mapping.
- //    float AccumulatedLength = 0.0f;
- //    for (int32 SegmentIndex = 0; SegmentIndex < NumSegments; ++SegmentIndex)
- //    {
- //    	for (int32 Step = 0; Step < ReparamStepsPerSegment; ++Step)
- //    	{
- //    		// These are with respect to current segment.
- //    		const float SegmentParam = static_cast<float>(Step) / ReparamStepsPerSegment;
- //    		const float OffsetLength = GetSegmentLength(SegmentIndex, SegmentParam);
- //    
- //    		// Map the above values to 'spline space' (i.e. [0, NumSegments]) by normalizing and converting.
- //    		const float SplineParam = SegmentIndex + SegmentParam;
- //    		const float CurrentLength = OffsetLength + AccumulatedLength;
- //    		AddAttributeValue<float>(ParamToLengthAttrName, CurrentLength, SplineParam);
- //    	}
- //    
- //    	AccumulatedLength += GetSegmentLength(SegmentIndex, 1.0f);
- //    }
- //    AddAttributeValue<float>(ParamToLengthAttrName, AccumulatedLength, static_cast<float>(NumSegments));
-	//
- //    // Here we get Length -> Param mapping
-	// const int32 TotalSteps = NumSegments * ReparamStepsPerSegment + 1;
- //    for (int32 Step = 0; Step <= TotalSteps; ++Step)
- //    {
- //    	const float SplineParam = FromNormalizedSpace(FMath::Clamp(static_cast<float>(Step) / TotalSteps, 0.0f, 1.0f));
- //    	const float NormalizedLength = EvaluateAttribute<float>(ParamToLengthAttrName, ToInternalSplineSpace(SplineParam)) / AccumulatedLength;
- //    	AddAttributeValue<float>(NormalizedLengthToParamAttrName, SplineParam, NormalizedLength);
- //    }
+	if (!bReparamTableDirty)
+	{
+		return;
+	}
+
+	const int32 NumPoints = GetSpline().GetNumPoints();
+	const int32 NumSegments = NumPoints - 1;
+	
+	// Start by clearing it
+	ClearAttributeChannel(ParamToLengthAttrName);
+	ClearAttributeChannel(NormalizedLengthToParamAttrName);
+	
+	// Here we get Param -> Length mapping.
+	float AccumulatedLength = 0.0f;
+	for (int32 SegmentIndex = 0; SegmentIndex < NumSegments; ++SegmentIndex)
+	{
+		for (int32 Step = 0; Step < ReparamStepsPerSegment; ++Step)
+		{
+			// These are with respect to current segment.
+			const float SegmentParam = static_cast<float>(Step) / ReparamStepsPerSegment;
+			const float OffsetLength = GetSegmentLength(SegmentIndex, SegmentParam);
+	
+			// Map the above values to 'spline space' (i.e. [0, NumSegments]) by normalizing and converting.
+			const float SplineParam = SegmentIndex + SegmentParam;
+			const float CurrentLength = OffsetLength + AccumulatedLength;
+			AddAttributeValue<float>(ParamToLengthAttrName, CurrentLength, SplineParam);
+		}
+	
+		AccumulatedLength += GetSegmentLength(SegmentIndex, 1.0f);
+	}
+	AddAttributeValue<float>(ParamToLengthAttrName, AccumulatedLength, static_cast<float>(NumSegments));
+	
+	// Here we get Length -> Param mapping
+	const int32 TotalSteps = NumSegments * ReparamStepsPerSegment + 1;
+	for (int32 Step = 0; Step <= TotalSteps; ++Step)
+	{
+		const float SplineParam = FromNormalizedSpace(FMath::Clamp(static_cast<float>(Step) / TotalSteps, 0.0f, 1.0f));
+		const float NormalizedLength = EvaluateAttribute<float>(ParamToLengthAttrName, ToInternalSplineSpace(SplineParam)) / AccumulatedLength;
+		AddAttributeValue<float>(NormalizedLengthToParamAttrName, SplineParam, NormalizedLength);
+	}
+
+	bReparamTableDirty = false;
 }
 
 void FNewSpline::RebuildLegacyCurves()
 {
-    PositionCurve.Reset();
-    RotationCurve.Reset();
-    ScaleCurve.Reset();
+	PositionCurve.Points.Reset(GetSpline().GetNumPoints());
+	RotationCurve.Points.Reset(GetSpline().GetNumPoints());
+	ScaleCurve.Points.Reset(GetSpline().GetNumPoints());
 
     PositionCurve.bIsLooped = IsClosedLoop();
     PositionCurve.LoopKeyOffset = 1.f;
     RotationCurve.bIsLooped = IsClosedLoop();
     ScaleCurve.bIsLooped = IsClosedLoop();
-
-	// if this fails, some other function which mutates InternalSpline is failing to keep Points appropriately sized
-	//check(Points.Num() == GetNumPoints());
 	
     for (int32 i = 0; i < GetSpline().GetNumPoints(); ++i)
     {
