Add some virtual texture CPU prefetch stats to the stat streaming display.
#rnx

diff --git a/Engine/Source/Runtime/Engine/Private/Streaming/AsyncTextureStreaming.cpp b/Engine/Source/Runtime/Engine/Private/Streaming/AsyncTextureStreaming.cpp
--- a/Engine/Source/Runtime/Engine/Private/Streaming/AsyncTextureStreaming.cpp
+++ b/Engine/Source/Runtime/Engine/Private/Streaming/AsyncTextureStreaming.cpp
@@ -1133,6 +1133,11 @@ void FRenderAssetStreamingMipCalcTask::UpdateStats_Async()
 	int32 TotalNumResidentLODs = 0;
 	int32 TotalNumEvictedLODs = 0;
 
+	Stats.VirtualTexturePool = Settings.VirtualTexturePoolSize * 1024 * 1024;
+	Stats.NumVirtualTextures = 0;
+	Stats.RequiredVirtualTexturePool = 0;
+	Stats.UsedVirtualTexturePool = 0;
+
 	for (FStreamingRenderAsset& StreamingRenderAsset : StreamingRenderAssets)
 	{
 		if (IsAborted()) break;
@@ -1144,6 +1149,14 @@ void FRenderAssetStreamingMipCalcTask::UpdateStats_Async()
 		const int64 MaxSize = StreamingRenderAsset.GetSize(StreamingRenderAsset.MaxAllowedMips);
 		const int64 VisibleWantedSize = StreamingRenderAsset.GetSize(StreamingRenderAsset.VisibleWantedMips);
 
+		if (StreamingRenderAsset.IsVirtualTexture())
+		{
+			++Stats.NumVirtualTextures;
+			Stats.RequiredVirtualTexturePool += RequiredSize;
+			Stats.UsedVirtualTexturePool += FMath::Min<int64>(RequiredSize, BudgetedSize);
+			continue;
+		}
+
 #if WITH_EDITOR
 		if (StreamingRenderAsset.IsTexture() && StreamingRenderAsset.RenderAsset->bMarkAsEditorStreamingPool)
 		{
@@ -1276,6 +1289,7 @@ void FRenderAssetStreamingMipCalcTask::UpdateCSVOnlyStats_Async()
 	{
 		if (IsAborted()) break;
 		if (!StreamingRenderAsset.RenderAsset) continue;
+		if (!StreamingRenderAsset.IsVirtualTexture()) continue;
 
 		const int64 ResidentSize = StreamingRenderAsset.GetSize(StreamingRenderAsset.ResidentMips);
 		const int64 RequiredSize = StreamingRenderAsset.GetSize(StreamingRenderAsset.GetPerfectWantedMips());
diff --git a/Engine/Source/Runtime/Engine/Private/Streaming/TextureStreamingHelpers.cpp b/Engine/Source/Runtime/Engine/Private/Streaming/TextureStreamingHelpers.cpp
--- a/Engine/Source/Runtime/Engine/Private/Streaming/TextureStreamingHelpers.cpp
+++ b/Engine/Source/Runtime/Engine/Private/Streaming/TextureStreamingHelpers.cpp
@@ -33,11 +33,16 @@ DECLARE_MEMORY_STAT_POOL(TEXT("Editor Required Pool"), STAT_Streaming15_EditorRe
 #endif
 DECLARE_MEMORY_STAT_POOL(TEXT("Desired Data Loaded Percent"), STAT_Streaming16_DesiredDataLoadedPct, STATGROUP_Streaming, FPlatformMemory::MCR_UsedStreamingPool);
 
+DECLARE_MEMORY_STAT(TEXT("Virtual Texture Pool"), STAT_Streaming17_VirtualTexturePool, STATGROUP_Streaming);
+DECLARE_MEMORY_STAT(TEXT("Virtual Texture Required"), STAT_Streaming18_VirtualTextureRequired, STATGROUP_Streaming);
+DECLARE_MEMORY_STAT(TEXT("Virtual Texture Used"), STAT_Streaming19_VirtualTextureUsed, STATGROUP_Streaming);
+
 DECLARE_CYCLE_STAT(TEXT("Setup Async Task"), STAT_Streaming01_SetupAsyncTask, STATGROUP_Streaming);
 DECLARE_CYCLE_STAT(TEXT("Update Streaming Data"), STAT_Streaming02_UpdateStreamingData, STATGROUP_Streaming);
 DECLARE_CYCLE_STAT(TEXT("Streaming Render Assets"), STAT_Streaming03_StreamRenderAssets, STATGROUP_Streaming);
 DECLARE_CYCLE_STAT(TEXT("Notifications"), STAT_Streaming04_Notifications, STATGROUP_Streaming);
 DECLARE_DWORD_COUNTER_STAT(TEXT("Pending 2D Update"), STAT_Streaming05_Pending2DUpdate, STATGROUP_Streaming);
+DECLARE_DWORD_COUNTER_STAT(TEXT("Virtual Texture Count"), STAT_Streaming06_VirtualTextureCount, STATGROUP_Streaming);
 
 /** Streaming Overview stats */
 
@@ -440,11 +445,17 @@ void FRenderAssetStreamingStats::Apply()
 #endif
 
 	SET_MEMORY_STAT(STAT_Streaming16_DesiredDataLoadedPct, (bHasBudgetLimit && StreamingPool != 0) ? 100.0f * float(WantedMips) / float(StreamingPool) : 0.0f);
+	
+	SET_MEMORY_STAT(STAT_Streaming17_VirtualTexturePool, VirtualTexturePool);
+	SET_MEMORY_STAT(STAT_Streaming18_VirtualTextureRequired, RequiredVirtualTexturePool);
+	SET_MEMORY_STAT(STAT_Streaming19_VirtualTextureUsed, UsedVirtualTexturePool);
+
 	SET_CYCLE_COUNTER(STAT_Streaming01_SetupAsyncTask, SetupAsyncTaskCycles);
 	SET_CYCLE_COUNTER(STAT_Streaming02_UpdateStreamingData, UpdateStreamingDataCycles);
 	SET_CYCLE_COUNTER(STAT_Streaming03_StreamRenderAssets, StreamRenderAssetsCycles);
 	SET_CYCLE_COUNTER(STAT_Streaming04_Notifications, CallbacksCycles);
 	INC_DWORD_STAT_BY(STAT_Streaming05_Pending2DUpdate, GPending2DUpdateCount);
+	INC_DWORD_STAT_BY(STAT_Streaming06_VirtualTextureCount, NumVirtualTextures);
 
 	/** Streaming Overview stats */
 
diff --git a/Engine/Source/Runtime/Engine/Private/Streaming/TextureStreamingHelpers.h b/Engine/Source/Runtime/Engine/Private/Streaming/TextureStreamingHelpers.h
--- a/Engine/Source/Runtime/Engine/Private/Streaming/TextureStreamingHelpers.h
+++ b/Engine/Source/Runtime/Engine/Private/Streaming/TextureStreamingHelpers.h
@@ -202,6 +202,11 @@ struct FRenderAssetStreamingStats
 	int64 StreamedMeshMem;		// Total memory in bytes of mesh LODs that can be streamed
 	int64 ResidentMeshMem;		// Total memory in bytes of resident mesh LODs
 	int64 EvictedMeshMem;		// Total memory in bytes of evicted mesh LODs
+
+	int64 VirtualTexturePool;			// Pool size used for virtual texture budget
+	int32 NumVirtualTextures;			// Number of virtual textures managed by the streamer
+	int64 RequiredVirtualTexturePool;	// Estimated memory in bytes the streamer would use if there was no limit
+	int64 UsedVirtualTexturePool;		// Estimated memory in bytes the streamer actually use
 };
 
 // Helper to access the level bStaticComponentsRegisteredInStreamingManager flag.
