Fix crash when deleting a world in the content browser, introduced by CL42342766.

diff --git a/Engine/Source/Runtime/Engine/Private/WorldPartition/DataLayer/WorldDataLayers.cpp b/Engine/Source/Runtime/Engine/Private/WorldPartition/DataLayer/WorldDataLayers.cpp
--- a/Engine/Source/Runtime/Engine/Private/WorldPartition/DataLayer/WorldDataLayers.cpp
+++ b/Engine/Source/Runtime/Engine/Private/WorldPartition/DataLayer/WorldDataLayers.cpp
@@ -109,7 +109,6 @@ void AWorldDataLayers::PostUnregisterAllComponents()
 	Super::PostUnregisterAllComponents();
 
 	ResetDataLayerRuntimeStates();
-	InitializeDataLayerRuntimeStates();
 }
 
 #define AWORLDDATALAYERS_UPDATE_REPLICATED_DATALAYERS(ReplicatedArray, SourceArray) \
