Re-enabled async deferred updates with fixes.

diff --git a/Engine/Source/Runtime/Engine/Classes/Engine/World.h b/Engine/Source/Runtime/Engine/Classes/Engine/World.h
--- a/Engine/Source/Runtime/Engine/Classes/Engine/World.h
+++ b/Engine/Source/Runtime/Engine/Classes/Engine/World.h
@@ -33,7 +33,7 @@
 #include "Engine/WorldInitializationValues.h"
 #include "Misc/TransactionallySafeRWLock.h"
 #include "Engine/ScopedMovementUpdate.h"
-
+#include "PSOPrecacheFwd.h"
 #include "World.generated.h"
 
 #define UE_API ENGINE_API
@@ -1639,6 +1639,10 @@ private:
 	UPROPERTY(Transient, NonTransactional)
 	TArray<TObjectPtr<UActorComponent>> ComponentsThatNeedEndOfFrameUpdate_OnGameThread;
 
+	/** Array of components that need to be marked dirty *after* updating, because it wasn't safe to do so during an update. */
+	UPROPERTY(Transient, NonTransactional)
+	TArray<TObjectPtr<UActorComponent>> ComponentsThatNeedEndOfFrameMarkRenderStateDirty;
+
 	struct FEndOfFrameUpdateContext;
 	TArray<FEndOfFrameUpdateContext*> EndOfFrameUpdateContexts;
 	TArray<UE::Tasks::FTask> EndOfFrameUpdateTasks;
@@ -3303,6 +3307,13 @@ public:
 	 */
 	UE_API void RunTickGroup(ETickingGroup Group, bool bBlockTillComplete);
 
+#if UE_WITH_PSO_PRECACHING
+	/**
+	 * Informs the world that a primitive component has completed a PSO pre-cache operation.
+	 */
+	UE_API void AddPrimitiveComponentPSOPrecacheCompleted(UPrimitiveComponent* Component);
+#endif
+
 	/**
 	 * Mark a component as needing an end of frame update
 	 * @param Component - Component to update at the end of the frame
diff --git a/Engine/Source/Runtime/Engine/Private/Components/PrimitiveComponent.cpp b/Engine/Source/Runtime/Engine/Private/Components/PrimitiveComponent.cpp
--- a/Engine/Source/Runtime/Engine/Private/Components/PrimitiveComponent.cpp
+++ b/Engine/Source/Runtime/Engine/Private/Components/PrimitiveComponent.cpp
@@ -5167,7 +5167,11 @@ struct FPSOPrecacheFinishedTask
 			QUICK_SCOPE_CYCLE_COUNTER(STAT_PSOPrecacheFinishedTask);
 			int32 CurrJobSetCompleted = PC->LatestPSOPrecacheJobSetCompleted.load();
 			while (CurrJobSetCompleted < JobSetThatJustCompleted && !PC->LatestPSOPrecacheJobSetCompleted.compare_exchange_weak(CurrJobSetCompleted, JobSetThatJustCompleted)){}
-			PC->MarkRenderStateDirty();
+
+			if (PC->GetWorld())
+			{
+				PC->GetWorld()->AddPrimitiveComponentPSOPrecacheCompleted(PC);
+			}
 		}
 	}
 
diff --git a/Engine/Source/Runtime/Engine/Private/LevelTick.cpp b/Engine/Source/Runtime/Engine/Private/LevelTick.cpp
--- a/Engine/Source/Runtime/Engine/Private/LevelTick.cpp
+++ b/Engine/Source/Runtime/Engine/Private/LevelTick.cpp
@@ -913,6 +913,20 @@ void UWorld::UpdateActorComponentEndOfFrameUpdateState(UActorComponent* Componen
 }
 #endif
 
+#if UE_WITH_PSO_PRECACHING
+void UWorld::AddPrimitiveComponentPSOPrecacheCompleted(UPrimitiveComponent* Component)
+{
+	if (Component->IsRenderStateUpdating())
+	{
+		ComponentsThatNeedEndOfFrameMarkRenderStateDirty.Emplace(Component);
+	}
+	else
+	{
+		Component->MarkRenderStateDirty();
+	}
+}
+#endif
+
 void UWorld::ClearActorComponentEndOfFrameUpdate(UActorComponent* Component)
 {
 	check(!bPostTickComponentUpdate); // can't call this while we are doing the updates
@@ -1318,6 +1332,12 @@ void UWorld::FinishAsyncSendAllEndOfFrameUpdates()
 
 		TRACE_END_REGION(UE_END_OF_FRAME_UPDATE_REGION);
 	}
+
+	for (UActorComponent* Component : ComponentsThatNeedEndOfFrameMarkRenderStateDirty)
+	{
+		Component->MarkRenderStateDirty();
+	}
+	ComponentsThatNeedEndOfFrameMarkRenderStateDirty.Empty();
 }
 
 /**
