Change r.Streaming.DropMips to default Cvar instead of Cheat

diff --git a/Engine/Source/Runtime/Engine/Private/Streaming/TextureStreamingHelpers.cpp b/Engine/Source/Runtime/Engine/Private/Streaming/TextureStreamingHelpers.cpp
--- a/Engine/Source/Runtime/Engine/Private/Streaming/TextureStreamingHelpers.cpp
+++ b/Engine/Source/Runtime/Engine/Private/Streaming/TextureStreamingHelpers.cpp
@@ -150,7 +150,7 @@ TAutoConsoleVariable<int32> CVarStreamingDropMips(
 	TEXT("1: Drop Cached Mips\n")
 	TEXT("2: Drop Cached and Hidden Mips\n")
 	TEXT("3: Drop cached mips and non-inlined LODs of no-ref meshes"),
-	ECVF_Cheat);
+	ECVF_Default);
 
 TAutoConsoleVariable<int32> CVarStreamingHLODStrategy(
 	TEXT("r.Streaming.HLODStrategy"),
