Remove PipelineStateCache.h from PSOPrecacheFwd.h to avoid pulling in large amounts of RHI code that is unneeded in most cases and causes warnings in modules that have enabled missing enum cases in switches

diff --git a/Engine/Source/Runtime/Engine/Classes/Components/PrimitiveComponent.h b/Engine/Source/Runtime/Engine/Classes/Components/PrimitiveComponent.h
--- a/Engine/Source/Runtime/Engine/Classes/Components/PrimitiveComponent.h
+++ b/Engine/Source/Runtime/Engine/Classes/Components/PrimitiveComponent.h
@@ -31,6 +31,7 @@
 #include "HLOD/HLODLevelExclusion.h"
 #include "Stats/Stats.h"
 #include "PSOPrecacheFwd.h"
+#include "PipelineStateCache.h"
 #include "MeshDrawCommandStatsDefines.h"
 #include "PrimitiveSceneInfoData.h"
 #include "PrimitiveComponent.generated.h"
diff --git a/Engine/Source/Runtime/Engine/Public/PSOPrecacheFwd.h b/Engine/Source/Runtime/Engine/Public/PSOPrecacheFwd.h
--- a/Engine/Source/Runtime/Engine/Public/PSOPrecacheFwd.h
+++ b/Engine/Source/Runtime/Engine/Public/PSOPrecacheFwd.h
@@ -9,7 +9,6 @@
 #include "Containers/Array.h"
 #include "EngineDefines.h"
 #include "HAL/Platform.h"
-#include "PipelineStateCache.h"
 
 // General switch that decides whether to compile out some PSO precaching code (most importantly, reduce sizeofs of common classes)
 #ifndef UE_WITH_PSO_PRECACHING
