Fixed 32 bits float scene color setup when Alpha is needed.

#tests pc editor dx12
#rnx

diff --git a/Engine/Source/Runtime/Engine/Private/SceneTexturesConfig.cpp b/Engine/Source/Runtime/Engine/Private/SceneTexturesConfig.cpp
--- a/Engine/Source/Runtime/Engine/Private/SceneTexturesConfig.cpp
+++ b/Engine/Source/Runtime/Engine/Private/SceneTexturesConfig.cpp
@@ -126,7 +126,7 @@ static EPixelFormat GetSceneColorFormat(bool bRequiresAlphaChannel)
 		Format = PF_FloatRGBA;
 	}
 
-	if (bRequiresAlphaChannel && Format != PF_FloatRGBA)
+	if (bRequiresAlphaChannel && !(Format == PF_FloatRGBA || Format == PF_A32B32G32R32F))
 	{
 		UE_CALL_ONCE([]()
 			{
