Spline: Fix reparam table when spline is closed

diff --git a/Engine/Source/Runtime/Engine/Private/Curves/Spline.cpp b/Engine/Source/Runtime/Engine/Private/Curves/Spline.cpp
--- a/Engine/Source/Runtime/Engine/Private/Curves/Spline.cpp
+++ b/Engine/Source/Runtime/Engine/Private/Curves/Spline.cpp
@@ -2329,8 +2329,7 @@ void FNewSpline::UpdateReparamTable()
 		return;
 	}
 
-	const int32 NumPoints = GetSpline().GetNumPoints();
-	const int32 NumSegments = NumPoints - 1;
+	const int32 NumSegments = GetSpline().GetNumSegments();
 	
 	// Start by clearing it
 	ClearAttributeChannel(ParamToLengthAttrName);
