Fixed FGPUSkinPassthroughVertexFactory TargetVertexFactory dangling pointer leading to memory access violation crash

#rnx

diff --git a/Engine/Source/Runtime/Engine/Private/GPUSkinCache.cpp b/Engine/Source/Runtime/Engine/Private/GPUSkinCache.cpp
--- a/Engine/Source/Runtime/Engine/Private/GPUSkinCache.cpp
+++ b/Engine/Source/Runtime/Engine/Private/GPUSkinCache.cpp
@@ -2849,7 +2849,14 @@ void FGPUSkinCache::SetEntryGPUSkin(FGPUSkinCacheEntry* Entry, FSkeletalMeshObje
 		//   accessing invalid bone data if the original owner is released. The original owner *does* dequeue on release,
 		//   however, the transfer nulls the old entry.
 		FGPUSkinCache::Dequeue(Entry);
-		
+
+		for (FGPUSkinCacheEntry::FSectionDispatchData& DispatchData : Entry->DispatchData)
+		{
+			// Reset VF pointers to ensure IsSourceFactoryValid returns false when entry will get updated in next ProcessEntry call
+			DispatchData.SourceVertexFactory = nullptr;
+			DispatchData.TargetVertexFactory = nullptr;
+		}
+
 		Entry->GPUSkin = Skin;
 	}
 }
