Added default initialization of fields of FMaterialEnumerationEntry.

diff --git a/Engine/Source/Runtime/Engine/Public/Materials/MaterialEnumeration.h b/Engine/Source/Runtime/Engine/Public/Materials/MaterialEnumeration.h
--- a/Engine/Source/Runtime/Engine/Public/Materials/MaterialEnumeration.h
+++ b/Engine/Source/Runtime/Engine/Public/Materials/MaterialEnumeration.h
@@ -36,11 +36,11 @@ struct FMaterialEnumerationEntry
 
 	// Entry identifier.
 	UPROPERTY(EditAnywhere, Category=MaterialEnumerationEntry)
-	FName Name; 
+	FName Name{};
 
 	// Associated integer value.
 	UPROPERTY(EditAnywhere, Category=MaterialEnumerationEntry)
-	int32 Value;
+	int32 Value{};
 };
 
 UCLASS(MinimalAPI)
