Restore problematic include wrapped in include order to avoid surprise breakages for licensees
#rnx

diff --git a/Engine/Source/Runtime/Engine/Public/PSOPrecacheFwd.h b/Engine/Source/Runtime/Engine/Public/PSOPrecacheFwd.h
--- a/Engine/Source/Runtime/Engine/Public/PSOPrecacheFwd.h
+++ b/Engine/Source/Runtime/Engine/Public/PSOPrecacheFwd.h
@@ -9,6 +9,9 @@
 #include "Containers/Array.h"
 #include "EngineDefines.h"
 #include "HAL/Platform.h"
+#if UE_ENABLE_INCLUDE_ORDER_DEPRECATED_IN_5_7
+#include "PipelineStateCache.h"
+#endif
 
 // General switch that decides whether to compile out some PSO precaching code (most importantly, reduce sizeofs of common classes)
 #ifndef UE_WITH_PSO_PRECACHING
