Improve warnings for scene component culling mismatch between parents and children

diff --git a/Engine/Source/Runtime/Engine/Private/Components/SceneComponent.cpp b/Engine/Source/Runtime/Engine/Private/Components/SceneComponent.cpp
--- a/Engine/Source/Runtime/Engine/Private/Components/SceneComponent.cpp
+++ b/Engine/Source/Runtime/Engine/Private/Components/SceneComponent.cpp
@@ -465,7 +465,7 @@ static void UpdateAttachedIsEditorOnly(USceneComponent* ComponentThatChanged)
 	}
 }
 
-static bool SceneComponentNeedsLoadForTarget(USceneComponent const* SceneComponentObject, const ITargetPlatform* TargetPlatform)
+static bool SceneComponentNeedsLoadForTarget(USceneComponent const* SceneComponentObject, const ITargetPlatform* TargetPlatform, USceneComponent const* ParentThatWillBeCulled = nullptr)
 {
 	if(UDeviceProfile* DeviceProfile = UDeviceProfileManager::Get().FindProfile(TargetPlatform->IniPlatformName()))
 	{
@@ -479,26 +479,34 @@ static bool SceneComponentNeedsLoadForTarget(USceneComponent const* SceneCompone
 				// Check component's detail mode.
 				// If e.g. the component's detail mode is High and the platform detail is Medium,
 				// then we should cull it.
-				if((int32)SceneComponentObject->DetailMode > CVarDetailMode)
-				{
-					return false;
-				}
+				const bool bKeep = (int32)SceneComponentObject->DetailMode <= CVarDetailMode;
+				UE_CLOGFMT(bKeep && ParentThatWillBeCulled, LogSceneComponent, Warning, 
+					"SceneComponent {ChildComponent} will not be culled because its detail mode {DetailMode} is less than the maximum detail mode for the target platform {PlatformMaxDetailMode}, "
+					"but its parent {ParentComponent} will be culled. Make sure that attached components have a detail mode equal or lower to their parent.", 
+					FSoftObjectPath(SceneComponentObject), SceneComponentObject->DetailMode, CVarDetailMode, FSoftObjectPath(ParentThatWillBeCulled));
+				return bKeep;
 			}
 		}
 	}
 
-	return TargetPlatform->AllowsEditorObjects() || !SceneComponentObject->IsEditorOnly();
+	const bool bKeep = TargetPlatform->AllowsEditorObjects() || !SceneComponentObject->IsEditorOnly();
+	UE_CLOGFMT(bKeep && ParentThatWillBeCulled, LogSceneComponent, Warning, 
+		"SceneComponent {ChildComponent} will not be culled because it is not marked editor-only, but its parent {ParentComponent} will be culled."
+		"Make sure that all attached children of editor-only components are also editor-only.", 
+		FSoftObjectPath(SceneComponentObject), FSoftObjectPath(ParentThatWillBeCulled));
+	return bKeep;
 }
 
-static bool CheckDescendantsAreAlsoCulledForTarget(USceneComponent const* SceneComponentObject, const ITargetPlatform* TargetPlatform)
+// Returns number of child components not culled
+static int32 CheckDescendantsAreAlsoCulledForTarget(USceneComponent const* SceneComponentObject, const ITargetPlatform* TargetPlatform)
 {
 	if (!ensure(SceneComponentObject != nullptr))
 	{
 		return 0;
 	}
 
-	TArray<USceneComponent*> AttachedChildren = SceneComponentObject->GetAttachChildren();
-
+	const TArray<USceneComponent*>& AttachedChildren = SceneComponentObject->GetAttachChildren();
+	int32 NumNotCulled = 0;
 	for (USceneComponent* ChildSceneComponent : AttachedChildren)
 	{
 		if (!ChildSceneComponent)
@@ -506,33 +514,34 @@ static bool CheckDescendantsAreAlsoCulledForTarget(USceneComponent const* SceneC
 			continue;
 		}
 
-		UE_LOG(LogSceneComponent, Display, TEXT("Checking attached component %s for culling"), *GetPathNameSafe(ChildSceneComponent));
-		if (SceneComponentNeedsLoadForTarget(ChildSceneComponent, TargetPlatform))
+		UE_LOGFMT(LogSceneComponent, Verbose, "Checking attached component {ChildComponent} for culling", FSoftObjectPath(ChildSceneComponent));
+		if (SceneComponentNeedsLoadForTarget(ChildSceneComponent, TargetPlatform, /* ParentThatWillBeCulled */ SceneComponentObject))
 		{
-			UE_LOG(LogSceneComponent, Display, TEXT("Scene component %s will not be culled"), *GetPathNameSafe(ChildSceneComponent));
-			return false;
+			UE_LOGFMT(LogSceneComponent, Verbose, "Scene component {ChildComponent} will not be culled", FSoftObjectPath(ChildSceneComponent));
+			NumNotCulled++;
 		}
 		else
 		{
-			UE_LOG(LogSceneComponent, Display, TEXT("Scene component %s will be culled"), *GetPathNameSafe(ChildSceneComponent));
+			UE_LOGFMT(LogSceneComponent, Verbose, "Scene component {ChildComponent} will be culled", FSoftObjectPath(ChildSceneComponent));
 		}
 	}
 
-	return true;
+	return NumNotCulled;
 }
 
 bool USceneComponent::NeedsLoadForTargetPlatform(const ITargetPlatform* TargetPlatform) const
 {
 	if(!SceneComponentNeedsLoadForTarget(this, TargetPlatform))
 	{
-		UE_LOG(LogSceneComponent, Display, TEXT("Scene component %s will be culled for target, checking children"), *GetPathName());
+		UE_LOGFMT(LogSceneComponent, Log, "Scene component {ParentComponent} will be culled for target, checking children", FSoftObjectPath(this));
 		// Also check whether any of our children are culled.
-		bool bDescendantsCulled = CheckDescendantsAreAlsoCulledForTarget(this, TargetPlatform);
+		int32 DescendantsNotCulled = CheckDescendantsAreAlsoCulledForTarget(this, TargetPlatform);
 
 		// Child not culled, so warn
-		if(!bDescendantsCulled)
+		if(DescendantsNotCulled != 0)
 		{
-			UE_LOG(LogSceneComponent, Warning, TEXT("Component %s not removed from client data because descendants were not also not removed."), *GetPathName());
+			UE_LOGFMT(LogSceneComponent, Warning, "Component {ParentComponent} not removed from client data because {NumNotCulled} descendants could not be removed. See prior warnings for details.",
+				FSoftObjectPath(this), DescendantsNotCulled);
 			return true;
 		}
 
