Make it possible to not change to fixed timestep at low dt (i.e. bypass the cvar). Use this in control rig physics.

diff --git a/Engine/Plugins/Experimental/PhysicsControl/Source/ControlRigPhysics/Private/RigPhysicsSimulation.cpp b/Engine/Plugins/Experimental/PhysicsControl/Source/ControlRigPhysics/Private/RigPhysicsSimulation.cpp
--- a/Engine/Plugins/Experimental/PhysicsControl/Source/ControlRigPhysics/Private/RigPhysicsSimulation.cpp
+++ b/Engine/Plugins/Experimental/PhysicsControl/Source/ControlRigPhysics/Private/RigPhysicsSimulation.cpp
@@ -1642,6 +1642,9 @@ void FRigPhysicsSimulation::StepSimulation(
 		SolverSettings.ProjectionIterations,
 		SolverSettings.bUseManifolds);
 
+	Simulation->SetUseMinStepTime(false);
+	Simulation->SetUseFixedStepTolerance(false);
+
 	Chaos::FCollisionDetectorSettings CollisionDetectorSettings = Simulation->GetCollisionDetectorSettings();
 	CollisionDetectorSettings.BoundsVelocityInflation = SolverSettings.BoundsVelocityMultiplier;
 	CollisionDetectorSettings.MaxVelocityBoundsExpansion = SolverSettings.MaxVelocityBoundsExpansion;
diff --git a/Engine/Plugins/Experimental/PhysicsControl/Source/ControlRigPhysics/Private/RigPhysicsSimulationSpace.cpp b/Engine/Plugins/Experimental/PhysicsControl/Source/ControlRigPhysics/Private/RigPhysicsSimulationSpace.cpp
--- a/Engine/Plugins/Experimental/PhysicsControl/Source/ControlRigPhysics/Private/RigPhysicsSimulationSpace.cpp
+++ b/Engine/Plugins/Experimental/PhysicsControl/Source/ControlRigPhysics/Private/RigPhysicsSimulationSpace.cpp
@@ -419,10 +419,11 @@ FRigPhysicsSimulation::FSimulationSpaceData FRigPhysicsSimulation::UpdateSimulat
 	if (GEngine && CVarControlRigPhysicsShowSimulationSpaceInfo.GetValueOnAnyThread())
 	{
 		GEngine->AddOnScreenDebugMessage(INDEX_NONE, 0.f, FColor::Yellow,
-			FString::Printf(TEXT("Sim %s dt %.2fms Space speed %.3f accel %.3f"),
+			FString::Printf(TEXT("Sim %s dt %6.2fms Space P %35s V %35s A %35s"),
 				*OwningControlRig->GetName(), Dt * 1000.0f,
-				SimulationSpaceData.LinearVelocity.Length(),
-				SimulationSpaceData.LinearAcceleration.Length()));
+				*SimulationSpaceState.SimulationSpaceTM.GetLocation().ToString(),
+				*SimulationSpaceData.LinearVelocity.ToString(),
+				*SimulationSpaceData.LinearAcceleration.ToString()));
 	}
 
 	return SimulationSpaceData;
diff --git a/Engine/Source/Runtime/Engine/Private/PhysicsEngine/ImmediatePhysics/ImmediatePhysicsChaos/ImmediatePhysicsSimulation_Chaos.cpp b/Engine/Source/Runtime/Engine/Private/PhysicsEngine/ImmediatePhysics/ImmediatePhysicsChaos/ImmediatePhysicsSimulation_Chaos.cpp
--- a/Engine/Source/Runtime/Engine/Private/PhysicsEngine/ImmediatePhysics/ImmediatePhysicsChaos/ImmediatePhysicsSimulation_Chaos.cpp
+++ b/Engine/Source/Runtime/Engine/Private/PhysicsEngine/ImmediatePhysics/ImmediatePhysicsChaos/ImmediatePhysicsSimulation_Chaos.cpp
@@ -355,6 +355,8 @@ namespace ImmediatePhysics_Chaos
 		Chaos::FReal RollingAverageStepTime;
 		int32 NumRollingAverageStepTimes;
 		int32 MaxNumRollingAverageStepTimes; // Set to 1 to disable smoothing
+		bool bUseMinStepTime = true;
+		bool bUseFixedStepTolerance = true;
 
 		bool bActorsDirty;
 	};
@@ -868,6 +870,17 @@ namespace ImmediatePhysics_Chaos
 		}
 	}
 
+	void FSimulation::SetUseMinStepTime(bool bUse)
+	{
+		Implementation->bUseMinStepTime = bUse;
+	}
+
+	void FSimulation::SetUseFixedStepTolerance(bool bUse)
+	{
+		Implementation->bUseFixedStepTolerance = bUse;
+	}
+
+
 	void FSimulation::DebugDraw()
 	{
 #if CHAOS_DEBUG_DRAW
@@ -1020,7 +1033,7 @@ namespace ImmediatePhysics_Chaos
 
 		// Use fixed step mode anyway if StepTime is too low
 		// This can prevent collision push resolution introducing large velocities at small DTs
-		if ((FixedStepTime <= 0.0f) && (StepTime < ChaosImmediate_Evolution_MinStepTime))
+		if (Implementation->bUseMinStepTime && (FixedStepTime <= 0.0f) && (StepTime < ChaosImmediate_Evolution_MinStepTime))
 		{
 			FixedStepTime = ChaosImmediate_Evolution_MinStepTime;
 		}
@@ -1032,7 +1045,7 @@ namespace ImmediatePhysics_Chaos
 			StepTime = FixedStepTime;
 			NumSteps = FMath::FloorToInt(DeltaTime / StepTime);
 			FReal RemainderTime = DeltaTime - NumSteps * StepTime;
-			if (RemainderTime > ChaosImmediate_Evolution_FixedStepTolerance * StepTime)
+			if (!Implementation->bUseFixedStepTolerance || RemainderTime > ChaosImmediate_Evolution_FixedStepTolerance * StepTime)
 			{
 				++NumSteps;
 				RewindTime = StepTime - RemainderTime;
diff --git a/Engine/Source/Runtime/Engine/Public/Physics/ImmediatePhysics/ImmediatePhysicsChaos/ImmediatePhysicsSimulation_Chaos.h b/Engine/Source/Runtime/Engine/Public/Physics/ImmediatePhysics/ImmediatePhysicsChaos/ImmediatePhysicsSimulation_Chaos.h
--- a/Engine/Source/Runtime/Engine/Public/Physics/ImmediatePhysics/ImmediatePhysicsChaos/ImmediatePhysicsSimulation_Chaos.h
+++ b/Engine/Source/Runtime/Engine/Public/Physics/ImmediatePhysics/ImmediatePhysicsChaos/ImmediatePhysicsSimulation_Chaos.h
@@ -146,8 +146,16 @@ namespace ImmediatePhysics_Chaos
 
 		// Additional settings for less commonly adjusted parameters. As in SetSolverSettings,
 		// invalid (negative) values will leave that value unchanged from the default.
-		ENGINE_API void SetMaxNumRollingAverageStepTimes(
-			const int32 MaxNumRollingAverageStepTimes);
+		UE_INTERNAL
+		ENGINE_API void SetMaxNumRollingAverageStepTimes(const int32 MaxNumRollingAverageStepTimes);
+
+		// Sets whether or not to use the minimum step time, which is itself set with a cvar p.Chaos.ImmPhys.MinStepTime
+		UE_INTERNAL
+		ENGINE_API void SetUseMinStepTime(bool bUse);
+
+		// Sets whether or not to use the fixed timestep tolerance for rewinding, which is itself set with a cvar p.Chaos.ImmPhys.FixedStepTolerance
+		UE_INTERNAL
+		ENGINE_API void SetUseFixedStepTolerance(bool bUse);
 
 		/** Explicit debug draw path if the use case needs it to happen at a point outside of the simulation **/
 		ENGINE_API void DebugDraw();
