[Asset LOD streaming] Add a new DropMips mode that drops only no-ref mesh LODs

diff --git a/Engine/Source/Runtime/Engine/Private/Streaming/AsyncTextureStreaming.cpp b/Engine/Source/Runtime/Engine/Private/Streaming/AsyncTextureStreaming.cpp
--- a/Engine/Source/Runtime/Engine/Private/Streaming/AsyncTextureStreaming.cpp
+++ b/Engine/Source/Runtime/Engine/Private/Streaming/AsyncTextureStreaming.cpp
@@ -313,7 +313,7 @@ void FAsyncRenderAssetStreamingData::UpdatePerfectWantedMips_Async(FStreamingRen
 
 			// Ignore bUseUnkownRefHeuristic if they haven't been used in the last 90 sec.
 			// If critical, it must be implemented using the ForceFullyLoad logic.
-			if (StreamingRenderAsset.bUseUnkownRefHeuristic && StreamingRenderAsset.LastRenderTime < 90.0f && (Settings.DropMips != 3 || AssetType == EStreamableRenderAssetType::Texture))
+			if (StreamingRenderAsset.bUseUnkownRefHeuristic && StreamingRenderAsset.LastRenderTime < 90.0f && ((Settings.DropMips != 3 && Settings.DropMips != 4) || AssetType == EStreamableRenderAssetType::Texture))
 			{
 				if (bOutputToLog) UE_LOG(LogContentStreaming, Log,  TEXT("  UnkownRef"));
 				MaxSize = FMath::Max(MaxSize, MaxAllowedSize); // affected by HiddenPrimitiveScale
@@ -841,7 +841,7 @@ void FRenderAssetStreamingMipCalcTask::UpdateBudgetedMips_Async()
 	// Handle drop mips debug option
 	//*************************************
 #if !UE_BUILD_SHIPPING
-	if (Settings.DropMips > 0)
+	if (Settings.DropMips > 0 && Settings.DropMips != 4)
 	{
 		for (FStreamingRenderAsset& StreamingRenderAsset : StreamingRenderAssets)
 		{
diff --git a/Engine/Source/Runtime/Engine/Private/Streaming/TextureStreamingHelpers.cpp b/Engine/Source/Runtime/Engine/Private/Streaming/TextureStreamingHelpers.cpp
--- a/Engine/Source/Runtime/Engine/Private/Streaming/TextureStreamingHelpers.cpp
+++ b/Engine/Source/Runtime/Engine/Private/Streaming/TextureStreamingHelpers.cpp
@@ -149,7 +149,8 @@ TAutoConsoleVariable<int32> CVarStreamingDropMips(
 	TEXT("0: Drop No Mips \n")
 	TEXT("1: Drop Cached Mips\n")
 	TEXT("2: Drop Cached and Hidden Mips\n")
-	TEXT("3: Drop cached mips and non-inlined LODs of no-ref meshes"),
+	TEXT("3: Drop cached mips and non-inlined LODs of no-ref meshes\n")
+	TEXT("4: Drop non-inlined LODs of no-ref meshes"),
 	ECVF_Default);
 
 TAutoConsoleVariable<int32> CVarStreamingHLODStrategy(
