[CrashReport][Ensure] UnrealEditor_Engine!UAnimInstance::Montage_Advance()

diff --git a/Engine/Source/Runtime/Engine/Private/Animation/AnimInstance.cpp b/Engine/Source/Runtime/Engine/Private/Animation/AnimInstance.cpp
--- a/Engine/Source/Runtime/Engine/Private/Animation/AnimInstance.cpp
+++ b/Engine/Source/Runtime/Engine/Private/Animation/AnimInstance.cpp
@@ -1940,8 +1940,9 @@ void UAnimInstance::Montage_Advance(float DeltaSeconds)
 	for (int32 InstanceIndex = 0; InstanceIndex < MontageInstances.Num(); InstanceIndex++)
 	{
 		FAnimMontageInstance* const MontageInstance = MontageInstances[InstanceIndex];
-		// should never be NULL and the montage should never be null or garbage
-		if (ensure(MontageInstance && IsValid(MontageInstance->Montage)))
+		// should never be NULL
+		ensure(MontageInstance);
+		if (MontageInstance && IsValid(MontageInstance->Montage))
 		{
 			bool const bUsingBlendedRootMotion = (RootMotionMode == ERootMotionMode::RootMotionFromEverything);
 			bool const bNoRootMotionExtraction = (RootMotionMode == ERootMotionMode::NoRootMotionExtraction);
