Set the initialisation of bWantsEditorEffects to false in order to avoid non determinism bugs when initialising FPrimitiveSceneProxyDesc with it's default constructor, resulting in incorrect selection outlines in certain editor views.

diff --git a/Engine/Source/Runtime/Engine/Public/PrimitiveSceneProxyDesc.h b/Engine/Source/Runtime/Engine/Public/PrimitiveSceneProxyDesc.h
--- a/Engine/Source/Runtime/Engine/Public/PrimitiveSceneProxyDesc.h
+++ b/Engine/Source/Runtime/Engine/Public/PrimitiveSceneProxyDesc.h
@@ -71,6 +71,7 @@ struct FPrimitiveSceneProxyDesc
 		bShadowIndirectOnly = false;
 		bShouldRenderProxyFallbackToDefaultMaterial = false;
 		bShouldRenderSelected = false;
+		bWantsEditorEffects = false;
 
 #if WITH_EDITOR
 		bIsOwnedByFoliage = false;
