[HWRT] Use inline allocator for GeometryHandles in FRayTracingGeometryGroup.
- Projects using default ray tracing proxies always have 2 geometries per group for Nanite meshes, so using inline allocator avoids cache misses in Tick(...).

(Project with ~5K geometries / ~2.5K groups: 0.7ms -> 0.5ms)

diff --git a/Engine/Source/Runtime/Engine/Private/Rendering/RayTracingGeometryManager.cpp b/Engine/Source/Runtime/Engine/Private/Rendering/RayTracingGeometryManager.cpp
--- a/Engine/Source/Runtime/Engine/Private/Rendering/RayTracingGeometryManager.cpp
+++ b/Engine/Source/Runtime/Engine/Private/Rendering/RayTracingGeometryManager.cpp
@@ -550,6 +550,7 @@ void FRayTracingGeometryManager::Tick(FRHICommandList& RHICmdList)
 #endif
 
 	const bool bUsingReferenceBasedResidency = IsRayTracingUsingReferenceBasedResidency();
+	const int32 NumAlwaysResidentLODs = GRayTracingNumAlwaysResidentLODs;
 
 	checkf(bUsingReferenceBasedResidency || (ReferencedGeometryHandles.IsEmpty() && ReferencedGeometryGroups.IsEmpty() && ReferencedGeometryGroupsForDynamicUpdate.IsEmpty()),
 		TEXT("ReferencedGeometryHandles, ReferencedGeometryGroups and ReferencedGeometryGroupsForDynamicUpdate are expected to be empty when not using reference based residency"));
@@ -579,7 +580,7 @@ void FRayTracingGeometryManager::Tick(FRHICommandList& RHICmdList)
 			}
 
 			const FRayTracingGeometryGroup& Group = RegisteredGroups[RegisteredGeometry.Geometry->GroupHandle];
-			if (RegisteredGeometry.Geometry->LODIndex >= Group.GeometryHandles.Num() - GRayTracingNumAlwaysResidentLODs)
+			if (RegisteredGeometry.Geometry->LODIndex >= Group.GeometryHandles.Num() - NumAlwaysResidentLODs)
 			{
 				AlwaysResidentGeometries.Add(RegisteredGeometry.Geometry->RayTracingGeometryHandle);
 				TotalAlwaysResidentSize += RegisteredGeometry.Size;
@@ -713,8 +714,9 @@ void FRayTracingGeometryManager::Tick(FRHICommandList& RHICmdList)
 			checkf(RegisteredGroups.IsValidIndex(GroupHandle), TEXT("RayTracingGeometryGroupHandle must be valid"));
 
 			const FRayTracingGeometryGroup& Group = RegisteredGroups[GroupHandle];
+			const int32 NumGeometriesInGroup = Group.GeometryHandles.Num();
 
-			for (uint8 LODIndex = Group.CurrentFirstLODIdx; LODIndex < Group.GeometryHandles.Num(); ++LODIndex)
+			for (uint8 LODIndex = Group.CurrentFirstLODIdx; LODIndex < NumGeometriesInGroup; ++LODIndex)
 			{
 				FGeometryHandle GeometryHandle = Group.GeometryHandles[LODIndex];
 
@@ -725,7 +727,7 @@ void FRayTracingGeometryManager::Tick(FRHICommandList& RHICmdList)
 
 					RequestedSize += RegisteredGeometry.Size;
 
-					if (LODIndex >= Group.GeometryHandles.Num() - GRayTracingNumAlwaysResidentLODs)
+					if (LODIndex >= NumGeometriesInGroup - NumAlwaysResidentLODs)
 					{
 						checkf(!RegisteredGeometry.bEvicted, TEXT("Always resident ray tracing geometry was unexpectely evicted."));
 					}
diff --git a/Engine/Source/Runtime/Engine/Public/Rendering/RayTracingGeometryManager.h b/Engine/Source/Runtime/Engine/Public/Rendering/RayTracingGeometryManager.h
--- a/Engine/Source/Runtime/Engine/Public/Rendering/RayTracingGeometryManager.h
+++ b/Engine/Source/Runtime/Engine/Public/Rendering/RayTracingGeometryManager.h
@@ -96,7 +96,7 @@ private:
 
 	struct FRayTracingGeometryGroup
 	{
-		TArray<FGeometryHandle> GeometryHandles;
+		TArray<FGeometryHandle, TInlineAllocator<2>> GeometryHandles;
 
 		TSet<FPrimitiveSceneProxy*> ProxiesWithCachedRayTracingState;
 
