Add a null pointer check in ADebugCameraController::ConsoleCommand to make sure after the console command has been executed on the previous player controller that its Player pointer hasn't been set to null, in case the console command changed the player controller

diff --git a/Engine/Source/Runtime/Engine/Private/DebugCameraController.cpp b/Engine/Source/Runtime/Engine/Private/DebugCameraController.cpp
--- a/Engine/Source/Runtime/Engine/Private/DebugCameraController.cpp
+++ b/Engine/Source/Runtime/Engine/Private/DebugCameraController.cpp
@@ -296,7 +296,14 @@ FString ADebugCameraController::ConsoleCommand(const FString& Cmd,bool bWriteToL
 			{
 				Player->PlayerController = OriginalControllerRef;
 				Player->Exec( GetWorld(), Line, StrOut);
-				Player->PlayerController = this;
+
+				// The Player pointer could be null if the calling console command has changed the player controller back to the
+				// original one. In that case, the Player pointer will be null so we can't switch the player controller back
+				// (and wouldn't want to)
+				if (Player != nullptr)
+				{
+					Player->PlayerController = this;
+				}
 			}
 		}
 
