FWorldCookPackageSplitter: introduce this new cook package splitter for UWorld which allows classes to register subsplitters, like FWorldPartitionCookPackageSplitter that is now a subsplitter for UWorld.
This will be used for scene graph streaming generation in the future.

diff --git a/Engine/Source/Editor/UnrealEd/Public/CookPackageSplitter.h b/Engine/Source/Editor/UnrealEd/Public/CookPackageSplitter.h
--- a/Engine/Source/Editor/UnrealEd/Public/CookPackageSplitter.h
+++ b/Engine/Source/Editor/UnrealEd/Public/CookPackageSplitter.h
@@ -428,6 +428,9 @@ private:
  * // In cpp
  * REGISTER_COOKPACKAGE_SPLITTER(FMyCookPackageSplitter, UMySplitDataClass);
  */
+#define DEFINE_COOKPACKAGE_SPLITTER(SplitterClass, SplitDataClass) \
+friend class PREPROCESSOR_JOIN(PREPROCESSOR_JOIN(SplitterClass, SplitDataClass), _Register)
+
 #define REGISTER_COOKPACKAGE_SPLITTER(SplitterClass, SplitDataClass) \
 class PREPROCESSOR_JOIN(PREPROCESSOR_JOIN(SplitterClass, SplitDataClass), _Register) \
 	: public UE::Cook::Private::FRegisteredCookPackageSplitter \
diff --git a/Engine/Source/Runtime/Engine/Private/WorldCookPackageSplitter.cpp b/Engine/Source/Runtime/Engine/Private/WorldCookPackageSplitter.cpp
new file mode 100644
--- /dev/null
+++ b/Engine/Source/Runtime/Engine/Private/WorldCookPackageSplitter.cpp
@@ -0,0 +1,121 @@
+// Copyright Epic Games, Inc. All Rights Reserved.
+
+#include "WorldCookPackageSplitter.h"
+
+#if WITH_EDITOR
+#include "Engine/World.h"
+
+REGISTER_COOKPACKAGE_SPLITTER(FWorldCookPackageSplitter, UWorld);
+
+TMap<UClass*, FWorldCookPackageSplitter::FSubSplitterFactory> FWorldCookPackageSplitter::RegisteredCookPackageSubSplitterFactories;
+
+void FWorldCookPackageSplitter::RegisterCookPackageSubSplitterFactory(UClass* Class, FWorldCookPackageSplitter::FSubSplitterFactory Factory)
+{
+	check(!RegisteredCookPackageSubSplitterFactories.Contains(Class));
+	RegisteredCookPackageSubSplitterFactories.Emplace(Class, MoveTemp(Factory));
+}
+
+void FWorldCookPackageSplitter::UnregisterCookPackageSubSplitterFactory(UClass* Class)
+{
+	check(RegisteredCookPackageSubSplitterFactories.Contains(Class));
+	RegisteredCookPackageSubSplitterFactories.Remove(Class);
+}
+
+bool FWorldCookPackageSplitter::ShouldSplit(UObject* SplitData)
+{
+	bool bResult = false;
+	if (UWorld* World = Cast<UWorld>(SplitData))
+	{
+		ForEachRegisteredCookPackageSubSplitterFactories(
+			[&bResult, World](const FWorldCookPackageSplitter::FSubSplitterFactory& WorldCookPackageSplitterFactory)
+			{
+				bResult = WorldCookPackageSplitterFactory.ShouldSplit(World);
+				return !bResult;
+			});
+	}
+	return bResult;
+}
+
+void FWorldCookPackageSplitter::Teardown(ICookPackageSplitter::ETeardown Status)
+{
+	ForEachCookPackageSubSplitters(
+		[Status](FWorldCookPackageSplitter::FSubSplitter* WorldCookPackageSplitter)
+		{
+			WorldCookPackageSplitter->Teardown(Status);
+		});
+}
+
+bool FWorldCookPackageSplitter::UseInternalReferenceToAvoidGarbageCollect()
+{
+	return true;
+}
+
+ICookPackageSplitter::EGeneratedRequiresGenerator FWorldCookPackageSplitter::DoesGeneratedRequireGenerator()
+{
+	return EGeneratedRequiresGenerator::Populate;
+}
+
+bool FWorldCookPackageSplitter::RequiresGeneratorPackageDestructBeforeResplit()
+{
+	return true;
+}
+
+TArray<ICookPackageSplitter::FGeneratedPackage> FWorldCookPackageSplitter::GetGenerateList(const UPackage* OwnerPackage, const UObject* OwnerObject)
+{
+	const UWorld* OwnerWorld = CastChecked<const UWorld>(OwnerObject);
+
+	check(!CookPackageSubSplitters.Num());
+	ForEachRegisteredCookPackageSubSplitterFactories(
+		[this, OwnerWorld](const FWorldCookPackageSplitter::FSubSplitterFactory& WorldCookPackageSplitterFactory)
+		{
+			if (WorldCookPackageSplitterFactory.ShouldSplit(OwnerWorld))
+			{
+				CookPackageSubSplitters.Emplace(WorldCookPackageSplitterFactory.MakeInstance());
+			}
+		});
+	check(CookPackageSubSplitters.Num());
+
+	TArray<ICookPackageSplitter::FGeneratedPackage> PackagesToGenerate;
+	ForEachCookPackageSubSplitters(
+		[&PackagesToGenerate, OwnerPackage, OwnerWorld](FWorldCookPackageSplitter::FSubSplitter* WorldCookPackageSplitter)
+		{
+			PackagesToGenerate.Append(WorldCookPackageSplitter->GetGenerateList(OwnerPackage, OwnerWorld));
+		});
+	return PackagesToGenerate;
+}
+
+bool FWorldCookPackageSplitter::PopulateGeneratedPackage(FPopulateContext& PopulateContext)
+{
+	bool bResult = false;
+	ForEachCookPackageSubSplitters(
+		[&bResult, &PopulateContext](FWorldCookPackageSplitter::FSubSplitter* WorldCookPackageSplitter)
+		{
+			bResult = WorldCookPackageSplitter->PopulateGeneratedPackage(PopulateContext);
+			return bResult;
+		});
+	return bResult;
+}
+
+bool FWorldCookPackageSplitter::PopulateGeneratorPackage(FPopulateContext& PopulateContext)
+{
+	bool bResult = false;
+	ForEachCookPackageSubSplitters(
+		[&bResult, &PopulateContext](FWorldCookPackageSplitter::FSubSplitter* WorldCookPackageSplitter)
+		{
+			bResult = WorldCookPackageSplitter->PopulateGeneratorPackage(PopulateContext);
+			return bResult;
+		});
+	return bResult;
+}
+
+void FWorldCookPackageSplitter::OnOwnerReloaded(UPackage* OwnerPackage, UObject* OwnerObject)
+{
+	const UWorld* OwnerWorld = CastChecked<const UWorld>(OwnerObject);
+	ForEachCookPackageSubSplitters(
+		[OwnerPackage, OwnerWorld](FWorldCookPackageSplitter::FSubSplitter* WorldCookPackageSplitter)
+		{
+			WorldCookPackageSplitter->OnOwnerReloaded(OwnerPackage, OwnerWorld);
+		});
+}
+
+#endif // WITH_EDITOR
\ No newline at end of file
diff --git a/Engine/Source/Runtime/Engine/Private/WorldPartition/WorldPartition.cpp b/Engine/Source/Runtime/Engine/Private/WorldPartition/WorldPartition.cpp
--- a/Engine/Source/Runtime/Engine/Private/WorldPartition/WorldPartition.cpp
+++ b/Engine/Source/Runtime/Engine/Private/WorldPartition/WorldPartition.cpp
@@ -18,6 +18,7 @@
 #include "WorldPartition/DataLayer/DataLayerManager.h"
 #include "WorldPartition/DataLayer/ExternalDataLayerManager.h"
 #include "WorldPartition/WorldPartitionSettings.h"
+#include "WorldPartition/WorldPartitionCookPackageSplitter.h"
 #include "GameFramework/WorldSettings.h"
 #include "ProfilingDebugging/ScopedTimers.h"
 #include "UObject/UE5MainStreamObjectVersion.h"
@@ -383,6 +384,24 @@ UWorldPartition::UWorldPartition(const FObjectInitializer& ObjectInitializer)
 #if WITH_EDITOR
 	bAllowShowingHLODsInEditor = true;
 	WorldPartitionStreamingPolicyClass = UWorldPartitionLevelStreamingPolicy::StaticClass();
+
+	if (HasAnyFlags(RF_ClassDefaultObject) && !HasAnyFlags(RF_ImmutableDefaultObject))
+	{
+		FWorldCookPackageSplitter::RegisterCookPackageSubSplitterFactory(
+			GetClass(),
+			{
+				.ShouldSplit =
+					[](TNotNull<const UWorld*> World) -> bool
+        			{
+						return World->IsPartitionedWorld();
+					},
+				.MakeInstance =
+					[]() -> FWorldCookPackageSplitter::FSubSplitter*
+        			{
+						return new FWorldPartitionCookPackageSplitter;
+					}
+			});
+	}
 #endif
 }
 
@@ -1723,6 +1742,14 @@ bool UWorldPartition::ResolveSubobject(const TCHAR* SubObjectPath, UObject*& Out
 void UWorldPartition::BeginDestroy()
 {
 	check(InitState == EWorldPartitionInitState::Uninitialized);
+
+#if WITH_EDITOR
+	if (HasAnyFlags(RF_ClassDefaultObject) && !HasAnyFlags(RF_ImmutableDefaultObject))
+	{
+		FWorldCookPackageSplitter::UnregisterCookPackageSubSplitterFactory(GetClass());
+	}
+#endif
+
 	Super::BeginDestroy();
 }
 
diff --git a/Engine/Source/Runtime/Engine/Private/WorldPartition/WorldPartitionCookPackageSplitter.cpp b/Engine/Source/Runtime/Engine/Private/WorldPartition/WorldPartitionCookPackageSplitter.cpp
--- a/Engine/Source/Runtime/Engine/Private/WorldPartition/WorldPartitionCookPackageSplitter.cpp
+++ b/Engine/Source/Runtime/Engine/Private/WorldPartition/WorldPartitionCookPackageSplitter.cpp
@@ -12,51 +12,11 @@
 #include "WorldPartition/WorldPartitionRuntimeCell.h"
 #include "Editor.h"
 
-// Register FWorldPartitionCookPackageSplitter for UWorld class
-REGISTER_COOKPACKAGE_SPLITTER(FWorldPartitionCookPackageSplitter, UWorld);
-
-bool FWorldPartitionCookPackageSplitter::ShouldSplit(UObject* SplitData)
-{
-	UWorld* World = Cast<UWorld>(SplitData);
-	return World && World->IsPartitionedWorld();
-}
-
-FWorldPartitionCookPackageSplitter::FWorldPartitionCookPackageSplitter()
-{
-}
-
 FWorldPartitionCookPackageSplitter::~FWorldPartitionCookPackageSplitter()
 {
 	check(!ReferencedWorld);
 }
 
-void FWorldPartitionCookPackageSplitter::Teardown(ETeardown Status)
-{
-	UE_LOG(LogWorldPartition, Display, TEXT("[Cook] Debug(TearDown): OwnerObject=%s bForceInitializedWorld=%d bInitializedPhysicsSceneForSave=%d"), *GetFullNameSafe(ReferencedWorld),  bForceInitializedWorld ? 1 : 0, bInitializedPhysicsSceneForSave ? 1 : 0);
-
-	FWorldDelegates::OnWorldCleanup.RemoveAll(this);
-
-	if (ReferencedWorld != nullptr)
-	{
-		// We were destructed without GetGenerateList being called; nothing to teardown
-		// Assume that the world is partitioned as per FWorldPartitionCookPackageSplitter::ShouldSplit
-		UWorldPartition* WorldPartition = ReferencedWorld->PersistentLevel->GetWorldPartition();
-		check(WorldPartition);
-
-		WorldPartition->EndCook(CookContext);
-		WorldPartition->Uninitialize();
-
-		if (bInitializedPhysicsSceneForSave)
-		{
-			GEditor->CleanupPhysicsSceneThatWasInitializedForSave(ReferencedWorld, bForceInitializedWorld);
-			bInitializedPhysicsSceneForSave = false;
-			bForceInitializedWorld = false;
-		}
-
-		ReferencedWorld = nullptr;
-	}
-}
-
 void FWorldPartitionCookPackageSplitter::OnWorldCleanup(UWorld* InWorld, bool bSessionEnded, bool bCleanupResources)
 {
 	checkf(InWorld != ReferencedWorld, TEXT("[Cook] %s is being cleaned up while still referenced by a package splitter."), *GetFullNameSafe(InWorld));
@@ -72,14 +32,14 @@ FString FWorldPartitionCookPackageSplitter::GetReferencerName() const
 	return TEXT("FWorldPartitionCookPackageSplitter");
 }
 
-TArray<ICookPackageSplitter::FGeneratedPackage> FWorldPartitionCookPackageSplitter::GetGenerateList(const UPackage* OwnerPackage, const UObject* OwnerObject)
+TArray<ICookPackageSplitter::FGeneratedPackage> FWorldPartitionCookPackageSplitter::GetGenerateList(const UPackage* OwnerPackage, TNotNull<const UWorld*> World)
 {
-	UE_LOG(LogWorldPartition, Display, TEXT("[Cook] Gathering packages to cook from generators for owner object %s."), *GetFullNameSafe(OwnerObject));
+	UE_LOG(LogWorldPartition, Display, TEXT("[Cook] Gathering packages to cook from generators for owner object %s."), *GetFullNameSafe(NotNullGet(World)));
 
 	// Store the World pointer to declare it to GarbageCollection; we do not want to allow the World to be Garbage Collected
 	// until we have finished all of our PreSaveGeneratedPackage calls, because we store information on the World 
 	// that is necessary for populate 
-	ReferencedWorld = const_cast<UWorld*>(CastChecked<const UWorld>(OwnerObject));
+	ReferencedWorld = const_cast<UWorld*>(NotNullGet(World));
 
 	check(!bInitializedPhysicsSceneForSave && !bForceInitializedWorld);
 	bInitializedPhysicsSceneForSave = GEditor->InitializePhysicsSceneForSaveIfNecessary(ReferencedWorld, bForceInitializedWorld);
@@ -94,7 +54,7 @@ TArray<ICookPackageSplitter::FGeneratedPackage> FWorldPartitionCookPackageSplitt
 	WorldPartition->BeginCook(CookContext);
 
 	bool bIsSuccess = CookContext.GatherPackagesToCook();
-	UE_CLOG(!bIsSuccess, LogWorldPartition, Warning, TEXT("[Cook] Errors while gathering packages to took from generators for owner object %s."), *GetFullNameSafe(OwnerObject));
+	UE_CLOG(!bIsSuccess, LogWorldPartition, Warning, TEXT("[Cook] Errors while gathering packages to took from generators for owner object %s."), *GetFullNameSafe(NotNullGet(World)));
 
 	UE_LOG(LogWorldPartition, Display, TEXT("[Cook] Gathered %u packages to generate from %u Generators."), CookContext.NumPackageToGenerate(), CookContext.NumGenerators());
 
@@ -102,17 +62,17 @@ TArray<ICookPackageSplitter::FGeneratedPackage> FWorldPartitionCookPackageSplitt
 	BuildPackagesToGenerateList(PackagesToGenerate);
 
 	UE_LOG(LogWorldPartition, Display, TEXT("[Cook] Sending %u packages to be generated."), PackagesToGenerate.Num());
-	UE_LOG(LogWorldPartition, Display, TEXT("[Cook] Debug(GetGenerateList) : OwnerObject=%s bForceInitializedWorld=%d bInitializedPhysicsSceneForSave=%d"), *GetFullNameSafe(OwnerObject), bForceInitializedWorld ? 1 : 0, bInitializedPhysicsSceneForSave ? 1 : 0);
+	UE_LOG(LogWorldPartition, Display, TEXT("[Cook] Debug(GetGenerateList) : OwnerWorld=%s bForceInitializedWorld=%d bInitializedPhysicsSceneForSave=%d"), *GetFullNameSafe(NotNullGet(World)), bForceInitializedWorld ? 1 : 0, bInitializedPhysicsSceneForSave ? 1 : 0);
 
 	FWorldDelegates::OnWorldCleanup.AddRaw(this, &FWorldPartitionCookPackageSplitter::OnWorldCleanup);
 
 	return PackagesToGenerate;
 }
 
-bool FWorldPartitionCookPackageSplitter::PopulateGeneratedPackage(FPopulateContext& PopulateContext)
+bool FWorldPartitionCookPackageSplitter::PopulateGeneratedPackage(ICookPackageSplitter::FPopulateContext& PopulateContext)
 {
-	const FGeneratedPackageForPopulate& GeneratedPackage = *PopulateContext.GetTargetGeneratedPackage();
-	TGuardValue<FPopulateContext*> Guard(CookContext.PopulateContext, &PopulateContext);
+	const ICookPackageSplitter::FGeneratedPackageForPopulate& GeneratedPackage = *PopulateContext.GetTargetGeneratedPackage();
+	TGuardValue<ICookPackageSplitter::FPopulateContext*> Guard(CookContext.PopulateContext, &PopulateContext);
 
 	UE_LOG(LogWorldPartition, Verbose, TEXT("[Cook][PopulateGeneratedPackage] Processing %s"), *FWorldPartitionCookPackage::MakeGeneratedFullPath(GeneratedPackage.GeneratedRootPath, GeneratedPackage.RelativePath));
 
@@ -137,10 +97,10 @@ bool FWorldPartitionCookPackageSplitter::PopulateGeneratedPackage(FPopulateConte
 	return bIsSuccess;
 }
 
-bool FWorldPartitionCookPackageSplitter::PopulateGeneratorPackage(FPopulateContext& PopulateContext)
+bool FWorldPartitionCookPackageSplitter::PopulateGeneratorPackage(ICookPackageSplitter::FPopulateContext& PopulateContext)
 {
 	TConstArrayView<ICookPackageSplitter::FGeneratedPackageForPopulate> GeneratedPackages = PopulateContext.GetGeneratedPackages();
-	TGuardValue<FPopulateContext*> Guard(CookContext.PopulateContext, &PopulateContext);
+	TGuardValue<ICookPackageSplitter::FPopulateContext*> Guard(CookContext.PopulateContext, &PopulateContext);
 
 	UE_LOG(LogWorldPartition, Display, TEXT("[Cook][PopulateGeneratorPackage] Processing %u packages"), GeneratedPackages.Num());
 
@@ -167,12 +127,39 @@ bool FWorldPartitionCookPackageSplitter::PopulateGeneratorPackage(FPopulateConte
 	return bIsSuccess;
 }
 
-void FWorldPartitionCookPackageSplitter::OnOwnerReloaded(UPackage* OwnerPackage, UObject* OwnerObject)
+void FWorldPartitionCookPackageSplitter::OnOwnerReloaded(UPackage* OwnerPackage, TNotNull<const UWorld*> World)
 {
 	// It should not be possible for the owner to reload due to garbage collection while we are active and keeping it referenced
 	check(!ReferencedWorld);
 }
 
+void FWorldPartitionCookPackageSplitter::Teardown(ICookPackageSplitter::ETeardown Status)
+{
+	UE_LOG(LogWorldPartition, Display, TEXT("[Cook] Debug(TearDown): OwnerWorld=%s bForceInitializedWorld=%d bInitializedPhysicsSceneForSave=%d"), *GetFullNameSafe(ReferencedWorld),  bForceInitializedWorld ? 1 : 0, bInitializedPhysicsSceneForSave ? 1 : 0);
+
+	FWorldDelegates::OnWorldCleanup.RemoveAll(this);
+
+	if (ReferencedWorld)
+	{
+		// We were destructed without GetGenerateList being called; nothing to teardown
+		// Assume that the world is partitioned as per FWorldPartitionCookPackageSplitter::ShouldSplit
+		UWorldPartition* WorldPartition = ReferencedWorld->PersistentLevel->GetWorldPartition();
+		check(WorldPartition);
+
+		WorldPartition->EndCook(CookContext);
+		WorldPartition->Uninitialize();
+
+		if (bInitializedPhysicsSceneForSave)
+		{
+			GEditor->CleanupPhysicsSceneThatWasInitializedForSave(ReferencedWorld, bForceInitializedWorld);
+			bInitializedPhysicsSceneForSave = false;
+			bForceInitializedWorld = false;
+		}
+
+		ReferencedWorld = nullptr;
+	}
+}
+
 void FWorldPartitionCookPackageSplitter::BuildPackagesToGenerateList(TArray<ICookPackageSplitter::FGeneratedPackage>& PackagesToGenerate) const
 {
 	for (const IWorldPartitionCookPackageGenerator* CookPackageGenerator : CookContext.GetCookPackageGenerators())
diff --git a/Engine/Source/Runtime/Engine/Private/WorldPartition/WorldPartitionCookPackageSplitter.h b/Engine/Source/Runtime/Engine/Private/WorldPartition/WorldPartitionCookPackageSplitter.h
--- a/Engine/Source/Runtime/Engine/Private/WorldPartition/WorldPartitionCookPackageSplitter.h
+++ b/Engine/Source/Runtime/Engine/Private/WorldPartition/WorldPartitionCookPackageSplitter.h
@@ -4,40 +4,20 @@
 
 #if WITH_EDITOR
 
-#include "CookPackageSplitter.h"
-#include "Engine/World.h"
-#include "UObject/WeakObjectPtrTemplates.h"
+#include "WorldCookPackageSplitter.h"
 #include "WorldPartition/Cook/WorldPartitionCookPackageContext.h"
 
-class FWorldPartitionCookPackageSplitter : public FGCObject, public ICookPackageSplitter
+class FWorldPartitionCookPackageSplitter : public FGCObject, public FWorldCookPackageSplitter::FSubSplitter
 {
 public:
-	//~ Begin of ICookPackageSplitter
-	static bool ShouldSplit(UObject* SplitData);
-	static FString GetSplitterDebugName() { return TEXT("FWorldPartitionCookPackageSplitter"); }
-
-	FWorldPartitionCookPackageSplitter();
+	//~Begin FWorldCookPackageSplitter::FSubSplitter interface
 	virtual ~FWorldPartitionCookPackageSplitter();
-
-	virtual void Teardown(ETeardown Status) override;
-	virtual bool UseInternalReferenceToAvoidGarbageCollect() override
-	{
-		return true;
-	}
-	virtual EGeneratedRequiresGenerator DoesGeneratedRequireGenerator() override
-	{
-		return EGeneratedRequiresGenerator::Populate;
-	}
-	virtual bool RequiresGeneratorPackageDestructBeforeResplit() override
-	{
-		return true;
-	}
-
-	virtual TArray<ICookPackageSplitter::FGeneratedPackage> GetGenerateList(const UPackage* OwnerPackage, const UObject* OwnerObject) override;
-	virtual bool PopulateGeneratedPackage(FPopulateContext& PopulateContext) override;
-	virtual bool PopulateGeneratorPackage(FPopulateContext& PopulateContext) override;
-	virtual void OnOwnerReloaded(UPackage* OwnerPackage, UObject* OwnerObject) override;
-	//~ End of ICookPackageSplitter
+	virtual TArray<ICookPackageSplitter::FGeneratedPackage> GetGenerateList(const UPackage* OwnerPackage, TNotNull<const UWorld*> World) override;
+	virtual bool PopulateGeneratedPackage(ICookPackageSplitter::FPopulateContext& PopulateContext) override;
+	virtual bool PopulateGeneratorPackage(ICookPackageSplitter::FPopulateContext& PopulateContext) override;
+	virtual void OnOwnerReloaded(UPackage* OwnerPackage, TNotNull<const UWorld*> World) override;
+	virtual void Teardown(ICookPackageSplitter::ETeardown Status) override;
+	//~End FWorldCookPackageSplitter::FSubSplitter interface
 
 private:
 	void OnWorldCleanup(UWorld* InWorld, bool bSessionEnded, bool bCleanupResources);
diff --git a/Engine/Source/Runtime/Engine/Public/WorldCookPackageSplitter.h b/Engine/Source/Runtime/Engine/Public/WorldCookPackageSplitter.h
new file mode 100644
--- /dev/null
+++ b/Engine/Source/Runtime/Engine/Public/WorldCookPackageSplitter.h
@@ -0,0 +1,88 @@
+// Copyright Epic Games, Inc. All Rights Reserved.
+
+#pragma once
+
+#if WITH_EDITOR
+#include "CookPackageSplitter.h"
+#include "OverrideVoidReturnInvoker.h"
+#include "Engine/World.h"
+#include "Templates/Function.h"
+#endif
+
+#define UE_API ENGINE_API
+
+#if WITH_EDITOR
+class FWorldCookPackageSplitter : public ICookPackageSplitter
+{
+	DEFINE_COOKPACKAGE_SPLITTER(FWorldCookPackageSplitter, UWorld);
+
+public:
+	struct FSubSplitter
+	{
+		virtual ~FSubSplitter() {}
+		virtual TArray<ICookPackageSplitter::FGeneratedPackage> GetGenerateList(const UPackage* OwnerPackage, TNotNull<const UWorld*> World) = 0;
+		virtual bool PopulateGeneratedPackage(ICookPackageSplitter::FPopulateContext& PopulateContext) = 0;
+		virtual bool PopulateGeneratorPackage(ICookPackageSplitter::FPopulateContext& PopulateContext) = 0;
+		virtual void OnOwnerReloaded(UPackage* OwnerPackage, TNotNull<const UWorld*> World) = 0;
+		virtual void Teardown(ICookPackageSplitter::ETeardown Status) = 0;
+	};
+
+	struct FSubSplitterFactory
+	{
+		TFunction<bool(TNotNull<const UWorld*>)> ShouldSplit;
+		TFunction<FSubSplitter*()> MakeInstance;
+	};
+
+	static UE_API void RegisterCookPackageSubSplitterFactory(UClass* Class, FSubSplitterFactory Factory);
+	static UE_API void UnregisterCookPackageSubSplitterFactory(UClass* Class);
+
+protected:
+	static TMap<UClass*, FSubSplitterFactory> RegisteredCookPackageSubSplitterFactories;
+
+	template<typename FuncType>
+	static bool ForEachRegisteredCookPackageSubSplitterFactories(FuncType Func)
+	{
+		TOverrideVoidReturnInvoker Invoker(true, Func);
+		for (const TPair<UClass*, FSubSplitterFactory>& RegisteredCookPackageSubSplitterFactory : RegisteredCookPackageSubSplitterFactories)
+		{
+			if (!Invoker(RegisteredCookPackageSubSplitterFactory.Value))
+			{
+				return false;
+			}
+		}
+		return true;
+	}
+
+	//~ Begin of ICookPackageSplitter
+	static bool ShouldSplit(UObject* SplitData);
+	static FString GetSplitterDebugName() { return TEXT("FWorldCookPackageSplitter"); }
+
+	virtual void Teardown(ETeardown Status) override;
+	virtual bool UseInternalReferenceToAvoidGarbageCollect() override;
+	virtual EGeneratedRequiresGenerator DoesGeneratedRequireGenerator() override;
+	virtual bool RequiresGeneratorPackageDestructBeforeResplit() override;
+	virtual TArray<ICookPackageSplitter::FGeneratedPackage> GetGenerateList(const UPackage* OwnerPackage, const UObject* OwnerObject) override;
+	virtual bool PopulateGeneratedPackage(FPopulateContext& PopulateContext) override;
+	virtual bool PopulateGeneratorPackage(FPopulateContext& PopulateContext) override;
+	virtual void OnOwnerReloaded(UPackage* OwnerPackage, UObject* OwnerObject) override;
+	//~ End of ICookPackageSplitter
+
+	TArray<TUniquePtr<FSubSplitter>> CookPackageSubSplitters;
+
+	template<typename FuncType>
+	bool ForEachCookPackageSubSplitters(FuncType Func)
+	{
+		TOverrideVoidReturnInvoker Invoker(true, Func);
+		for (const TUniquePtr<FWorldCookPackageSplitter::FSubSplitter>& CookPackageSubSplitter : CookPackageSubSplitters)
+		{
+			if (!Invoker(CookPackageSubSplitter.Get()))
+			{
+				return false;
+			}
+		}
+		return true;
+	}
+};
+#endif // WITH_EDITOR
+
+#undef UE_API
\ No newline at end of file
