[SkeletalMesh HalfEdge] Removed an unsafe asset load when querying whether half edge data is needed. Swtich to building half edge data by default if the query fails and show a warning to ask for asset resave

diff --git a/Engine/Source/Runtime/Engine/Private/Rendering/SkeletalMeshHalfEdgeBufferAccessor.cpp b/Engine/Source/Runtime/Engine/Private/Rendering/SkeletalMeshHalfEdgeBufferAccessor.cpp
--- a/Engine/Source/Runtime/Engine/Private/Rendering/SkeletalMeshHalfEdgeBufferAccessor.cpp
+++ b/Engine/Source/Runtime/Engine/Private/Rendering/SkeletalMeshHalfEdgeBufferAccessor.cpp
@@ -5,6 +5,8 @@
 #include "AssetRegistry/AssetRegistryModule.h"
 #include "AssetRegistry/AssetData.h"
 
+DEFINE_LOG_CATEGORY_STATIC(LogSkeletalMeshHalfEdgeBufferAccessor, Log, All);
+
 FName SkeletalMeshHalfEdgeBufferAccessor::GetHalfEdgeRequirementAssetTagName()
 {
 	static FName Tag("bRequiresSkeletalMeshHalfEdgeBuffer");
@@ -23,31 +25,27 @@ bool SkeletalMeshHalfEdgeBufferAccessor::IsHalfEdgeRequired(TSoftObjectPtr<UObje
 	FAssetData AssetData;
 	AssetRegistryModule.TryGetAssetByObjectPath(InAssetSoftPtr.ToSoftObjectPath(), AssetData);
 
-	
-	// if (AssetData.GetClass()->ImplementsInterface(USkeletalMeshHalfEdgeBufferAccessor::StaticClass()))
+	// Asset types that implement USkeletalMeshHalfEdgeBufferAccessor should have this tag,
+	// if not, it is likely old and requires a resave to have this information saved as asset tag
+	bool bIsHalfEdgeRequired = false;
+	if (AssetData.GetTagValue<bool>(GetHalfEdgeRequirementAssetTagName(), bIsHalfEdgeRequired))
 	{
-		bool bIsHalfEdgeRequired = false;
-		if (AssetData.GetTagValue<bool>(GetHalfEdgeRequirementAssetTagName(), bIsHalfEdgeRequired))
-		{
-			if (bIsHalfEdgeRequired)
-			{
-				return true;
-			}
-		}
-		else
-		{
-			// Fallback path
-			// Older assets may not have this tag serialized yet, lets trigger a load and query it directly
-			// However this isn't bulletproof since loading new assets isn't allowed when this function is called by certain loaders (ZenLoader)
-			if (ISkeletalMeshHalfEdgeBufferAccessor* HalfEdgeBufferAccessor = Cast<ISkeletalMeshHalfEdgeBufferAccessor>(InAssetSoftPtr.LoadSynchronous()))
-			{
-				if (HalfEdgeBufferAccessor->IsSkeletalMeshHalfEdgeBufferRequired())
-				{
-					return true;
-				}
-			}
-		}
+		return bIsHalfEdgeRequired;
 	}
+	
+	// Fallback path
+	//
+	// If an asset do not have this tag serialized is likely old and haven't been re-saved recently.
+	// Therefore, we have to simply assume that the deformer needs half-edge buffer here to make sure it still works,
+	// while giving out a warning to ask for manual resave to avoid building half edge buffers unnecessarily.
+	//
+	// Ideally we want to load the asset and check it directly but loading isn't safe here
+	// because this function can be called from worker threads during load/build.
+	
+	UE_LOG(LogSkeletalMeshHalfEdgeBufferAccessor, Warning,
+		TEXT("Unable to determine if Skeletal Mesh Half Edge data is required for Asset %s, default to required. "
+	   "Resaving the asset may help avoid building half edge data unnecessarily"), *InAssetSoftPtr.ToString());
+	
+	return true;
 
-	return false;
 }
