Add missing breaks from Modulation Inheritance Statement

diff --git a/Engine/Source/Runtime/Engine/Private/Audio.cpp b/Engine/Source/Runtime/Engine/Private/Audio.cpp
--- a/Engine/Source/Runtime/Engine/Private/Audio.cpp
+++ b/Engine/Source/Runtime/Engine/Private/Audio.cpp
@@ -327,6 +327,7 @@ void FWaveInstance::CalculateInheritedModulationSettingsAndValue(const EModulati
 			InOutSettings = UnionSettings;
 			OutValue = UnionSettings.Value;
 		}
+		break;
 
 		case EModulationRouting::Inherit:
 		{
@@ -358,6 +359,7 @@ void FWaveInstance::CalculateInheritedModulationSettingsAndValue(const EModulati
 		InOutSettings = OverrideSettings;
 		OutValue = OverrideSettings.Value;	
 	}
+	break;
 
 	case EModulationRouting::Disable:
 	default:
