Only use the RepLayout CompareObject logic from CL 38890523 if UE_WITH_REMOTE_OBJECT_HANDLE is defined. Using TObjectPtrs here is not safe if shadow state refers to garbage-collected objects and gc.AllowIncrementalReachability is enabled. This is because copying FObjectPtr in this case can cause a dereference of the object via ConditionallyMarkAsReachable.

diff --git a/Engine/Source/Runtime/Engine/Private/RepLayout.cpp b/Engine/Source/Runtime/Engine/Private/RepLayout.cpp
--- a/Engine/Source/Runtime/Engine/Private/RepLayout.cpp
+++ b/Engine/Source/Runtime/Engine/Private/RepLayout.cpp
@@ -602,10 +602,17 @@ static FORCEINLINE bool CompareObject(
 {
 	// Until FObjectPropertyBase::Identical is made safe for GC'd objects, we need to do it manually
 	// This saves us from having to add referenced objects during GC
-	FObjectPropertyBase* ObjProperty = CastFieldChecked<FObjectPropertyBase>(Cmd.Property);
+	const FObjectPropertyBase* ObjProperty = CastFieldChecked<FObjectPropertyBase>(Cmd.Property);
 
-	TObjectPtr<UObject> ObjectA = ObjProperty->GetObjectPtrPropertyValue(A);
-	TObjectPtr<UObject> ObjectB = ObjProperty->GetObjectPtrPropertyValue(B);
+	// With remote object handles, use TObjectPtrs so that the objects don't have to be resolved.
+	// Since that's not safe for GC'd objects, use raw pointers otherwise.
+#if UE_WITH_REMOTE_OBJECT_HANDLE
+	const TObjectPtr<UObject> ObjectA = ObjProperty->GetObjectPtrPropertyValue(A);
+	const TObjectPtr<UObject> ObjectB = ObjProperty->GetObjectPtrPropertyValue(B);
+#else
+	const UObject* ObjectA = ObjProperty->GetObjectPropertyValue(A);
+	const UObject* ObjectB = ObjProperty->GetObjectPropertyValue(B);
+#endif
 
 	return ObjectA == ObjectB;
 }
