Fix for validation error in GPU skin cache when dynamic mesh bounds are enabled.

diff --git a/Engine/Source/Runtime/Engine/Private/GPUSkinCache.cpp b/Engine/Source/Runtime/Engine/Private/GPUSkinCache.cpp
--- a/Engine/Source/Runtime/Engine/Private/GPUSkinCache.cpp
+++ b/Engine/Source/Runtime/Engine/Private/GPUSkinCache.cpp
@@ -1186,7 +1186,7 @@ public:
 		uint32 NumElementedToAllocate = uint32(FMath::Max(MinSlotsToAllocate, int32(FMath::RoundUpToPowerOfTwo(GetNumSlotsAllocated()))) * 2);
 		ResizeResourceIfNeeded(RHICmdList, MeshBoundsBuffer, sizeof(FIntVector4) * NumElementedToAllocate, TEXT("DynamicMeshBoundsBuffer.MeshBoundsBuffer"));
 
-		RHICmdList.Transition(FRHITransitionInfo(MeshBoundsBuffer.UAV, ERHIAccess::Unknown, ERHIAccess::UAVCompute), ERHITransitionCreateFlags::NoFence);
+		RHICmdList.Transition(FRHITransitionInfo(MeshBoundsBuffer.UAV, ERHIAccess::Unknown, ERHIAccess::UAVCompute), ERHITransitionCreateFlags::AllowDecayPipelines);
 
 		DispatchClearSlots(RHICmdList, SlotsToClearMask);
 
