FStreamableTextureResource texture initializing functions should take a command list.
Making sure the correct names are passed to Texture creation so that we don't have to set them post-creation.

diff --git a/Engine/Source/Runtime/Engine/Private/Rendering/StreamableTextureResource.cpp b/Engine/Source/Runtime/Engine/Private/Rendering/StreamableTextureResource.cpp
--- a/Engine/Source/Runtime/Engine/Private/Rendering/StreamableTextureResource.cpp
+++ b/Engine/Source/Runtime/Engine/Private/Rendering/StreamableTextureResource.cpp
@@ -200,7 +200,7 @@ void FStreamableTextureResource::DecrementTextureStats() const
 }
 #endif
 
-void FStreamableTextureResource::InitRHI(FRHICommandListBase&)
+void FStreamableTextureResource::InitRHI(FRHICommandListBase& RHICmdList)
 {
 	SCOPED_LOADTIMER(FStreamableTextureResource_InitRHI);
 
@@ -216,23 +216,19 @@ void FStreamableTextureResource::InitRHI(FRHICommandListBase&)
 	check(CVarReducedMode);
 	if (bUsePartiallyResidentMips && (!CVarReducedMode->GetValueOnRenderThread() || State.NumRequestedLODs > State.NumNonStreamingLODs))
 	{
-		CreatePartiallyResidentTexture();
+		CreatePartiallyResidentTexture(RHICmdList);
 	}
 	else
 	{
-		CreateTexture();
+		CreateTexture(RHICmdList);
 	}
 
 	// Update mip-level fading.
 	MipBiasFade.SetNewMipCount( State.NumRequestedLODs, State.NumRequestedLODs, LastRenderTime, MipFadeSetting );
 
-	TextureRHI->SetOwnerName(GetOwnerName());
-	TextureRHI->SetName(TextureName);
-	RHIBindDebugLabelName(TextureRHI, *TextureName.ToString());
-
 	if (ensure(TextureReferenceRHI.IsValid()))
 	{
-		RHIUpdateTextureReference(TextureReferenceRHI, TextureRHI);
+		RHICmdList.UpdateTextureReference(TextureReferenceRHI, TextureRHI);
 	}
 }
 
diff --git a/Engine/Source/Runtime/Engine/Private/Rendering/Texture2DArrayResource.cpp b/Engine/Source/Runtime/Engine/Private/Rendering/Texture2DArrayResource.cpp
--- a/Engine/Source/Runtime/Engine/Private/Rendering/Texture2DArrayResource.cpp
+++ b/Engine/Source/Runtime/Engine/Private/Rendering/Texture2DArrayResource.cpp
@@ -9,6 +9,8 @@
 #include "Engine/Texture2DArray.h"
 #include "RenderUtils.h"
 
+static const FName NAME_FTexture2DArrayResource(TEXT("FTexture2DArrayResource"));
+
 //*****************************************************************************
 //************************* FTexture2DArrayResource ***************************
 //*****************************************************************************
@@ -47,18 +49,22 @@ FTexture2DArrayResource::FTexture2DArrayResource(UTexture2DArray* InOwner, const
 {
 }
 
-void FTexture2DArrayResource::CreateTexture()
+void FTexture2DArrayResource::CreateTexture(FRHICommandListBase& RHICmdList)
 {
 	TArrayView<const FTexture2DMipMap*> MipsViewPostLODBias = GetPlatformMipsView();
 	const FTexture2DMipMap& FirstMip = *MipsViewPostLODBias[State.RequestedFirstLODIdx()];
 
+	const FString Name = TextureName.ToString();
+
 	const FRHITextureCreateDesc Desc =
-		FRHITextureCreateDesc::Create2DArray(TEXT("FTexture2DArrayResource"), FirstMip.SizeX, FirstMip.SizeY, FirstMip.SizeZ, PixelFormat)
+		FRHITextureCreateDesc::Create2DArray(*Name, FirstMip.SizeX, FirstMip.SizeY, FirstMip.SizeZ, PixelFormat)
 		.SetNumMips(State.NumRequestedLODs)
 		.SetFlags(CreationFlags)
-		.SetExtData(PlatformData->GetExtData());
+		.SetExtData(PlatformData->GetExtData())
+		.SetClassName(NAME_FTexture2DArrayResource)
+		.SetOwnerName(GetOwnerName());
 
-	TextureRHI = RHICreateTexture(Desc);
+	TextureRHI = RHICmdList.CreateTexture(Desc);
 
 	// Copy the mip data in to the texture.
 	for (int32 MipIndex = 0; MipIndex < AllMipsData.Num(); MipIndex++)
@@ -66,7 +72,7 @@ void FTexture2DArrayResource::CreateTexture()
 		for (uint32 ArrayIndex = 0; ArrayIndex < SizeZ; ++ArrayIndex)
 		{
 			uint32 DestStride = 0;
-			void* DestData = RHILockTexture2DArray(TextureRHI, ArrayIndex, MipIndex, RLM_WriteOnly, DestStride, false);
+			void* DestData = RHICmdList.GetAsImmediate().LockTexture2DArray(TextureRHI, ArrayIndex, MipIndex, RLM_WriteOnly, DestStride, false);
 			if (DestData)
 			{
 				GetData(FirstMip.SizeX, FirstMip.SizeY, ArrayIndex, MipIndex, DestData, DestStride);
@@ -75,14 +81,14 @@ void FTexture2DArrayResource::CreateTexture()
 			{
 				UE_LOG(LogTexture, Warning, TEXT("Failed to lock texture 2d array mip/slice %d / %d (%s)"), MipIndex, ArrayIndex, *TextureName.ToString());
 			}
-			RHIUnlockTexture2DArray(TextureRHI, ArrayIndex, MipIndex, false);
+			RHICmdList.GetAsImmediate().UnlockTexture2DArray(TextureRHI, ArrayIndex, MipIndex, false);
 		}
 	}
 
 	AllMipsData.Empty();
 }
 
-void FTexture2DArrayResource::CreatePartiallyResidentTexture()
+void FTexture2DArrayResource::CreatePartiallyResidentTexture(FRHICommandListBase& RHICmdList)
 {
 	unimplemented();
 	TextureRHI.SafeRelease();
diff --git a/Engine/Source/Runtime/Engine/Private/Rendering/Texture2DArrayResource.h b/Engine/Source/Runtime/Engine/Private/Rendering/Texture2DArrayResource.h
--- a/Engine/Source/Runtime/Engine/Private/Rendering/Texture2DArrayResource.h
+++ b/Engine/Source/Runtime/Engine/Private/Rendering/Texture2DArrayResource.h
@@ -40,8 +40,8 @@ public:
 
 protected:
 
-	void CreateTexture() final override;
-	void CreatePartiallyResidentTexture() final override;
+	void CreateTexture(FRHICommandListBase& RHICmdList) final override;
+	void CreatePartiallyResidentTexture(FRHICommandListBase& RHICmdList) final override;
 
 	void GetData(int32 BaseRHIMipSizeX, int32 BaseRHIMipSizeY, uint32 ArrayIndex, uint32 MipIndex, void* Dest, uint32 DestPitch) const;
 
diff --git a/Engine/Source/Runtime/Engine/Private/Rendering/Texture2DResource.cpp b/Engine/Source/Runtime/Engine/Private/Rendering/Texture2DResource.cpp
--- a/Engine/Source/Runtime/Engine/Private/Rendering/Texture2DResource.cpp
+++ b/Engine/Source/Runtime/Engine/Private/Rendering/Texture2DResource.cpp
@@ -11,7 +11,7 @@
 #include "RenderingThread.h"
 #include "Containers/ResourceArray.h"
 
-const static FLazyName Texture2DResourceName(TEXT("FTexture2DResource"));
+const static FName NAME_FTexture2DResource(TEXT("FTexture2DResource"));
 
 // TODO Only adding this setting to allow backwards compatibility to be forced.  The default
 // behavior is to NOT do this.  This variable should be removed in the future.  #ADDED 4.13
@@ -101,7 +101,7 @@ void FTexture2DResource::InitRHI(FRHICommandListBase& RHICmdList)
 	if (ProxiedResource)
 	{
 		TextureRHI = ProxiedResource->TextureRHI;
-		RHIUpdateTextureReference(TextureReferenceRHI, TextureRHI);
+		RHICmdList.UpdateTextureReference(TextureReferenceRHI, TextureRHI);
 		SamplerStateRHI = ProxiedResource->SamplerStateRHI;
 		DeferredPassSamplerStateRHI = ProxiedResource->DeferredPassSamplerStateRHI;
 	}
@@ -111,18 +111,20 @@ void FTexture2DResource::InitRHI(FRHICommandListBase& RHICmdList)
 	}
 }
 
-void FTexture2DResource::CreateTexture()
+void FTexture2DResource::CreateTexture(FRHICommandListBase& RHICmdList)
 {
 	const int32 RequestedMipIdx = State.RequestedFirstLODIdx();
 	const FTexture2DMipMap* RequestedMip = GetPlatformMip(RequestedMipIdx);
 
+	const FString Name = TextureName.ToString();
+
 	// create texture with ResourceMem data when available
 	FRHITextureCreateDesc Desc =
-		FRHITextureCreateDesc::Create2D(TEXT("FTexture2DResource"), RequestedMip->SizeX, RequestedMip->SizeY, PixelFormat)
+		FRHITextureCreateDesc::Create2D(*Name, RequestedMip->SizeX, RequestedMip->SizeY, PixelFormat)
 		.SetNumMips(State.NumRequestedLODs)
 		.SetFlags(CreationFlags)
 		.SetExtData(PlatformData->GetExtData())
-		.SetClassName(Texture2DResourceName)
+		.SetClassName(NAME_FTexture2DResource)
 		.SetOwnerName(GetOwnerName());
 
 	if (ResourceMem)
@@ -130,7 +132,7 @@ void FTexture2DResource::CreateTexture()
 		Desc.SetInitActionBulkData(ResourceMem);
 	}
 
-	TextureRHI = RHICreateTexture(Desc);
+	TextureRHI = RHICmdList.CreateTexture(Desc);
 
 	// if( ResourceMem && !State.bReadyForStreaming) // To investigate!
 	if (ResourceMem)
@@ -154,34 +156,35 @@ void FTexture2DResource::CreateTexture()
 			{
 				uint32 DestPitch = -1;
 				uint64 Size = ~0ULL;
-				void* TheMipData = RHILockTexture2D(TextureRHI, RHIMipIdx, RLM_WriteOnly, DestPitch, false /* bLockWithinMiptail */, true /* bFlushRHIThread */, &Size);
+				void* TheMipData = RHICmdList.GetAsImmediate().LockTexture2D(TextureRHI, RHIMipIdx, RLM_WriteOnly, DestPitch, false /* bLockWithinMiptail */, true /* bFlushRHIThread */, &Size);
 				//check(Size != ~0ULL);
 
 				GetData( ResourceMipIdx, TheMipData, DestPitch, Size );
-				RHIUnlockTexture2D(TextureRHI, RHIMipIdx, false );
+				RHICmdList.GetAsImmediate().UnlockTexture2D(TextureRHI, RHIMipIdx, false );
 			}
 		}
 	}
 }
 
-void FTexture2DResource::CreatePartiallyResidentTexture()
+void FTexture2DResource::CreatePartiallyResidentTexture(FRHICommandListBase& RHICmdList)
 {
+	const FString Name = TextureName.ToString();
 	const int32 CurrentFirstMip = State.RequestedFirstLODIdx();
 
 	check(bUsePartiallyResidentMips);
 
 	const FRHITextureCreateDesc Desc =
-		FRHITextureCreateDesc::Create2D(TEXT("FTexture2DResource-PRT"), SizeX, SizeY, PixelFormat)
+		FRHITextureCreateDesc::Create2D(*Name, SizeX, SizeY, PixelFormat)
 		.SetNumMips(State.MaxNumLODs)
 		.SetFlags(CreationFlags | ETextureCreateFlags::Virtual)
 		.SetExtData(PlatformData->GetExtData())
-		.SetClassName(Texture2DResourceName)
+		.SetClassName(NAME_FTexture2DResource)
 		.SetOwnerName(GetOwnerName());
 
-	TextureRHI = RHICreateTexture(Desc);
+	TextureRHI = RHICmdList.CreateTexture(Desc);
 
-	RHIVirtualTextureSetFirstMipInMemory(TextureRHI, CurrentFirstMip);
-	RHIVirtualTextureSetFirstMipVisible(TextureRHI, CurrentFirstMip);
+	RHICmdList.GetAsImmediate().VirtualTextureSetFirstMipInMemory(TextureRHI, CurrentFirstMip);
+	RHICmdList.GetAsImmediate().VirtualTextureSetFirstMipVisible(TextureRHI, CurrentFirstMip);
 
 	check( ResourceMem == nullptr );
 
@@ -192,11 +195,11 @@ void FTexture2DResource::CreatePartiallyResidentTexture()
 		{
 			uint32 DestPitch = -1;
 			uint64 Size = ~0ULL;
-			void* TheMipData = RHILockTexture2D(TextureRHI, MipIndex, RLM_WriteOnly, DestPitch, false /* bLockWithinMiptail */, true /* bFlushRHIThread */, &Size);
+			void* TheMipData = RHICmdList.GetAsImmediate().LockTexture2D(TextureRHI, MipIndex, RLM_WriteOnly, DestPitch, false /* bLockWithinMiptail */, true /* bFlushRHIThread */, &Size);
 			//check(Size != ~0ULL);
 
 			GetData( MipIndex, TheMipData, DestPitch, Size);
-			RHIUnlockTexture2D(TextureRHI, MipIndex, false );
+			RHICmdList.GetAsImmediate().UnlockTexture2D(TextureRHI, MipIndex, false);
 		}
 	}
 }
diff --git a/Engine/Source/Runtime/Engine/Private/Rendering/Texture3DResource.cpp b/Engine/Source/Runtime/Engine/Private/Rendering/Texture3DResource.cpp
--- a/Engine/Source/Runtime/Engine/Private/Rendering/Texture3DResource.cpp
+++ b/Engine/Source/Runtime/Engine/Private/Rendering/Texture3DResource.cpp
@@ -9,6 +9,8 @@
 #include "EngineLogs.h"
 #include "RenderUtils.h"
 
+static const FName NAME_FTexture3DResource(TEXT("FTexture3DResource"));
+
 //*****************************************************************************
 //************************* FVolumeTextureBulkData ****************************
 //*****************************************************************************
@@ -73,7 +75,7 @@ FTexture3DResource::FTexture3DResource(UVolumeTexture* InOwner, const FStreamabl
 	{
 		// Compute the size of each mips so that they can be merged into a single allocation.
 		// -> this should be unnecessary now that TryLoadMipsWithSizes reports the sizes
-		if (GUseTexture3DBulkDataRHI)
+		if (GRHIGlobals.UseTexture3DBulkData)
 		{
 			for (int32 MipIndex = FirstLODIdx; MipIndex < InState.MaxNumLODs; ++MipIndex)
 			{
@@ -115,7 +117,7 @@ FTexture3DResource::FTexture3DResource(UVolumeTexture* InOwner, const FTexture3D
 {
 }
 
-void FTexture3DResource::CreateTexture()
+void FTexture3DResource::CreateTexture(FRHICommandListBase& RHICmdList)
 {
 	TArrayView<const FTexture2DMipMap*> MipsView = GetPlatformMipsView();
 	const int32 FirstMipIdx = InitialData.GetFirstMipIdx(); // == State.RequestedFirstLODIdx()
@@ -123,17 +125,17 @@ void FTexture3DResource::CreateTexture()
 	// Create the RHI texture.
 	{
 		const FTexture2DMipMap& FirstMip = *MipsView[FirstMipIdx];
-		const static FLazyName ClassName(TEXT("FTexture3DResource"));
+		const FString Name = TextureName.ToString();
 
 		FRHITextureCreateDesc Desc =
-			FRHITextureCreateDesc::Create3D(TEXT("FTexture3DResource"), FirstMip.SizeX, FirstMip.SizeY, FirstMip.SizeZ, PixelFormat)
+			FRHITextureCreateDesc::Create3D(*Name, FirstMip.SizeX, FirstMip.SizeY, FirstMip.SizeZ, PixelFormat)
 			.SetNumMips(State.NumRequestedLODs)
 			.SetFlags(CreationFlags)
 			.SetExtData(PlatformData->GetExtData())
-			.SetClassName(ClassName)
+			.SetClassName(NAME_FTexture3DResource)
 			.SetOwnerName(GetOwnerName());
 
-		if (GUseTexture3DBulkDataRHI)
+		if (GRHIGlobals.UseTexture3DBulkData)
 		{
 			InitialData.MergeMips(State.MaxNumLODs);
 			Desc.SetInitActionBulkData(&InitialData);
@@ -141,10 +143,10 @@ void FTexture3DResource::CreateTexture()
 
 		Desc.DetermineInititialState();
 
-		TextureRHI = RHICreateTexture(Desc);
+		TextureRHI = RHICmdList.CreateTexture(Desc);
 	}
 
-	if (!GUseTexture3DBulkDataRHI) 
+	if (!GRHIGlobals.UseTexture3DBulkData) 
 	{
 		const int32 BlockSizeX = GPixelFormats[PixelFormat].BlockSizeX;
 		const int32 BlockSizeY = GPixelFormats[PixelFormat].BlockSizeY;
@@ -177,14 +179,14 @@ void FTexture3DResource::CreateTexture()
 					(int)MipBytes);
 				}
 
-				RHIUpdateTexture3D(TextureRHI, RHIMipIdx, UpdateRegion, NumBlockX * BlockBytes, NumBlockX * NumBlockY * BlockBytes, MipData);
+				RHICmdList.UpdateTexture3D(TextureRHI, RHIMipIdx, UpdateRegion, NumBlockX * BlockBytes, NumBlockX * NumBlockY * BlockBytes, MipData);
 			}
 		}
 		InitialData.Discard();
 	}
 }
 
-void FTexture3DResource::CreatePartiallyResidentTexture()
+void FTexture3DResource::CreatePartiallyResidentTexture(FRHICommandListBase& RHICmdList)
 {
 	unimplemented();
 	TextureRHI.SafeRelease();
diff --git a/Engine/Source/Runtime/Engine/Private/Rendering/Texture3DResource.h b/Engine/Source/Runtime/Engine/Private/Rendering/Texture3DResource.h
--- a/Engine/Source/Runtime/Engine/Private/Rendering/Texture3DResource.h
+++ b/Engine/Source/Runtime/Engine/Private/Rendering/Texture3DResource.h
@@ -88,8 +88,8 @@ public:
 
 private:
 
-	void CreateTexture() final override;
-	void CreatePartiallyResidentTexture() final override;
+	void CreateTexture(FRHICommandListBase& RHICmdList) final override;
+	void CreatePartiallyResidentTexture(FRHICommandListBase& RHICmdList) final override;
 
 	/** Another resource being proxied by this one. */
 	const FTexture3DResource* const ProxiedResource = nullptr;
diff --git a/Engine/Source/Runtime/Engine/Public/Rendering/StreamableTextureResource.h b/Engine/Source/Runtime/Engine/Public/Rendering/StreamableTextureResource.h
--- a/Engine/Source/Runtime/Engine/Public/Rendering/StreamableTextureResource.h
+++ b/Engine/Source/Runtime/Engine/Public/Rendering/StreamableTextureResource.h
@@ -79,8 +79,8 @@ public:
 
 protected:
 
-	virtual void CreateTexture() = 0;
-	virtual void CreatePartiallyResidentTexture() = 0;
+	virtual void CreateTexture(FRHICommandListBase& RHICmdList) = 0;
+	virtual void CreatePartiallyResidentTexture(FRHICommandListBase& RHICmdList) = 0;
 
 	/** Platform data should be ref counted eventually to remove any synchronization needed when the asset is rebuild. */
 	const FTexturePlatformData* PlatformData;
diff --git a/Engine/Source/Runtime/Engine/Public/Rendering/Texture2DResource.h b/Engine/Source/Runtime/Engine/Public/Rendering/Texture2DResource.h
--- a/Engine/Source/Runtime/Engine/Public/Rendering/Texture2DResource.h
+++ b/Engine/Source/Runtime/Engine/Public/Rendering/Texture2DResource.h
@@ -58,8 +58,8 @@ private:
 	 */
 	FTexture2DResource(UTexture2D* InOwner, const FTexture2DResource* InProxiedResource);
 
-	virtual void CreateTexture() final override;
-	virtual void CreatePartiallyResidentTexture() final override;
+	virtual void CreateTexture(FRHICommandListBase& RHICmdList) final override;
+	virtual void CreatePartiallyResidentTexture(FRHICommandListBase& RHICmdList) final override;
 
 	/** Texture streaming command classes that need to be friends in order to call Update/FinalizeMipCount.	*/
 	friend class UTexture2D;
