Added a way to ignore the physics query handler for the current thread.

diff --git a/Engine/Source/Runtime/Engine/Private/Collision/SceneQuery.cpp b/Engine/Source/Runtime/Engine/Private/Collision/SceneQuery.cpp
--- a/Engine/Source/Runtime/Engine/Private/Collision/SceneQuery.cpp
+++ b/Engine/Source/Runtime/Engine/Private/Collision/SceneQuery.cpp
@@ -33,6 +33,18 @@ float DebugLineLifetime = 2.f;
 
 CSV_DEFINE_CATEGORY(SceneQuery, false);
 
+#if UE_WITH_REMOTE_OBJECT_HANDLE
+thread_local bool bIgnoreQueryHandler = false;
+#endif
+
+bool GetIgnoreQueryHandler()
+{
+#if UE_WITH_REMOTE_OBJECT_HANDLE
+return bIgnoreQueryHandler
+#endif
+return false;
+}
+
 enum class ESingleMultiOrTest : uint8
 {
 	Single,
@@ -880,7 +892,7 @@ bool TSceneCastCommon(const UWorld* World, OutHitsType& OutHits, const GeomInput
 	else
 	{
 		UPhysicsQueryHandler* QueryHandler = World->PhysicsQueryHandler;
-		if (QueryHandler == nullptr)
+		if (GetIgnoreQueryHandler() || QueryHandler == nullptr)
 		{
 			return TraceCommonImp<SweepOrRay, SingleMultiOrTest>(ThreadContext, World, OutHits, GeomInputs, Start, End, TraceChannel, Params, ResponseParams, ObjectParams, AccelContainer);
 		}
@@ -1306,7 +1318,7 @@ bool GeomOverlapMultiHelper(const UWorld* World, const GeomType& InGeom, const F
 	else
 	{
 		UPhysicsQueryHandler* QueryHandler = World->PhysicsQueryHandler;
-		if (QueryHandler == nullptr)
+		if (GetIgnoreQueryHandler() || QueryHandler == nullptr)
 		{
 			return GeomOverlapMultiHelper<InfoType>(ThreadContext, World, InGeom, GeomPose, OutOverlaps, TraceChannel, Params, ResponseParams, ObjectParams, AccelContainer);
 		}
@@ -1613,4 +1625,11 @@ namespace Chaos::Private
 	{
 		return SweepCommon<ESingleMultiOrTest::Multi>(ThreadContext, World, SweepData, CommonData, OutHits);
 	}
+
+#if UE_WITH_REMOTE_OBJECT_HANDLE
+	void FQueryInterface_Internal::SetIgnoreQueryHandlerThreadLocal(bool bIgnore)
+	{
+		bIgnoreQueryHandler = bIgnore;
+	}
+#endif
 }
diff --git a/Engine/Source/Runtime/Engine/Public/Physics/GenericPhysicsInterface.h b/Engine/Source/Runtime/Engine/Public/Physics/GenericPhysicsInterface.h
--- a/Engine/Source/Runtime/Engine/Public/Physics/GenericPhysicsInterface.h
+++ b/Engine/Source/Runtime/Engine/Public/Physics/GenericPhysicsInterface.h
@@ -153,5 +153,9 @@ namespace Chaos::Private
 		static bool SweepTest(const EThreadQueryContext ThreadContext, const UWorld* World, const FSweepQueryData& SweepData, const FCommonQueryData& CommonData, FHitResult& OutHits);
 		static bool SweepSingle(const EThreadQueryContext ThreadContext, const UWorld* World, const FSweepQueryData& SweepData, const FCommonQueryData& CommonData, FHitResult& OutHits);
 		static bool SweepMulti(const EThreadQueryContext ThreadContext, const UWorld* World, const FSweepQueryData& SweepData, const FCommonQueryData& CommonData, TArray<FHitResult>& OutHits);
+
+#if UE_WITH_REMOTE_OBJECT_HANDLE
+		UE_INTERNAL ENGINE_API static void SetIgnoreQueryHandlerThreadLocal(bool bIgnore);
+#endif
 	};
 }
