[Scheduler] Make sure workers are restarted during InitThreadConfig if any of those CVars have been set during startup so that .ini file configuration can be applied properly
TaskGraph.UseDynamicPrioritization
TaskGraph.UseDynamicThreadCreation
TaskGraph.OversubscriptionRatio

diff --git a/Engine/Source/Runtime/Core/Private/Async/TaskGraph.cpp b/Engine/Source/Runtime/Core/Private/Async/TaskGraph.cpp
--- a/Engine/Source/Runtime/Core/Private/Async/TaskGraph.cpp
+++ b/Engine/Source/Runtime/Core/Private/Async/TaskGraph.cpp
@@ -99,7 +99,7 @@ static FAutoConsoleVariableRef CVarForkedProcessMaxWorkerThreads(
 	TEXT("Configures the number of worker threads a forked process should spawn if it allows multithreading.")
 );
 
-// NOTE: the ECVF_ReadOnly cvars currently cannot be modified with ini files as they are applied too early
+// NOTE: the ECVF_ReadOnly cvars must be set in ini files that are read before InitThreadConfig
 
 CORE_API bool GTaskGraphUseDynamicPrioritization = 1;
 static FAutoConsoleVariableRef CVarTaskDynamicPrioritization(
diff --git a/Engine/Source/Runtime/Engine/Private/UnrealEngine.cpp b/Engine/Source/Runtime/Engine/Private/UnrealEngine.cpp
--- a/Engine/Source/Runtime/Engine/Private/UnrealEngine.cpp
+++ b/Engine/Source/Runtime/Engine/Private/UnrealEngine.cpp
@@ -19675,7 +19675,24 @@ void InitThreadConfig()
 			GConfig->GetArray(TEXT("/Script/Engine.Engine"), *FString::Printf(TEXT("ThreadConfigs_%s"), *PlatformVariant), ThreadConfigEntries, GEngineIni);
 		}
 
-		if (!ThreadConfigEntries.IsEmpty())
+		auto IsCVarModified =
+			[](const TCHAR* CVarName) -> bool
+			{
+				if (IConsoleObject* ConsoleObject = IConsoleManager::Get().FindConsoleObject(CVarName))
+				{
+					return (ConsoleObject->GetFlags() & EConsoleVariableFlags::ECVF_SetByMask) != 0;
+				}
+
+				return false;
+			};
+
+		// These cvars are not read until after the initial thread setup, so we need to reapply if they are different
+		bool bRequiresRestart =
+			IsCVarModified(TEXT("TaskGraph.UseDynamicPrioritization")) ||
+			IsCVarModified(TEXT("TaskGraph.UseDynamicThreadCreation")) ||
+			IsCVarModified(TEXT("TaskGraph.OversubscriptionRatio"));
+
+		if (bRequiresRestart || !ThreadConfigEntries.IsEmpty())
 		{
 			SetupThreadConfig(ThreadConfigEntries);
 		}
