[HWRT] Workaround for ray tracing geometry of skeletal meshes not being created when any segment is disabled.
- use static vertex buffer/factory when section is disabled since skin cache didn't generate valid data for the section.

diff --git a/Engine/Source/Runtime/Engine/Private/SkeletalMeshSceneProxy.cpp b/Engine/Source/Runtime/Engine/Private/SkeletalMeshSceneProxy.cpp
--- a/Engine/Source/Runtime/Engine/Private/SkeletalMeshSceneProxy.cpp
+++ b/Engine/Source/Runtime/Engine/Private/SkeletalMeshSceneProxy.cpp
@@ -707,7 +707,19 @@ const FCardRepresentationData* FSkeletalMeshSceneProxy::GetMeshCardRepresentatio
 
 void FSkeletalMeshSceneProxy::CreateBaseMeshBatch(const FSceneView* View, const FSkeletalMeshLODRenderData& LODData, const int32 LODIndex, const int32 SectionIndex, const FSectionElementInfo& SectionElementInfo, FMeshBatch& Mesh, ESkinVertexFactoryMode VFMode) const
 {
-	Mesh.VertexFactory = MeshObject->GetSkinVertexFactory(View, LODIndex, SectionIndex, VFMode);
+	const FSkelMeshRenderSection& Section = LODData.RenderSections[SectionIndex];
+
+	if (Section.bDisabled && VFMode == ESkinVertexFactoryMode::RayTracing)
+	{
+		// Use static skin vertex factory when section is disabled since the StaticVertexBuffers.PositionVertexBuffer is used to build BLAS 
+		// (see FSkeletalMeshObjectGPUSkin::UpdateRayTracingGeometry_Internal)
+		// TODO: Should also do this for ESkinVertexFactoryMode::Default?
+		Mesh.VertexFactory = MeshObject->GetStaticSkinVertexFactory(LODIndex, SectionIndex, VFMode);
+	}
+	else
+	{
+		Mesh.VertexFactory = MeshObject->GetSkinVertexFactory(View, LODIndex, SectionIndex, VFMode);
+	}
 	Mesh.MaterialRenderProxy = SectionElementInfo.Material->GetRenderProxy();
 #if RHI_RAYTRACING
 	Mesh.SegmentIndex = SectionIndex;
@@ -715,13 +727,13 @@ void FSkeletalMeshSceneProxy::CreateBaseMeshBatch(const FSceneView* View, const
 #endif
 
 	FMeshBatchElement& BatchElement = Mesh.Elements[0];
-	BatchElement.FirstIndex = LODData.RenderSections[SectionIndex].BaseIndex;
+	BatchElement.FirstIndex = Section.BaseIndex;
 	BatchElement.IndexBuffer = LODData.MultiSizeIndexContainer.GetIndexBuffer();
-	BatchElement.MinVertexIndex = LODData.RenderSections[SectionIndex].GetVertexBufferIndex();
-	BatchElement.MaxVertexIndex = LODData.RenderSections[SectionIndex].GetVertexBufferIndex() + LODData.RenderSections[SectionIndex].GetNumVertices() - 1;
+	BatchElement.MinVertexIndex = Section.GetVertexBufferIndex();
+	BatchElement.MaxVertexIndex = Section.GetVertexBufferIndex() + Section.GetNumVertices() - 1;
 
 	BatchElement.PrimitiveUniformBuffer = GetUniformBuffer();
-	BatchElement.NumPrimitives = LODData.RenderSections[SectionIndex].NumTriangles;
+	BatchElement.NumPrimitives = Section.NumTriangles;
 }
 
 uint8 FSkeletalMeshSceneProxy::GetCurrentFirstLODIdx_Internal() const
diff --git a/Engine/Source/Runtime/Engine/Private/SkeletalRenderGPUSkin.cpp b/Engine/Source/Runtime/Engine/Private/SkeletalRenderGPUSkin.cpp
--- a/Engine/Source/Runtime/Engine/Private/SkeletalRenderGPUSkin.cpp
+++ b/Engine/Source/Runtime/Engine/Private/SkeletalRenderGPUSkin.cpp
@@ -1127,14 +1127,6 @@ void FSkeletalMeshObjectGPUSkin::UpdateRayTracingGeometry_Internal(
 			FBufferRHIRef IndexBufferRHI = LODModel.MultiSizeIndexContainer.GetIndexBuffer()->IndexBufferRHI;
 			uint32 VertexBufferNumVertices = LODModel.StaticVertexBuffers.PositionVertexBuffer.GetNumVertices();
 			uint32 VertexBufferStride = LODModel.StaticVertexBuffers.PositionVertexBuffer.GetStride();
-			
-			//#dxr_todo: do we need support for separate sections in FRayTracingGeometryData?
-			uint32 TotalNumTriangles = 0;
-			for (int32 SectionIndex = 0; SectionIndex < LODModel.RenderSections.Num(); SectionIndex++)
-			{
-				const FSkelMeshRenderSection& Section = LODModel.RenderSections[SectionIndex];
-				TotalNumTriangles += Section.NumTriangles;
-			}
 
 			FRayTracingGeometryInitializer Initializer;
 
@@ -1153,7 +1145,6 @@ void FSkeletalMeshObjectGPUSkin::UpdateRayTracingGeometry_Internal(
 
 			Initializer.OwnerName = MeshObject->GetAssetPathName(LODIndex);
 			Initializer.IndexBuffer = IndexBufferRHI;
-			Initializer.TotalPrimitiveCount = TotalNumTriangles;
 			Initializer.GeometryType = RTGT_Triangles;
 			Initializer.bFastBuild = true;
 			Initializer.bAllowUpdate = true;
@@ -1175,7 +1166,22 @@ void FSkeletalMeshObjectGPUSkin::UpdateRayTracingGeometry_Internal(
 
 				// TODO: If we are at a dropped LOD, route material index through the LODMaterialMap in the LODInfo struct.
 				Segment.bEnabled = !MeshObject->IsMaterialHidden(LODIndex, Section.MaterialIndex) && Section.IsValid() && Section.bVisibleInRayTracing;
+
+				if (Segment.bEnabled)
+				{
+					checkf(Segment.VertexBuffer != nullptr, TEXT("Must provide valid vertex buffer when segment is enabled."));
+				}
+				else if (Segment.VertexBuffer == nullptr) // TODO: Consider always doing this when !Segment.bEnabled
+				{
+					// TODO: RHIs should accept NULL VertexBuffer when Segment is disabled.
+					// For now use the static PositionVertexBuffer as placeholder.
+					// TODO: Consider also setting Segment.NumPrimitives = 0; in this case
+					Segment.VertexBuffer = LODModel.StaticVertexBuffers.PositionVertexBuffer.GetRHI();
+				}
+
 				Initializer.Segments.Add(Segment);
+
+				Initializer.TotalPrimitiveCount += Segment.NumPrimitives;
 			}
 
 			if (RayTracingGeometry.GetRHI() != nullptr)
