Fix compiler error for SceneQuery.cpp in Dstm

diff --git a/Engine/Source/Runtime/Engine/Private/Collision/SceneQuery.cpp b/Engine/Source/Runtime/Engine/Private/Collision/SceneQuery.cpp
--- a/Engine/Source/Runtime/Engine/Private/Collision/SceneQuery.cpp
+++ b/Engine/Source/Runtime/Engine/Private/Collision/SceneQuery.cpp
@@ -40,7 +40,7 @@ thread_local bool bIgnoreQueryHandler = false;
 bool GetIgnoreQueryHandler()
 {
 #if UE_WITH_REMOTE_OBJECT_HANDLE
-return bIgnoreQueryHandler
+return bIgnoreQueryHandler;
 #endif
 return false;
 }
