Allow bIsPreviewSound when testing SoundWaveInstance::IsPlaying() checks.

[RN] Fix previewing Sounds in AudioLink enabled builds that use ForceSpatialize cvar

diff --git a/Engine/Source/Runtime/Engine/Private/Audio.cpp b/Engine/Source/Runtime/Engine/Private/Audio.cpp
--- a/Engine/Source/Runtime/Engine/Private/Audio.cpp
+++ b/Engine/Source/Runtime/Engine/Private/Audio.cpp
@@ -1110,6 +1110,11 @@ bool FWaveInstance::IsPlaying() const
 		return false;
 	}
 
+	if (ActiveSound->bIsPreviewSound)
+	{
+		return true;
+	}
+
 	// TODO: move out of audio.  Subtitle system should be separate and just set VirtualizationMode to PlayWhenSilent
 	const bool bHasSubtitles = ActiveSound->bHandleSubtitles && (ActiveSound->bHasExternalSubtitles || WaveData->Subtitles.Num() > 0);
 	if (bHasSubtitles)
