UE: Fix Update Source button on Pose Asset and warning about guid being out of date

diff --git a/Engine/Source/Runtime/Engine/Private/Animation/PoseAsset.cpp b/Engine/Source/Runtime/Engine/Private/Animation/PoseAsset.cpp
--- a/Engine/Source/Runtime/Engine/Private/Animation/PoseAsset.cpp
+++ b/Engine/Source/Runtime/Engine/Private/Animation/PoseAsset.cpp
@@ -946,14 +946,15 @@ void UPoseAsset::PostLoad()
 
 			IAnimationDataModel::FGuidGenerationSettings Settings;
 			Settings.bIncludeTimingData = false;
-			if (SourceAnimationRawDataGUID.IsValid() && SourceAnimationRawDataGUID != SourceAnimation->GetDataModel()->GenerateGuid(Settings))
+			FGuid DataModelGuid = SourceAnimation->GetDataModel()->GenerateGuid(Settings);
+			if (SourceAnimationRawDataGUID.IsValid() && SourceAnimationRawDataGUID != DataModelGuid)
 			{
 				FFormatNamedArguments Args;
 				Args.Add(TEXT("AssetName"), FText::FromString(GetPathName()));
 				Args.Add(TEXT("SourceAsset"), FText::FromString(SourceAnimation->GetPathName()));
 
 				Args.Add(TEXT("Stored"), FText::FromString(SourceAnimationRawDataGUID.ToString()));
-				Args.Add(TEXT("Found"), FText::FromString(SourceAnimation->GetDataModel()->GenerateGuid().ToString()));
+				Args.Add(TEXT("Found"), FText::FromString(DataModelGuid.ToString()));
 				
 				const FText ResultText = FText::Format(LOCTEXT("PoseAssetSourceOutOfDate", "PoseAsset {AssetName} is out-of-date with its source animation {SourceAsset} {Stored} vs {Found}"), Args);
 				UE_LOG(LogAnimation, Warning,TEXT("%s"), *ResultText.ToString());
@@ -1552,7 +1553,9 @@ void UPoseAsset::CreatePoseFromAnimation(class UAnimSequence* AnimSequence, cons
 				PostProcessData();
 			}
 
-			SourceAnimationRawDataGUID = AnimSequence->GetDataModel()->GenerateGuid();
+			IAnimationDataModel::FGuidGenerationSettings Settings;
+			Settings.bIncludeTimingData = false;
+			SourceAnimationRawDataGUID = AnimSequence->GetDataModel()->GenerateGuid(Settings);
 		}
 	}
 }
