FRenderAssetStreamingManager - Store TObjectPtr<UStreamableRenderAsset> in ReferencedRenderAssets but ensure the TSet keying is done with raw pointers
* Under certain circumstances, lookups can be performed on ReferencedRenderAssets with invalid pointers (freed resources).
* Constructing a TObjectPtr<UStreamableRenderAsset>() from an invalid pointer is unsupported.

diff --git a/Engine/Source/Runtime/Engine/Private/Streaming/StreamingManagerTexture.cpp b/Engine/Source/Runtime/Engine/Private/Streaming/StreamingManagerTexture.cpp
--- a/Engine/Source/Runtime/Engine/Private/Streaming/StreamingManagerTexture.cpp
+++ b/Engine/Source/Runtime/Engine/Private/Streaming/StreamingManagerTexture.cpp
@@ -3189,7 +3189,10 @@ void FRenderAssetStreamingManager::AddReferencedObjects(FReferenceCollector& Col
 	 */
 	if(!AsyncWork->IsDone())
 	{
-		Collector.AddReferencedObjects(ReferencedRenderAssets);
+		for (TObjectPtr<UStreamableRenderAsset>& ReferencedRenderAsset : ReferencedRenderAssets)
+		{
+			Collector.AddReferencedObject(ReferencedRenderAsset);
+		}
 	}
 }
 FString FRenderAssetStreamingManager::GetReferencerName() const
diff --git a/Engine/Source/Runtime/Engine/Private/Streaming/StreamingManagerTexture.h b/Engine/Source/Runtime/Engine/Private/Streaming/StreamingManagerTexture.h
--- a/Engine/Source/Runtime/Engine/Private/Streaming/StreamingManagerTexture.h
+++ b/Engine/Source/Runtime/Engine/Private/Streaming/StreamingManagerTexture.h
@@ -271,8 +271,21 @@ private:
 		*/
 		TArray<FStreamingRenderAsset> AsyncUnsafeStreamingRenderAssets;
 
-		/** All the textures/meshes referenced in StreamingRenderAssets. Used to handled deleted textures/meshes.  */
-		TSet<TObjectPtr<UStreamableRenderAsset>> ReferencedRenderAssets;
+		/** 
+		 * KeyFunc to ensure lookups in the TSet below can be done with raw pointers.
+		 * Under certain circumstances, lookups can be performed on ReferencedRenderAssets with invalid pointers (freed resources).
+		 * Constructing a TObjectPtr<UStreamableRenderAsset>() from an invalid pointer is unsupported.
+		 */
+		struct FPointerKeyFunc : BaseKeyFuncs<UStreamableRenderAsset*, const UStreamableRenderAsset*, false>
+		{
+			static KeyInitType	GetSetKey(ElementInitType& Element) { return Element; }
+			static KeyInitType	GetSetKey(const TObjectPtr<UStreamableRenderAsset>& Element) { return Element.Get(); }			
+			static bool			Matches(KeyInitType A, KeyInitType B) { return A == B; }
+			static uint32		GetKeyHash(KeyInitType Key) { return GetTypeHash(Key); }
+		};
+
+		/** All the textures/meshes referenced in StreamingRenderAssets. Used to handle deleted textures/meshes. */
+		TSet<TObjectPtr<UStreamableRenderAsset>, FPointerKeyFunc> ReferencedRenderAssets;
 
 		/** All the currently installed optional bulkdata files */
 		TSet<FString> OptionalBulkDataFiles;
