[Waveform Editor] - Editor Crashes when closing out a Sound Wave asset after canceling to Export Waveform

diff --git a/Engine/Source/Runtime/Engine/Private/SampleBufferIO.cpp b/Engine/Source/Runtime/Engine/Private/SampleBufferIO.cpp
--- a/Engine/Source/Runtime/Engine/Private/SampleBufferIO.cpp
+++ b/Engine/Source/Runtime/Engine/Private/SampleBufferIO.cpp
@@ -425,7 +425,7 @@ namespace Audio
 	{
 		CancelWrite();
 
-		if (CurrentSoundWave && !bWasPreviouslyAddedToRoot)
+		if (CurrentSoundWave && CurrentSoundWave->IsValidLowLevel() && !bWasPreviouslyAddedToRoot)
 		{
 			CurrentSoundWave->RemoveFromRoot();
 		}
