New material translator:
- Added support for material enumeration parameters.
- Added support for float-to-int and int-to-float expressions.
- Added CastToScalarKind() API to the Emitter (with its specialized CastToBool(), etc)
- Renamed Cast instruction to Scalar, to better express its semantic (constructing a single scalar value from another).
- Moved local function in MaterialExpressionOperator from header to cpp only.

#rnx

diff --git a/Engine/Source/Runtime/Engine/Private/Materials/MaterialExpressions.cpp b/Engine/Source/Runtime/Engine/Private/Materials/MaterialExpressions.cpp
--- a/Engine/Source/Runtime/Engine/Private/Materials/MaterialExpressions.cpp
+++ b/Engine/Source/Runtime/Engine/Private/Materials/MaterialExpressions.cpp
@@ -23982,7 +23982,7 @@ UMaterialExpressionOperator::UMaterialExpressionOperator(const FObjectInitialize
 
 #if WITH_EDITOR
 
-bool UMaterialExpressionOperator::OperatorSupportsVariablePinCount(EMaterialExpressionOperatorKind InOperator)
+static bool OperatorSupportsVariablePinCount(EMaterialExpressionOperatorKind InOperator)
 {
 	auto begin = std::begin(VariableInputsExpressionOperators);
 	auto end = std::end(VariableInputsExpressionOperators);
diff --git a/Engine/Source/Runtime/Engine/Private/Materials/MaterialExpressionsIR.cpp b/Engine/Source/Runtime/Engine/Private/Materials/MaterialExpressionsIR.cpp
--- a/Engine/Source/Runtime/Engine/Private/Materials/MaterialExpressionsIR.cpp
+++ b/Engine/Source/Runtime/Engine/Private/Materials/MaterialExpressionsIR.cpp
@@ -73,6 +73,7 @@
 #include "Materials/MaterialExpressionDivide.h"
 #include "Materials/MaterialExpressionDotProduct.h"
 #include "Materials/MaterialExpressionDoubleVectorParameter.h"
+#include "Materials/MaterialExpressionEnumerationParameter.h"
 #include "Materials/MaterialExpressionDynamicParameter.h"
 #include "Materials/MaterialExpressionExponential.h"
 #include "Materials/MaterialExpressionExponential2.h"
@@ -234,6 +235,7 @@
 #include "Materials/MaterialExpressionVolumetricAdvancedMaterialInput.h"
 #include "Materials/MaterialExpressionVolumetricAdvancedMaterialOutput.h"
 #include "Materials/MaterialExpressionWorldPosition.h"
+#include "Materials/MaterialExpressionFloatToUInt.h"
 #include "Materials/MaterialExternalCodeRegistry.h"
 #include "Materials/MaterialFunctionInterface.h"
 #include "Materials/MaterialParameterCollection.h"
@@ -411,6 +413,22 @@ void UMaterialExpressionStaticSwitchParameter::Build(MIR::FEmitter& Em)
 	Em.Output(0, Em.Input(bCondition ? &A : &B));
 }
 
+void UMaterialExpressionEnumerationParameter::Build(MIR::FEmitter& Em)
+{
+	if (!Enumeration)
+	{
+		return Em.Errorf(TEXT("Enumeration definition unspecified in enumeration parameter '%s'."), *ParameterName.ToString());
+	}
+
+	FMaterialParameterMetadata Metadata;
+	if (!GetParameterValue(Metadata))
+	{
+		return Em.Errorf(TEXT("Could fetch parameter metadata from enumeration parameter '%s'."), *ParameterName.ToString());
+	}
+
+	Em.Output(0, Em.Parameter(ParameterName, Metadata));
+}
+
 void UMaterialExpressionAppendVector::Build(MIR::FEmitter& Em)
 {
 	FValueRef AVal = Em.CheckIsScalarOrVector(Em.Input(&A));
@@ -1356,7 +1374,16 @@ void UMaterialExpressionOperator::Build(MIR::FEmitter& Em)
 
 		Em.Output(0, Em.Operator(OpMIR, AValue, BValue, CValue));
 	}
+}
 
+void UMaterialExpressionFloatToUInt::Build(MIR::FEmitter& Em)
+{
+	Em.Output(0, Em.CastToInt(Em.Input(&Input)));
+}
+
+void UMaterialExpressionUIntToFloat::Build(MIR::FEmitter& Em)
+{
+	Em.Output(0, Em.CastToFloat(Em.Input(&Input)));
 }
 
 #endif // WITH_EDITOR
diff --git a/Engine/Source/Runtime/Engine/Private/Materials/MaterialIR.cpp b/Engine/Source/Runtime/Engine/Private/Materials/MaterialIR.cpp
--- a/Engine/Source/Runtime/Engine/Private/Materials/MaterialIR.cpp
+++ b/Engine/Source/Runtime/Engine/Private/Materials/MaterialIR.cpp
@@ -48,7 +48,7 @@ const TCHAR* LexToString(EValueKind Kind)
 		case VK_Operator: return TEXT("Operator");
 		case VK_Branch: return TEXT("Branch");
 		case VK_Subscript: return TEXT("Subscript");
-		case VK_Cast: return TEXT("Cast");
+		case VK_Scalar: return TEXT("Scalar");
 		case VK_TextureRead: return TEXT("TextureRead");
 		case VK_InlineHLSL: return TEXT("InlineHLSL");
 		case VK_StageSwitch: return TEXT("StageSwitch");
@@ -106,7 +106,7 @@ uint32 FValue::GetSizeInBytes() const
 		case VK_Operator: return sizeof(FOperator);
 		case VK_Branch: return sizeof(FBranch);
 		case VK_Subscript: return sizeof(FSubscript);
-		case VK_Cast: return sizeof(FCast);
+		case VK_Scalar: return sizeof(FScalar);
 		case VK_TextureRead: return sizeof(FTextureRead);
 		case VK_InlineHLSL: return sizeof(FInlineHLSL);
 		case VK_StageSwitch: return sizeof(FStageSwitch);
@@ -187,10 +187,10 @@ TConstArrayView<FValue*> FValue::GetUses() const
 			return { &This->Arg, FSubscript::NumStaticUses };
 		}
 
-		case VK_Cast:
+		case VK_Scalar:
 		{
-			auto This = static_cast<const FCast*>(this);
-			return { &This->Arg, FCast::NumStaticUses };
+			auto This = static_cast<const FScalar*>(this);
+			return { &This->Arg, FScalar::NumStaticUses };
 		}
 			
 		case VK_TextureRead:
diff --git a/Engine/Source/Runtime/Engine/Private/Materials/MaterialIREmitter.cpp b/Engine/Source/Runtime/Engine/Private/Materials/MaterialIREmitter.cpp
--- a/Engine/Source/Runtime/Engine/Private/Materials/MaterialIREmitter.cpp
+++ b/Engine/Source/Runtime/Engine/Private/Materials/MaterialIREmitter.cpp
@@ -618,13 +618,18 @@ FValueRef FEmitter::Parameter(FName Name, FMaterialParameterMetadata& Metadata,
 	switch (Metadata.Value.Type)
 	{
 		case EMaterialParameterType::Scalar:
+		case EMaterialParameterType::Enumeration:
 		{
 			if (Metadata.PrimitiveDataIndex != INDEX_NONE)
 			{
 				return CustomPrimitiveData(Metadata.PrimitiveDataIndex);
 			}
 
-			MIR::FUniformParameter Prototype = MakePrototype<MIR::FUniformParameter>(FPrimitiveType::GetFloat());
+			const FType* Type = Metadata.Value.Type == EMaterialParameterType::Enumeration
+				? FPrimitiveType::GetInt()
+				: FPrimitiveType::GetFloat();
+
+			MIR::FUniformParameter Prototype = MakePrototype<MIR::FUniformParameter>(Type);
 			Prototype.ParameterIdInModule = RegisterParameter(Module, Info, Metadata);
 			return EmitPrototype(*this, Prototype);
 		}
@@ -2013,25 +2018,25 @@ static FValue* CastValueToPrimitiveType(FEmitter& Emitter, FValueRef Value, cons
 	// Construct a scalar from another scalar.
 	if (TargetPrimitiveType->IsScalar())
 	{
-		// 
+		// Get the first component of value. We already know value's type is primitive, so this will return a scalar.
 		Value = Emitter.Subscript(Value, 0);
 		ValuePrimitiveType = Value->Type->AsPrimitive();
+		check(ValuePrimitiveType->IsScalar());
 		
-		//
+		// If types are identical, return the component value as is.
 		if (ValuePrimitiveType == TargetPrimitiveType)
 		{
 			return Value;
 		}
-
-		// Construct the scalar from a constant.
-		if (FConstant* ConstantInitializer = As<FConstant>(Value))
+		// If value is a constant, we can cast the constant now.
+		else if (FConstant* ConstantInitializer = As<FConstant>(Value)) 
 		{
 			return CastConstant(Emitter, ConstantInitializer, ValuePrimitiveType->ScalarKind, TargetPrimitiveType->ScalarKind);
 		}
 		else
 		{
-			// Emit the cast to the target type of the subscript value
-			FCast Prototype = MakePrototype<FCast>(TargetPrimitiveType);
+			// Otherwise emit a cast instruction of the subscript value to the target type.
+			FScalar Prototype = MakePrototype<FScalar>(TargetPrimitiveType);
 			Prototype.Arg = Value;
 			return EmitPrototype(Emitter, Prototype);
 		}
@@ -2097,36 +2102,25 @@ static FValue* CastValueToPrimitiveType(FEmitter& Emitter, FValueRef Value, cons
 	}
 	
 	// The two primitive types are identical matrices that differ only by their scalar type.
-	if (TargetPrimitiveType->NumRows 	 == ValuePrimitiveType->NumRows &&
+	if (TargetPrimitiveType->NumRows    == ValuePrimitiveType->NumRows &&
 		TargetPrimitiveType->NumColumns == ValuePrimitiveType->NumColumns)
 	{
 		check(TargetPrimitiveType->IsMatrix());
 
-		//  If initializer is a dimensional, we can subscript components of interest.
-		if (FDimensional* DimensionalInitializer = As<FDimensional>(Value))
-		{
-			// Create the result dimensional value.
-			FDimensional* Result = NewDimensionalValue(Emitter, TargetPrimitiveType);
-			TArrayView<FValue*> ResultComponents = Result->GetMutableComponents();
+		// Create the result dimensional value.
+		FDimensional* Result = NewDimensionalValue(Emitter, TargetPrimitiveType);
+		TArrayView<FValue*> ResultComponents = Result->GetMutableComponents();
 			
-			// Determine the result component type (scalar).
-			const FPrimitiveType* ResultComponentType = TargetPrimitiveType->ToScalar();
-
-			// Convert components from the initializer vector.
-			for (int32 Index = 0, Num = TargetPrimitiveType->GetNumComponents(); Index < Num; ++Index)
-			{
-				ResultComponents[Index] = Emitter.Cast(DimensionalInitializer->GetComponents()[Index], ResultComponentType);
-			}
+		// Determine the result component type (scalar).
+		const FPrimitiveType* ResultComponentType = TargetPrimitiveType->ToScalar();
 
-			return EmitNew(Emitter, Result);
-		}
-		else
+		// Convert components from the initializer vector.
+		for (int32 Index = 0, Num = TargetPrimitiveType->GetNumComponents(); Index < Num; ++Index)
 		{
-			// Initializer is an unknown value, construct target value casting initializer.
-			FCast Prototype = MakePrototype<FCast>(TargetPrimitiveType);
-			Prototype.Arg = Value;
-			return EmitPrototype(Emitter, Prototype);
+			ResultComponents[Index] = Emitter.Cast(Emitter.Subscript(Value, Index), ResultComponentType);
 		}
+
+		return EmitNew(Emitter, Result);
 	}
 
 	// Initializer value cannot be used to construct this primitive type.
@@ -2169,6 +2163,29 @@ FValueRef FEmitter::CastToScalar(FValueRef Value)
 		: Value;
 }
 
+FValueRef FEmitter::CastToScalarKind(FValueRef Value, EScalarKind ToScalarKind)
+{
+	Value = CheckIsPrimitive(Value);
+	return Value.IsValid()
+		? Cast(Value, Value->Type->AsPrimitive()->WithScalarKind(ToScalarKind))
+		: Value;
+}
+
+FValueRef FEmitter::CastToBool(FValueRef Value)
+{
+	return CastToScalarKind(Value, ScalarKind_Bool);
+}
+
+FValueRef FEmitter::CastToInt(FValueRef Value)
+{
+	return CastToScalarKind(Value, ScalarKind_Bool);
+}
+
+FValueRef FEmitter::CastToFloat(FValueRef Value)
+{
+	return CastToScalarKind(Value, ScalarKind_Float);
+}
+
 FValueRef FEmitter::CastToBool(FValueRef Value, int NumRows)
 {
 	return Cast(Value, FPrimitiveType::GetVector(ScalarKind_Bool, NumRows));
@@ -2665,9 +2682,9 @@ FValueRef FEmitter::AnalyticalPartialDerivative(FValueRef Value, EDerivativeAxis
 			return Subscript(AnalyticalPartialDerivative(AsSubscript->Arg, Axis), AsSubscript->Index);
 		}
 
-		case MIR::VK_Cast:
+		case MIR::VK_Scalar:
 		{
-			FCast* AsCast = Value->As<FCast>();
+			FScalar* AsCast = Value->As<FScalar>();
 			return Cast(AnalyticalPartialDerivative(AsCast->Arg, Axis), AsCast->Type);
 		}
 
diff --git a/Engine/Source/Runtime/Engine/Private/Materials/MaterialIRToHLSLTranslator.cpp b/Engine/Source/Runtime/Engine/Private/Materials/MaterialIRToHLSLTranslator.cpp
--- a/Engine/Source/Runtime/Engine/Private/Materials/MaterialIRToHLSLTranslator.cpp
+++ b/Engine/Source/Runtime/Engine/Private/Materials/MaterialIRToHLSLTranslator.cpp
@@ -392,6 +392,7 @@ struct FTranslator : FMaterialIRToHLSLTranslation
 			case MIR::VK_Operator: LowerOperator(static_cast<MIR::FOperator*>(Instr)); break;
 			case MIR::VK_Branch: LowerBranch(static_cast<MIR::FBranch*>(Instr)); break;
 			case MIR::VK_Subscript: LowerSubscript(static_cast<MIR::FSubscript*>(Instr)); break;
+			case MIR::VK_Scalar: LowerScalar(static_cast<MIR::FScalar*>(Instr)); break;
 			case MIR::VK_TextureRead: LowerTextureRead(static_cast<MIR::FTextureRead*>(Instr)); break;
 			case MIR::VK_InlineHLSL: LowerInlineHLSL(static_cast<const MIR::FInlineHLSL*>(Instr)); break;
 			case MIR::VK_StageSwitch: LowerStageSwitch(static_cast<const MIR::FStageSwitch*>(Instr)); break;
@@ -968,6 +969,11 @@ struct FTranslator : FMaterialIRToHLSLTranslation
 		}
 	}
 
+	void LowerScalar(const MIR::FScalar* Scalar)
+	{
+		Printer << LowerType(Scalar->Type) << "(" << LowerValue(Scalar->Arg) << ")";
+	}
+
 	void LowerTextureRead(const MIR::FTextureRead* TextureRead)
 	{
 		bool bSamplerNeedsBrackets = LowerSamplerType(TextureRead->SamplerType);
diff --git a/Engine/Source/Runtime/Engine/Private/Materials/MaterialIRValueAnalyzer.cpp b/Engine/Source/Runtime/Engine/Private/Materials/MaterialIRValueAnalyzer.cpp
--- a/Engine/Source/Runtime/Engine/Private/Materials/MaterialIRValueAnalyzer.cpp
+++ b/Engine/Source/Runtime/Engine/Private/Materials/MaterialIRValueAnalyzer.cpp
@@ -13,6 +13,7 @@
 #include "Materials/MaterialIRInternal.h"
 #include "Materials/MaterialInsights.h"
 #include "Materials/Material.h"
+#include "Materials/MaterialEnumeration.h"
 #include "MaterialShared.h"
 #include "Engine/Texture.h"
 #include "UObject/Package.h"
@@ -116,6 +117,7 @@ static void AnalyzePrimitiveUniformParameter(FMaterialIRValueAnalyzer& Analyzer,
 	{
 		case EMaterialParameterType::Scalar: DefaultValue = ParameterMetadata.Value.AsScalar(); break;
 		case EMaterialParameterType::Vector: DefaultValue = ParameterMetadata.Value.AsLinearColor(); break;
+		case EMaterialParameterType::Enumeration: DefaultValue = Cast<IMaterialEnumerationProvider>(ParameterMetadata.Enumeration)->GetValueOrDefault(ParameterMetadata.Value.EnumEntry, 0); break;
 		default: UE_MIR_TODO();
 	}
 
diff --git a/Engine/Source/Runtime/Engine/Private/Materials/MaterialUniformExpressions.cpp b/Engine/Source/Runtime/Engine/Private/Materials/MaterialUniformExpressions.cpp
--- a/Engine/Source/Runtime/Engine/Private/Materials/MaterialUniformExpressions.cpp
+++ b/Engine/Source/Runtime/Engine/Private/Materials/MaterialUniformExpressions.cpp
@@ -1015,11 +1015,9 @@ void FUniformExpressionSet::FillUniformBuffer(const FMaterialRenderContext& Mate
 			UE::Shader::FValue ParameterShaderValue;
 
 			// First allow proxy the chance to override parameter
-			FMaterialParameterValue ParameterValue;
 			if (MaterialRenderContext.MaterialRenderProxy 
-				&& MaterialRenderContext.MaterialRenderProxy->GetParameterValue(Parameter.ParameterType, Parameter.ParameterInfo, ParameterValue, MaterialRenderContext))
+				&& MaterialRenderContext.MaterialRenderProxy->GetParameterShaderValue(Parameter.ParameterType, Parameter.ParameterInfo, ParameterShaderValue, MaterialRenderContext))
 			{
-				ParameterShaderValue = ParameterValue.AsShaderValue();
 				bFoundParameter = true;
 			}
 
diff --git a/Engine/Source/Runtime/Engine/Public/Materials/MaterialExpressionEnumerationParameter.h b/Engine/Source/Runtime/Engine/Public/Materials/MaterialExpressionEnumerationParameter.h
--- a/Engine/Source/Runtime/Engine/Public/Materials/MaterialExpressionEnumerationParameter.h
+++ b/Engine/Source/Runtime/Engine/Public/Materials/MaterialExpressionEnumerationParameter.h
@@ -26,6 +26,7 @@ public:
 #if WITH_EDITOR
 	//~ Begin UMaterialExpression Interface
 	virtual int32 Compile(class FMaterialCompiler* Compiler, int32 OutputIndex) override;
+	virtual void Build(MIR::FEmitter& em);
 	virtual void GetCaption(TArray<FString>& OutCaptions) const override;
 	virtual EMaterialValueType GetOutputValueType(int32 InIndex) override;
 	virtual bool GetParameterValue(FMaterialParameterMetadata& OutMeta) const override;
diff --git a/Engine/Source/Runtime/Engine/Public/Materials/MaterialExpressionFloatToUInt.h b/Engine/Source/Runtime/Engine/Public/Materials/MaterialExpressionFloatToUInt.h
--- a/Engine/Source/Runtime/Engine/Public/Materials/MaterialExpressionFloatToUInt.h
+++ b/Engine/Source/Runtime/Engine/Public/Materials/MaterialExpressionFloatToUInt.h
@@ -33,6 +33,7 @@ public:
 	//~ Begin UMaterialExpression Interface
 #if WITH_EDITOR
 	virtual int32 Compile(class FMaterialCompiler* Compiler, int32 OutputIndex) override;
+	virtual void Build(MIR::FEmitter& Emitter) override;
 	virtual void GetCaption(TArray<FString>& OutCaptions) const override;
 	virtual EMaterialValueType GetInputValueType(int32 InputIndex) override;
 	virtual EMaterialValueType GetOutputValueType(int32 OutputIndex) override;
@@ -53,6 +54,7 @@ public:
 	//~ Begin UMaterialExpression Interface
 #if WITH_EDITOR
 	virtual int32 Compile(class FMaterialCompiler* Compiler, int32 OutputIndex) override;
+	virtual void Build(MIR::FEmitter& Emitter) override;
 	virtual void GetCaption(TArray<FString>& OutCaptions) const override;
 	virtual EMaterialValueType GetInputValueType(int32 InputIndex) override;
 	virtual EMaterialValueType GetOutputValueType(int32 OutputIndex) override;
diff --git a/Engine/Source/Runtime/Engine/Public/Materials/MaterialExpressionOperator.h b/Engine/Source/Runtime/Engine/Public/Materials/MaterialExpressionOperator.h
--- a/Engine/Source/Runtime/Engine/Public/Materials/MaterialExpressionOperator.h
+++ b/Engine/Source/Runtime/Engine/Public/Materials/MaterialExpressionOperator.h
@@ -122,7 +122,6 @@ public:
 
 	//~ Begin UMaterialExpression Interface
 #if WITH_EDITOR
-	static bool OperatorSupportsVariablePinCount(EMaterialExpressionOperatorKind InOperator);
 	virtual void AddInputPin();	
 	virtual bool CanDeletePin(EEdGraphPinDirection PinDirection, int32 PinIndex) const override;
 	virtual void DeletePin(EEdGraphPinDirection PinDirection, int32 PinIndex) override;	
diff --git a/Engine/Source/Runtime/Engine/Public/Materials/MaterialIR.h b/Engine/Source/Runtime/Engine/Public/Materials/MaterialIR.h
--- a/Engine/Source/Runtime/Engine/Public/Materials/MaterialIR.h
+++ b/Engine/Source/Runtime/Engine/Public/Materials/MaterialIR.h
@@ -100,7 +100,7 @@ enum EValueKind
 	VK_Operator,
 	VK_Branch,
 	VK_Subscript,
-	VK_Cast,
+	VK_Scalar,
 	VK_TextureRead,
 	VK_InlineHLSL,
 	VK_StageSwitch,
@@ -722,12 +722,12 @@ struct FSubscript : TInstruction<VK_Subscript, 1>
 	int Index;
 };
 
-// A type cast instruction.
+// A scalar construction instruction.
 //
-// This instruction casts an argument value to a different type.
-struct FCast : TInstruction<VK_Cast, 1>
+// This instruction constructs a scalar from another (of possibly a different scalar kind).
+struct FScalar : TInstruction<VK_Scalar, 1>
 {
-	// The argument value to cast.
+	// The initializing argument value.
 	FValue* Arg;
 };
 
diff --git a/Engine/Source/Runtime/Engine/Public/Materials/MaterialIREmitter.h b/Engine/Source/Runtime/Engine/Public/Materials/MaterialIREmitter.h
--- a/Engine/Source/Runtime/Engine/Public/Materials/MaterialIREmitter.h
+++ b/Engine/Source/Runtime/Engine/Public/Materials/MaterialIREmitter.h
@@ -339,10 +339,16 @@ public:
 	// Casts the given value to the specified target type. Returns poison if the cast is invalid.
 	FValueRef Cast(FValueRef Value, const FType* TargetType);
 	
-	// Casts the given arithmetic value to a scalar. If this is a vector, it will convert to its first component.
-	// Returns if the cast is invalid.
+	// Casts the given primitive value to a scalar. If this is a vector, it will convert to its first component.
 	FValueRef CastToScalar(FValueRef Value);
 
+	// Casts the given primitive value to a given scalar kind, maintaining the same
+	// number of rows and columns.
+	FValueRef CastToScalarKind(FValueRef Value, EScalarKind ToScalarKind);
+	FValueRef CastToBool(FValueRef Value);
+	FValueRef CastToInt(FValueRef Value);
+	FValueRef CastToFloat(FValueRef Value);
+
 	// These functions cast the given value to the a primitive scalar or column vector type of given rows.
 	// If the value is missing ore could not be cast to the relative primitive type, they repors an error and return poison.
 	FValueRef CastToBool(FValueRef Value, int NumRows);
