Fixed ensure when UActorComponent::DestroyPhysicsState() is called.
This could happen when enabling p.Chaos.EnableAsyncInitBody and LevelStreaming.AllowIncrementalPreRegisterComponents.
If a child actor is recreated by UChildActorComponent::OnRegister() and a component part of that child actor was pre-registered (but not registered), then a future GC would destroy the deleted child actor, DestroyPhysicsState would trigger the ensure that validated that the component was registered. The fix is to also accept PreRegistered components.
#rnx

diff --git a/Engine/Source/Runtime/Engine/Private/Components/ActorComponent.cpp b/Engine/Source/Runtime/Engine/Private/Components/ActorComponent.cpp
--- a/Engine/Source/Runtime/Engine/Private/Components/ActorComponent.cpp
+++ b/Engine/Source/Runtime/Engine/Private/Components/ActorComponent.cpp
@@ -2350,7 +2350,7 @@ void UActorComponent::DestroyPhysicsState()
 		// Broadcast delegate
 		GlobalDestroyPhysicsDelegate.Broadcast(this);
 
-		ensureMsgf(bRegistered, TEXT("Component has physics state when not registered (%s)"), *GetFullName()); // should not have physics state unless we are registered
+		ensureMsgf(bRegistered || IsPreRegistered(), TEXT("Component has physics state when not registered (%s)"), *GetFullName()); // should not have physics state unless we are registered
 
 		// Call virtual
 		OnDestroyPhysicsState();
