[Mobile][ArmASR][WIP] Allow velocity rendering when using a temporal upscaler. This is for supporting rendering velocities for use with ArmASR.
- Added in API for post processes on mobile to pass down their scene texture uniform buffers / gbuffer details. Modified the mobile renderer to pass these down.
- Added a bunch of caching of the various buffers created by ArmASR to reduce RHI / Render thread impact of enabling the feature. Turns out creating 4 fullscreen buffers per frame is pretty expensive.
- Modified velocity rendering detection to see if a temporal upscaler is being used. If it is (or we're using a temporal AA version), we will enable velocity rendering.
- Added a platform support variable for ArmASR that can be used to tell games (per device profile) that we would like to allow ArmASR to be enabled at runtime.
- TODO: Add the platform support variables for the specific devices / GPU families we want to target to Device Profiles.
- TODO Cache of current screen percentage and set one whenever "r.ArmASR.Enable" is triggered. We'll probably use the balanced value of 59% as the default for now.
#rnx

diff --git a/Engine/Source/Runtime/Engine/Public/SceneViewExtension.h b/Engine/Source/Runtime/Engine/Public/SceneViewExtension.h
--- a/Engine/Source/Runtime/Engine/Public/SceneViewExtension.h
+++ b/Engine/Source/Runtime/Engine/Public/SceneViewExtension.h
@@ -86,6 +86,7 @@ struct FMinimalViewInfo;
 struct FSceneViewProjectionData;
 class FRDGBuilder;
 struct FPostProcessingInputs;
+struct FMobilePostProcessingInputs;
 struct FPostProcessMaterialInputs;
 struct FScreenPassTexture;
 class FViewport;
@@ -188,6 +189,11 @@ public:
 	 */
 	virtual void PrePostProcessPass_RenderThread(FRDGBuilder& GraphBuilder, const FSceneView& InView, const FPostProcessingInputs& Inputs) {};
 
+	/**
+	 * Called right before Post Processing rendering begins for the mobile renderer
+	 */
+	virtual void PrePostProcessPassMobile_RenderThread(FRDGBuilder& GraphBuilder, const FSceneView& InView, const FMobilePostProcessingInputs& Inputs) {};
+
 	/**
 	* This will be called at the beginning of post processing to make sure that each view extension gets a chance to subscribe to a post-processing pass event.
 	*  - The pass MUST write to the override output texture if it is active (this occurs when the pass is the last in the post processing chain writing to the back buffer).
diff --git a/Engine/Source/Runtime/Renderer/Private/MobileShadingRenderer.cpp b/Engine/Source/Runtime/Renderer/Private/MobileShadingRenderer.cpp
--- a/Engine/Source/Runtime/Renderer/Private/MobileShadingRenderer.cpp
+++ b/Engine/Source/Runtime/Renderer/Private/MobileShadingRenderer.cpp
@@ -1550,6 +1550,16 @@ void FMobileSceneRenderer::Render(FRDGBuilder& GraphBuilder, const FSceneRenderU
 					PostProcessingInputs.ViewFamilyTexture = ViewFamilyTexture;
 					PostProcessingInputs.SceneTextures = CreateMobileSceneTextureUniformBuffer(GraphBuilder, &SceneTextures, EMobileSceneTextureSetupMode::All);
 
+					for (int32 ViewExt = 0; ViewExt < ViewFamily.ViewExtensions.Num(); ++ViewExt)
+					{
+						for (int32 ViewIndex = 0; ViewIndex < ViewFamily.Views.Num(); ++ViewIndex)
+						{
+							FViewInfo& View = Views[ViewIndex];
+							RDG_GPU_MASK_SCOPE(GraphBuilder, View.GPUMask);
+							ViewFamily.ViewExtensions[ViewExt]->PrePostProcessPassMobile_RenderThread(GraphBuilder, View, PostProcessingInputs);
+						}
+					}
+
 					for (int32 ViewIndex = 0; ViewIndex < Views.Num(); ViewIndex++)
 					{
 						if (Views[ViewIndex].ShouldRenderView())
diff --git a/Engine/Source/Runtime/Renderer/Private/VelocityRendering.cpp b/Engine/Source/Runtime/Renderer/Private/VelocityRendering.cpp
--- a/Engine/Source/Runtime/Renderer/Private/VelocityRendering.cpp
+++ b/Engine/Source/Runtime/Renderer/Private/VelocityRendering.cpp
@@ -220,13 +220,15 @@ bool FMobileSceneRenderer::ShouldRenderVelocities() const
 	}
 
 	bool bNeedsVelocity = false;
-	for (int32 ViewIndex = 0; ViewIndex < Views.Num(); ViewIndex++)
+	for (int32 ViewIndex = 0; ViewIndex < Views.Num() && !bNeedsVelocity; ViewIndex++)
 	{
 		const FViewInfo& View = Views[ViewIndex];
 
-		bool bTemporalAA = IsTemporalAccumulationBasedMethod(View.AntiAliasingMethod) && !View.bCameraCut;
+		const bool bTemporalAA = IsTemporalAccumulationBasedMethod(View.AntiAliasingMethod);
+		const bool bIsUsingTemporalUpscaler = View.Family->GetTemporalUpscalerInterface() != nullptr;
+		const bool bVelocityRendering = (bIsUsingTemporalUpscaler || bTemporalAA) && !View.bCameraCut;
 
-		bNeedsVelocity |= bTemporalAA;
+		bNeedsVelocity |= bVelocityRendering;
 	}
 
 	return bNeedsVelocity;
