Changed the physics query ignore handler to be a RAII object to prevent accidental leakage.

diff --git a/Engine/Source/Runtime/Engine/Private/Collision/SceneQuery.cpp b/Engine/Source/Runtime/Engine/Private/Collision/SceneQuery.cpp
--- a/Engine/Source/Runtime/Engine/Private/Collision/SceneQuery.cpp
+++ b/Engine/Source/Runtime/Engine/Private/Collision/SceneQuery.cpp
@@ -1627,9 +1627,15 @@ namespace Chaos::Private
 	}
 
 #if UE_WITH_REMOTE_OBJECT_HANDLE
-	void FQueryInterface_Internal::SetIgnoreQueryHandlerThreadLocal(bool bIgnore)
+	FScopedIgnoreQueryHandler::FScopedIgnoreQueryHandler(bool bIgnore)
 	{
+		bOldState = bIgnoreQueryHandler;
 		bIgnoreQueryHandler = bIgnore;
 	}
+
+	FScopedIgnoreQueryHandler::~FScopedIgnoreQueryHandler()
+	{
+		bIgnoreQueryHandler = bOldState;
+	}
 #endif
 }
diff --git a/Engine/Source/Runtime/Engine/Public/Physics/GenericPhysicsInterface.h b/Engine/Source/Runtime/Engine/Public/Physics/GenericPhysicsInterface.h
--- a/Engine/Source/Runtime/Engine/Public/Physics/GenericPhysicsInterface.h
+++ b/Engine/Source/Runtime/Engine/Public/Physics/GenericPhysicsInterface.h
@@ -153,9 +153,17 @@ namespace Chaos::Private
 		static bool SweepTest(const EThreadQueryContext ThreadContext, const UWorld* World, const FSweepQueryData& SweepData, const FCommonQueryData& CommonData, FHitResult& OutHits);
 		static bool SweepSingle(const EThreadQueryContext ThreadContext, const UWorld* World, const FSweepQueryData& SweepData, const FCommonQueryData& CommonData, FHitResult& OutHits);
 		static bool SweepMulti(const EThreadQueryContext ThreadContext, const UWorld* World, const FSweepQueryData& SweepData, const FCommonQueryData& CommonData, TArray<FHitResult>& OutHits);
+	};
 
 #if UE_WITH_REMOTE_OBJECT_HANDLE
-		UE_INTERNAL ENGINE_API static void SetIgnoreQueryHandlerThreadLocal(bool bIgnore);
-#endif
+	// Scoped object that allows the physics query handle to be ignored for the current thread.
+	struct FScopedIgnoreQueryHandler
+	{
+		FScopedIgnoreQueryHandler(bool bIgnore);
+		~FScopedIgnoreQueryHandler();
+
+	private:
+		bool bOldState;
 	};
+#endif
 }
