Fix CIS issues 903187 & 903165 (StaticAnalysis + Android)

diff --git a/Engine/Source/Runtime/Engine/Private/ObjectTrace.cpp b/Engine/Source/Runtime/Engine/Private/ObjectTrace.cpp
--- a/Engine/Source/Runtime/Engine/Private/ObjectTrace.cpp
+++ b/Engine/Source/Runtime/Engine/Private/ObjectTrace.cpp
@@ -596,7 +596,7 @@ void FObjectTrace::OutputChildElement(const UObject* InObject, const uint64 InCh
 	}
 
 	const uint64 ObjectId = GetObjectId(InObject);
-	if (!AddObjectToTracedSet(UE::ObjectTrace::FObjectId(ObjectId, InChildElementId)))
+	if (!AddObjectToTracedSet(UE::ObjectTrace::FObjectId{ObjectId, InChildElementId}))
 	{
 		return;
 	}
@@ -610,6 +610,7 @@ void FObjectTrace::OutputChildElement(const UObject* InObject, const uint64 InCh
 		OutputChildElement(InOuterObject, InOuterChildElementId);
 	}
 
+	check(InObject);
 	FString ObjectName = InObject->GetName();
 	ObjectName += FString::Printf(TEXT(" : [0x%llX]"), InChildElementId);
 
@@ -634,6 +635,7 @@ void FObjectTrace::OutputChildElementLifetimeEnd(const UObject* InObject, const
 		return;
 	}
 
+	check(InObject);
 	TRACE_CHILD_ELEMENT(InObject, InChildElementId);
 
 	UE_TRACE_LOG(Object, ChildElementLifetimeEnd, ObjectChannel)
