Fixed DSTM crash with server console commands.

The general fix is from Mike. On top of that I removed UE::RemoteExecutor::ExecuteTransactional from the dstm.StopSimpleCoinGame console command as it is no longer needed with the general fix, and it was causing a crash when used with ServerExec

#rnx

diff --git a/Engine/Source/Runtime/Engine/Private/PlayerController.cpp b/Engine/Source/Runtime/Engine/Private/PlayerController.cpp
--- a/Engine/Source/Runtime/Engine/Private/PlayerController.cpp
+++ b/Engine/Source/Runtime/Engine/Private/PlayerController.cpp
@@ -76,7 +76,9 @@
 #endif // UE_WITH_IRIS
 
 #if UE_WITH_REMOTE_OBJECT_HANDLE
+#include "AutoRTFM/AutoRTFM.h"
 #include "UObject/UObjectMigrationContext.h"
+#include "UObject/RemoteExecutor.h"
 #endif
 
 #include UE_INLINE_GENERATED_CPP_BY_NAME(PlayerController)
@@ -514,7 +516,25 @@ FString APlayerController::ConsoleCommand(const FString& Cmd, bool bWriteToLog)
 {
 	if (Player != nullptr)
 	{
-		return Player->ConsoleCommand(Cmd, bWriteToLog);
+		FString Result;
+#if UE_WITH_REMOTE_OBJECT_HANDLE
+		if (!AutoRTFM::IsClosed())
+		{
+			UE_AUTORTFM_OPEN
+			{
+				static FName TransactionalWorkName(TEXT("APlyerController::ConsoleComand"));
+				UE::RemoteExecutor::ExecuteTransactional(TransactionalWorkName, [this, &Result, Cmd, bWriteToLog]()
+				{
+					Result = Player->ConsoleCommand(Cmd, bWriteToLog); 
+				});
+			};
+		}
+		else
+#endif
+		{
+			Result = Player->ConsoleCommand(Cmd, bWriteToLog);
+		}
+		return Result;		
 	}
 
 	return TEXT("");
