Skinned Mesh Component: Ensure pending skin weight profile requests are cleared out when setting a new mesh
- Otherwise switching LODs may end up getting the previous mesh's skin weight profile applied, resulting in visual artifacts or crashes.

#rnx

diff --git a/Engine/Source/Runtime/Engine/Private/Components/SkinnedMeshComponent.cpp b/Engine/Source/Runtime/Engine/Private/Components/SkinnedMeshComponent.cpp
--- a/Engine/Source/Runtime/Engine/Private/Components/SkinnedMeshComponent.cpp
+++ b/Engine/Source/Runtime/Engine/Private/Components/SkinnedMeshComponent.cpp
@@ -2479,6 +2479,10 @@ void USkinnedMeshComponent::SetSkinnedAssetAndUpdate(USkinnedAsset* InSkinnedAss
 		//Handle destroying and recreating the renderstate
 		FRenderStateRecreator RenderStateRecreator(this);
 
+		// Ensure we clear all skin weight profiles and requests, since any that might be in flight could 
+		// get applied to this new mesh on a LOD switch.
+		ClearAllSkinWeightProfiles();
+
 		SetSkinnedAsset(InSkinnedAsset);
 		SetPredictedLODLevel(0);
 
