Fixed crash in SafeRenameEditorOnlyData() when EditorOnlyCurrent is null.

diff --git a/Engine/Source/Runtime/Engine/Private/Materials/MaterialInterface.cpp b/Engine/Source/Runtime/Engine/Private/Materials/MaterialInterface.cpp
--- a/Engine/Source/Runtime/Engine/Private/Materials/MaterialInterface.cpp
+++ b/Engine/Source/Runtime/Engine/Private/Materials/MaterialInterface.cpp
@@ -121,7 +121,7 @@ static bool SafeRenameEditorOnlyData(UMaterialInterface* MaterialInterface, cons
 		EditorOnlyWithIncorrectName->MarkAsGarbage();
 	}
 
-	if (EditorOnlyCurrent->GetName() != EditorOnlyDataName)
+	if (EditorOnlyCurrent && EditorOnlyCurrent->GetName() != EditorOnlyDataName)
 	{
 		if (bLogMessageIfRenamed)
 		{
