Skeletal Mesh Component: make sure that the bone container is valid before using it

#prefight tbd
#rnx

diff --git a/Engine/Source/Runtime/Engine/Private/Components/SkeletalMeshComponent.cpp b/Engine/Source/Runtime/Engine/Private/Components/SkeletalMeshComponent.cpp
--- a/Engine/Source/Runtime/Engine/Private/Components/SkeletalMeshComponent.cpp
+++ b/Engine/Source/Runtime/Engine/Private/Components/SkeletalMeshComponent.cpp
@@ -1327,7 +1327,11 @@ bool USkeletalMeshComponent::InitializeAnimScriptInstance(bool bForceReinit, boo
 				
 				if(FAnimNode_LinkedInputPose* InputNode = PostProcessAnimInstance->GetLinkedInputPoseNode())
 				{
-					InputNode->CachedInputPose.SetBoneContainer(&PostProcessAnimInstance->GetRequiredBones());
+					const FBoneContainer& BoneContainer = PostProcessAnimInstance->GetRequiredBones();
+					if (BoneContainer.IsValid())
+					{
+						InputNode->CachedInputPose.SetBoneContainer(&BoneContainer);
+					}
 
 					// SetBoneContainer allocates space for bone data but leaves it uninitalized.
 					InputNode->bIsCachedInputPoseInitialized = false;
