Level Streaming - Added SetPriorityOverride()/ResetPriorityOverride() that override the user specified StreamingPriority property.
* Original value is backed up to be properly restored once the priority override is not necessary (when the level is no longer making visible or making invisible)
* Fixes an issue where a level that is not making visible/invsible could be processed before one that is making visible/invisible which could lead to delaying the completion of an already making visible/invisible (when LevelStreaming.MaximumMakingVisibleLevels > 0 or  LevelStreaming.MaximumMakingInvisibleLevels > 0).

diff --git a/Engine/Source/Runtime/Engine/Classes/Engine/LevelStreaming.h b/Engine/Source/Runtime/Engine/Classes/Engine/LevelStreaming.h
--- a/Engine/Source/Runtime/Engine/Classes/Engine/LevelStreaming.h
+++ b/Engine/Source/Runtime/Engine/Classes/Engine/LevelStreaming.h
@@ -205,6 +205,9 @@ protected:
 	UPROPERTY(EditAnywhere, Category = LevelStreaming, BlueprintSetter = SetPriority)
 	int32 StreamingPriority;
 
+private:
+	TOptional<int32> StreamingPriorityOverrideBackup;
+
 public:
 
 	/** If this isn't Name_None, then we load from this package on disk to the new package named PackageName					*/
@@ -453,6 +456,12 @@ public:
 	UFUNCTION(BlueprintSetter)
 	ENGINE_API void SetPriority(int32 NewPriority);
 
+	/** Provide a temporary priority override for the streaming level. Only one override can be specified at a time. Original priority must be restored with ResetPriorityOverride(). */
+	ENGINE_API void SetPriorityOverride(int32 PriorityOverride);
+
+	/** Restore the original streaming level priority after a call to SetPriorityOverride(). */
+	ENGINE_API void ResetPriorityOverride();
+
 	/** Returns whether the streaming level is in the loading state. */
 	bool HasLoadRequestPending() const { return CurrentState == ELevelStreamingState::Loading; }
 
diff --git a/Engine/Source/Runtime/Engine/Private/LevelStreaming.cpp b/Engine/Source/Runtime/Engine/Private/LevelStreaming.cpp
--- a/Engine/Source/Runtime/Engine/Private/LevelStreaming.cpp
+++ b/Engine/Source/Runtime/Engine/Private/LevelStreaming.cpp
@@ -2194,6 +2194,25 @@ void ULevelStreaming::SetPriority(const int32 NewPriority)
 	}
 }
 
+void ULevelStreaming::SetPriorityOverride(const int32 PriorityOverride)
+{
+	if (!StreamingPriorityOverrideBackup.IsSet())
+	{
+		StreamingPriorityOverrideBackup = StreamingPriority;
+	}
+
+	SetPriority(PriorityOverride);
+}
+
+void ULevelStreaming::ResetPriorityOverride()
+{
+	if (StreamingPriorityOverrideBackup.IsSet())
+	{
+		SetPriority(StreamingPriorityOverrideBackup.GetValue());
+		StreamingPriorityOverrideBackup.Reset();
+	}
+}
+
 void ULevelStreaming::SetLevelLODIndex(const int32 LODIndex)
 {
 	if (LODIndex != LevelLODIndex)
diff --git a/Engine/Source/Runtime/Engine/Private/World.cpp b/Engine/Source/Runtime/Engine/Private/World.cpp
--- a/Engine/Source/Runtime/Engine/Private/World.cpp
+++ b/Engine/Source/Runtime/Engine/Private/World.cpp
@@ -3430,10 +3430,17 @@ void UWorld::SetLevelPendingVisibilityState(ULevel* InLevel, ELevelPendingVisibi
 	switch (InState)
 	{
 	case ELevelPendingVisibilityState::None:
-		MakingVisibleLevels.Remove(InLevel);
-		MakingInvisibleLevels.Remove(InLevel);
-		check(!MakingVisibleLevels.Contains(InLevel));
-		check(!MakingInvisibleLevels.Contains(InLevel));
+	    {
+		    MakingVisibleLevels.Remove(InLevel);
+		    MakingInvisibleLevels.Remove(InLevel);
+		    check(!MakingVisibleLevels.Contains(InLevel));
+		    check(!MakingInvisibleLevels.Contains(InLevel));
+		    // Reset streaming level's priority so that it is no longer prioritized in the list
+		    if (ULevelStreaming* StreamingLevel = FLevelUtils::FindStreamingLevel(InLevel))
+		    {
+			    StreamingLevel->ResetPriorityOverride();
+		    }
+	    }
 		break;
 	case ELevelPendingVisibilityState::MakingVisible:
 		check(!MakingInvisibleLevels.Contains(InLevel));
@@ -4863,7 +4870,7 @@ void UWorld::UpdateLevelStreaming(const TOptional<const UE::FTimeout>& ExternalT
 			if (ULevelStreaming* StreamingLevel = ULevelStreaming::FindStreamingLevel(NextLevel))
 			{
 				check(StreamingLevelsToConsider.GetStreamingLevels().Contains(StreamingLevel));
-				StreamingLevel->SetPriority(Priority--);
+				StreamingLevel->SetPriorityOverride(Priority--);
 			}
 
 			if (bAlternatePriority)
diff --git a/Engine/Source/Runtime/Engine/Private/WorldPartition/WorldPartitionRuntimeLevelStreamingCell.cpp b/Engine/Source/Runtime/Engine/Private/WorldPartition/WorldPartitionRuntimeLevelStreamingCell.cpp
--- a/Engine/Source/Runtime/Engine/Private/WorldPartition/WorldPartitionRuntimeLevelStreamingCell.cpp
+++ b/Engine/Source/Runtime/Engine/Private/WorldPartition/WorldPartitionRuntimeLevelStreamingCell.cpp
@@ -698,7 +698,7 @@ void UWorldPartitionRuntimeLevelStreamingCell::SetStreamingPriority(int32 InStre
 {
 	if (LevelStreaming)
 	{
-		LevelStreaming->SetPriority(InStreamingPriority);
+		LevelStreaming->SetPriorityOverride(InStreamingPriority);
 	}
 }
 
