- Added object version to handle pre/post DataModel Guid generation change during pose-asset validation
- Added helper function to UPoseAsset used for retrieve correct Guid
- Added Object Annotation map to keep track of re-generated pose-assets (which should then use new GUID format regardless of asset/linker object version state)
- Update PoseAsset usages of GenerateGuid to use helper function instead

diff --git a/Engine/Plugins/Editor/EngineAssetDefinitions/Source/Private/Animation/AssetDefinition_PoseAsset.cpp b/Engine/Plugins/Editor/EngineAssetDefinitions/Source/Private/Animation/AssetDefinition_PoseAsset.cpp
--- a/Engine/Plugins/Editor/EngineAssetDefinitions/Source/Private/Animation/AssetDefinition_PoseAsset.cpp
+++ b/Engine/Plugins/Editor/EngineAssetDefinitions/Source/Private/Animation/AssetDefinition_PoseAsset.cpp
@@ -27,9 +27,7 @@ namespace MenuExtension_PoseAsset
 			{
 				if (IAnimationDataModel* SourceAnimationDataModel = SourceAnimation->GetDataModel())
 				{
-					IAnimationDataModel::FGuidGenerationSettings Settings;
-					Settings.bIncludeTimingData = false;
-					if (!PoseAsset->SourceAnimationRawDataGUID.IsValid() || PoseAsset->SourceAnimationRawDataGUID != SourceAnimationDataModel->GenerateGuid(Settings))
+					if (!PoseAsset->SourceAnimationRawDataGUID.IsValid() || PoseAsset->SourceAnimationRawDataGUID != PoseAsset->GetSourceAnimationGuid())
 					{
 						if (PoseAsset->GetSkeleton()->IsCompatibleForEditor(SourceAnimation->GetSkeleton()))
 						{
diff --git a/Engine/Source/Editor/DetailCustomizations/Private/PoseAssetDetails.cpp b/Engine/Source/Editor/DetailCustomizations/Private/PoseAssetDetails.cpp
--- a/Engine/Source/Editor/DetailCustomizations/Private/PoseAssetDetails.cpp
+++ b/Engine/Source/Editor/DetailCustomizations/Private/PoseAssetDetails.cpp
@@ -703,9 +703,7 @@ bool FPoseAssetDetails::IsUpdateSourceEnabled() const
 	const UAnimSequence* AnimSequenceSelected = Cast<UAnimSequence>(ObjectSet);
 	if (AnimSequenceSelected && PoseAsset.IsValid())
 	{
-		IAnimationDataModel::FGuidGenerationSettings Settings;
-		Settings.bIncludeTimingData = false;
-		return (PoseAsset->SourceAnimation && PoseAsset->SourceAnimation->GetDataModel()->GenerateGuid(Settings) != PoseAsset->SourceAnimationRawDataGUID);
+		return (PoseAsset->SourceAnimation && PoseAsset->GetSourceAnimationGuid() != PoseAsset->SourceAnimationRawDataGUID);		
 	}
 
 	return false;
diff --git a/Engine/Source/Editor/Persona/Private/AnimationEditorViewportClient.cpp b/Engine/Source/Editor/Persona/Private/AnimationEditorViewportClient.cpp
--- a/Engine/Source/Editor/Persona/Private/AnimationEditorViewportClient.cpp
+++ b/Engine/Source/Editor/Persona/Private/AnimationEditorViewportClient.cpp
@@ -1443,7 +1443,7 @@ FText FAnimationViewportClient::GetDisplayInfo(bool bDisplayAllInfo) const
 
 	if (const UPoseAsset* PoseAsset = Cast<UPoseAsset>(GetAnimPreviewScene()->GetPreviewAnimationAsset()))
 	{
-		if (PoseAsset->GetLinkerCustomVersion(FUE5MainStreamObjectVersion::GUID) >= FUE5MainStreamObjectVersion::PoseAssetRawDataGUIDUpdate && PoseAsset->SourceAnimation && PoseAsset->SourceAnimationRawDataGUID.IsValid() && PoseAsset->SourceAnimation->GetDataModel()->GenerateGuid() != PoseAsset->SourceAnimationRawDataGUID)
+		if (PoseAsset->GetLinkerCustomVersion(FUE5MainStreamObjectVersion::GUID) >= FUE5MainStreamObjectVersion::PoseAssetRawDataGUIDUpdate && PoseAsset->SourceAnimation && PoseAsset->SourceAnimationRawDataGUID.IsValid() && PoseAsset->GetSourceAnimationGuid() != PoseAsset->SourceAnimationRawDataGUID)
 		{
 			TextValue = ConcatenateLine(TextValue, LOCTEXT("PoseAssetOutOfDateWarning", "<AnimViewport.WarningText>Poses are out-of-sync with the source animation. To update them click \"Update Source\"</>"));
 		}
diff --git a/Engine/Source/Runtime/Core/Public/UObject/UE5ReleaseStreamObjectVersions.inl b/Engine/Source/Runtime/Core/Public/UObject/UE5ReleaseStreamObjectVersions.inl
--- a/Engine/Source/Runtime/Core/Public/UObject/UE5ReleaseStreamObjectVersions.inl
+++ b/Engine/Source/Runtime/Core/Public/UObject/UE5ReleaseStreamObjectVersions.inl
@@ -168,5 +168,8 @@ UE5_RELEASE_VERSION(AnimLayeredBoneBlendMasksFix, 52)
 // Separated lens flare from bloom intensity
 UE5_RELEASE_VERSION(BloomIndependentLensFlare, 53)
 
+// Add settings to IAnimationDataModel GUID generation
+UE5_RELEASE_VERSION(AnimModelGuidGenerationSettings, 54)
+
 // -----<new versions can be added above this line>-------------------------------------------------
 #undef UE5_RELEASE_VERSION
\ No newline at end of file
diff --git a/Engine/Source/Runtime/Engine/Classes/Animation/PoseAsset.h b/Engine/Source/Runtime/Engine/Classes/Animation/PoseAsset.h
--- a/Engine/Source/Runtime/Engine/Classes/Animation/PoseAsset.h
+++ b/Engine/Source/Runtime/Engine/Classes/Animation/PoseAsset.h
@@ -12,6 +12,7 @@
 #include "Animation/SmartName.h"
 #include "Animation/AnimationAsset.h"
 #include "Animation/AnimCurveTypes.h"
+#include "AnimData/IAnimationDataModel.h"
 #include "PoseAsset.generated.h"
 
 class UAnimSequence;
@@ -398,6 +399,7 @@ public:
 	ENGINE_API static FSmartName GetUniquePoseSmartName(USkeleton* Skeleton);
 	
 	ENGINE_API static FName GetUniquePoseName(UPoseAsset* PoseAsset);
+	ENGINE_API FGuid GetSourceAnimationGuid() const;
 
 protected:
 	virtual void RemapTracksToNewSkeleton(USkeleton* NewSkeleton, bool bConvertSpaces) override;
diff --git a/Engine/Source/Runtime/Engine/Private/Animation/PoseAsset.cpp b/Engine/Source/Runtime/Engine/Private/Animation/PoseAsset.cpp
--- a/Engine/Source/Runtime/Engine/Private/Animation/PoseAsset.cpp
+++ b/Engine/Source/Runtime/Engine/Private/Animation/PoseAsset.cpp
@@ -23,6 +23,7 @@
 #include "UObject/ObjectSaveContext.h"
 #include "UObject/Package.h"
 #include "UObject/UnrealType.h"
+#include "UObject/UObjectAnnotation.h"
 #include "Widgets/Notifications/SNotificationList.h"
 
 #include UE_INLINE_GENERATED_CPP_BY_NAME(PoseAsset)
@@ -35,6 +36,11 @@
 //
 //////////////////////////////////////////////////////////////////////////////////////////////////////////////////
 
+
+#if WITH_EDITOR
+namespace UE { namespace Anim { FUObjectAnnotationSparseBool GRegeneratedPoseAssets; } }
+#endif // WITH_EDITOR
+
 bool FPoseDataContainer::Serialize(FArchive& Ar)
 {
 #if WITH_EDITORONLY_DATA
@@ -944,9 +950,7 @@ void UPoseAsset::PostLoad()
 			}
 			SourceAnimation->ConditionalPostLoad();
 
-			IAnimationDataModel::FGuidGenerationSettings Settings;
-			Settings.bIncludeTimingData = false;
-			FGuid DataModelGuid = SourceAnimation->GetDataModel()->GenerateGuid(Settings);
+			const FGuid DataModelGuid = GetSourceAnimationGuid();
 			if (SourceAnimationRawDataGUID.IsValid() && SourceAnimationRawDataGUID != DataModelGuid)
 			{
 				FFormatNamedArguments Args;
@@ -1553,9 +1557,10 @@ void UPoseAsset::CreatePoseFromAnimation(class UAnimSequence* AnimSequence, cons
 				PostProcessData();
 			}
 
-			IAnimationDataModel::FGuidGenerationSettings Settings;
-			Settings.bIncludeTimingData = false;
-			SourceAnimationRawDataGUID = AnimSequence->GetDataModel()->GenerateGuid(Settings);
+			// Mark PoseAsset as (re-)generated
+			UE::Anim::GRegeneratedPoseAssets.Set(this);
+			
+			SourceAnimationRawDataGUID = GetSourceAnimationGuid();
 		}
 	}
 }
@@ -1935,6 +1940,28 @@ FName UPoseAsset::GetUniquePoseName(UPoseAsset* PoseAsset)
 	return NewName;
 }
 
+FGuid UPoseAsset::GetSourceAnimationGuid() const
+{	
+	if (SourceAnimation)
+	{
+		IAnimationDataModel::FGuidGenerationSettings Settings;
+		if (GetLinkerCustomVersion(FUE5ReleaseStreamObjectVersion::GUID) < FUE5ReleaseStreamObjectVersion::AnimModelGuidGenerationSettings)
+		{
+			Settings.bIncludeTimingData = false;
+		}		
+
+		// If pose asset was re-generated during editor runtime, use the new GUID format
+		if(UE::Anim::GRegeneratedPoseAssets.Get(this))
+		{
+			Settings.bIncludeTimingData = true;	
+		}
+		
+		return SourceAnimation->GetDataModel()->GenerateGuid(Settings);
+	}
+
+	return FGuid();
+}
+
 void UPoseAsset::RemapTracksToNewSkeleton(USkeleton* NewSkeleton, bool bConvertSpaces)
 {
 	Super::RemapTracksToNewSkeleton(NewSkeleton, bConvertSpaces);
