Rename "Binaural" attenuation spatialization method to "Plugin-Spatialized"

diff --git a/Engine/Source/Runtime/Engine/Classes/Sound/SoundAttenuation.h b/Engine/Source/Runtime/Engine/Classes/Sound/SoundAttenuation.h
--- a/Engine/Source/Runtime/Engine/Classes/Sound/SoundAttenuation.h
+++ b/Engine/Source/Runtime/Engine/Classes/Sound/SoundAttenuation.h
@@ -33,8 +33,8 @@ enum ESoundSpatializationAlgorithm : int
 	// Standard panning method for spatialization (linear or equal power method defined in project settings)
 	SPATIALIZATION_Default UMETA(DisplayName = "Panning"),
 
-	// Binaural spatialization method if available (requires headphones, enabled by plugins)
-	SPATIALIZATION_HRTF UMETA(DisplayName = "Binaural"),
+	// Spatialization method provided by spatialization plugin currently selected in the project settings
+	SPATIALIZATION_HRTF UMETA(DisplayName = "Plugin-Spatialized"),
 };
 
 UENUM(BlueprintType)
