Refactor Bouncer Gallery to be more Networked Physics compliant.

#changelist validated
#virtualized

diff --git a/Engine/Source/Runtime/Engine/Public/Physics/GenericPhysicsInterface.h b/Engine/Source/Runtime/Engine/Public/Physics/GenericPhysicsInterface.h
--- a/Engine/Source/Runtime/Engine/Public/Physics/GenericPhysicsInterface.h
+++ b/Engine/Source/Runtime/Engine/Public/Physics/GenericPhysicsInterface.h
@@ -127,7 +127,6 @@ namespace Chaos::Private
 		*  Only the single closest blocking result will be generated, no tests will be done after that.
 		*  Falls back to a raycast is the shape size is less than or equal to UE_KINDA_SMALL_NUMBER.
 		*/
-		UE_EXPERIMENTAL(5.6, "Physics Thread sweep trace, might have API changes")
 		static ENGINE_API bool GeomSweepMulti(const UWorld* World, const FCollisionShape& InGeom, const FQuat& InGeomRot, TArray<FHitResult>& OutHits, const FVector& Start, const FVector& End, ECollisionChannel TraceChannel, const FCollisionQueryParams& Params, const FCollisionResponseParams& ResponseParams, const FCollisionObjectQueryParams& ObjectParams = FCollisionObjectQueryParams::DefaultObjectQueryParam);
 
 		/**
@@ -137,7 +136,6 @@ namespace Chaos::Private
 		*  Results are sorted, so a blocking hit (if found) will be the last element of the array
 		*  Only the single closest blocking result will be generated, no tests will be done after that.
 		*/
-		UE_EXPERIMENTAL(5.6, "Physics Thread raycast, might have API changes")
 		static ENGINE_API bool RaycastMulti(const UWorld* World, TArray<FHitResult>& OutHits, const FVector& Start, const FVector& End, ECollisionChannel TraceChannel, const FCollisionQueryParams& Params, const FCollisionResponseParams& ResponseParams, const FCollisionObjectQueryParams& ObjectParams = FCollisionObjectQueryParams::DefaultObjectQueryParam);
 	};
 
