Fixed 32 bits float scene color setup when Alpha is needed.

+ Enable PF_A32R32G32B32F format to work in Read Raw Pixel Value and Read Pixel Value to return the float beyond the 16f precisoin if opted into r.SceneColorFormat=5.

#tests pc editor dx12
#rnx

diff --git a/Engine/Source/Runtime/Engine/Private/KismetRenderingLibrary.cpp b/Engine/Source/Runtime/Engine/Private/KismetRenderingLibrary.cpp
--- a/Engine/Source/Runtime/Engine/Private/KismetRenderingLibrary.cpp
+++ b/Engine/Source/Runtime/Engine/Private/KismetRenderingLibrary.cpp
@@ -345,6 +345,7 @@ EPixelFormat ReadRenderTargetHelper(
 		}
 		break;
 	case PF_FloatRGBA:
+	case PF_A32B32G32R32F:
 		if (!RenderTarget->ReadLinearColorPixels(OutHDRValues, ReadSurfaceDataFlags, SampleRect))
 		{
 			OutFormat = PF_Unknown;
@@ -388,6 +389,7 @@ FColor UKismetRenderingLibrary::ReadRenderTargetPixel(UObject* WorldContextObjec
 		check(Samples.Num() == 1 && LinearSamples.Num() == 0);
 		return Samples[0];
 	case PF_FloatRGBA:
+	case PF_A32B32G32R32F:
 		check(Samples.Num() == 0 && LinearSamples.Num() == 1);
 		return LinearSamples[0].ToFColor(true);
 	case PF_Unknown:
@@ -413,6 +415,7 @@ bool UKismetRenderingLibrary::ReadRenderTarget(UObject* WorldContextObject, UTex
 			check(OutSamples.Num() == NumSamples && LinearSamples.Num() == 0);
 			return true;
 		case PF_FloatRGBA:
+		case PF_A32B32G32R32F:
 			check(OutSamples.Num() == 0 && LinearSamples.Num() == NumSamples);
 			for (int32 SampleIndex = 0; SampleIndex < NumSamples; ++SampleIndex)
 			{
@@ -439,6 +442,7 @@ FLinearColor UKismetRenderingLibrary::ReadRenderTargetRawPixel(UObject * WorldCo
 		check(Samples.Num() == 1 && LinearSamples.Num() == 0);
 		return FLinearColor(float(Samples[0].R), float(Samples[0].G), float(Samples[0].B), float(Samples[0].A));
 	case PF_FloatRGBA:
+	case PF_A32B32G32R32F:
 		check(Samples.Num() == 0 && LinearSamples.Num() == 1);
 		return LinearSamples[0];
 	case PF_Unknown:
@@ -464,6 +468,7 @@ ENGINE_API TArray<FLinearColor> UKismetRenderingLibrary::ReadRenderTargetRawPixe
 		return LinearSamples;
 
 	case PF_FloatRGBA:
+	case PF_A32B32G32R32F:
 		check(Samples.Num() == 0 && LinearSamples.Num() > 0);
 		return LinearSamples;
 
@@ -509,6 +514,7 @@ bool UKismetRenderingLibrary::ReadRenderTargetRaw(UObject* WorldContextObject, U
 			}
 			return true;
 		case PF_FloatRGBA:
+		case PF_A32B32G32R32F:
 			check(Samples.Num() == 0 && OutLinearSamples.Num() == NumSamples);
 			return true;
 		case PF_Unknown:
