[HWRT] Fix stale RegisteredGeometry.bEvicted when Landscape LOD changes.
- When landscape LOD changes FRayTracingDynamicGeometryUpdateManager::AddDynamicGeometryToUpdate(...) recreates the geometry which means it's no longer evicted, but the cached flag in FRayTracingGeometryManager was not being updated.

diff --git a/Engine/Source/Runtime/Engine/Private/Rendering/RayTracingGeometryManager.cpp b/Engine/Source/Runtime/Engine/Private/Rendering/RayTracingGeometryManager.cpp
--- a/Engine/Source/Runtime/Engine/Private/Rendering/RayTracingGeometryManager.cpp
+++ b/Engine/Source/Runtime/Engine/Private/Rendering/RayTracingGeometryManager.cpp
@@ -463,13 +463,15 @@ void FRayTracingGeometryManager::RefreshRegisteredGeometry(FGeometryHandle Handl
 			RegisteredGeometry.Size = bAllSegmentsAreValid ? RHICalcRayTracingGeometrySize(RegisteredGeometry.Geometry->Initializer).ResultSize : 0;
 		}
 
+		RegisteredGeometry.bEvicted = RegisteredGeometry.Geometry->IsEvicted();
+
 		if (AlwaysResidentGeometries.Contains(Handle))
 		{
 			TotalAlwaysResidentSize -= OldSize;
 			TotalAlwaysResidentSize += RegisteredGeometry.Size;
 		}
 
-		if (RegisteredGeometry.Geometry->IsValid() && !RegisteredGeometry.Geometry->IsEvicted())
+		if (RegisteredGeometry.Geometry->IsValid() && !RegisteredGeometry.bEvicted)
 		{
 			bool bAlreadyInSet;
 			ResidentGeometries.Add(Handle, &bAlreadyInSet);
@@ -508,7 +510,7 @@ void FRayTracingGeometryManager::RefreshRegisteredGeometry(FGeometryHandle Handl
 			EvictableGeometries.Remove(Handle);
 		}
 
-		checkf(!AlwaysResidentGeometries.Contains(Handle) || !RegisteredGeometry.Geometry->IsEvicted() || !IsRayTracingEnabled(), TEXT("Always resident geometries can't be evicted"));
+		checkf(!AlwaysResidentGeometries.Contains(Handle) || !RegisteredGeometry.bEvicted || !IsRayTracingEnabled(), TEXT("Always resident geometries can't be evicted"));
 
 		if (RegisteredGeometry.Geometry->Initializer.Type == ERayTracingGeometryInitializerType::StreamingDestination)
 		{
