World Partition - Fix broken support for PIE "Launch Separate Server" option
* PIE vs the separate server process would use different package paths for the world partition cells
* Regression introduced in 39414208

diff --git a/Engine/Source/Runtime/Engine/Private/WorldPartition/WorldPartitionLevelStreamingPolicy.cpp b/Engine/Source/Runtime/Engine/Private/WorldPartition/WorldPartitionLevelStreamingPolicy.cpp
--- a/Engine/Source/Runtime/Engine/Private/WorldPartition/WorldPartitionLevelStreamingPolicy.cpp
+++ b/Engine/Source/Runtime/Engine/Private/WorldPartition/WorldPartitionLevelStreamingPolicy.cpp
@@ -65,7 +65,7 @@ FString UWorldPartitionLevelStreamingPolicy::GetCellPackagePath(const FName& InC
 {
 	if (InWorld->IsGameWorld())
 	{
-		if (InWorld->WorldType == EWorldType::PIE)
+		if (InWorld->IsPlayInEditor() || InWorld->IsPlayInPreview())
 		{
 			// Set as memory package to avoid wasting time in FPackageName::DoesPackageExist
 			return FString::Printf(TEXT("/Memory/%s_%s"), *FPackageName::GetShortName(InWorld->GetPackage()), *InCellName.ToString());
