Remove SCS nodes based on editor only components

[RN] Fixed crash when editing children from a cooked blueprint that had editor only components

diff --git a/Engine/Source/Runtime/Engine/Classes/Engine/SimpleConstructionScript.h b/Engine/Source/Runtime/Engine/Classes/Engine/SimpleConstructionScript.h
--- a/Engine/Source/Runtime/Engine/Classes/Engine/SimpleConstructionScript.h
+++ b/Engine/Source/Runtime/Engine/Classes/Engine/SimpleConstructionScript.h
@@ -115,6 +115,9 @@ class USimpleConstructionScript : public UObject
 #if WITH_EDITOR
 	/** Returns Valid if this object has data validation rules set up for it and the data for this object is valid. Returns Invalid if it does not pass the rules. Returns NotValidated if no rules are set for this object. */
 	virtual EDataValidationResult IsDataValid(class FDataValidationContext& Context) const override;
+
+	/** Internal function, intentionally not exported. Removes components that are editor only */
+	void RemoveEditorOnlyComponents();
 #endif // WITH_EDITOR
 
 	FStaticMeshComponentBulkReregisterContext* GetReregisterContext() const { return ReregisterContext; }
diff --git a/Engine/Source/Runtime/Engine/Private/Blueprint.cpp b/Engine/Source/Runtime/Engine/Private/Blueprint.cpp
--- a/Engine/Source/Runtime/Engine/Private/Blueprint.cpp
+++ b/Engine/Source/Runtime/Engine/Private/Blueprint.cpp
@@ -384,6 +384,12 @@ void UBlueprint::PreSave(FObjectPreSaveContext ObjectSaveContext)
 	UpgradeNotesLog.Reset();
 	const ITargetPlatform* TargetPlatform = ObjectSaveContext.GetTargetPlatform();
 
+	if(ObjectSaveContext.IsCooking() && SimpleConstructionScript && TargetPlatform && !TargetPlatform->HasEditorOnlyData())
+	{
+		// we need to remove SCS nodes that are using editor only components:
+		SimpleConstructionScript->RemoveEditorOnlyComponents();
+	}
+
 	if (!TargetPlatform || TargetPlatform->HasEditorOnlyData())
 	{
 		// This will force an immediate (synchronous) update of this Blueprint's index tag value.
diff --git a/Engine/Source/Runtime/Engine/Private/SimpleConstructionScript.cpp b/Engine/Source/Runtime/Engine/Private/SimpleConstructionScript.cpp
--- a/Engine/Source/Runtime/Engine/Private/SimpleConstructionScript.cpp
+++ b/Engine/Source/Runtime/Engine/Private/SimpleConstructionScript.cpp
@@ -1250,6 +1250,18 @@ EDataValidationResult USimpleConstructionScript::IsDataValid(FDataValidationCont
 	return Result;
 }
 
+void USimpleConstructionScript::RemoveEditorOnlyComponents()
+{
+	TArray<USCS_Node*> Nodes = GetAllNodes();
+	for(USCS_Node* Node : Nodes)
+	{
+		if(Node->ComponentTemplate && Node->ComponentTemplate->IsEditorOnly())
+		{
+			RemoveNodeAndPromoteChildren(Node);
+		}
+	}
+}
+
 void USimpleConstructionScript::GenerateListOfExistingNames(TSet<FName>& CurrentNames) const
 {
 	const UBlueprintGeneratedClass* OwnerClass = Cast<const UBlueprintGeneratedClass>(GetOuter());
