[HWRT] Improvements to RayTracingGeometryManager stats.

Context:
- "Requested Memory" stat didn't take the size of "always resident geometry" into account, so we didn't have a clear metric for how much memory the manager was trying to stream in, making it hard for developer to determine the optimal ray tracing geometry pool size for their project.

Changes:
- Fixed "Requested Memory" to include always resident geometry sizes - this is the total amount of memory that the manager is trying to stream in.
- Added a new stat "Referenced Memory" that only includes ray tracing geometry currently referenced by the ray tracing scene - this doesn't include always resident geometry not referenced by the ray tracing scene (this metric is the same as "Requested Memory" before this change).
- Also added logic to cache "bAlwaysResident" flag in FRegisteredGeometry to avoid cache misses in the main loop of the manager.

diff --git a/Engine/Source/Runtime/Engine/Private/Rendering/RayTracingGeometryManager.cpp b/Engine/Source/Runtime/Engine/Private/Rendering/RayTracingGeometryManager.cpp
--- a/Engine/Source/Runtime/Engine/Private/Rendering/RayTracingGeometryManager.cpp
+++ b/Engine/Source/Runtime/Engine/Private/Rendering/RayTracingGeometryManager.cpp
@@ -144,6 +144,7 @@ DECLARE_DWORD_ACCUMULATOR_STAT(TEXT("Geometry Group Count"), STAT_RayTracingGeom
 
 DECLARE_MEMORY_STAT(TEXT("Resident Memory"), STAT_RayTracingGeometryResidentMemory, STATGROUP_RayTracingGeometry);
 DECLARE_MEMORY_STAT(TEXT("Always Resident Memory"), STAT_RayTracingGeometryAlwaysResidentMemory, STATGROUP_RayTracingGeometry);
+DECLARE_MEMORY_STAT(TEXT("Referenced Memory"), STAT_RayTracingGeometryReferencedMemory, STATGROUP_RayTracingGeometry);
 DECLARE_MEMORY_STAT(TEXT("Requested Memory"), STAT_RayTracingGeometryRequestedMemory, STATGROUP_RayTracingGeometry);
 
 DECLARE_DWORD_ACCUMULATOR_STAT(TEXT("Pending Builds"), STAT_RayTracingPendingBuilds, STATGROUP_RayTracingGeometry);
@@ -225,6 +226,11 @@ void FRayTracingGeometryManager::RemoveBuildRequest(FBuildRequestIndex InRequest
 	GeometryBuildRequests.RemoveAt(InRequestIndex);
 }
 
+bool FRayTracingGeometryManager::IsAlwaysResidentGeometry(const FRayTracingGeometry* InGeometry, const FRayTracingGeometryGroup& Group)
+{
+	return InGeometry->LODIndex >= Group.GeometryHandles.Num() - GRayTracingNumAlwaysResidentLODs;
+}
+
 RayTracing::FGeometryGroupHandle FRayTracingGeometryManager::RegisterRayTracingGeometryGroup(uint32 NumLODs, uint32 CurrentFirstLODIdx)
 {
 	FScopeLock ScopeLock(&MainCS);
@@ -297,14 +303,14 @@ FRayTracingGeometryManager::FGeometryHandle FRayTracingGeometryManager::Register
 		Group.GeometryHandles[InGeometry->LODIndex] = Handle;
 		++Group.NumReferences;
 
-		const bool bAlwaysResident = InGeometry->LODIndex >= Group.GeometryHandles.Num() - GRayTracingNumAlwaysResidentLODs;
+		RegisteredGeometry.bAlwaysResident = IsAlwaysResidentGeometry(InGeometry, Group);
 
-		if (bAlwaysResident)
+		if (RegisteredGeometry.bAlwaysResident)
 		{
 			AlwaysResidentGeometries.Add(Handle);
 		}
 
-		if (IsRayTracingEnabled() && InGeometry->LODIndex >= Group.CurrentFirstLODIdx && (!IsRayTracingUsingReferenceBasedResidency() || bAlwaysResident))
+		if (IsRayTracingEnabled() && InGeometry->LODIndex >= Group.CurrentFirstLODIdx && (!IsRayTracingUsingReferenceBasedResidency() || RegisteredGeometry.bAlwaysResident))
 		{
 			PendingStreamingRequests.Add(Handle);
 			INC_DWORD_STAT(STAT_RayTracingPendingStreamingRequests);
@@ -359,6 +365,8 @@ void FRayTracingGeometryManager::ReleaseRayTracingGeometryHandle(FGeometryHandle
 
 		if (NumRemoved > 0)
 		{
+			checkf(RegisteredGeometry.bAlwaysResident, TEXT("Geometry should have the bAlwaysResident flag enabled since it was in the AlwaysResidentGeometries set."));
+
 			TotalAlwaysResidentSize -= RegisteredGeometry.Size;
 		}
 	}
@@ -465,8 +473,10 @@ void FRayTracingGeometryManager::RefreshRegisteredGeometry(FGeometryHandle Handl
 
 		RegisteredGeometry.bEvicted = RegisteredGeometry.Geometry->IsEvicted();
 
-		if (AlwaysResidentGeometries.Contains(Handle))
+		if (RegisteredGeometry.bAlwaysResident)
 		{
+			checkf(AlwaysResidentGeometries.Contains(Handle), TEXT("Geometry with bAlwaysResident flag set should be in the AlwaysResidentGeometries set."));
+
 			TotalAlwaysResidentSize -= OldSize;
 			TotalAlwaysResidentSize += RegisteredGeometry.Size;
 		}
@@ -483,18 +493,8 @@ void FRayTracingGeometryManager::RefreshRegisteredGeometry(FGeometryHandle Handl
 
 			TotalResidentSize += RegisteredGeometry.Size;
 
-			if (RegisteredGeometry.Geometry->GroupHandle != INDEX_NONE)
-			{
-				const FRayTracingGeometryGroup& Group = RegisteredGroups[RegisteredGeometry.Geometry->GroupHandle];
-
-				if (RegisteredGeometry.Geometry->LODIndex < Group.GeometryHandles.Num() - GRayTracingNumAlwaysResidentLODs) // don't want to evict lowest LODs
-				{
-					EvictableGeometries.Add(Handle);
-				}
-			}
-			else
+			if (!RegisteredGeometry.bAlwaysResident)
 			{
-				// geometries not assigned to a group (eg: dynamic geometry) are always evictable
 				EvictableGeometries.Add(Handle);
 			}
 		}
@@ -552,21 +552,10 @@ void FRayTracingGeometryManager::Tick(FRHICommandList& RHICmdList)
 #endif
 
 	const bool bUsingReferenceBasedResidency = IsRayTracingUsingReferenceBasedResidency();
-	const int32 NumAlwaysResidentLODs = GRayTracingNumAlwaysResidentLODs;
 
 	checkf(bUsingReferenceBasedResidency || (ReferencedGeometryHandles.IsEmpty() && ReferencedGeometryGroups.IsEmpty() && ReferencedGeometryGroupsForDynamicUpdate.IsEmpty()),
 		TEXT("ReferencedGeometryHandles, ReferencedGeometryGroups and ReferencedGeometryGroupsForDynamicUpdate are expected to be empty when not using reference based residency"));
 
-#if DO_CHECK
-	if (GRayTracingTestCheckIntegrity)
-	{
-		for (FRegisteredGeometry& RegisteredGeometry : RegisteredGeometries)
-		{
-			check(RegisteredGeometry.bEvicted == RegisteredGeometry.Geometry->IsEvicted());
-		}
-	}
-#endif
-
 	if (bRefreshAlwaysResidentRayTracingGeometries)
 	{
 		bRefreshAlwaysResidentRayTracingGeometries = false;
@@ -576,13 +565,19 @@ void FRayTracingGeometryManager::Tick(FRHICommandList& RHICmdList)
 
 		for (FRegisteredGeometry& RegisteredGeometry : RegisteredGeometries)
 		{
-			if (RegisteredGeometry.Geometry->GroupHandle == INDEX_NONE)
+			const RayTracing::FGeometryGroupHandle GroupHandle = RegisteredGeometry.Geometry->GroupHandle;
+
+			if (GroupHandle == INDEX_NONE)
 			{
+				RegisteredGeometry.bAlwaysResident = false;
 				continue;
 			}
 
-			const FRayTracingGeometryGroup& Group = RegisteredGroups[RegisteredGeometry.Geometry->GroupHandle];
-			if (RegisteredGeometry.Geometry->LODIndex >= Group.GeometryHandles.Num() - NumAlwaysResidentLODs)
+			const FRayTracingGeometryGroup& Group = RegisteredGroups[GroupHandle];
+
+			RegisteredGeometry.bAlwaysResident = IsAlwaysResidentGeometry(RegisteredGeometry.Geometry, Group);
+
+			if (RegisteredGeometry.bAlwaysResident)
 			{
 				AlwaysResidentGeometries.Add(RegisteredGeometry.Geometry->RayTracingGeometryHandle);
 				TotalAlwaysResidentSize += RegisteredGeometry.Size;
@@ -608,6 +603,33 @@ void FRayTracingGeometryManager::Tick(FRHICommandList& RHICmdList)
 		}
 	}
 
+#if DO_CHECK
+	if (GRayTracingTestCheckIntegrity)
+	{
+		for (FRegisteredGeometry& RegisteredGeometry : RegisteredGeometries)
+		{
+			const FGeometryHandle GeometryHandle = RegisteredGeometry.Geometry->RayTracingGeometryHandle;
+			const RayTracing::FGeometryGroupHandle GroupHandle = RegisteredGeometry.Geometry->GroupHandle;
+
+			bool bAlwaysResident = false;
+
+			if (GroupHandle != INDEX_NONE)
+			{
+				const FRayTracingGeometryGroup& Group = RegisteredGroups[GroupHandle];
+				bAlwaysResident = IsAlwaysResidentGeometry(RegisteredGeometry.Geometry, Group);
+			}
+			else
+			{
+				// geometries not assigned to a group (eg: dynamic geometry) are always evictable
+			}
+
+			checkf(RegisteredGeometry.bEvicted == RegisteredGeometry.Geometry->IsEvicted(), TEXT("Cached bEvicted flag in FRegisteredGeometry is stale"));
+			checkf(RegisteredGeometry.bAlwaysResident == bAlwaysResident, TEXT("Cached bAlwaysResident flag in FRegisteredGeometry is stale"));
+			checkf(RegisteredGeometry.bAlwaysResident == AlwaysResidentGeometries.Contains(GeometryHandle), TEXT("Geometry with bAlwaysResident flag set should be in the AlwaysResidentGeometries set."));
+		}
+	}
+#endif
+
 	if (!IsRayTracingEnabled())
 	{
 		if (bHasRayTracingEnableChanged)
@@ -644,7 +666,10 @@ void FRayTracingGeometryManager::Tick(FRHICommandList& RHICmdList)
 
 		check(PendingStreamingRequests.IsEmpty());
 
+		SET_MEMORY_STAT(STAT_RayTracingGeometryReferencedMemory, 0);
 		SET_MEMORY_STAT(STAT_RayTracingGeometryRequestedMemory, 0);
+		CSV_CUSTOM_STAT(RayTracingGeometry, ReferencedSizeMB, 0, ECsvCustomStatOp::Set);
+		CSV_CUSTOM_STAT(RayTracingGeometry, RequestedSizeMB, 0, ECsvCustomStatOp::Set);
 	}
 	else if (bUsingReferenceBasedResidency)
 	{
@@ -685,7 +710,8 @@ void FRayTracingGeometryManager::Tick(FRHICommandList& RHICmdList)
 
 		TArray<FGeometryHandle> ReferencedGeometries;
 
-		uint64 RequestedSize = 0;
+		uint64 ReferencedSize = 0;
+		uint64 RequestedSize = TotalAlwaysResidentSize;
 		uint64 RequestedButEvictedSize = 0;
 
 		// Step 1
@@ -698,9 +724,14 @@ void FRayTracingGeometryManager::Tick(FRHICommandList& RHICmdList)
 			ReferencedGeometries.Add(GeometryHandle);
 			NotReferencedResidentGeometries.Remove(GeometryHandle);
 
-			RequestedSize += RegisteredGeometry.Size;
+			ReferencedSize += RegisteredGeometry.Size;
+
+			if (!RegisteredGeometry.bAlwaysResident)
+			{
+				RequestedSize += RegisteredGeometry.Size;
+			}
 
-			if (RegisteredGeometry.Geometry->IsEvicted())
+			if (RegisteredGeometry.bEvicted)
 			{
 				RequestedButEvictedSize += RegisteredGeometry.Size;
 			}
@@ -727,9 +758,9 @@ void FRayTracingGeometryManager::Tick(FRHICommandList& RHICmdList)
 					FRegisteredGeometry& RegisteredGeometry = RegisteredGeometries[GeometryHandle];
 					RegisteredGeometry.LastReferencedFrame = GFrameCounterRenderThread;
 
-					RequestedSize += RegisteredGeometry.Size;
+					ReferencedSize += RegisteredGeometry.Size;
 
-					if (LODIndex >= NumGeometriesInGroup - NumAlwaysResidentLODs)
+					if (RegisteredGeometry.bAlwaysResident)
 					{
 						checkf(!RegisteredGeometry.bEvicted, TEXT("Always resident ray tracing geometry was unexpectely evicted."));
 					}
@@ -738,6 +769,8 @@ void FRayTracingGeometryManager::Tick(FRHICommandList& RHICmdList)
 						ReferencedGeometries.Add(GeometryHandle);
 						NotReferencedResidentGeometries.Remove(GeometryHandle);
 
+						RequestedSize += RegisteredGeometry.Size;
+
 						if (RegisteredGeometry.bEvicted)
 						{
 							RequestedButEvictedSize += RegisteredGeometry.Size;
@@ -832,7 +865,9 @@ void FRayTracingGeometryManager::Tick(FRHICommandList& RHICmdList)
 			}
 		}
 
+		SET_MEMORY_STAT(STAT_RayTracingGeometryReferencedMemory, ReferencedSize);
 		SET_MEMORY_STAT(STAT_RayTracingGeometryRequestedMemory, RequestedSize);
+		CSV_CUSTOM_STAT(RayTracingGeometry, ReferencedSizeMB, ReferencedSize / 1024.0f / 1024.0f, ECsvCustomStatOp::Set);
 		CSV_CUSTOM_STAT(RayTracingGeometry, RequestedSizeMB, RequestedSize / 1024.0f / 1024.0f, ECsvCustomStatOp::Set);
 
 #if !UE_BUILD_SHIPPING
@@ -888,7 +923,9 @@ void FRayTracingGeometryManager::Tick(FRHICommandList& RHICmdList)
 #endif
 		}
 
+		SET_MEMORY_STAT(STAT_RayTracingGeometryReferencedMemory, 0);
 		SET_MEMORY_STAT(STAT_RayTracingGeometryRequestedMemory, TotalResidentSize);
+		CSV_CUSTOM_STAT(RayTracingGeometry, ReferencedSizeMB, 0, ECsvCustomStatOp::Set);
 		CSV_CUSTOM_STAT(RayTracingGeometry, RequestedSizeMB, TotalResidentSize / 1024.0f / 1024.0f, ECsvCustomStatOp::Set);
 	}
 
diff --git a/Engine/Source/Runtime/Engine/Public/Rendering/RayTracingGeometryManager.h b/Engine/Source/Runtime/Engine/Public/Rendering/RayTracingGeometryManager.h
--- a/Engine/Source/Runtime/Engine/Public/Rendering/RayTracingGeometryManager.h
+++ b/Engine/Source/Runtime/Engine/Public/Rendering/RayTracingGeometryManager.h
@@ -130,7 +130,8 @@ private:
 		};
 
 		FStatus Status = FStatus::StreamedOut;
-		bool bEvicted = false;
+		bool bEvicted : 1 = false;
+		bool bAlwaysResident : 1 = false;
 	};
 
 	struct FStreamingRequest
@@ -181,6 +182,8 @@ private:
 	// Must call this instead of FRayTracingGeometry::Evict() directly
 	void EvictGeometry(FRHICommandListBase& RHICmdList, FRegisteredGeometry& RegisteredGeometry);
 
+	static bool IsAlwaysResidentGeometry(const FRayTracingGeometry* InGeometry, const FRayTracingGeometryGroup& Group);
+
 	TSparseArray<FRayTracingGeometryGroup> RegisteredGroups;
 
 	// Used for keeping track of geometries when ray tracing is dynamic
