Mark SeekRestart as Experimental

diff --git a/Engine/Source/Runtime/Engine/Classes/Sound/SoundBase.h b/Engine/Source/Runtime/Engine/Classes/Sound/SoundBase.h
--- a/Engine/Source/Runtime/Engine/Classes/Sound/SoundBase.h
+++ b/Engine/Source/Runtime/Engine/Classes/Sound/SoundBase.h
@@ -70,10 +70,10 @@ enum class EVirtualizationMode : uint8
 	/** If sound is looping, sound restarts from beginning upon realization from being virtual */
 	Restart,
 
-	/** Sound fully virtualizes, but attempts to seek to the proper part of its looping duration when realized.
+	/** EXPERIMENTAL - Sound fully virtualizes, but attempts to seek to the proper part of its looping duration when realized.
 	 * The playback behavior is similar to PlayWhenSilent, but does not take up a voice slot.
 	 */
-	SeekRestart
+	SeekRestart UE_EXPERIMENTAL(5.6, "This virtualization mode is experimental and may not function as expected in all use cases.")
 };
 
 /**
