Implement UMaterialExpressionParticleRadius in the new translator.

- This simply uses the build function of the base class.
- This expression already had the external code identifier setup in BaseMaterialExpressions.ini so I only needed to inherit the base class.

#rnx

diff --git a/Engine/Source/Runtime/Engine/Public/Materials/MaterialExpressionParticleRadius.h b/Engine/Source/Runtime/Engine/Public/Materials/MaterialExpressionParticleRadius.h
--- a/Engine/Source/Runtime/Engine/Public/Materials/MaterialExpressionParticleRadius.h
+++ b/Engine/Source/Runtime/Engine/Public/Materials/MaterialExpressionParticleRadius.h
@@ -6,10 +6,11 @@
 #include "CoreMinimal.h"
 #include "UObject/ObjectMacros.h"
 #include "Materials/MaterialExpression.h"
+#include "Materials/MaterialExpressionExternalCodeBase.h"
 #include "MaterialExpressionParticleRadius.generated.h"
 
 UCLASS(collapsecategories, hidecategories=Object)
-class UMaterialExpressionParticleRadius : public UMaterialExpression
+class UMaterialExpressionParticleRadius : public UMaterialExpressionExternalCodeBase
 {
 	GENERATED_BODY()
 
