Wrap GetCachedRayTracingInstance with #if RHI_RAYTRACING
#rnx

diff --git a/Engine/Source/Runtime/Engine/Private/Components/InstancedSkinnedMeshComponent.cpp b/Engine/Source/Runtime/Engine/Private/Components/InstancedSkinnedMeshComponent.cpp
--- a/Engine/Source/Runtime/Engine/Private/Components/InstancedSkinnedMeshComponent.cpp
+++ b/Engine/Source/Runtime/Engine/Private/Components/InstancedSkinnedMeshComponent.cpp
@@ -183,6 +183,7 @@ public:
 		return AnimBankIds;
 	}
 
+#if RHI_RAYTRACING
 	virtual ERayTracingPrimitiveFlags GetCachedRayTracingInstance(FRayTracingInstance& RayTracingInstance) override
 	{
 		static const auto RayTracingInstancedSkeletalMeshesCVar = IConsoleManager::Get().FindTConsoleVariableDataInt(TEXT("r.RayTracing.Geometry.InstancedSkeletalMeshes"));
@@ -194,6 +195,7 @@ public:
 
 		return Nanite::FSkinnedSceneProxy::GetCachedRayTracingInstance(RayTracingInstance);
 	}
+#endif
 
 	virtual bool GetInstanceDrawDistanceMinMax(FVector2f& OutCullRange) const override
 	{
