[Runtime virtual texture] Do not reinitialize resource when it is used by other RVT scene proxies. This can happen when the RVT is used in multiple scenes. A FRuntimeVirtualTextureSceneProxy is constructed in each scene. If the resource is reinitialized when the second scene proxy is constructed, cached resource pointers in the first proxy will become invalid and cause a crash if used

diff --git a/Engine/Source/Runtime/Engine/Classes/VT/RuntimeVirtualTexture.h b/Engine/Source/Runtime/Engine/Classes/VT/RuntimeVirtualTexture.h
--- a/Engine/Source/Runtime/Engine/Classes/VT/RuntimeVirtualTexture.h
+++ b/Engine/Source/Runtime/Engine/Classes/VT/RuntimeVirtualTexture.h
@@ -206,6 +206,7 @@ protected:
 private:
 	/** Render thread resource container. */
 	class FRuntimeVirtualTextureRenderResource* Resource;
+	int32 ResourceRefCount = 0;
 
 	/** Material uniform parameters to support transform from world to UV coordinates. */
 	FVector4 WorldToUVTransformParameters[3];
diff --git a/Engine/Source/Runtime/Engine/Private/VT/RuntimeVirtualTexture.cpp b/Engine/Source/Runtime/Engine/Private/VT/RuntimeVirtualTexture.cpp
--- a/Engine/Source/Runtime/Engine/Private/VT/RuntimeVirtualTexture.cpp
+++ b/Engine/Source/Runtime/Engine/Private/VT/RuntimeVirtualTexture.cpp
@@ -605,19 +605,33 @@ EVTProducerPriority URuntimeVirtualTexture::GetDefaultPriority() const
 
 void URuntimeVirtualTexture::Initialize(IVirtualTexture* InProducer, FVTProducerDescription const& InProducerDesc, FTransform const& InVolumeToWorld, FBox const& InWorldBounds)
 {
-	WorldToUVTransformParameters[0] = InVolumeToWorld.GetTranslation();
-	WorldToUVTransformParameters[1] = InVolumeToWorld.GetUnitAxis(EAxis::X) * 1.f / InVolumeToWorld.GetScale3D().X;
-	WorldToUVTransformParameters[2] = InVolumeToWorld.GetUnitAxis(EAxis::Y) * 1.f / InVolumeToWorld.GetScale3D().Y;
+	// TODO: Look into multi-scene support if needed. The ref count prevents resource being reallocated and invalidating
+	// cached pointers. However, it assumes that the RVT is initialized using the same transform, bounds, and producer
+	// all scenes which seems to be the case for PIE. So this currently won't work if the RVT needs to be initialized
+	// differently in different scenes
+	if (ResourceRefCount == 0)
+	{
+		WorldToUVTransformParameters[0] = InVolumeToWorld.GetTranslation();
+		WorldToUVTransformParameters[1] = InVolumeToWorld.GetUnitAxis(EAxis::X) * 1.f / InVolumeToWorld.GetScale3D().X;
+		WorldToUVTransformParameters[2] = InVolumeToWorld.GetUnitAxis(EAxis::Y) * 1.f / InVolumeToWorld.GetScale3D().Y;
+
+		const FVector::FReal HeightRange = FMath::Max<FVector::FReal>(InWorldBounds.Max.Z - InWorldBounds.Min.Z, 1.f);
+		WorldHeightUnpackParameter = FVector4(HeightRange, InWorldBounds.Min.Z, 0.f, 0.f);
 
-	const FVector::FReal HeightRange = FMath::Max<FVector::FReal>(InWorldBounds.Max.Z - InWorldBounds.Min.Z, 1.f);
-	WorldHeightUnpackParameter = FVector4(HeightRange, InWorldBounds.Min.Z, 0.f, 0.f);
+		InitResource(InProducer, InProducerDesc);
+	}
 
-	InitResource(InProducer, InProducerDesc);
+	++ResourceRefCount;
 }
 
 void URuntimeVirtualTexture::Release()
 {
-	InitNullResource();
+	--ResourceRefCount;
+
+	if (ResourceRefCount == 0)
+	{
+		InitNullResource();
+	}
 }
 
 void URuntimeVirtualTexture::InitResource(IVirtualTexture* InProducer, FVTProducerDescription const& InProducerDesc)
diff --git a/Engine/Source/Runtime/Renderer/Private/RendererScene.cpp b/Engine/Source/Runtime/Renderer/Private/RendererScene.cpp
--- a/Engine/Source/Runtime/Renderer/Private/RendererScene.cpp
+++ b/Engine/Source/Runtime/Renderer/Private/RendererScene.cpp
@@ -3192,6 +3192,8 @@ void FScene::AddRuntimeVirtualTexture(URuntimeVirtualTextureComponent* Component
 		// This is a component update.
 		// Delete old proxy and add new one on the render thread.
 		FRuntimeVirtualTextureSceneProxy* SceneProxyToRemove = Component->SceneProxy;
+		SceneProxyToRemove->Release();
+
 		Component->SceneProxy = new FRuntimeVirtualTextureSceneProxy(Component);
 
 		FScene* Scene = this;
