Fixed incorrect registration of components with the scene that are marked 'CastHiddenShadow' but not 'CastShadow'. This matches the behavior in FPrimitiveSceneProxy when it assigns bCastHiddenShadow.

diff --git a/Engine/Source/Runtime/Engine/Private/Components/PrimitiveComponent.cpp b/Engine/Source/Runtime/Engine/Private/Components/PrimitiveComponent.cpp
--- a/Engine/Source/Runtime/Engine/Private/Components/PrimitiveComponent.cpp
+++ b/Engine/Source/Runtime/Engine/Private/Components/PrimitiveComponent.cpp
@@ -1997,7 +1997,7 @@ bool UPrimitiveComponent::ShouldComponentAddToScene() const
 	const bool bIsTemporarilyHiddenInEditor = false;
 #endif
 
-	return bSceneAdd && !bIsTemporarilyHiddenInEditor && (ShouldRender() || (bCastHiddenShadow && !bIsHiddenInEditor) || bAffectIndirectLightingWhileHidden || bRayTracingFarField);
+	return bSceneAdd && !bIsTemporarilyHiddenInEditor && (ShouldRender() || (CastShadow && bCastHiddenShadow && !bIsHiddenInEditor) || bAffectIndirectLightingWhileHidden || bRayTracingFarField);
 }
 
 bool UPrimitiveComponent::ShouldCreatePhysicsState() const
