[HWRT] Fix Ray Tracing Geometry Manager streaming in too many geometries leading to going over budget.

Context:
When deciding whether more geometries fit in budget, we were checking TotalResidentSize, which was originally correct because calling MakeGeometryResident(...) was increasing that value.
However that's no longer always the case since if the geometry needs to be streamed in, because TotalResidentSize is only updated when the streaming requests complete.
This means that the manager could end up kicking off too many streaming requests (only getting limited by r.RayTracing.Streaming.MaxPendingRequests or geometries that don't require streaming).

Changes:
- Keep track of how much memory will be used once the in-flight streaming requests complete (TotalStreamingSize).
- Take both TotalResidentSize and TotalStreamingSize into account when deciding whether geometries need to be evicted or made resident.

diff --git a/Engine/Source/Runtime/Engine/Private/Rendering/RayTracingGeometryManager.cpp b/Engine/Source/Runtime/Engine/Private/Rendering/RayTracingGeometryManager.cpp
--- a/Engine/Source/Runtime/Engine/Private/Rendering/RayTracingGeometryManager.cpp
+++ b/Engine/Source/Runtime/Engine/Private/Rendering/RayTracingGeometryManager.cpp
@@ -23,6 +23,8 @@
 
 #if RHI_RAYTRACING
 
+DECLARE_LOG_CATEGORY_CLASS(LogRayTracingGeometryManager, Log, All);
+
 static bool bHasRayTracingEnableChanged = false;
 static TAutoConsoleVariable<int32> CVarRayTracingEnable(
 	TEXT("r.RayTracing.Enable"),
@@ -788,7 +790,7 @@ void FRayTracingGeometryManager::Tick(FRHICommandList& RHICmdList)
 
 		// Step 3
 		// - if making requested geometries resident will put us over budget -> evict some geometry not referenced by TLAS
-		if (TotalResidentSize + RequestedButEvictedSize > ResidentGeometryMemoryPoolSize)
+		if (TotalResidentSize + TotalStreamingSize + RequestedButEvictedSize > ResidentGeometryMemoryPoolSize)
 		{
 			TArray<FGeometryHandle> NotReferencedResidentGeometriesArray = NotReferencedResidentGeometries.Array();
 
@@ -817,7 +819,7 @@ void FRayTracingGeometryManager::Tick(FRHICommandList& RHICmdList)
 			// Step 3.2
 			// - evict geometries until we are in budget
 			int32 Index = 0;
-			while (TotalResidentSize + RequestedButEvictedSize > ResidentGeometryMemoryPoolSize && Index < NotReferencedResidentGeometriesArray.Num())
+			while (TotalResidentSize + TotalStreamingSize + RequestedButEvictedSize > ResidentGeometryMemoryPoolSize && Index < NotReferencedResidentGeometriesArray.Num())
 			{
 				FGeometryHandle GeometryHandle = NotReferencedResidentGeometriesArray[Index];
 				FRegisteredGeometry& RegisteredGeometry = RegisteredGeometries[GeometryHandle];
@@ -832,7 +834,7 @@ void FRayTracingGeometryManager::Tick(FRHICommandList& RHICmdList)
 
 		// Step 4
 		// - make referenced geometries resident until we go over budget
-		if (TotalResidentSize < ResidentGeometryMemoryPoolSize)
+		if (TotalResidentSize + TotalStreamingSize < ResidentGeometryMemoryPoolSize)
 		{
 			// Step 4.1
 			//	- sort by size to prioritize smaller geometries
@@ -844,13 +846,11 @@ void FRayTracingGeometryManager::Tick(FRHICommandList& RHICmdList)
 					return LHS.Size < RHS.Size;
 				});
 
-			// Step 3.2
+			// Step 4.2
 			// - make geometries resident until we go over budget
 			int32 Index = 0;
-			while (TotalResidentSize < ResidentGeometryMemoryPoolSize && Index < ReferencedGeometries.Num())
+			while (TotalResidentSize + TotalStreamingSize < ResidentGeometryMemoryPoolSize && Index < ReferencedGeometries.Num())
 			{
-				// if referenced this frame, mark for eviction and add to pending list
-
 				FGeometryHandle GeometryHandle = ReferencedGeometries[Index];
 				FRegisteredGeometry& RegisteredGeometry = RegisteredGeometries[GeometryHandle];
 
@@ -1029,13 +1029,16 @@ bool FRayTracingGeometryManager::RequestRayTracingGeometryStreamIn(FRHICommandLi
 			StreamableDataSize += RegisteredGeometry.StreamableBVHSize;
 		}
 
-		StreamingRequest.GeometryHandle = Geometry->RayTracingGeometryHandle;
+		StreamingRequest.GeometryHandle = GeometryHandle;
+		StreamingRequest.GeometrySize = RegisteredGeometry.Size;
 		StreamingRequest.bBuffersOnly = !bStreamBVH;
 		StreamingRequest.RequestBuffer = FIoBuffer(StreamableDataSize); // TODO: Use FIoBuffer::Wrap with preallocated memory
 
 		RegisteredGeometry.StreamableAsset->IssueRequest(StreamingRequest.Request, StreamingRequest.RequestBuffer, StreamingRequest.bBuffersOnly);
 
 		RegisteredGeometry.Status = FRegisteredGeometry::FStatus::Streaming;
+
+		TotalStreamingSize += StreamingRequest.GeometrySize;
 	}
 
 	if (RegisteredGeometry.Status == FRegisteredGeometry::FStatus::StreamedIn)
@@ -1081,6 +1084,9 @@ void FRayTracingGeometryManager::ProcessCompletedStreamingRequests(FRHICommandLi
 		{
 			++NumCompletedRequests;
 
+			TotalStreamingSize -= PendingRequest.GeometrySize;
+			check(TotalStreamingSize >= 0);
+
 			if (PendingRequest.bCancelled)
 			{
 				PendingRequest.Reset();
@@ -1099,11 +1105,9 @@ void FRayTracingGeometryManager::ProcessCompletedStreamingRequests(FRHICommandLi
 			}
 			else if (!PendingRequest.Request.IsOk())
 			{
-				// Retry if IO request failed for some reason
-				RegisteredGeometry.StreamableAsset->IssueRequest(PendingRequest.Request, PendingRequest.RequestBuffer, PendingRequest.bBuffersOnly);
-						
-				// TODO: Could other requests already be completed?
-				break;
+				UE_LOG(LogRayTracingGeometryManager, Warning, TEXT("Ray Tracing Geometry IO Request failed (%s)"), *RegisteredGeometry.Geometry->Initializer.DebugName.ToString());
+
+				// Manager will retry again if still necessary on the next frame
 			}
 			else
 			{
diff --git a/Engine/Source/Runtime/Engine/Public/Rendering/RayTracingGeometryManager.h b/Engine/Source/Runtime/Engine/Public/Rendering/RayTracingGeometryManager.h
--- a/Engine/Source/Runtime/Engine/Public/Rendering/RayTracingGeometryManager.h
+++ b/Engine/Source/Runtime/Engine/Public/Rendering/RayTracingGeometryManager.h
@@ -140,6 +140,8 @@ private:
 		FBulkDataBatchRequest Request;
 
 		FGeometryHandle GeometryHandle = INDEX_NONE;
+		uint32 GeometrySize = 0;
+
 		bool bBuffersOnly = false;
 
 		bool bCancelled = false;
@@ -208,6 +210,10 @@ private:
 	int32 NumStreamingRequests = 0;
 	int32 NextStreamingRequestIndex = 0;
 
+	// Total size of the BLASes currently being streamed.
+	// Used to keep track of how much TotalResidentSize will increase when the requests complete.
+	uint64 TotalStreamingSize = 0;
+
 	// Feedback
 	TSet<int32> VisibleGeometryHandles;
 
