Moved duplicated world initialization sequence and garbage collection code from FWorldPartitionCookPackageSplitter into FWorldCookPackageSplitter to avoid duplicating that logic for other subsplitters.

diff --git a/Engine/Source/Runtime/Engine/Private/WorldCookPackageSplitter.cpp b/Engine/Source/Runtime/Engine/Private/WorldCookPackageSplitter.cpp
--- a/Engine/Source/Runtime/Engine/Private/WorldCookPackageSplitter.cpp
+++ b/Engine/Source/Runtime/Engine/Private/WorldCookPackageSplitter.cpp
@@ -4,11 +4,17 @@
 
 #if WITH_EDITOR
 #include "Engine/World.h"
+#include "Editor.h"
 
 REGISTER_COOKPACKAGE_SPLITTER(FWorldCookPackageSplitter, UWorld);
 
 TMap<UClass*, FWorldCookPackageSplitter::FSubSplitterFactory> FWorldCookPackageSplitter::RegisteredCookPackageSubSplitterFactories;
 
+FWorldCookPackageSplitter::~FWorldCookPackageSplitter()
+{
+	check(!ReferencedWorld);
+}
+
 void FWorldCookPackageSplitter::RegisterCookPackageSubSplitterFactory(UClass* Class, FWorldCookPackageSplitter::FSubSplitterFactory Factory)
 {
 	check(!RegisteredCookPackageSubSplitterFactories.Contains(Class));
@@ -21,6 +27,21 @@ void FWorldCookPackageSplitter::UnregisterCookPackageSubSplitterFactory(UClass*
 	RegisteredCookPackageSubSplitterFactories.Remove(Class);
 }
 
+void FWorldCookPackageSplitter::OnWorldCleanup(UWorld* InWorld, bool bSessionEnded, bool bCleanupResources)
+{
+	checkf(InWorld != ReferencedWorld, TEXT("[Cook] %s is being cleaned up while still referenced by a package splitter."), *GetFullNameSafe(InWorld));
+}
+
+void FWorldCookPackageSplitter::AddReferencedObjects(FReferenceCollector& Collector)
+{
+	Collector.AddReferencedObject(ReferencedWorld);
+}
+
+FString FWorldCookPackageSplitter::GetReferencerName() const
+{
+	return TEXT("FWorldCookPackageSplitter");
+}
+
 bool FWorldCookPackageSplitter::ShouldSplit(UObject* SplitData)
 {
 	bool bResult = false;
@@ -36,15 +57,6 @@ bool FWorldCookPackageSplitter::ShouldSplit(UObject* SplitData)
 	return bResult;
 }
 
-void FWorldCookPackageSplitter::Teardown(ICookPackageSplitter::ETeardown Status)
-{
-	ForEachCookPackageSubSplitters(
-		[Status](FWorldCookPackageSplitter::FSubSplitter* WorldCookPackageSplitter)
-		{
-			WorldCookPackageSplitter->Teardown(Status);
-		});
-}
-
 bool FWorldCookPackageSplitter::UseInternalReferenceToAvoidGarbageCollect()
 {
 	return true;
@@ -62,7 +74,12 @@ bool FWorldCookPackageSplitter::RequiresGeneratorPackageDestructBeforeResplit()
 
 TArray<ICookPackageSplitter::FGeneratedPackage> FWorldCookPackageSplitter::GetGenerateList(const UPackage* OwnerPackage, const UObject* OwnerObject)
 {
-	const UWorld* OwnerWorld = CastChecked<const UWorld>(OwnerObject);
+	TNotNull<const UWorld*> OwnerWorld = CastChecked<const UWorld>(OwnerObject);
+
+	ReferencedWorld = const_cast<UWorld*>(NotNullGet(OwnerWorld));
+
+	check(!bInitializedPhysicsSceneForSave && !bForceInitializedWorld);
+	bInitializedPhysicsSceneForSave = GEditor->InitializePhysicsSceneForSaveIfNecessary(ReferencedWorld, bForceInitializedWorld);
 
 	check(!CookPackageSubSplitters.Num());
 	ForEachRegisteredCookPackageSubSplitterFactories(
@@ -70,7 +87,7 @@ TArray<ICookPackageSplitter::FGeneratedPackage> FWorldCookPackageSplitter::GetGe
 		{
 			if (WorldCookPackageSplitterFactory.ShouldSplit(OwnerWorld))
 			{
-				CookPackageSubSplitters.Emplace(WorldCookPackageSplitterFactory.MakeInstance());
+				CookPackageSubSplitters.Emplace(WorldCookPackageSplitterFactory.MakeInstance(OwnerWorld));
 			}
 		});
 	check(CookPackageSubSplitters.Num());
@@ -79,8 +96,11 @@ TArray<ICookPackageSplitter::FGeneratedPackage> FWorldCookPackageSplitter::GetGe
 	ForEachCookPackageSubSplitters(
 		[&PackagesToGenerate, OwnerPackage, OwnerWorld](FWorldCookPackageSplitter::FSubSplitter* WorldCookPackageSplitter)
 		{
-			PackagesToGenerate.Append(WorldCookPackageSplitter->GetGenerateList(OwnerPackage, OwnerWorld));
+			PackagesToGenerate.Append(WorldCookPackageSplitter->GetGenerateList(OwnerPackage));
 		});
+
+	FWorldDelegates::OnWorldCleanup.AddRaw(this, &FWorldCookPackageSplitter::OnWorldCleanup);
+
 	return PackagesToGenerate;
 }
 
@@ -110,12 +130,27 @@ bool FWorldCookPackageSplitter::PopulateGeneratorPackage(FPopulateContext& Popul
 
 void FWorldCookPackageSplitter::OnOwnerReloaded(UPackage* OwnerPackage, UObject* OwnerObject)
 {
-	const UWorld* OwnerWorld = CastChecked<const UWorld>(OwnerObject);
+	check(!ReferencedWorld);
+}
+
+void FWorldCookPackageSplitter::Teardown(ICookPackageSplitter::ETeardown Status)
+{
+	FWorldDelegates::OnWorldCleanup.RemoveAll(this);
+
 	ForEachCookPackageSubSplitters(
-		[OwnerPackage, OwnerWorld](FWorldCookPackageSplitter::FSubSplitter* WorldCookPackageSplitter)
+		[Status](FWorldCookPackageSplitter::FSubSplitter* WorldCookPackageSplitter)
 		{
-			WorldCookPackageSplitter->OnOwnerReloaded(OwnerPackage, OwnerWorld);
+			WorldCookPackageSplitter->Teardown(Status);
 		});
+
+	if (bInitializedPhysicsSceneForSave)
+	{
+		GEditor->CleanupPhysicsSceneThatWasInitializedForSave(ReferencedWorld, bForceInitializedWorld);
+		bInitializedPhysicsSceneForSave = false;
+		bForceInitializedWorld = false;
+	}
+
+	ReferencedWorld = nullptr;
 }
 
 #endif // WITH_EDITOR
\ No newline at end of file
diff --git a/Engine/Source/Runtime/Engine/Private/WorldPartition/WorldPartition.cpp b/Engine/Source/Runtime/Engine/Private/WorldPartition/WorldPartition.cpp
--- a/Engine/Source/Runtime/Engine/Private/WorldPartition/WorldPartition.cpp
+++ b/Engine/Source/Runtime/Engine/Private/WorldPartition/WorldPartition.cpp
@@ -385,7 +385,7 @@ UWorldPartition::UWorldPartition(const FObjectInitializer& ObjectInitializer)
 	bAllowShowingHLODsInEditor = true;
 	WorldPartitionStreamingPolicyClass = UWorldPartitionLevelStreamingPolicy::StaticClass();
 
-	if (HasAnyFlags(RF_ClassDefaultObject) && !HasAnyFlags(RF_ImmutableDefaultObject))
+	if (HasAnyFlags(RF_ClassDefaultObject) && !HasAnyFlags(RF_ImmutableDefaultObject) && IsRunningCookCommandlet())
 	{
 		FWorldCookPackageSplitter::RegisterCookPackageSubSplitterFactory(
 			GetClass(),
@@ -396,9 +396,9 @@ UWorldPartition::UWorldPartition(const FObjectInitializer& ObjectInitializer)
 						return World->IsPartitionedWorld();
 					},
 				.MakeInstance =
-					[]() -> FWorldCookPackageSplitter::FSubSplitter*
+					[](TNotNull<const UWorld*> World) -> FWorldCookPackageSplitter::FSubSplitter*
         			{
-						return new FWorldPartitionCookPackageSplitter;
+						return new FWorldPartitionCookPackageSplitter(World);
 					}
 			});
 	}
@@ -1744,7 +1744,7 @@ void UWorldPartition::BeginDestroy()
 	check(InitState == EWorldPartitionInitState::Uninitialized);
 
 #if WITH_EDITOR
-	if (HasAnyFlags(RF_ClassDefaultObject) && !HasAnyFlags(RF_ImmutableDefaultObject))
+	if (HasAnyFlags(RF_ClassDefaultObject) && !HasAnyFlags(RF_ImmutableDefaultObject) && IsRunningCookCommandlet())
 	{
 		FWorldCookPackageSplitter::UnregisterCookPackageSubSplitterFactory(GetClass());
 	}
diff --git a/Engine/Source/Runtime/Engine/Private/WorldPartition/WorldPartitionCookPackageSplitter.cpp b/Engine/Source/Runtime/Engine/Private/WorldPartition/WorldPartitionCookPackageSplitter.cpp
--- a/Engine/Source/Runtime/Engine/Private/WorldPartition/WorldPartitionCookPackageSplitter.cpp
+++ b/Engine/Source/Runtime/Engine/Private/WorldPartition/WorldPartitionCookPackageSplitter.cpp
@@ -10,42 +10,15 @@
 #include "WorldPartition/WorldPartition.h"
 #include "WorldPartition/WorldPartitionLog.h"
 #include "WorldPartition/WorldPartitionRuntimeCell.h"
-#include "Editor.h"
 
-FWorldPartitionCookPackageSplitter::~FWorldPartitionCookPackageSplitter()
+TArray<ICookPackageSplitter::FGeneratedPackage> FWorldPartitionCookPackageSplitter::GetGenerateList(const UPackage* OwnerPackage)
 {
-	check(!ReferencedWorld);
-}
-
-void FWorldPartitionCookPackageSplitter::OnWorldCleanup(UWorld* InWorld, bool bSessionEnded, bool bCleanupResources)
-{
-	checkf(InWorld != ReferencedWorld, TEXT("[Cook] %s is being cleaned up while still referenced by a package splitter."), *GetFullNameSafe(InWorld));
-}
-
-void FWorldPartitionCookPackageSplitter::AddReferencedObjects(FReferenceCollector& Collector)
-{
-	Collector.AddReferencedObject(ReferencedWorld);
-}
+	UWorld* ReferencedWorld = const_cast<UWorld*>(NotNullGet(World));
 
-FString FWorldPartitionCookPackageSplitter::GetReferencerName() const
-{
-	return TEXT("FWorldPartitionCookPackageSplitter");
-}
-
-TArray<ICookPackageSplitter::FGeneratedPackage> FWorldPartitionCookPackageSplitter::GetGenerateList(const UPackage* OwnerPackage, TNotNull<const UWorld*> World)
-{
-	UE_LOG(LogWorldPartition, Display, TEXT("[Cook] Gathering packages to cook from generators for owner object %s."), *GetFullNameSafe(NotNullGet(World)));
-
-	// Store the World pointer to declare it to GarbageCollection; we do not want to allow the World to be Garbage Collected
-	// until we have finished all of our PreSaveGeneratedPackage calls, because we store information on the World 
-	// that is necessary for populate 
-	ReferencedWorld = const_cast<UWorld*>(NotNullGet(World));
-
-	check(!bInitializedPhysicsSceneForSave && !bForceInitializedWorld);
-	bInitializedPhysicsSceneForSave = GEditor->InitializePhysicsSceneForSaveIfNecessary(ReferencedWorld, bForceInitializedWorld);
+	UE_LOG(LogWorldPartition, Display, TEXT("[Cook] Gathering packages to cook from generators for owner object %s."), *GetFullNameSafe(ReferencedWorld));
 
 	// Assume that the world is partitioned as per FWorldPartitionCookPackageSplitter::ShouldSplit
-	UWorldPartition* WorldPartition = ReferencedWorld->PersistentLevel->GetWorldPartition();
+	UWorldPartition* WorldPartition = World->PersistentLevel->GetWorldPartition();
 	check(WorldPartition);
 
 	// We expect the WorldPartition has not yet been initialized
@@ -54,7 +27,7 @@ TArray<ICookPackageSplitter::FGeneratedPackage> FWorldPartitionCookPackageSplitt
 	WorldPartition->BeginCook(CookContext);
 
 	bool bIsSuccess = CookContext.GatherPackagesToCook();
-	UE_CLOG(!bIsSuccess, LogWorldPartition, Warning, TEXT("[Cook] Errors while gathering packages to took from generators for owner object %s."), *GetFullNameSafe(NotNullGet(World)));
+	UE_CLOG(!bIsSuccess, LogWorldPartition, Warning, TEXT("[Cook] Errors while gathering packages to took from generators for owner object %s."), *GetFullNameSafe(ReferencedWorld));
 
 	UE_LOG(LogWorldPartition, Display, TEXT("[Cook] Gathered %u packages to generate from %u Generators."), CookContext.NumPackageToGenerate(), CookContext.NumGenerators());
 
@@ -62,9 +35,7 @@ TArray<ICookPackageSplitter::FGeneratedPackage> FWorldPartitionCookPackageSplitt
 	BuildPackagesToGenerateList(PackagesToGenerate);
 
 	UE_LOG(LogWorldPartition, Display, TEXT("[Cook] Sending %u packages to be generated."), PackagesToGenerate.Num());
-	UE_LOG(LogWorldPartition, Display, TEXT("[Cook] Debug(GetGenerateList) : OwnerWorld=%s bForceInitializedWorld=%d bInitializedPhysicsSceneForSave=%d"), *GetFullNameSafe(NotNullGet(World)), bForceInitializedWorld ? 1 : 0, bInitializedPhysicsSceneForSave ? 1 : 0);
-
-	FWorldDelegates::OnWorldCleanup.AddRaw(this, &FWorldPartitionCookPackageSplitter::OnWorldCleanup);
+	UE_LOG(LogWorldPartition, Display, TEXT("[Cook] Debug(GetGenerateList) : OwnerWorld=%s"), *GetFullNameSafe(NotNullGet(World)));
 
 	return PackagesToGenerate;
 }
@@ -127,37 +98,19 @@ bool FWorldPartitionCookPackageSplitter::PopulateGeneratorPackage(ICookPackageSp
 	return bIsSuccess;
 }
 
-void FWorldPartitionCookPackageSplitter::OnOwnerReloaded(UPackage* OwnerPackage, TNotNull<const UWorld*> World)
-{
-	// It should not be possible for the owner to reload due to garbage collection while we are active and keeping it referenced
-	check(!ReferencedWorld);
-}
-
 void FWorldPartitionCookPackageSplitter::Teardown(ICookPackageSplitter::ETeardown Status)
 {
-	UE_LOG(LogWorldPartition, Display, TEXT("[Cook] Debug(TearDown): OwnerWorld=%s bForceInitializedWorld=%d bInitializedPhysicsSceneForSave=%d"), *GetFullNameSafe(ReferencedWorld),  bForceInitializedWorld ? 1 : 0, bInitializedPhysicsSceneForSave ? 1 : 0);
-
-	FWorldDelegates::OnWorldCleanup.RemoveAll(this);
-
-	if (ReferencedWorld)
-	{
-		// We were destructed without GetGenerateList being called; nothing to teardown
-		// Assume that the world is partitioned as per FWorldPartitionCookPackageSplitter::ShouldSplit
-		UWorldPartition* WorldPartition = ReferencedWorld->PersistentLevel->GetWorldPartition();
-		check(WorldPartition);
+	UWorld* ReferencedWorld = const_cast<UWorld*>(NotNullGet(World));
 
-		WorldPartition->EndCook(CookContext);
-		WorldPartition->Uninitialize();
+	UE_LOG(LogWorldPartition, Display, TEXT("[Cook] Debug(TearDown): OwnerWorld=%s"), *GetFullNameSafe(ReferencedWorld));
 
-		if (bInitializedPhysicsSceneForSave)
-		{
-			GEditor->CleanupPhysicsSceneThatWasInitializedForSave(ReferencedWorld, bForceInitializedWorld);
-			bInitializedPhysicsSceneForSave = false;
-			bForceInitializedWorld = false;
-		}
+	// We were destructed without GetGenerateList being called; nothing to teardown
+	// Assume that the world is partitioned as per FWorldPartitionCookPackageSplitter::ShouldSplit
+	UWorldPartition* WorldPartition = ReferencedWorld->PersistentLevel->GetWorldPartition();
+	check(WorldPartition);
 
-		ReferencedWorld = nullptr;
-	}
+	WorldPartition->EndCook(CookContext);
+	WorldPartition->Uninitialize();
 }
 
 void FWorldPartitionCookPackageSplitter::BuildPackagesToGenerateList(TArray<ICookPackageSplitter::FGeneratedPackage>& PackagesToGenerate) const
diff --git a/Engine/Source/Runtime/Engine/Private/WorldPartition/WorldPartitionCookPackageSplitter.h b/Engine/Source/Runtime/Engine/Private/WorldPartition/WorldPartitionCookPackageSplitter.h
--- a/Engine/Source/Runtime/Engine/Private/WorldPartition/WorldPartitionCookPackageSplitter.h
+++ b/Engine/Source/Runtime/Engine/Private/WorldPartition/WorldPartitionCookPackageSplitter.h
@@ -7,33 +7,24 @@
 #include "WorldCookPackageSplitter.h"
 #include "WorldPartition/Cook/WorldPartitionCookPackageContext.h"
 
-class FWorldPartitionCookPackageSplitter : public FGCObject, public FWorldCookPackageSplitter::FSubSplitter
+class FWorldPartitionCookPackageSplitter : public FWorldCookPackageSplitter::FSubSplitter
 {
 public:
+	FWorldPartitionCookPackageSplitter(TNotNull<const UWorld*> World)
+		: FWorldCookPackageSplitter::FSubSplitter(World)
+	{}
+
 	//~Begin FWorldCookPackageSplitter::FSubSplitter interface
-	virtual ~FWorldPartitionCookPackageSplitter();
-	virtual TArray<ICookPackageSplitter::FGeneratedPackage> GetGenerateList(const UPackage* OwnerPackage, TNotNull<const UWorld*> World) override;
+	virtual TArray<ICookPackageSplitter::FGeneratedPackage> GetGenerateList(const UPackage* OwnerPackage) override;
 	virtual bool PopulateGeneratedPackage(ICookPackageSplitter::FPopulateContext& PopulateContext) override;
 	virtual bool PopulateGeneratorPackage(ICookPackageSplitter::FPopulateContext& PopulateContext) override;
-	virtual void OnOwnerReloaded(UPackage* OwnerPackage, TNotNull<const UWorld*> World) override;
 	virtual void Teardown(ICookPackageSplitter::ETeardown Status) override;
 	//~End FWorldCookPackageSplitter::FSubSplitter interface
 
 private:
-	void OnWorldCleanup(UWorld* InWorld, bool bSessionEnded, bool bCleanupResources);
-
-	/** FGCObject interface */
-	virtual void AddReferencedObjects(FReferenceCollector& Collector) override;
-	virtual FString GetReferencerName() const override;
-
 	void BuildPackagesToGenerateList(TArray<ICookPackageSplitter::FGeneratedPackage>& PackagesToGenerate) const;
 	
-	TObjectPtr<UWorld> ReferencedWorld = nullptr;
-
 	FWorldPartitionCookPackageContext CookContext;
-
-	bool bForceInitializedWorld = false;
-	bool bInitializedPhysicsSceneForSave = false;
 };
 
 #endif
\ No newline at end of file
diff --git a/Engine/Source/Runtime/Engine/Public/WorldCookPackageSplitter.h b/Engine/Source/Runtime/Engine/Public/WorldCookPackageSplitter.h
--- a/Engine/Source/Runtime/Engine/Public/WorldCookPackageSplitter.h
+++ b/Engine/Source/Runtime/Engine/Public/WorldCookPackageSplitter.h
@@ -12,25 +12,36 @@
 #define UE_API ENGINE_API
 
 #if WITH_EDITOR
-class FWorldCookPackageSplitter : public ICookPackageSplitter
+class FWorldCookPackageSplitter : public FGCObject, ICookPackageSplitter
 {
 	DEFINE_COOKPACKAGE_SPLITTER(FWorldCookPackageSplitter, UWorld);
 
 public:
+	~FWorldCookPackageSplitter();
+
 	struct FSubSplitter
 	{
-		virtual ~FSubSplitter() {}
-		virtual TArray<ICookPackageSplitter::FGeneratedPackage> GetGenerateList(const UPackage* OwnerPackage, TNotNull<const UWorld*> World) = 0;
+		FSubSplitter(TNotNull<const UWorld*> InWorld)
+			: World(InWorld)
+		{}
+
+		virtual ~FSubSplitter()
+		{}
+
+		//~Begin FSubSplitter interface
+		virtual TArray<ICookPackageSplitter::FGeneratedPackage> GetGenerateList(const UPackage* OwnerPackage) = 0;
 		virtual bool PopulateGeneratedPackage(ICookPackageSplitter::FPopulateContext& PopulateContext) = 0;
 		virtual bool PopulateGeneratorPackage(ICookPackageSplitter::FPopulateContext& PopulateContext) = 0;
-		virtual void OnOwnerReloaded(UPackage* OwnerPackage, TNotNull<const UWorld*> World) = 0;
 		virtual void Teardown(ICookPackageSplitter::ETeardown Status) = 0;
+		//~End FSubSplitter interface
+
+		TNotNull<const UWorld*> World;
 	};
 
 	struct FSubSplitterFactory
 	{
 		TFunction<bool(TNotNull<const UWorld*>)> ShouldSplit;
-		TFunction<FSubSplitter*()> MakeInstance;
+		TFunction<FSubSplitter*(TNotNull<const UWorld*>)> MakeInstance;
 	};
 
 	static UE_API void RegisterCookPackageSubSplitterFactory(UClass* Class, FSubSplitterFactory Factory);
@@ -53,11 +64,15 @@ protected:
 		return true;
 	}
 
+	//~Begin FGCObject interface
+	virtual void AddReferencedObjects(FReferenceCollector& Collector) override;
+	virtual FString GetReferencerName() const override;
+	//~End FGCObject interface
+
 	//~ Begin of ICookPackageSplitter
 	static bool ShouldSplit(UObject* SplitData);
 	static FString GetSplitterDebugName() { return TEXT("FWorldCookPackageSplitter"); }
 
-	virtual void Teardown(ETeardown Status) override;
 	virtual bool UseInternalReferenceToAvoidGarbageCollect() override;
 	virtual EGeneratedRequiresGenerator DoesGeneratedRequireGenerator() override;
 	virtual bool RequiresGeneratorPackageDestructBeforeResplit() override;
@@ -65,6 +80,7 @@ protected:
 	virtual bool PopulateGeneratedPackage(FPopulateContext& PopulateContext) override;
 	virtual bool PopulateGeneratorPackage(FPopulateContext& PopulateContext) override;
 	virtual void OnOwnerReloaded(UPackage* OwnerPackage, UObject* OwnerObject) override;
+	virtual void Teardown(ETeardown Status) override;
 	//~ End of ICookPackageSplitter
 
 	TArray<TUniquePtr<FSubSplitter>> CookPackageSubSplitters;
@@ -82,6 +98,12 @@ protected:
 		}
 		return true;
 	}
+
+	void OnWorldCleanup(UWorld* InWorld, bool bSessionEnded, bool bCleanupResources);
+
+	TObjectPtr<UWorld> ReferencedWorld = nullptr;
+	bool bForceInitializedWorld = false;
+	bool bInitializedPhysicsSceneForSave = false;
 };
 #endif // WITH_EDITOR
 
