WorldSubSystem:
- Added UWorld::GetSubsystemChecked that returns a specific world subsystem, checked.
- Replace Cast/check by CastChecked in UWorldSubsystem::ShouldCreateSubsystem.

diff --git a/Engine/Source/Runtime/Engine/Classes/Engine/World.h b/Engine/Source/Runtime/Engine/Classes/Engine/World.h
--- a/Engine/Source/Runtime/Engine/Classes/Engine/World.h
+++ b/Engine/Source/Runtime/Engine/Classes/Engine/World.h
@@ -4187,6 +4187,15 @@ public:
 		return SubsystemCollection.GetSubsystem<TSubsystemClass>(TSubsystemClass::StaticClass());
 	}
 
+	/**
+	 * Get a Subsystem of specified type, checked
+	 */
+	template <typename TSubsystemClass>
+	TNotNull<TSubsystemClass*> GetSubsystemChecked() const
+	{
+		return GetSubsystem<TSubsystemClass>();
+	}
+
 	/**
 	 * Get a Subsystem of specified type from the provided GameInstance
 	 * returns nullptr if the Subsystem cannot be found or the GameInstance is null
diff --git a/Engine/Source/Runtime/Engine/Private/Subsystems/WorldSubsystem.cpp b/Engine/Source/Runtime/Engine/Private/Subsystems/WorldSubsystem.cpp
--- a/Engine/Source/Runtime/Engine/Private/Subsystems/WorldSubsystem.cpp
+++ b/Engine/Source/Runtime/Engine/Private/Subsystems/WorldSubsystem.cpp
@@ -32,8 +32,7 @@ bool UWorldSubsystem::ShouldCreateSubsystem(UObject* Outer) const
 		return false;
 	}
 
-	UWorld* World = Cast<UWorld>(Outer);
-	check(World);
+	UWorld* World = CastChecked<UWorld>(Outer);
 	return DoesSupportWorldType(World->WorldType);
 }
 
