Include diff of nested arrays and other containers in edgraphnodes, which were accidently being filtered because the inner FProperty has no meaningful flags on it

[RN] EdGraphNode property changes nested in containers are now detected by the Diff tool

diff --git a/Engine/Source/Runtime/Engine/Private/EdGraph/EdGraphNode.cpp b/Engine/Source/Runtime/Engine/Private/EdGraph/EdGraphNode.cpp
--- a/Engine/Source/Runtime/Engine/Private/EdGraph/EdGraphNode.cpp
+++ b/Engine/Source/Runtime/Engine/Private/EdGraph/EdGraphNode.cpp
@@ -327,11 +327,19 @@ void UEdGraphNode::DiffProperties(UStruct* StructA, UStruct* StructB, uint8* Dat
 
 		auto ShouldSkipProperty = [](const FProperty& Prop)
 			{
-				return !Prop.HasAnyPropertyFlags(CPF_Edit | CPF_BlueprintVisible) ||
-					Prop.HasAnyPropertyFlags(CPF_Transient) ||
-					Prop.HasAnyPropertyFlags(CPF_DisableEditOnInstance) ||
-					Prop.IsA(FDelegateProperty::StaticClass()) ||
-					Prop.IsA(FMulticastDelegateProperty::StaticClass());
+				// if the property is an inner property (e.g. for a container) then use the outer property for visibility:
+				const FProperty* PropertyToTest = &Prop;
+				const FProperty* Owner = Prop.GetOwnerProperty();
+				if(Owner)
+				{
+					PropertyToTest = Owner;
+				}
+
+				return !PropertyToTest->HasAnyPropertyFlags(CPF_Edit | CPF_BlueprintVisible) ||
+					PropertyToTest->HasAnyPropertyFlags(CPF_Transient) ||
+					PropertyToTest->HasAnyPropertyFlags(CPF_DisableEditOnInstance) ||
+					PropertyToTest->IsA(FDelegateProperty::StaticClass()) ||
+					PropertyToTest->IsA(FMulticastDelegateProperty::StaticClass());
 			};
 
 		// skip properties we cant see
