Skip LOD levels below MinLOD in FSkeletalMeshSceneProxy::DrawStaticElements

This also fixes PSO precaching mismatches for TGPUSkinVertexFactory where PSO precaching was respecting MinLOD, and DrawStaticElements was not. For example, LOD0 might have 8 bones and LOD1 might have 4 bones.

diff --git a/Engine/Source/Runtime/Engine/Private/SkeletalMeshSceneProxy.cpp b/Engine/Source/Runtime/Engine/Private/SkeletalMeshSceneProxy.cpp
--- a/Engine/Source/Runtime/Engine/Private/SkeletalMeshSceneProxy.cpp
+++ b/Engine/Source/Runtime/Engine/Private/SkeletalMeshSceneProxy.cpp
@@ -78,14 +78,19 @@ bool AllowLumenCardGenerationForSkeletalMeshes(EShaderPlatform Platform)
  * @param	Component - skeletal mesh primitive being added
  */
 FSkeletalMeshSceneProxy::FSkeletalMeshSceneProxy(const USkinnedMeshComponent* Component, FSkeletalMeshRenderData* InSkelMeshRenderData)
-	: FSkeletalMeshSceneProxy(FSkinnedMeshSceneProxyDesc(Component), InSkelMeshRenderData)
+	: FSkeletalMeshSceneProxy(FSkinnedMeshSceneProxyDesc(Component), InSkelMeshRenderData, Component->ComputeMinLOD())
 { }
 
 FSkeletalMeshSceneProxy::FSkeletalMeshSceneProxy(const FSkinnedMeshSceneProxyDesc& InMeshDesc, FSkeletalMeshRenderData* InSkelMeshRenderData)
+	: FSkeletalMeshSceneProxy(InMeshDesc, InSkelMeshRenderData, 0)
+{ }
+
+FSkeletalMeshSceneProxy::FSkeletalMeshSceneProxy(const FSkinnedMeshSceneProxyDesc& InMeshDesc, FSkeletalMeshRenderData* InSkelMeshRenderData, int32 MinLODIndex)
 		:	FPrimitiveSceneProxy(InMeshDesc, InMeshDesc.GetSkinnedAsset()->GetFName())
 		,	Owner(Cast<AActor>(InMeshDesc.GetOwner()))
 		,	MeshObject(InMeshDesc.MeshObject)
 		,	SkeletalMeshRenderData(InSkelMeshRenderData)
+		,	ClampedLODIndex(MinLODIndex)
 		,	SkeletalMeshForDebug(InMeshDesc.GetSkinnedAsset())
 		,	PhysicsAssetForDebug(InMeshDesc.GetPhysicsAsset())
 		,	OverlayMaterial(InMeshDesc.GetOverlayMaterial())
@@ -476,9 +481,8 @@ void FSkeletalMeshSceneProxy::DrawStaticElements(FStaticPrimitiveDrawInterface*
 		bool bUseSelectedMaterial = false;
 
 		int32 NumLODs = SkeletalMeshRenderData->LODRenderData.Num();
-		int32 ClampedMinLOD = 0; // TODO: MinLOD, Bias?
 
-		for (int32 LODIndex = ClampedMinLOD; LODIndex < NumLODs; ++LODIndex)
+		for (int32 LODIndex = ClampedLODIndex; LODIndex < NumLODs; ++LODIndex)
 		{
 			const FSkeletalMeshLODRenderData& LODData = SkeletalMeshRenderData->LODRenderData[LODIndex];
 
diff --git a/Engine/Source/Runtime/Engine/Private/SkinnedMeshSceneProxyDesc.cpp b/Engine/Source/Runtime/Engine/Private/SkinnedMeshSceneProxyDesc.cpp
--- a/Engine/Source/Runtime/Engine/Private/SkinnedMeshSceneProxyDesc.cpp
+++ b/Engine/Source/Runtime/Engine/Private/SkinnedMeshSceneProxyDesc.cpp
@@ -286,7 +286,7 @@ FPrimitiveSceneProxy* FSkinnedMeshSceneProxyDesc::CreateSceneProxy(const FSkinne
 			}
 			else
 			{
-				Result = ::new FSkeletalMeshSceneProxy(Desc, SkelMeshRenderData);
+				Result = ::new FSkeletalMeshSceneProxy(Desc, SkelMeshRenderData, MinLODIndex);
 			}
 		}
 	}
diff --git a/Engine/Source/Runtime/Engine/Public/SkeletalMeshSceneProxy.h b/Engine/Source/Runtime/Engine/Public/SkeletalMeshSceneProxy.h
--- a/Engine/Source/Runtime/Engine/Public/SkeletalMeshSceneProxy.h
+++ b/Engine/Source/Runtime/Engine/Public/SkeletalMeshSceneProxy.h
@@ -26,7 +26,9 @@ public:
 	 * @param	Component - skeletal mesh primitive being added
 	 */
 	ENGINE_API FSkeletalMeshSceneProxy(const USkinnedMeshComponent* Component, FSkeletalMeshRenderData* InSkelMeshRenderData);
+	UE_DEPRECATED(5.7, "Pass MinLODIndex to FSkeletalMeshSceneProxy constructor.")
 	ENGINE_API FSkeletalMeshSceneProxy(const FSkinnedMeshSceneProxyDesc& InMeshDesc, FSkeletalMeshRenderData* InSkelMeshRenderData);
+	ENGINE_API FSkeletalMeshSceneProxy(const FSkinnedMeshSceneProxyDesc& InMeshDesc, FSkeletalMeshRenderData* InSkelMeshRenderData, int32 MinLODIndex);
 
 	ENGINE_API virtual ~FSkeletalMeshSceneProxy();
 
@@ -136,6 +138,7 @@ protected:
 	AActor* Owner;
 	FSkeletalMeshObject* MeshObject;
 	FSkeletalMeshRenderData* SkeletalMeshRenderData;
+	int32 ClampedLODIndex;
 
 	/** The points to the skeletal mesh and physics assets are purely for debug purposes. Access is NOT thread safe! */
 	const class USkinnedAsset* SkeletalMeshForDebug;
