[FSimpleStreamableAssetManager] Removing double Scale application from TexelFactor

diff --git a/Engine/Source/Runtime/Engine/Private/StaticMeshSceneProxy.cpp b/Engine/Source/Runtime/Engine/Private/StaticMeshSceneProxy.cpp
--- a/Engine/Source/Runtime/Engine/Private/StaticMeshSceneProxy.cpp
+++ b/Engine/Source/Runtime/Engine/Private/StaticMeshSceneProxy.cpp
@@ -559,7 +559,7 @@ void FStaticMeshSceneProxy::GetStreamableRenderAssetInfo(const FBoxSphereBounds&
 				MaterialData.UVChannelData = UVChannelData ? UVChannelData : &Fallback;
 				MaterialData.Material = ShadingMaterial;
 		
-				LevelContext.ProcessMaterial(InPrimitiveBounds, MaterialData, InPrimitiveBounds.SphereRadius, OutStreamableRenderAssets, bIsValidTextureStreamingBuiltData, nullptr);
+				LevelContext.ProcessMaterial(InPrimitiveBounds, MaterialData, 1.0f, OutStreamableRenderAssets, bIsValidTextureStreamingBuiltData, nullptr);
 			}
 		}
 
