Fix crash when baking runtime virtual texture streaming mips.
The old code was doing a direct Release() call on the runtime virtual texture asset which broke the new ref counting.
We do need to temporarily release the RVT producers because they ref the old streaming mips that we are replacing.
Changed the mechanism for that to simple Unregister/Register component calls.

diff --git a/Engine/Source/Runtime/Engine/Private/Components/RuntimeVirtualTextureComponent.cpp b/Engine/Source/Runtime/Engine/Private/Components/RuntimeVirtualTextureComponent.cpp
--- a/Engine/Source/Runtime/Engine/Private/Components/RuntimeVirtualTextureComponent.cpp
+++ b/Engine/Source/Runtime/Engine/Private/Components/RuntimeVirtualTextureComponent.cpp
@@ -456,40 +456,6 @@ FLinearColor URuntimeVirtualTextureComponent::GetStreamingMipsFixedColor() const
 	return Color;
 }
 
-// RAII class to release and recreate runtime virtual texture producers associated with a UVirtualTextureBuilder.
-// Required around modifications of a UVirtualTextureBuilder because virtual producers hold pointers to the internal data.
-class FScopedRuntimeVirtualTextureRecreate
-{
-public:
-	FScopedRuntimeVirtualTextureRecreate(UVirtualTextureBuilder* VirtualTextureBuilder)
-	{
-		for (TObjectIterator<URuntimeVirtualTextureComponent> It(RF_ClassDefaultObject, false, EInternalObjectFlags::Garbage); It; ++It)
-		{
-			if (It->GetStreamingTexture() == VirtualTextureBuilder)
-			{
-				URuntimeVirtualTexture* VirtualTexture = It->GetVirtualTexture();
-				if (VirtualTexture != nullptr)
-				{
-					VirtualTextures.Add(VirtualTexture);
-					VirtualTexture->Release();
-				}
-			}
-		}
-	}
-
-	~FScopedRuntimeVirtualTextureRecreate()
-	{
-		for (URuntimeVirtualTexture* VirtualTexture : VirtualTextures)
-		{
-			// PostEditChange will trigger the correct notifications and recreation of virtual texture producers.
-			VirtualTexture->PostEditChange();
-		}
-	}
-
-private:
-	TArray<URuntimeVirtualTexture*> VirtualTextures;
-};
-
 static TextureCompressionSettings GetCompressionSettingFromLayerFormat(EPixelFormat LayerFormat)
 {
 	switch (LayerFormat)
@@ -516,7 +482,9 @@ void URuntimeVirtualTextureComponent::InitializeStreamingTexture(EShadingPath Sh
 	// We need an existing StreamingTexture object to update.
 	if (IsActiveInWorld() && VirtualTexture != nullptr && StreamingTexture != nullptr)
 	{
-		FScopedRuntimeVirtualTextureRecreate ProducerRecreate(StreamingTexture);
+		// Call Unregister()/Register() around this to release and recreate runtime virtual texture producers associated with a UVirtualTextureBuilder.
+		// Required around modifications of a UVirtualTextureBuilder because virtual producers hold pointers to the internal data.
+		UnregisterComponent();
 
 		FVirtualTextureBuildDesc BuildDesc;
 		BuildDesc.bSinglePhysicalSpace = VirtualTexture->GetSinglePhysicalSpace();
@@ -550,9 +518,9 @@ void URuntimeVirtualTextureComponent::InitializeStreamingTexture(EShadingPath Sh
 		constexpr bool bWaitForCompilation = true;
 		StreamingTexture->BuildTexture(ShadingPath, BuildDesc, bWaitForCompilation);
 		StreamingTexture->Modify();
-	}
 
-	MarkRenderStateDirty();
+		RegisterComponent();
+	}
 }
 
 bool URuntimeVirtualTextureComponent::CanEditChange(const FProperty* InProperty) const
