Fix Soundfield Submix Dynamic Connect related crash.

diff --git a/Engine/Source/Runtime/Engine/Private/SoundSubmix.cpp b/Engine/Source/Runtime/Engine/Private/SoundSubmix.cpp
--- a/Engine/Source/Runtime/Engine/Private/SoundSubmix.cpp
+++ b/Engine/Source/Runtime/Engine/Private/SoundSubmix.cpp
@@ -532,6 +532,13 @@ bool USoundSubmixWithParentBase::DynamicConnect(FAudioDeviceHandle Handle, USoun
 		return false;
 	}
 
+	// Reject anything that's Soundfield related. That system is being retired/
+	if (Cast<USoundfieldSubmix>(this) || Cast<USoundfieldSubmix>(InParent))
+	{
+		UE_LOG(LogAudio, Warning, TEXT("Submix (DynamicConnect): Soundfield submixes do not support dynamic connections. Submix=%s, Parent=%s "), *GetName(), *GetNameSafe(InParent));
+		return false;
+	}
+	
 	if (InParent && !SubmixUtils::AreSubmixFormatsCompatible(this, InParent))
 	{
 		UE_LOG(LogAudio, Warning, TEXT("Submix (DynamicConnect): Submix connot be connected are they are incompatible [%s] and [%s]"), *GetName(), *GetNameSafe(InParent));
