Exporting IsEmpty method to check if world contains any data layers

diff --git a/Engine/Source/Runtime/Engine/Public/WorldPartition/DataLayer/WorldDataLayers.h b/Engine/Source/Runtime/Engine/Public/WorldPartition/DataLayer/WorldDataLayers.h
--- a/Engine/Source/Runtime/Engine/Public/WorldPartition/DataLayer/WorldDataLayers.h
+++ b/Engine/Source/Runtime/Engine/Public/WorldPartition/DataLayer/WorldDataLayers.h
@@ -112,7 +112,7 @@ public:
 	static ENGINE_API AWorldDataLayers* Create(UWorld* World, FName InWorldDataLayerName = NAME_None);
 	static ENGINE_API AWorldDataLayers* Create(const FActorSpawnParameters& SpawnParameters);
 
-	bool IsEmpty() const;
+	ENGINE_API bool IsEmpty() const;
 	bool HasDeprecatedDataLayers() const { return bHasDeprecatedDataLayers; }
 
 	template<class DataLayerInstanceType, typename ...CreationsArgs>
