Fixup API macros for merged modules, PURE_VIRTUAL does not need API export

diff --git a/Engine/Source/Runtime/Engine/Classes/Engine/ReplicationDriver.h b/Engine/Source/Runtime/Engine/Classes/Engine/ReplicationDriver.h
--- a/Engine/Source/Runtime/Engine/Classes/Engine/ReplicationDriver.h
+++ b/Engine/Source/Runtime/Engine/Classes/Engine/ReplicationDriver.h
@@ -62,49 +62,49 @@ public:
 	// -----------------------------------------------------------------------
 
 	/** Called to associate a world with a rep driver. This will be called before  InitForNetDriver */
-	ENGINE_API virtual void SetRepDriverWorld(UWorld* InWorld) PURE_VIRTUAL(UReplicationDriver::SetRepDriverWorld, );
+	virtual void SetRepDriverWorld(UWorld* InWorld) PURE_VIRTUAL(UReplicationDriver::SetRepDriverWorld, );
 
 	/** Called to associate a netdriver with a rep driver. The rep driver can "get itself ready" here. SetRepDriverWorld() will have already been caleld */
-	ENGINE_API virtual void InitForNetDriver(UNetDriver* InNetDriver) PURE_VIRTUAL(UReplicationDriver::InitForNetDriver, );
+	virtual void InitForNetDriver(UNetDriver* InNetDriver) PURE_VIRTUAL(UReplicationDriver::InitForNetDriver, );
 
 	/** Called after World and NetDriver have been set. This is where RepDriver should possibly look at existing actors in the world */
-	ENGINE_API virtual void InitializeActorsInWorld(UWorld* InWorld) PURE_VIRTUAL(UReplicationDriver::InitializeActorsInWorld, );
+	virtual void InitializeActorsInWorld(UWorld* InWorld) PURE_VIRTUAL(UReplicationDriver::InitializeActorsInWorld, );
 
 	virtual void TearDown() { MarkAsGarbage(); }
 
-	ENGINE_API virtual void ResetGameWorldState() PURE_VIRTUAL(UReplicationDriver::ResetGameWorldState, );
+	virtual void ResetGameWorldState() PURE_VIRTUAL(UReplicationDriver::ResetGameWorldState, );
 
-	ENGINE_API virtual void AddClientConnection(UNetConnection* NetConnection) PURE_VIRTUAL(UReplicationDriver::AddClientConnection, );
+	virtual void AddClientConnection(UNetConnection* NetConnection) PURE_VIRTUAL(UReplicationDriver::AddClientConnection, );
 
-	ENGINE_API virtual void RemoveClientConnection(UNetConnection* NetConnection) PURE_VIRTUAL(UReplicationDriver::RemoveClientConnection, );
+	virtual void RemoveClientConnection(UNetConnection* NetConnection) PURE_VIRTUAL(UReplicationDriver::RemoveClientConnection, );
 
-	ENGINE_API virtual void AddNetworkActor(AActor* Actor) PURE_VIRTUAL(UReplicationDriver::AddNetworkActor, );
+	virtual void AddNetworkActor(AActor* Actor) PURE_VIRTUAL(UReplicationDriver::AddNetworkActor, );
 
-	ENGINE_API virtual void RemoveNetworkActor(AActor* Actor) PURE_VIRTUAL(UReplicationDriver::RemoveNetworkActor, );
+	virtual void RemoveNetworkActor(AActor* Actor) PURE_VIRTUAL(UReplicationDriver::RemoveNetworkActor, );
 
-	ENGINE_API virtual void ForceNetUpdate(AActor* Actor) PURE_VIRTUAL(UReplicationDriver::ForceNetUpdate, );
+	virtual void ForceNetUpdate(AActor* Actor) PURE_VIRTUAL(UReplicationDriver::ForceNetUpdate, );
 
-	ENGINE_API virtual void FlushNetDormancy(AActor* Actor, bool WasDormInitial) PURE_VIRTUAL(UReplicationDriver::FlushNetDormancy, );
+	virtual void FlushNetDormancy(AActor* Actor, bool WasDormInitial) PURE_VIRTUAL(UReplicationDriver::FlushNetDormancy, );
 
-	ENGINE_API virtual void NotifyActorTearOff(AActor* Actor) PURE_VIRTUAL(UReplicationDriver::NotifyActorTearOff, );
+	virtual void NotifyActorTearOff(AActor* Actor) PURE_VIRTUAL(UReplicationDriver::NotifyActorTearOff, );
 
-	ENGINE_API virtual void NotifyActorFullyDormantForConnection(AActor* Actor, UNetConnection* Connection) PURE_VIRTUAL(UReplicationDriver::NotifyActorFullyDormantForConnection, );
+	virtual void NotifyActorFullyDormantForConnection(AActor* Actor, UNetConnection* Connection) PURE_VIRTUAL(UReplicationDriver::NotifyActorFullyDormantForConnection, );
 
-	ENGINE_API virtual void NotifyActorDormancyChange(AActor* Actor, ENetDormancy OldDormancyState) PURE_VIRTUAL(UReplicationDriver::NotifyActorDormancyChange, );
+	virtual void NotifyActorDormancyChange(AActor* Actor, ENetDormancy OldDormancyState) PURE_VIRTUAL(UReplicationDriver::NotifyActorDormancyChange, );
 
 	/** Called when a destruction info is created for an actor. Can be used to override some of the destruction info struct */
-	ENGINE_API virtual void NotifyDestructionInfoCreated(AActor* Actor, FActorDestructionInfo& DestructionInfo) PURE_VIRTUAL(UReplicationDriver::NotifyDestructionInfoCreated, );
+	virtual void NotifyDestructionInfoCreated(AActor* Actor, FActorDestructionInfo& DestructionInfo) PURE_VIRTUAL(UReplicationDriver::NotifyDestructionInfoCreated, );
 
 	/** Called when an actor is renamed. The implementation should update any cached actor/level mapping if the actor's outer changed. */
-	ENGINE_API virtual void NotifyActorRenamed(AActor* Actor, UObject* PreviousOuter, FName PreviousName) PURE_VIRTUAL(UReplicationDriver::NotifyActorRenamed, );
+	virtual void NotifyActorRenamed(AActor* Actor, UObject* PreviousOuter, FName PreviousName) PURE_VIRTUAL(UReplicationDriver::NotifyActorRenamed, );
 
-	ENGINE_API virtual void SetRoleSwapOnReplicate(AActor* Actor, bool bSwapRoles) PURE_VIRTUAL(UReplicationDriver::SetRoleSwapOnReplicate, );
+	virtual void SetRoleSwapOnReplicate(AActor* Actor, bool bSwapRoles) PURE_VIRTUAL(UReplicationDriver::SetRoleSwapOnReplicate, );
 
 	/** Handles an RPC. Returns true if it actually handled it. Returning false will cause the rep driver function to handle it instead */
 	virtual bool ProcessRemoteFunction(class AActor* Actor, UFunction* Function, void* Parameters, FOutParmRec* OutParms, FFrame* Stack, UObject* SubObject ) { return false; }
 
 	/** The main function that will actually replicate actors. Called every server tick. */
-	ENGINE_API virtual int32 ServerReplicateActors(float DeltaSeconds) PURE_VIRTUAL(UReplicationDriver::ServerReplicateActors, return 0; );
+	virtual int32 ServerReplicateActors(float DeltaSeconds) PURE_VIRTUAL(UReplicationDriver::ServerReplicateActors, return 0; );
 
 	/** Called after the netdriver has handled TickDispatch */
 	virtual void PostTickDispatch() { }
@@ -118,23 +118,23 @@ class UReplicationConnectionDriver : public UObject
 
 public:
 
-	ENGINE_API virtual void NotifyActorChannelAdded(AActor* Actor, UActorChannel* Channel) PURE_VIRTUAL(UReplicationConnectionDriver::NotifyActorChannelAdded, );
+	virtual void NotifyActorChannelAdded(AActor* Actor, UActorChannel* Channel) PURE_VIRTUAL(UReplicationConnectionDriver::NotifyActorChannelAdded, );
 
-	ENGINE_API virtual void NotifyActorChannelRemoved(AActor* Actor) PURE_VIRTUAL(UReplicationConnectionDriver::NotifyActorChannelRemoved, );
+	virtual void NotifyActorChannelRemoved(AActor* Actor) PURE_VIRTUAL(UReplicationConnectionDriver::NotifyActorChannelRemoved, );
 
-	ENGINE_API virtual void NotifyActorChannelCleanedUp(UActorChannel* Channel) PURE_VIRTUAL(UReplicationConnectionDriver::NotifyActorChannelCleanedUp, );
+	virtual void NotifyActorChannelCleanedUp(UActorChannel* Channel) PURE_VIRTUAL(UReplicationConnectionDriver::NotifyActorChannelCleanedUp, );
 
-	ENGINE_API virtual void NotifyAddDestructionInfo(FActorDestructionInfo* DestructInfo) PURE_VIRTUAL(UReplicationConnectionDriver::NotifyAddDestructionInfo, );
+	virtual void NotifyAddDestructionInfo(FActorDestructionInfo* DestructInfo) PURE_VIRTUAL(UReplicationConnectionDriver::NotifyAddDestructionInfo, );
 
-	ENGINE_API virtual void NotifyAddDormantDestructionInfo(AActor* Actor) PURE_VIRTUAL(UReplicationConnectionDriver::NotifyAddDormantDestructionInfo, );
+	virtual void NotifyAddDormantDestructionInfo(AActor* Actor) PURE_VIRTUAL(UReplicationConnectionDriver::NotifyAddDormantDestructionInfo, );
 
-	ENGINE_API virtual void NotifyRemoveDestructionInfo(FActorDestructionInfo* DestructInfo) PURE_VIRTUAL(UReplicationConnectionDriver::NotifyRemoveDestructionInfo, );
+	virtual void NotifyRemoveDestructionInfo(FActorDestructionInfo* DestructInfo) PURE_VIRTUAL(UReplicationConnectionDriver::NotifyRemoveDestructionInfo, );
 
-	ENGINE_API virtual void NotifyResetDestructionInfo() PURE_VIRTUAL(UReplicationConnectionDriver::NotifyResetDestructionInfo, );
+	virtual void NotifyResetDestructionInfo() PURE_VIRTUAL(UReplicationConnectionDriver::NotifyResetDestructionInfo, );
 
-	ENGINE_API virtual void NotifyClientVisibleLevelNamesAdd(FName LevelName, UWorld* StreamingWorld) PURE_VIRTUAL(UReplicationConnectionDriver::NotifyClientVisibleLevelNamesAdd, );
+	virtual void NotifyClientVisibleLevelNamesAdd(FName LevelName, UWorld* StreamingWorld) PURE_VIRTUAL(UReplicationConnectionDriver::NotifyClientVisibleLevelNamesAdd, );
 
-	ENGINE_API virtual void NotifyClientVisibleLevelNamesRemove(FName LevelName) PURE_VIRTUAL(UReplicationConnectionDriver::NotifyClientVisibleLevelNamesRemove, );
+	virtual void NotifyClientVisibleLevelNamesRemove(FName LevelName) PURE_VIRTUAL(UReplicationConnectionDriver::NotifyClientVisibleLevelNamesRemove, );
 
 	virtual void TearDown() { MarkAsGarbage(); }
 };
diff --git a/Engine/Source/Runtime/Engine/Classes/Engine/SkinnedAsset.h b/Engine/Source/Runtime/Engine/Classes/Engine/SkinnedAsset.h
--- a/Engine/Source/Runtime/Engine/Classes/Engine/SkinnedAsset.h
+++ b/Engine/Source/Runtime/Engine/Classes/Engine/SkinnedAsset.h
@@ -50,15 +50,15 @@ public:
 	ENGINE_API virtual ~USkinnedAsset();
 
 	/** Return the reference skeleton. */
-	ENGINE_API virtual struct FReferenceSkeleton& GetRefSkeleton()
+	virtual struct FReferenceSkeleton& GetRefSkeleton()
 	PURE_VIRTUAL(USkinnedAsset::GetRefSkeleton, static FReferenceSkeleton Dummy; return Dummy;);
-	ENGINE_API virtual const struct FReferenceSkeleton& GetRefSkeleton() const
+	virtual const struct FReferenceSkeleton& GetRefSkeleton() const
 	PURE_VIRTUAL(USkinnedAsset::GetRefSkeleton, static const FReferenceSkeleton Dummy; return Dummy;);
 
 	/** Return the LOD information for the specified LOD index. */
-	ENGINE_API virtual FSkeletalMeshLODInfo* GetLODInfo(int32 Index)
+	virtual FSkeletalMeshLODInfo* GetLODInfo(int32 Index)
 	PURE_VIRTUAL(USkinnedAsset::GetLODInfo, return nullptr;);
-	ENGINE_API virtual const FSkeletalMeshLODInfo* GetLODInfo(int32 Index) const
+	virtual const FSkeletalMeshLODInfo* GetLODInfo(int32 Index) const
 	PURE_VIRTUAL(USkinnedAsset::GetLODInfo, return nullptr;);
 
 	/** Return if the material index is valid. */
@@ -68,15 +68,15 @@ public:
 	ENGINE_API virtual int32 GetNumMaterials() const;
 
 	/** Return the physics asset whose shapes will be used for shadowing. */
-	ENGINE_API virtual UPhysicsAsset* GetShadowPhysicsAsset() const
+	virtual UPhysicsAsset* GetShadowPhysicsAsset() const
 	PURE_VIRTUAL(USkinnedAsset::GetShadowPhysicsAsset, return nullptr;);
 
 	/** Return the component orientation of a bone or socket. */
-	ENGINE_API virtual FMatrix GetComposedRefPoseMatrix(FName InBoneName) const
+	virtual FMatrix GetComposedRefPoseMatrix(FName InBoneName) const
 	PURE_VIRTUAL(USkinnedAsset::GetComposedRefPoseMatrix, return FMatrix::Identity;);
 
 	/** Return the component orientation of a bone or socket. */
-	ENGINE_API virtual FMatrix GetComposedRefPoseMatrix(int32 InBoneIndex) const
+	virtual FMatrix GetComposedRefPoseMatrix(int32 InBoneIndex) const
 	PURE_VIRTUAL(USkinnedAsset::GetComposedRefPoseMatrix, return FMatrix::Identity;);
 
 	/**
@@ -86,66 +86,66 @@ public:
 	 * @param MaterialIndex		the material index for which to get the data for.
 	 * @return the data, or null if none exists.
 	 */
-	ENGINE_API virtual const FMeshUVChannelInfo* GetUVChannelData(int32 MaterialIndex) const
+	virtual const FMeshUVChannelInfo* GetUVChannelData(int32 MaterialIndex) const
 	PURE_VIRTUAL(USkinnedAsset::GetUVChannelData, return nullptr;);
 
 	/** Return whether ray tracing is supported on this mesh. */
-	ENGINE_API virtual bool GetSupportRayTracing() const
+	virtual bool GetSupportRayTracing() const
 	PURE_VIRTUAL(USkinnedAsset::GetSupportRayTracing, return false;);
 
 	/** Return the minimum ray tracing LOD of this mesh. */
-	ENGINE_API virtual int32 GetRayTracingMinLOD() const
+	virtual int32 GetRayTracingMinLOD() const
 	PURE_VIRTUAL(USkinnedAsset::GetRayTracingMinLOD, return 0;);
 
 	/** Return the reference skeleton precomputed bases. */
-	ENGINE_API virtual TArray<FMatrix44f>& GetRefBasesInvMatrix()
+	virtual TArray<FMatrix44f>& GetRefBasesInvMatrix()
 	PURE_VIRTUAL(USkinnedAsset::GetRefBasesInvMatrix, static TArray<FMatrix44f> Dummy; return Dummy;);
-	ENGINE_API virtual const TArray<FMatrix44f>& GetRefBasesInvMatrix() const
+	virtual const TArray<FMatrix44f>& GetRefBasesInvMatrix() const
 	PURE_VIRTUAL(USkinnedAsset::GetRefBasesInvMatrix, static const TArray<FMatrix44f> Dummy; return Dummy;);
 
 	/** Return the whole array of LOD info. */
 	UE_DEPRECATED(5.5, "Use GetLODInfo and GetLODNum instead.")
-	ENGINE_API virtual TArray<FSkeletalMeshLODInfo>& GetLODInfoArray()
+	virtual TArray<FSkeletalMeshLODInfo>& GetLODInfoArray()
 	PURE_VIRTUAL(USkinnedAsset::GetLODInfoArray, return GetMeshLodInfoDummyArray(););
 	UE_DEPRECATED(5.5, "Use GetLODInfo and GetLODNum instead.")
-	ENGINE_API virtual const TArray<FSkeletalMeshLODInfo>& GetLODInfoArray() const
+	virtual const TArray<FSkeletalMeshLODInfo>& GetLODInfoArray() const
 	PURE_VIRTUAL(USkinnedAsset::GetLODInfoArray, return GetMeshLodInfoDummyArray(););
 
 	/** Get the data to use for rendering. */
-	ENGINE_API virtual class FSkeletalMeshRenderData* GetResourceForRendering() const
+	virtual class FSkeletalMeshRenderData* GetResourceForRendering() const
 	PURE_VIRTUAL(USkinnedAsset::GetResourceForRendering, return nullptr;);
 
-	ENGINE_API virtual int32 GetDefaultMinLod() const
+	virtual int32 GetDefaultMinLod() const
 	PURE_VIRTUAL(USkinnedAsset::GetDefaultMinLod, return 0;);
 
-	ENGINE_API virtual const FPerPlatformInt& GetMinLod() const
+	virtual const FPerPlatformInt& GetMinLod() const
 	PURE_VIRTUAL(USkinnedAsset::GetMinLod, static const FPerPlatformInt Dummy;  return Dummy;);
 
 	/** Check the QualitLevel property is enabled for MinLod. */
 	virtual bool IsMinLodQualityLevelEnable() const { return false; }
 
-	ENGINE_API virtual UPhysicsAsset* GetPhysicsAsset() const
+	virtual UPhysicsAsset* GetPhysicsAsset() const
 	PURE_VIRTUAL(USkinnedAsset::GetPhysicsAsset, return nullptr;);
 
-	ENGINE_API virtual TArray<FSkeletalMaterial>& GetMaterials()
+	virtual TArray<FSkeletalMaterial>& GetMaterials()
 	PURE_VIRTUAL(USkinnedAsset::GetMaterials, return GetSkeletalMaterialDummyArray(););
-	ENGINE_API virtual const TArray<FSkeletalMaterial>& GetMaterials() const
+	virtual const TArray<FSkeletalMaterial>& GetMaterials() const
 	PURE_VIRTUAL(USkinnedAsset::GetMaterials, return GetSkeletalMaterialDummyArray(););
 
-	ENGINE_API virtual int32 GetLODNum() const
+	virtual int32 GetLODNum() const
 	PURE_VIRTUAL(USkinnedAsset::GetLODNum, return 0;);
 
-	ENGINE_API virtual bool IsMaterialUsed(int32 MaterialIndex) const
+	virtual bool IsMaterialUsed(int32 MaterialIndex) const
 	PURE_VIRTUAL(USkinnedAsset::IsMaterialUsed, return false;);
 
-	ENGINE_API virtual FBoxSphereBounds GetBounds() const
+	virtual FBoxSphereBounds GetBounds() const
 	PURE_VIRTUAL(USkinnedAsset::GetBounds, return FBoxSphereBounds(););
 
 	/**
 	 * Returns the "active" socket list - all sockets from this mesh plus all non-duplicates from the skeleton
 	 * Const ref return value as this cannot be modified externally
 	 */
-	ENGINE_API virtual TArray<class USkeletalMeshSocket*> GetActiveSocketList() const
+	virtual TArray<class USkeletalMeshSocket*> GetActiveSocketList() const
 	PURE_VIRTUAL(USkinnedAsset::GetActiveSocketList, static TArray<class USkeletalMeshSocket*> Dummy; return Dummy;);
 
 	/**
@@ -153,7 +153,7 @@ public:
 	 * Entering NAME_None will return NULL. If there are multiple sockets with the same name, will return the first one.
 	 */
 	UFUNCTION(BlueprintCallable, Category = "Animation")
-	ENGINE_API virtual USkeletalMeshSocket* FindSocket(FName InSocketName) const
+	virtual USkeletalMeshSocket* FindSocket(FName InSocketName) const
 	PURE_VIRTUAL(USkinnedAsset::FindSocket, return nullptr;);
 
 	/**
@@ -163,35 +163,35 @@ public:
 	 * Also returns the socket transform and the bone index (if any)
 	 */
 	UFUNCTION(BlueprintCallable, Category = "Animation")
-	ENGINE_API virtual USkeletalMeshSocket* FindSocketInfo(FName InSocketName, FTransform& OutTransform, int32& OutBoneIndex, int32& OutIndex) const
+	virtual USkeletalMeshSocket* FindSocketInfo(FName InSocketName, FTransform& OutTransform, int32& OutBoneIndex, int32& OutIndex) const
 	PURE_VIRTUAL(USkinnedAsset::FindSocketInfo, return nullptr;);
 
-	ENGINE_API virtual USkeleton* GetSkeleton()
+	virtual USkeleton* GetSkeleton()
 	PURE_VIRTUAL(USkinnedAsset::GetSkeleton, return nullptr;);
-	ENGINE_API virtual const USkeleton* GetSkeleton() const
+	virtual const USkeleton* GetSkeleton() const
 	PURE_VIRTUAL(USkinnedAsset::GetSkeleton, return nullptr;);
-	ENGINE_API virtual void SetSkeleton(USkeleton* InSkeleton)
+	virtual void SetSkeleton(USkeleton* InSkeleton)
 	PURE_VIRTUAL(USkinnedAsset::SetSkeleton,);
 
-	ENGINE_API virtual UMeshDeformer* GetDefaultMeshDeformer() const
+	virtual UMeshDeformer* GetDefaultMeshDeformer() const
 	PURE_VIRTUAL(USkinnedAsset::GetDefaultMeshDeformer, return nullptr;);
 	
-	ENGINE_API virtual UMeshDeformerCollection* GetTargetMeshDeformers() const
+	virtual UMeshDeformerCollection* GetTargetMeshDeformers() const
 	PURE_VIRTUAL(USkinnedAsset::GetTargetMeshDeformers, return nullptr;);
 	
-	ENGINE_API virtual bool HasHalfEdgeBuffer(int32 LODIndex) const
+	virtual bool HasHalfEdgeBuffer(int32 LODIndex) const
 	PURE_VIRTUAL(USkinnedAsset::HasHalfEdgeBuffer, return false;);
 	
-	ENGINE_API virtual class UMaterialInterface* GetOverlayMaterial() const
+	virtual class UMaterialInterface* GetOverlayMaterial() const
 	PURE_VIRTUAL(USkinnedAsset::GetOverlayMaterial, return nullptr;);
-	ENGINE_API virtual float GetOverlayMaterialMaxDrawDistance() const
+	virtual float GetOverlayMaterialMaxDrawDistance() const
 	PURE_VIRTUAL(USkinnedAsset::GetOverlayMaterialMaxDrawDistance, return 0.f;);
 	
 	/** Return true if given index's LOD is valid */
-	ENGINE_API virtual bool IsValidLODIndex(int32 Index) const
+	virtual bool IsValidLODIndex(int32 Index) const
 	PURE_VIRTUAL(USkinnedAsset::IsValidLODIndex, return false;);
 
-	ENGINE_API virtual int32 GetMinLodIdx(bool bForceLowestLODIdx = false) const
+	virtual int32 GetMinLodIdx(bool bForceLowestLODIdx = false) const
 	PURE_VIRTUAL(USkinnedAsset::GetMinLodIdx, return 0;);
 
 	/** Return the morph targets. */
@@ -205,17 +205,17 @@ public:
 	{ return nullptr; }
 
 	/** True if this mesh LOD needs to keep it's data on CPU. USkinnedAsset interface. */
-	ENGINE_API virtual bool NeedCPUData(int32 LODIndex) const
+	virtual bool NeedCPUData(int32 LODIndex) const
 	PURE_VIRTUAL(USkinnedAsset::NeedCPUData, return false;);
 
 	/** Return whether or not the mesh has vertex colors. */
-	ENGINE_API virtual bool GetHasVertexColors() const
+	virtual bool GetHasVertexColors() const
 	PURE_VIRTUAL(USkinnedAsset::GetHasVertexColors, return false;);
 
-	ENGINE_API virtual int32 GetPlatformMinLODIdx(const ITargetPlatform* TargetPlatform) const
+	virtual int32 GetPlatformMinLODIdx(const ITargetPlatform* TargetPlatform) const
 	PURE_VIRTUAL(USkinnedAsset::GetPlatformMinLODIdx, return 0;);
 
-	ENGINE_API virtual const FPerPlatformBool& GetDisableBelowMinLodStripping() const
+	virtual const FPerPlatformBool& GetDisableBelowMinLodStripping() const
 	PURE_VIRTUAL(USkinnedAsset::GetDisableBelowMinLodStripping, static const FPerPlatformBool Dummy; return Dummy;);
 
 	virtual void SetSkinWeightProfilesData(int32 LODIndex, struct FSkinWeightProfilesData& SkinWeightProfilesData) {}
@@ -258,11 +258,11 @@ public:
 	/** IInterface_AsyncCompilation end*/
 
 #if WITH_EDITOR
-	ENGINE_API virtual FString BuildDerivedDataKey(const ITargetPlatform* TargetPlatform)
+	virtual FString BuildDerivedDataKey(const ITargetPlatform* TargetPlatform)
 	PURE_VIRTUAL(USkinnedAsset::BuildDerivedDataKey, return TEXT(""););
 
 	/* Return true if this asset was never build since its creation. */
-	ENGINE_API virtual bool IsInitialBuildDone() const
+	virtual bool IsInitialBuildDone() const
 	PURE_VIRTUAL(USkinnedAsset::IsInitialBuildDone, return false;);
 
 	/* Build a LOD model before creating its render data. */
@@ -273,14 +273,14 @@ public:
 	virtual void BuildLODModel(FSkeletalMeshRenderData& RenderData, const ITargetPlatform* TargetPlatform, int32 LODIndex) {}
 
 	/** Get whether this mesh should use LOD streaming for the given platform. */
-	ENGINE_API virtual bool GetEnableLODStreaming(const ITargetPlatform* TargetPlatform) const
+	virtual bool GetEnableLODStreaming(const ITargetPlatform* TargetPlatform) const
 	PURE_VIRTUAL(USkinnedAsset::GetEnableLODStreaming, return false;);
 
 	/** Get the maximum number of LODs that can be streamed. */
-	ENGINE_API virtual int32 GetMaxNumStreamedLODs(const ITargetPlatform* TargetPlatform) const
+	virtual int32 GetMaxNumStreamedLODs(const ITargetPlatform* TargetPlatform) const
 	PURE_VIRTUAL(USkinnedAsset::GetMaxNumStreamedLODs, return 0;);
 
-	ENGINE_API virtual int32 GetMaxNumOptionalLODs(const ITargetPlatform* TargetPlatform) const
+	virtual int32 GetMaxNumOptionalLODs(const ITargetPlatform* TargetPlatform) const
 	PURE_VIRTUAL(USkinnedAsset::GetMaxNumOptionalLODs, return 0;);
 #endif // WITH_EDITOR
 
@@ -289,7 +289,7 @@ public:
 	static bool GetUseLegacyMeshDerivedDataKey()	{ return false; }
 
 	/** Get the source mesh data. */
-	ENGINE_API virtual class FSkeletalMeshModel* GetImportedModel() const
+	virtual class FSkeletalMeshModel* GetImportedModel() const
 	PURE_VIRTUAL(USkinnedAsset::GetImportedModel, return nullptr;);
 #endif // WITH_EDITORONLY_DATA
 
diff --git a/Engine/Source/Runtime/Engine/Classes/Engine/Texture.h b/Engine/Source/Runtime/Engine/Classes/Engine/Texture.h
--- a/Engine/Source/Runtime/Engine/Classes/Engine/Texture.h
+++ b/Engine/Source/Runtime/Engine/Classes/Engine/Texture.h
@@ -1675,7 +1675,7 @@ public:
 	ENGINE_API virtual void PostEditImport() override;
 
 	/** Get Texture Class */
-	ENGINE_API virtual ETextureClass GetTextureClass() const PURE_VIRTUAL(UTexture::GetTextureClass, return ETextureClass::Invalid; );
+	virtual ETextureClass GetTextureClass() const PURE_VIRTUAL(UTexture::GetTextureClass, return ETextureClass::Invalid; );
 
 	/**
 	 * Resets the resource for the texture.
diff --git a/Engine/Source/Runtime/Engine/Classes/Engine/TextureAllMipDataProviderFactory.h b/Engine/Source/Runtime/Engine/Classes/Engine/TextureAllMipDataProviderFactory.h
--- a/Engine/Source/Runtime/Engine/Classes/Engine/TextureAllMipDataProviderFactory.h
+++ b/Engine/Source/Runtime/Engine/Classes/Engine/TextureAllMipDataProviderFactory.h
@@ -44,12 +44,12 @@ public:
 	 * @param DebugContext - A string used for debug tracking and logging. Usually Texture->GetPathName()
 	 * @returns true if the requested mip data has been successfully returned.
 	 */
-	ENGINE_API virtual bool GetInitialMipData(int32 FirstMipToLoad, TArrayView<void*> OutMipData, TArrayView<int64> OutMipSize, FStringView DebugContext)
+	virtual bool GetInitialMipData(int32 FirstMipToLoad, TArrayView<void*> OutMipData, TArrayView<int64> OutMipSize, FStringView DebugContext)
 		PURE_VIRTUAL(UTextureAllMipDataProviderFactory::GetInitialMipData, return false;);
 
 	/**
 	  * Get the initial streaming state (after texture is first loaded)
 	  */
-	ENGINE_API virtual FStreamableRenderResourceState GetResourcePostInitState(const UTexture* Owner, bool bAllowStreaming)
+	virtual FStreamableRenderResourceState GetResourcePostInitState(const UTexture* Owner, bool bAllowStreaming)
 		PURE_VIRTUAL(UTextureAllMipDataProviderFactory::GetResourcePostInitState, return FStreamableRenderResourceState(););
 };
diff --git a/Engine/Source/Runtime/Engine/Classes/Engine/TextureRenderTarget.h b/Engine/Source/Runtime/Engine/Classes/Engine/TextureRenderTarget.h
--- a/Engine/Source/Runtime/Engine/Classes/Engine/TextureRenderTarget.h
+++ b/Engine/Source/Runtime/Engine/Classes/Engine/TextureRenderTarget.h
@@ -62,7 +62,7 @@ class UTextureRenderTarget : public UTexture
 	 * @param OutErrorMessage - (Optional) If passed, will contain an error message in case conversion is impossible
 	 * @return True if render target can be converted to a texture
 	 */
-	ENGINE_API virtual bool CanConvertToTexture(ETextureSourceFormat& OutTextureSourceFormat, EPixelFormat& OutPixelFormat, FText* OutErrorMessage = nullptr) const PURE_VIRTUAL(UTextureRenderTarget, return false;);
+	virtual bool CanConvertToTexture(ETextureSourceFormat& OutTextureSourceFormat, EPixelFormat& OutPixelFormat, FText* OutErrorMessage = nullptr) const PURE_VIRTUAL(UTextureRenderTarget, return false;);
 
 	/**
 	 * Validates that the UTextureRenderTarget can be converted to a UTexture (e.g. supported format, valid size, etc.)
@@ -74,13 +74,13 @@ class UTextureRenderTarget : public UTexture
 	/**
 	 * Returns the UTexture class that corresponds to this render target (e.g. UTexture2D for UTextureRenderTarget2D)
 	 */
-	ENGINE_API virtual TSubclassOf<UTexture> GetTextureUClass() const PURE_VIRTUAL(GetTextureUClass, return nullptr;)
+	virtual TSubclassOf<UTexture> GetTextureUClass() const PURE_VIRTUAL(GetTextureUClass, return nullptr;)
 	// GetTextureClass() will just return "RenderTarget" ; to get the sub-type (2d/cube), use GetRenderTargetTextureClass
-	ENGINE_API virtual ETextureClass GetRenderTargetTextureClass() const PURE_VIRTUAL(GetRenderTargetTextureClass, return ETextureClass::Invalid;)
+	virtual ETextureClass GetRenderTargetTextureClass() const PURE_VIRTUAL(GetRenderTargetTextureClass, return ETextureClass::Invalid;)
 
-	ENGINE_API virtual EPixelFormat GetFormat() const PURE_VIRTUAL(GetFormat,return PF_Unknown;)
-	ENGINE_API virtual bool IsSRGB() const PURE_VIRTUAL(IsSRGB,return false;)
-	ENGINE_API virtual float GetDisplayGamma() const PURE_VIRTUAL(GetDisplayGamma,return 0.f;)
+	virtual EPixelFormat GetFormat() const PURE_VIRTUAL(GetFormat,return PF_Unknown;)
+	virtual bool IsSRGB() const PURE_VIRTUAL(IsSRGB,return false;)
+	virtual float GetDisplayGamma() const PURE_VIRTUAL(GetDisplayGamma,return 0.f;)
 
 	// UTextureRenderTarget default display gamma if none is set
 	//	returns hard-coded 2.2
diff --git a/Engine/Source/Runtime/Engine/Classes/Engine/TimecodeProvider.h b/Engine/Source/Runtime/Engine/Classes/Engine/TimecodeProvider.h
--- a/Engine/Source/Runtime/Engine/Classes/Engine/TimecodeProvider.h
+++ b/Engine/Source/Runtime/Engine/Classes/Engine/TimecodeProvider.h
@@ -68,7 +68,7 @@ public:
 	 * Since it may be called several times per frame, it is suggested to return a cached value.
 	*/
 	UFUNCTION(BlueprintCallable, Category = "Provider")
-	ENGINE_API virtual FQualifiedFrameTime GetQualifiedFrameTime() const PURE_VIRTUAL(UTimecodeProvider::GetQualifiedFrameTime, return FQualifiedFrameTime(););
+	virtual FQualifiedFrameTime GetQualifiedFrameTime() const PURE_VIRTUAL(UTimecodeProvider::GetQualifiedFrameTime, return FQualifiedFrameTime(););
 
 	/**
 	 * Return current frame time with FrameDelay applied.
@@ -91,14 +91,14 @@ public:
 
 	/** The state of the TimecodeProvider and if it's currently synchronized and the Timecode and FrameRate getters are valid. */
 	UFUNCTION(BlueprintCallable, Category = "Provider")
-	ENGINE_API virtual ETimecodeProviderSynchronizationState GetSynchronizationState() const PURE_VIRTUAL(UTimecodeProvider::IsSynchronized, return ETimecodeProviderSynchronizationState::Closed;);
+	virtual ETimecodeProviderSynchronizationState GetSynchronizationState() const PURE_VIRTUAL(UTimecodeProvider::IsSynchronized, return ETimecodeProviderSynchronizationState::Closed;);
 
 public:
 	/** This Provider became the Engine's Provider. */
-	ENGINE_API virtual bool Initialize(class UEngine* InEngine) PURE_VIRTUAL(UTimecodeProvider::Initialize, return false;);
+	virtual bool Initialize(class UEngine* InEngine) PURE_VIRTUAL(UTimecodeProvider::Initialize, return false;);
 
 	/** This Provider stopped being the Engine's Provider. */
-	ENGINE_API virtual void Shutdown(class UEngine* InEngine) PURE_VIRTUAL(UTimecodeProvider::Shutdown, );
+	virtual void Shutdown(class UEngine* InEngine) PURE_VIRTUAL(UTimecodeProvider::Shutdown, );
 
 	/** Whether this provider supports format autodetection. */
 	virtual bool SupportsAutoDetected() const
diff --git a/Engine/Source/Runtime/Engine/Classes/GameFramework/LocalMessage.h b/Engine/Source/Runtime/Engine/Classes/GameFramework/LocalMessage.h
--- a/Engine/Source/Runtime/Engine/Classes/GameFramework/LocalMessage.h
+++ b/Engine/Source/Runtime/Engine/Classes/GameFramework/LocalMessage.h
@@ -58,7 +58,7 @@ class ULocalMessage : public UObject
 {
 	GENERATED_UCLASS_BODY()
 	/** send message to client */
-	ENGINE_API virtual void ClientReceive(const FClientReceiveData& ClientData) const PURE_VIRTUAL(ULocalMessage::ClientReceive,);
+	virtual void ClientReceive(const FClientReceiveData& ClientData) const PURE_VIRTUAL(ULocalMessage::ClientReceive,);
 };
 
 
diff --git a/Engine/Source/Runtime/Engine/Classes/Interfaces/Interface_BoneReferenceSkeletonProvider.h b/Engine/Source/Runtime/Engine/Classes/Interfaces/Interface_BoneReferenceSkeletonProvider.h
--- a/Engine/Source/Runtime/Engine/Classes/Interfaces/Interface_BoneReferenceSkeletonProvider.h
+++ b/Engine/Source/Runtime/Engine/Classes/Interfaces/Interface_BoneReferenceSkeletonProvider.h
@@ -32,5 +32,5 @@ public:
 	 *
 	 * @return The skeleton we should use.
 	 */
-	ENGINE_API virtual class USkeleton* GetSkeleton(bool& bInvalidSkeletonIsError, const class IPropertyHandle* PropertyHandle) PURE_VIRTUAL(IBoneReferenceSkeletonProvider::GetSkeleton, return nullptr; );
+	virtual class USkeleton* GetSkeleton(bool& bInvalidSkeletonIsError, const class IPropertyHandle* PropertyHandle) PURE_VIRTUAL(IBoneReferenceSkeletonProvider::GetSkeleton, return nullptr; );
 };
diff --git a/Engine/Source/Runtime/Engine/Classes/Interfaces/NetworkPredictionInterface.h b/Engine/Source/Runtime/Engine/Classes/Interfaces/NetworkPredictionInterface.h
--- a/Engine/Source/Runtime/Engine/Classes/Interfaces/NetworkPredictionInterface.h
+++ b/Engine/Source/Runtime/Engine/Classes/Interfaces/NetworkPredictionInterface.h
@@ -27,39 +27,39 @@ class INetworkPredictionInterface
 	//--------------------------------
 
 	/** (Server) Send position to client if necessary, or just ack good moves. */
-	ENGINE_API virtual void SendClientAdjustment()					PURE_VIRTUAL(INetworkPredictionInterface::SendClientAdjustment,);
+	virtual void SendClientAdjustment()					PURE_VIRTUAL(INetworkPredictionInterface::SendClientAdjustment,);
 
 	/** (Server) Trigger a position update on clients, if the server hasn't heard from them in a while. @return Whether movement is performed. */
-	ENGINE_API virtual bool ForcePositionUpdate(float DeltaTime)	PURE_VIRTUAL(INetworkPredictionInterface::ForcePositionUpdate, return false;);
+	virtual bool ForcePositionUpdate(float DeltaTime)	PURE_VIRTUAL(INetworkPredictionInterface::ForcePositionUpdate, return false;);
 
 	//--------------------------------
 	// Client hooks
 	//--------------------------------
 
 	/** (Client) After receiving a network update of position, allow some custom smoothing, given the old transform before the correction and new transform from the update. */
-	ENGINE_API virtual void SmoothCorrection(const FVector& OldLocation, const FQuat& OldRotation, const FVector& NewLocation, const FQuat& NewRotation) PURE_VIRTUAL(INetworkPredictionInterface::SmoothCorrection,);
+	virtual void SmoothCorrection(const FVector& OldLocation, const FQuat& OldRotation, const FVector& NewLocation, const FQuat& NewRotation) PURE_VIRTUAL(INetworkPredictionInterface::SmoothCorrection,);
 
 	//--------------------------------
 	// Other
 	//--------------------------------
 
 	/** @return FNetworkPredictionData_Client instance used for network prediction. */
-	ENGINE_API virtual class FNetworkPredictionData_Client* GetPredictionData_Client() const PURE_VIRTUAL(INetworkPredictionInterface::GetPredictionData_Client, return NULL;);
+	virtual class FNetworkPredictionData_Client* GetPredictionData_Client() const PURE_VIRTUAL(INetworkPredictionInterface::GetPredictionData_Client, return NULL;);
 	
 	/** @return FNetworkPredictionData_Server instance used for network prediction. */
-	ENGINE_API virtual class FNetworkPredictionData_Server* GetPredictionData_Server() const PURE_VIRTUAL(INetworkPredictionInterface::GetPredictionData_Server, return NULL;);
+	virtual class FNetworkPredictionData_Server* GetPredictionData_Server() const PURE_VIRTUAL(INetworkPredictionInterface::GetPredictionData_Server, return NULL;);
 
 	/** Accessor to check if there is already client data, without potentially allocating it on demand.*/
-	ENGINE_API virtual bool HasPredictionData_Client() const PURE_VIRTUAL(INetworkPredictionInterface::HasPredictionData_Client, return false;);
+	virtual bool HasPredictionData_Client() const PURE_VIRTUAL(INetworkPredictionInterface::HasPredictionData_Client, return false;);
 
 	/** Accessor to check if there is already server data, without potentially allocating it on demand.*/
-	ENGINE_API virtual bool HasPredictionData_Server() const PURE_VIRTUAL(INetworkPredictionInterface::HasPredictionData_Server, return false;);
+	virtual bool HasPredictionData_Server() const PURE_VIRTUAL(INetworkPredictionInterface::HasPredictionData_Server, return false;);
 
 	/** Resets client prediction data. */
-	ENGINE_API virtual void ResetPredictionData_Client() PURE_VIRTUAL(INetworkPredictionInterface::ResetPredictionData_Client,);
+	virtual void ResetPredictionData_Client() PURE_VIRTUAL(INetworkPredictionInterface::ResetPredictionData_Client,);
 
 	/** Resets server prediction data. */
-	ENGINE_API virtual void ResetPredictionData_Server() PURE_VIRTUAL(INetworkPredictionInterface::ResetPredictionData_Server,);
+	virtual void ResetPredictionData_Server() PURE_VIRTUAL(INetworkPredictionInterface::ResetPredictionData_Server,);
 };
 
 
diff --git a/Engine/Source/Runtime/Engine/Classes/Sound/SoundEffectPreset.h b/Engine/Source/Runtime/Engine/Classes/Sound/SoundEffectPreset.h
--- a/Engine/Source/Runtime/Engine/Classes/Sound/SoundEffectPreset.h
+++ b/Engine/Source/Runtime/Engine/Classes/Sound/SoundEffectPreset.h
@@ -30,12 +30,12 @@ public:
 
 
 	virtual bool CanFilter() const { return true; }
-	ENGINE_API virtual FText GetAssetActionName() const PURE_VIRTUAL(USoundEffectPreset::GetAssetActionName, return FText(););
-	ENGINE_API virtual UClass* GetSupportedClass() const PURE_VIRTUAL(USoundEffectPreset::GetSupportedClass, return nullptr;);
-	ENGINE_API virtual USoundEffectPreset* CreateNewPreset(UObject* InParent, FName Name, EObjectFlags Flags) const PURE_VIRTUAL(USoundEffectPreset::CreateNewPreset, return nullptr;);
-	ENGINE_API virtual FSoundEffectBase* CreateNewEffect() const PURE_VIRTUAL(USoundEffectPreset::CreateNewEffect, return nullptr;);
+	virtual FText GetAssetActionName() const PURE_VIRTUAL(USoundEffectPreset::GetAssetActionName, return FText(););
+	virtual UClass* GetSupportedClass() const PURE_VIRTUAL(USoundEffectPreset::GetSupportedClass, return nullptr;);
+	virtual USoundEffectPreset* CreateNewPreset(UObject* InParent, FName Name, EObjectFlags Flags) const PURE_VIRTUAL(USoundEffectPreset::CreateNewPreset, return nullptr;);
+	virtual FSoundEffectBase* CreateNewEffect() const PURE_VIRTUAL(USoundEffectPreset::CreateNewEffect, return nullptr;);
 	virtual bool HasAssetActions() const { return false; }
-	ENGINE_API virtual void Init() PURE_VIRTUAL(USoundEffectPreset::Init, );
+	virtual void Init() PURE_VIRTUAL(USoundEffectPreset::Init, );
 	virtual void OnInit() {};
 	virtual FColor GetPresetColor() const { return FColor(200, 100, 100); }
 
diff --git a/Engine/Source/Runtime/Engine/Classes/Sound/SoundNodeAssetReferencer.h b/Engine/Source/Runtime/Engine/Classes/Sound/SoundNodeAssetReferencer.h
--- a/Engine/Source/Runtime/Engine/Classes/Sound/SoundNodeAssetReferencer.h
+++ b/Engine/Source/Runtime/Engine/Classes/Sound/SoundNodeAssetReferencer.h
@@ -17,8 +17,8 @@ class USoundNodeAssetReferencer : public USoundNode
 	GENERATED_BODY()
 
 public:
-	ENGINE_API virtual void LoadAsset(bool bAddToRoot = false) PURE_VIRTUAL(USoundNodeAssetReferencer::LoadAsset,);
-	ENGINE_API virtual void ClearAssetReferences() PURE_VIRTUAL(USoundNodeAssetReferencer::ClearAssetReferences, );
+	virtual void LoadAsset(bool bAddToRoot = false) PURE_VIRTUAL(USoundNodeAssetReferencer::LoadAsset,);
+	virtual void ClearAssetReferences() PURE_VIRTUAL(USoundNodeAssetReferencer::ClearAssetReferences, );
 	virtual bool ContainsProceduralSoundReference() const { return false; }
 
 	ENGINE_API bool ShouldHardReferenceAsset(const class ITargetPlatform*) const;
diff --git a/Engine/Source/Runtime/Engine/Public/ActorPartition/PartitionActor.h b/Engine/Source/Runtime/Engine/Public/ActorPartition/PartitionActor.h
--- a/Engine/Source/Runtime/Engine/Public/ActorPartition/PartitionActor.h
+++ b/Engine/Source/Runtime/Engine/Public/ActorPartition/PartitionActor.h
@@ -24,7 +24,7 @@ public:
 	//~ Begin AActor Interface
 	virtual bool CanChangeIsSpatiallyLoadedFlag() const override { return false; }
 	ENGINE_API virtual TUniquePtr<class FWorldPartitionActorDesc> CreateClassActorDesc() const override;
-	ENGINE_API virtual uint32 GetDefaultGridSize(UWorld* InWorld) const PURE_VIRTUAL(APartitionActor, return 0;)
+	virtual uint32 GetDefaultGridSize(UWorld* InWorld) const PURE_VIRTUAL(APartitionActor, return 0;)
 	ENGINE_API virtual bool ShouldIncludeGridSizeInName(UWorld * InWorld, const FActorPartitionIdentifier& InIdentifier) const;
 	virtual bool ShouldIncludeGridSizeInLabel() const { return false; }
 	virtual FGuid GetGridGuid() const { return FGuid(); }
diff --git a/Engine/Source/Runtime/Engine/Public/Animation/AnimCurveTypes.h b/Engine/Source/Runtime/Engine/Public/Animation/AnimCurveTypes.h
--- a/Engine/Source/Runtime/Engine/Public/Animation/AnimCurveTypes.h
+++ b/Engine/Source/Runtime/Engine/Public/Animation/AnimCurveTypes.h
@@ -221,7 +221,7 @@ struct FFloatCurve : public FAnimCurveBase
 	FFloatCurve(){}
 
 	UE_DEPRECATED(5.3, "Please use the constructor that takes an FName.")
-	ENGINE_API FFloatCurve(FSmartName InName, int32 InCurveTypeFlags)
+	FFloatCurve(FSmartName InName, int32 InCurveTypeFlags)
 	PRAGMA_DISABLE_DEPRECATION_WARNINGS
 		: FAnimCurveBase(InName, InCurveTypeFlags)
 	PRAGMA_ENABLE_DEPRECATION_WARNINGS
@@ -271,7 +271,7 @@ struct FVectorCurve : public FAnimCurveBase
 	FVectorCurve(){}
 
 	UE_DEPRECATED(5.3, "Please use the constructor that takes an FName.")
-	ENGINE_API FVectorCurve(FSmartName InName, int32 InCurveTypeFlags)
+	FVectorCurve(FSmartName InName, int32 InCurveTypeFlags)
 	PRAGMA_DISABLE_DEPRECATION_WARNINGS
 		: FAnimCurveBase(InName, InCurveTypeFlags)
 	PRAGMA_ENABLE_DEPRECATION_WARNINGS
@@ -325,7 +325,7 @@ struct FTransformCurve: public FAnimCurveBase
 	FTransformCurve(){}
 
 	UE_DEPRECATED(5.3, "Please use the constructor that takes an FName.")
-	ENGINE_API FTransformCurve(FSmartName InName, int32 InCurveTypeFlags)
+	FTransformCurve(FSmartName InName, int32 InCurveTypeFlags)
 	PRAGMA_DISABLE_DEPRECATION_WARNINGS
 		: FAnimCurveBase(InName, InCurveTypeFlags)
 	PRAGMA_ENABLE_DEPRECATION_WARNINGS
diff --git a/Engine/Source/Runtime/Engine/Public/Audio.h b/Engine/Source/Runtime/Engine/Public/Audio.h
--- a/Engine/Source/Runtime/Engine/Public/Audio.h
+++ b/Engine/Source/Runtime/Engine/Public/Audio.h
@@ -558,7 +558,7 @@ public:
 
 	ENGINE_API virtual ~FSoundBuffer();
 
-	ENGINE_API virtual int32 GetSize() PURE_VIRTUAL(FSoundBuffer::GetSize,return 0;);
+	virtual int32 GetSize() PURE_VIRTUAL(FSoundBuffer::GetSize,return 0;);
 
 	/**
 	 * Describe the buffer (platform can override to add to the description, but should call the base class version)
diff --git a/Engine/Source/Runtime/Engine/Public/Materials/MaterialFunctionInterface.h b/Engine/Source/Runtime/Engine/Public/Materials/MaterialFunctionInterface.h
--- a/Engine/Source/Runtime/Engine/Public/Materials/MaterialFunctionInterface.h
+++ b/Engine/Source/Runtime/Engine/Public/Materials/MaterialFunctionInterface.h
@@ -129,7 +129,7 @@ public:
 
 #if WITH_EDITORONLY_DATA
 	/** @return true if this function is dependent on the passed in function, directly or indirectly. */
-	ENGINE_API virtual bool IsDependent(UMaterialFunctionInterface* OtherFunction)
+	virtual bool IsDependent(UMaterialFunctionInterface* OtherFunction)
 		PURE_VIRTUAL(UMaterialFunctionInterface::IsDependent,return false;);
 
 	/**
@@ -139,11 +139,11 @@ public:
 	 *
 	 * @return true if all dependent functions were visited, false if the Predicate did break iteration
 	 */
-	ENGINE_API virtual bool IterateDependentFunctions(TFunctionRef<bool(UMaterialFunctionInterface*)> Predicate) const
+	virtual bool IterateDependentFunctions(TFunctionRef<bool(UMaterialFunctionInterface*)> Predicate) const
 		PURE_VIRTUAL(UMaterialFunctionInterface::IterateDependentFunctions,return false;);
 
 	/** Returns an array of the functions that this function is dependent on, directly or indirectly. */
-	ENGINE_API virtual void GetDependentFunctions(TArray<UMaterialFunctionInterface*>& DependentFunctions) const
+	virtual void GetDependentFunctions(TArray<UMaterialFunctionInterface*>& DependentFunctions) const
 		PURE_VIRTUAL(UMaterialFunctionInterface::GetDependentFunctions,);
 
 	/** Returns If returns an empty string, use the default class name for the material function. Otherwise, the string will be the name shown when the function is exposed to users in the material graph as a node, or from the contextual menu when searching for nodes. */
@@ -151,7 +151,7 @@ public:
 #endif // WITH_EDITORONLY_DATA
 
 #if WITH_EDITOR
-	ENGINE_API virtual UMaterialInterface* GetPreviewMaterial()
+	virtual UMaterialInterface* GetPreviewMaterial()
 		PURE_VIRTUAL(UMaterialFunctionInterface::GetPreviewMaterial,return nullptr;);
 
 	virtual void UpdateInputOutputTypes()
diff --git a/Engine/Source/Runtime/Engine/Public/Materials/MaterialInterface.h b/Engine/Source/Runtime/Engine/Public/Materials/MaterialInterface.h
--- a/Engine/Source/Runtime/Engine/Public/Materials/MaterialInterface.h
+++ b/Engine/Source/Runtime/Engine/Public/Materials/MaterialInterface.h
@@ -681,7 +681,7 @@ public:
 	/**
 	 * Update layers functions RuntimeGraphCache for compiling (primarily for Preview Materials)
 	 */
-	ENGINE_API virtual void SyncLayersRuntimeGraphCache(FMaterialLayersFunctions* OverrideLayers) PURE_VIRTUAL(UMaterialInterface::SyncLayersRuntimeGraphCache);
+	virtual void SyncLayersRuntimeGraphCache(FMaterialLayersFunctions* OverrideLayers) PURE_VIRTUAL(UMaterialInterface::SyncLayersRuntimeGraphCache);
 #endif // WITH_EDITORONLY_DATA
 
 #if WITH_EDITOR
@@ -701,7 +701,7 @@ public:
 	* @param	OutSortPriority	Will contain the sort priority of the parameter group if successful
 	* @return					True if successful
 	*/
-	ENGINE_API virtual bool GetGroupSortPriority(const FString& InGroupName, int32& OutSortPriority) const
+	virtual bool GetGroupSortPriority(const FString& InGroupName, int32& OutSortPriority) const
 		PURE_VIRTUAL(UMaterialInterface::GetGroupSortPriority, return false;);
 #endif // WITH_EDITOR
 
diff --git a/Engine/Source/Runtime/Engine/Public/Subsystems/WorldSubsystem.h b/Engine/Source/Runtime/Engine/Public/Subsystems/WorldSubsystem.h
--- a/Engine/Source/Runtime/Engine/Public/Subsystems/WorldSubsystem.h
+++ b/Engine/Source/Runtime/Engine/Public/Subsystems/WorldSubsystem.h
@@ -68,7 +68,7 @@ public:
 	ENGINE_API virtual ETickableTickType GetTickableTickType() const override;
 	ENGINE_API virtual bool IsAllowedToTick() const override final;
 	ENGINE_API virtual void Tick(float DeltaTime) override;
-	ENGINE_API virtual TStatId GetStatId() const override PURE_VIRTUAL(UTickableWorldSubsystem::GetStatId, return TStatId(););
+	virtual TStatId GetStatId() const override PURE_VIRTUAL(UTickableWorldSubsystem::GetStatId, return TStatId(););
 	// FTickableGameObject implementation End
 
 	// USubsystem implementation Begin
diff --git a/Engine/Source/Runtime/Engine/Public/WorldPartition/HLOD/DestructibleHLODComponent.h b/Engine/Source/Runtime/Engine/Public/WorldPartition/HLOD/DestructibleHLODComponent.h
--- a/Engine/Source/Runtime/Engine/Public/WorldPartition/HLOD/DestructibleHLODComponent.h
+++ b/Engine/Source/Runtime/Engine/Public/WorldPartition/HLOD/DestructibleHLODComponent.h
@@ -22,9 +22,9 @@ class UWorldPartitionDestructibleHLODComponent : public USceneComponent
 	GENERATED_UCLASS_BODY()
 
 public:
-	ENGINE_API virtual void DestroyActor(int32 ActorIndex) PURE_VIRTUAL(UWorldPartitionDestructibleHLODComponent::DestroyActor, );
-	ENGINE_API virtual void DamageActor(int32 ActorIndex, float DamagePercent) PURE_VIRTUAL(UWorldPartitionDestructibleHLODComponent::DamageActor, );
-	ENGINE_API virtual void OnDestructionStateUpdated() PURE_VIRTUAL(UWorldPartitionDestructibleHLODComponent::OnDestructionStateUpdated, );
+	virtual void DestroyActor(int32 ActorIndex) PURE_VIRTUAL(UWorldPartitionDestructibleHLODComponent::DestroyActor, );
+	virtual void DamageActor(int32 ActorIndex, float DamagePercent) PURE_VIRTUAL(UWorldPartitionDestructibleHLODComponent::DamageActor, );
+	virtual void OnDestructionStateUpdated() PURE_VIRTUAL(UWorldPartitionDestructibleHLODComponent::OnDestructionStateUpdated, );
 
 	ENGINE_API const TArray<FName>& GetDestructibleActors() const;
 	ENGINE_API void SetDestructibleActors(const TArray<FName>& InDestructibleActors);
diff --git a/Engine/Source/Runtime/Engine/Public/WorldPartition/HLOD/HLODBuilder.h b/Engine/Source/Runtime/Engine/Public/WorldPartition/HLOD/HLODBuilder.h
--- a/Engine/Source/Runtime/Engine/Public/WorldPartition/HLOD/HLODBuilder.h
+++ b/Engine/Source/Runtime/Engine/Public/WorldPartition/HLOD/HLODBuilder.h
@@ -147,7 +147,7 @@ public:
 	/**
 	 * Components created with this method need to be properly outered & assigned to your target actor.
 	 */
-	ENGINE_API virtual TArray<UActorComponent*> Build(const FHLODBuildContext& InHLODBuildContext, const TArray<UActorComponent*>& InSourceComponents) const PURE_VIRTUAL(UHLODBuilder::Build, return {};);
+	virtual TArray<UActorComponent*> Build(const FHLODBuildContext& InHLODBuildContext, const TArray<UActorComponent*>& InSourceComponents) const PURE_VIRTUAL(UHLODBuilder::Build, return {};);
 
 	/**
 	 * From a set of components, compute a unique hash from their properties that are relevant to HLOD generation.
diff --git a/Engine/Source/Runtime/Engine/Public/WorldPartition/HLOD/HLODModifier.h b/Engine/Source/Runtime/Engine/Public/WorldPartition/HLOD/HLODModifier.h
--- a/Engine/Source/Runtime/Engine/Public/WorldPartition/HLOD/HLODModifier.h
+++ b/Engine/Source/Runtime/Engine/Public/WorldPartition/HLOD/HLODModifier.h
@@ -21,8 +21,8 @@ class UWorldPartitionHLODModifier : public UObject
 	 GENERATED_UCLASS_BODY()
 
 public:
-	ENGINE_API virtual bool CanModifyHLOD(TSubclassOf<UHLODBuilder> InHLODBuilderClass) const PURE_VIRTUAL(UWorldPartitionHLODModifier::CanModifyHLOD, return false;);
+	virtual bool CanModifyHLOD(TSubclassOf<UHLODBuilder> InHLODBuilderClass) const PURE_VIRTUAL(UWorldPartitionHLODModifier::CanModifyHLOD, return false;);
 	
-	ENGINE_API virtual void BeginHLODBuild(const FHLODBuildContext& InHLODBuildContext) PURE_VIRTUAL(UWorldPartitionHLODModifier::BeginHLODBuild, );
-	ENGINE_API virtual void EndHLODBuild(TArray<UActorComponent*>& InOutComponents) PURE_VIRTUAL(UWorldPartitionHLODModifier::EndHLODBuild, );
+	virtual void BeginHLODBuild(const FHLODBuildContext& InHLODBuildContext) PURE_VIRTUAL(UWorldPartitionHLODModifier::BeginHLODBuild, );
+	virtual void EndHLODBuild(TArray<UActorComponent*>& InOutComponents) PURE_VIRTUAL(UWorldPartitionHLODModifier::EndHLODBuild, );
 };
diff --git a/Engine/Source/Runtime/Engine/Public/WorldPartition/HLOD/HLODSourceActors.h b/Engine/Source/Runtime/Engine/Public/WorldPartition/HLOD/HLODSourceActors.h
--- a/Engine/Source/Runtime/Engine/Public/WorldPartition/HLOD/HLODSourceActors.h
+++ b/Engine/Source/Runtime/Engine/Public/WorldPartition/HLOD/HLODSourceActors.h
@@ -16,7 +16,7 @@ class UWorldPartitionHLODSourceActors : public UObject
 
 public:
 #if WITH_EDITOR
-	ENGINE_API virtual ULevelStreaming* LoadSourceActors(bool& bOutDirty) const PURE_VIRTUAL(UWorldPartitionHLODSourceActors::LoadSourceActors, return nullptr; );
+	virtual ULevelStreaming* LoadSourceActors(bool& bOutDirty) const PURE_VIRTUAL(UWorldPartitionHLODSourceActors::LoadSourceActors, return nullptr; );
 	ENGINE_API virtual uint32 GetHLODHash() const;
 
 	ENGINE_API void SetHLODLayer(const UHLODLayer* HLODLayer);
diff --git a/Engine/Source/Runtime/Engine/Public/WorldPartition/RuntimeHashSet/RuntimePartition.h b/Engine/Source/Runtime/Engine/Public/WorldPartition/RuntimeHashSet/RuntimePartition.h
--- a/Engine/Source/Runtime/Engine/Public/WorldPartition/RuntimeHashSet/RuntimePartition.h
+++ b/Engine/Source/Runtime/Engine/Public/WorldPartition/RuntimeHashSet/RuntimePartition.h
@@ -73,13 +73,13 @@ public:
 	};
 
 	UE_API virtual void SetDefaultValues();
-	UE_API virtual bool SupportsHLODs() const PURE_VIRTUAL(URuntimePartition::SupportsHLODs, return true;);
+	virtual bool SupportsHLODs() const PURE_VIRTUAL(URuntimePartition::SupportsHLODs, return true;);
 	UE_API virtual void InitHLODRuntimePartitionFrom(const URuntimePartition* InRuntimePartition, int32 InHLODIndex);
 	virtual void UpdateHLODRuntimePartitionFrom(const URuntimePartition* InRuntimePartition) {}
 #endif
-	UE_API virtual bool IsValidPartitionTokens(const TArray<FName>& InPartitionTokens) const PURE_VIRTUAL(URuntimePartition::IsValidPartitionTokens, return false;);
+	virtual bool IsValidPartitionTokens(const TArray<FName>& InPartitionTokens) const PURE_VIRTUAL(URuntimePartition::IsValidPartitionTokens, return false;);
 #if WITH_EDITOR
-	UE_API virtual bool GenerateStreaming(const FGenerateStreamingParams& InParams, FGenerateStreamingResult& OutResult) PURE_VIRTUAL(URuntimePartition::GenerateStreaming, return false;);
+	virtual bool GenerateStreaming(const FGenerateStreamingParams& InParams, FGenerateStreamingResult& OutResult) PURE_VIRTUAL(URuntimePartition::GenerateStreaming, return false;);
 	virtual FArchive& AppendCellGuid(FArchive& InAr) { return InAr << Name << HLODIndex; }
 #endif
 
diff --git a/Engine/Source/Runtime/Engine/Public/WorldPartition/WorldPartitionEditorHash.h b/Engine/Source/Runtime/Engine/Public/WorldPartition/WorldPartitionEditorHash.h
--- a/Engine/Source/Runtime/Engine/Public/WorldPartition/WorldPartitionEditorHash.h
+++ b/Engine/Source/Runtime/Engine/Public/WorldPartition/WorldPartitionEditorHash.h
@@ -15,16 +15,16 @@ class UWorldPartitionEditorHash : public UObject
 
 public:
 #if WITH_EDITOR
-	ENGINE_API virtual void Initialize() PURE_VIRTUAL(UWorldPartitionEditorHash::Initialize, return;);
-	ENGINE_API virtual void SetDefaultValues() PURE_VIRTUAL(UWorldPartitionEditorHash::SetDefaultValues, return;);
-	ENGINE_API virtual FName GetWorldPartitionEditorName() const PURE_VIRTUAL(UWorldPartitionEditorHash::GetWorldPartitionEditorName, return FName(NAME_None););
-	ENGINE_API virtual FBox GetEditorWorldBounds() const  PURE_VIRTUAL(UWorldPartitionEditorHash::GetEditorWorldBounds, return FBox(ForceInit););
-	ENGINE_API virtual FBox GetRuntimeWorldBounds() const  PURE_VIRTUAL(UWorldPartitionEditorHash::GetRuntimeWorldBounds, return FBox(ForceInit););
-	ENGINE_API virtual FBox GetNonSpatialBounds() const PURE_VIRTUAL(UWorldPartitionEditorHash::GetNonSpatialBounds, return FBox(ForceInit););
-	ENGINE_API virtual void Tick(float DeltaSeconds) PURE_VIRTUAL(UWorldPartitionEditorHash::Tick, return;);
+	virtual void Initialize() PURE_VIRTUAL(UWorldPartitionEditorHash::Initialize, return;);
+	virtual void SetDefaultValues() PURE_VIRTUAL(UWorldPartitionEditorHash::SetDefaultValues, return;);
+	virtual FName GetWorldPartitionEditorName() const PURE_VIRTUAL(UWorldPartitionEditorHash::GetWorldPartitionEditorName, return FName(NAME_None););
+	virtual FBox GetEditorWorldBounds() const  PURE_VIRTUAL(UWorldPartitionEditorHash::GetEditorWorldBounds, return FBox(ForceInit););
+	virtual FBox GetRuntimeWorldBounds() const  PURE_VIRTUAL(UWorldPartitionEditorHash::GetRuntimeWorldBounds, return FBox(ForceInit););
+	virtual FBox GetNonSpatialBounds() const PURE_VIRTUAL(UWorldPartitionEditorHash::GetNonSpatialBounds, return FBox(ForceInit););
+	virtual void Tick(float DeltaSeconds) PURE_VIRTUAL(UWorldPartitionEditorHash::Tick, return;);
 
-	ENGINE_API virtual void HashActor(FWorldPartitionHandle& InActorHandle) PURE_VIRTUAL(UWorldPartitionEditorHash::HashActor, ;);
-	ENGINE_API virtual void UnhashActor(FWorldPartitionHandle& InActorHandle) PURE_VIRTUAL(UWorldPartitionEditorHash::UnhashActor, ;);
+	virtual void HashActor(FWorldPartitionHandle& InActorHandle) PURE_VIRTUAL(UWorldPartitionEditorHash::HashActor, ;);
+	virtual void UnhashActor(FWorldPartitionHandle& InActorHandle) PURE_VIRTUAL(UWorldPartitionEditorHash::UnhashActor, ;);
 
 	/* Struct of optional parameters passed to ForEachIntersectingActor. */
 	struct FForEachIntersectingActorParams
@@ -48,6 +48,6 @@ public:
 	UE_DEPRECATED(5.4, "Use ForEachIntersectingActor with FWorldPartitionActorDescInstance")
 	virtual int32 ForEachIntersectingActor(const FBox& Box, TFunctionRef<void(FWorldPartitionActorDesc*)> InOperation, const FForEachIntersectingActorParams& Params = FForEachIntersectingActorParams()) { return 0; }
 
-	ENGINE_API virtual int32 ForEachIntersectingActor(const FBox& Box, TFunctionRef<void(FWorldPartitionActorDescInstance*)> InOperation, const FForEachIntersectingActorParams& Params = FForEachIntersectingActorParams()) PURE_VIRTUAL(UWorldPartitionEditorHash::ForEachIntersectingActor, return 0;);
+	virtual int32 ForEachIntersectingActor(const FBox& Box, TFunctionRef<void(FWorldPartitionActorDescInstance*)> InOperation, const FForEachIntersectingActorParams& Params = FForEachIntersectingActorParams()) PURE_VIRTUAL(UWorldPartitionEditorHash::ForEachIntersectingActor, return 0;);
 #endif
 };
diff --git a/Engine/Source/Runtime/Engine/Public/WorldPartition/WorldPartitionRuntimeCell.h b/Engine/Source/Runtime/Engine/Public/WorldPartition/WorldPartitionRuntimeCell.h
--- a/Engine/Source/Runtime/Engine/Public/WorldPartition/WorldPartitionRuntimeCell.h
+++ b/Engine/Source/Runtime/Engine/Public/WorldPartition/WorldPartitionRuntimeCell.h
@@ -248,24 +248,24 @@ class UWorldPartitionRuntimeCell : public UObject, public IWorldPartitionCell, p
 		return Super::GetTypedOuter<T>();
 	}
 
-	ENGINE_API virtual void Load() const PURE_VIRTUAL(UWorldPartitionRuntimeCell::Load,);
-	ENGINE_API virtual void Unload() const PURE_VIRTUAL(UWorldPartitionRuntimeCell::Unload,);
-	ENGINE_API virtual bool CanUnload() const PURE_VIRTUAL(UWorldPartitionRuntimeCell::CanUnload, return true;);
-	ENGINE_API virtual void Activate() const PURE_VIRTUAL(UWorldPartitionRuntimeCell::Activate,);
-	ENGINE_API virtual void Deactivate() const PURE_VIRTUAL(UWorldPartitionRuntimeCell::Deactivate,);
-	ENGINE_API virtual ULevel* GetLevel() const PURE_VIRTUAL(UWorldPartitionRuntimeCell::GetLevel, return nullptr;);
-	ENGINE_API virtual EWorldPartitionRuntimeCellState GetCurrentState() const PURE_VIRTUAL(UWorldPartitionRuntimeCell::GetCurrentState, return EWorldPartitionRuntimeCellState::Unloaded;);
+	virtual void Load() const PURE_VIRTUAL(UWorldPartitionRuntimeCell::Load,);
+	virtual void Unload() const PURE_VIRTUAL(UWorldPartitionRuntimeCell::Unload,);
+	virtual bool CanUnload() const PURE_VIRTUAL(UWorldPartitionRuntimeCell::CanUnload, return true;);
+	virtual void Activate() const PURE_VIRTUAL(UWorldPartitionRuntimeCell::Activate,);
+	virtual void Deactivate() const PURE_VIRTUAL(UWorldPartitionRuntimeCell::Deactivate,);
+	virtual ULevel* GetLevel() const PURE_VIRTUAL(UWorldPartitionRuntimeCell::GetLevel, return nullptr;);
+	virtual EWorldPartitionRuntimeCellState GetCurrentState() const PURE_VIRTUAL(UWorldPartitionRuntimeCell::GetCurrentState, return EWorldPartitionRuntimeCellState::Unloaded;);
 	virtual FLinearColor GetDebugColor(EWorldPartitionRuntimeCellVisualizeMode VisualizeMode) const { static const FLinearColor DefaultColor = FLinearColor::Black.CopyWithNewOpacity(0.25f); return DefaultColor; }
 	virtual bool IsAlwaysLoaded() const { return bIsAlwaysLoaded; }
 	virtual bool IsSpatiallyLoaded() const { return bIsSpatiallyLoaded; }
 	virtual void SetIsAlwaysLoaded(bool bInIsAlwaysLoaded) { bIsAlwaysLoaded = bInIsAlwaysLoaded; }
 	virtual void SetIsSpatiallyLoaded(bool bInIsSpatiallyLoaded) { bIsSpatiallyLoaded = bInIsSpatiallyLoaded; }
-	ENGINE_API virtual void SetStreamingPriority(int32 InStreamingPriority) const PURE_VIRTUAL(UWorldPartitionRuntimeCell::SetStreamingPriority,);
+	virtual void SetStreamingPriority(int32 InStreamingPriority) const PURE_VIRTUAL(UWorldPartitionRuntimeCell::SetStreamingPriority,);
 	virtual EStreamingStatus GetStreamingStatus() const { return LEVEL_Unloaded; }
 	void SetClientOnlyVisible(bool bInClientOnlyVisible) { bClientOnlyVisible = bInClientOnlyVisible; }
 	bool GetClientOnlyVisible() const { return bClientOnlyVisible; }
 	virtual FGuid const& GetContentBundleID() const { return ContentBundleID; }
-	ENGINE_API virtual TArray<FName> GetActors() const PURE_VIRTUAL(UWorldPartitionRuntimeCell::GetActors, return TArray<FName>(););
+	virtual TArray<FName> GetActors() const PURE_VIRTUAL(UWorldPartitionRuntimeCell::GetActors, return TArray<FName>(););
 
 	//~Begin IWorldPartitionCell Interface
 	ENGINE_API virtual TArray<const UDataLayerInstance*> GetDataLayerInstances() const override;
@@ -314,7 +314,7 @@ class UWorldPartitionRuntimeCell : public UObject, public IWorldPartitionCell, p
 	virtual TSet<FName> GetActorPackageNames() const override { return TSet<FName>(); }
 	//~End IWorldPartitionCell Interface
 
-	ENGINE_API virtual void AddActorToCell(const FStreamingGenerationActorDescView& ActorDescView) PURE_VIRTUAL(UWorldPartitionRuntimeCell::AddActorToCell, );
+	virtual void AddActorToCell(const FStreamingGenerationActorDescView& ActorDescView) PURE_VIRTUAL(UWorldPartitionRuntimeCell::AddActorToCell, );
 	
 	PRAGMA_DISABLE_DEPRECATION_WARNINGS
 
@@ -323,8 +323,8 @@ class UWorldPartitionRuntimeCell : public UObject, public IWorldPartitionCell, p
 
 	PRAGMA_ENABLE_DEPRECATION_WARNINGS
 
-	ENGINE_API virtual void Fixup() PURE_VIRTUAL(UWorldPartitionRuntimeCell::Fixup, );
-	ENGINE_API virtual int32 GetActorCount() const PURE_VIRTUAL(UWorldPartitionRuntimeCell::GetActorCount, return 0;);
+	virtual void Fixup() PURE_VIRTUAL(UWorldPartitionRuntimeCell::Fixup, );
+	virtual int32 GetActorCount() const PURE_VIRTUAL(UWorldPartitionRuntimeCell::GetActorCount, return 0;);
 
 	// Cook methods
 	UE_DEPRECATED(5.5, "Use version with IWorldPartitionCookPackageContext")
@@ -342,7 +342,7 @@ class UWorldPartitionRuntimeCell : public UObject, public IWorldPartitionCell, p
 	virtual bool OnPrepareGeneratorPackageForCook(TArray<UPackage*>& OutModifiedPackages) override { return true; }
 	virtual bool OnPopulateGeneratorPackageForCook(const IWorldPartitionCookPackageContext& InCookContext, UPackage* InPackage) override { return true; }
 	virtual bool OnPopulateGeneratedPackageForCook(const IWorldPartitionCookPackageContext& InCookContext, UPackage* InPackage, TArray<UPackage*>& OutModifiedPackages) override { return true; }
-	ENGINE_API virtual FWorldPartitionPackageHash GetGenerationHash() const override PURE_VIRTUAL(UWorldPartitionRuntimeCell::GetGenerationHash, return FWorldPartitionPackageHash(););
+	virtual FWorldPartitionPackageHash GetGenerationHash() const override PURE_VIRTUAL(UWorldPartitionRuntimeCell::GetGenerationHash, return FWorldPartitionPackageHash(););
 	//~End IWorldPartitionCookPackageObject
 
 	ENGINE_API virtual void DumpStateLog(FHierarchicalLogArchive& Ar) const;
diff --git a/Engine/Source/Runtime/Engine/Public/WorldPartition/WorldPartitionRuntimeHash.h b/Engine/Source/Runtime/Engine/Public/WorldPartition/WorldPartitionRuntimeHash.h
--- a/Engine/Source/Runtime/Engine/Public/WorldPartition/WorldPartitionRuntimeHash.h
+++ b/Engine/Source/Runtime/Engine/Public/WorldPartition/WorldPartitionRuntimeHash.h
@@ -207,7 +207,7 @@ class UWorldPartitionRuntimeHash : public UObject
 	URuntimeHashExternalStreamingObjectBase* StoreStreamingContentToExternalStreamingObject(FName InStreamingObjectName) { return StoreStreamingContentToExternalStreamingObject(); }
 
 	ENGINE_API virtual void FlushStreamingContent();
-	ENGINE_API virtual TSubclassOf<URuntimeHashExternalStreamingObjectBase> GetExternalStreamingObjectClass() const PURE_VIRTUAL(UWorldPartitionRuntimeHash::GetExternalStreamingObjectClass, return nullptr;);
+	virtual TSubclassOf<URuntimeHashExternalStreamingObjectBase> GetExternalStreamingObjectClass() const PURE_VIRTUAL(UWorldPartitionRuntimeHash::GetExternalStreamingObjectClass, return nullptr;);
 	ENGINE_API virtual void DumpStateLog(FHierarchicalLogArchive& Ar) const;
 
 	//~Begin Deprecation
diff --git a/Engine/Source/Runtime/Experimental/Animation/Constraints/Public/ConstraintsManager.h b/Engine/Source/Runtime/Experimental/Animation/Constraints/Public/ConstraintsManager.h
--- a/Engine/Source/Runtime/Experimental/Animation/Constraints/Public/ConstraintsManager.h
+++ b/Engine/Source/Runtime/Experimental/Animation/Constraints/Public/ConstraintsManager.h
@@ -79,7 +79,7 @@ public:
 	virtual ~UTickableConstraint() {}
 
 	/** Returns the actual function that the tick function needs to evaluate. */
-	CONSTRAINTS_API virtual FConstraintTickFunction::ConstraintFunction GetFunction() const PURE_VIRTUAL(GetFunction, return {};);
+	virtual FConstraintTickFunction::ConstraintFunction GetFunction() const PURE_VIRTUAL(GetFunction, return {};);
 
 	/** Evaluates the constraint in a context where it's mot done thru the ConstraintTick's tick function. */
 	CONSTRAINTS_API virtual void Evaluate(bool bTickHandlesAlso = false) const;
@@ -92,22 +92,22 @@ public:
 	/** Get whether or not it's fully active, it's set to active and all pieces are set up,e.g. objects really exist that are being constrainted*/
 	CONSTRAINTS_API virtual bool IsFullyActive() const;
 	/** If true it contains objects bound to an external system, like sequencer so we don't do certain things, like remove constraints when they don't resolve*/
-	CONSTRAINTS_API virtual bool HasBoundObjects() const PURE_VIRTUAL(HasBoundObjects, return false;);
+	virtual bool HasBoundObjects() const PURE_VIRTUAL(HasBoundObjects, return false;);
 	/** Resolve the bound objects so that any object it references are resovled and correctly set up*/
-	CONSTRAINTS_API virtual void ResolveBoundObjects(FMovieSceneSequenceID LocalSequenceID, TSharedRef<UE::MovieScene::FSharedPlaybackState> SharedPlaybackState, UObject* SubObject = nullptr)  PURE_VIRTUAL(ResolveBoundObjects);
+	virtual void ResolveBoundObjects(FMovieSceneSequenceID LocalSequenceID, TSharedRef<UE::MovieScene::FSharedPlaybackState> SharedPlaybackState, UObject* SubObject = nullptr)  PURE_VIRTUAL(ResolveBoundObjects);
 
 	/** @todo document */
-	CONSTRAINTS_API virtual uint32 GetTargetHash() const PURE_VIRTUAL(GetTargetHash, return 0;);
+	virtual uint32 GetTargetHash() const PURE_VIRTUAL(GetTargetHash, return 0;);
 	/** Whether or not this object references this object. If it HasBoundObjects you should ResolveBoundObjects if you expect this to be up to date */
-	CONSTRAINTS_API virtual bool ReferencesObject(TWeakObjectPtr<UObject> InObject) const PURE_VIRTUAL(ReferencesObject, return false;);
+	virtual bool ReferencesObject(TWeakObjectPtr<UObject> InObject) const PURE_VIRTUAL(ReferencesObject, return false;);
 	/** Create duplicate with new Outer*/
 	CONSTRAINTS_API virtual UTickableConstraint* Duplicate(UObject* NewOuter) const;
 	/** Initialize and register the constraint with that world*/
-	CONSTRAINTS_API virtual void InitConstraint(UWorld * InWorld)  PURE_VIRTUAL(UTickableConstraint::InitConstraint);
+	virtual void InitConstraint(UWorld * InWorld)  PURE_VIRTUAL(UTickableConstraint::InitConstraint);
 	/** Teardown and unregister the constraint with that world*/
-	CONSTRAINTS_API virtual void TeardownConstraint(UWorld* InWorld)  PURE_VIRTUAL(UTickableConstraint::DeInitConstraint);
+	virtual void TeardownConstraint(UWorld* InWorld)  PURE_VIRTUAL(UTickableConstraint::DeInitConstraint);
 	/** */
-	CONSTRAINTS_API virtual void AddedToWorld(UWorld* InWorld) PURE_VIRTUAL(UTickableConstraint::AddedToWorld);
+	virtual void AddedToWorld(UWorld* InWorld) PURE_VIRTUAL(UTickableConstraint::AddedToWorld);
 	
 	/** Get Tick function in that world*/
 	CONSTRAINTS_API FConstraintTickFunction& GetTickFunction(UWorld* InWorld);
diff --git a/Engine/Source/Runtime/Experimental/Animation/Constraints/Public/Transform/TransformConstraint.h b/Engine/Source/Runtime/Experimental/Animation/Constraints/Public/Transform/TransformConstraint.h
--- a/Engine/Source/Runtime/Experimental/Animation/Constraints/Public/Transform/TransformConstraint.h
+++ b/Engine/Source/Runtime/Experimental/Animation/Constraints/Public/Transform/TransformConstraint.h
@@ -131,7 +131,7 @@ protected:
 	 * Computes the initial offset that is needed to keep the child's global transform unchanged when creating the
 	 * constraint. This can be called whenever necessary to update the current offset.
 	*/
-	CONSTRAINTS_API virtual void ComputeOffset() PURE_VIRTUAL(ComputeOffset, return;);
+	virtual void ComputeOffset() PURE_VIRTUAL(ComputeOffset, return;);
 	
 	/**
 	 * Sets up dependencies between the parent, the constraint and the child using their respective tick functions.
diff --git a/Engine/Source/Runtime/Experimental/Animation/Constraints/Public/Transform/TransformableHandle.h b/Engine/Source/Runtime/Experimental/Animation/Constraints/Public/Transform/TransformableHandle.h
--- a/Engine/Source/Runtime/Experimental/Animation/Constraints/Public/Transform/TransformableHandle.h
+++ b/Engine/Source/Runtime/Experimental/Animation/Constraints/Public/Transform/TransformableHandle.h
@@ -57,25 +57,25 @@ public:
 	 * pointer and control name are valid, so they must be fully loaded so that the constraint can be updated later. bDeepCheck = false will be used
 	 * for that purpose.
 	 */
-	CONSTRAINTS_API virtual bool IsValid(const bool bDeepCheck = true) const PURE_VIRTUAL(IsValid, return false;);
+	virtual bool IsValid(const bool bDeepCheck = true) const PURE_VIRTUAL(IsValid, return false;);
 	
 	/** Sets the global transform of the underlying transformable object. */
-	CONSTRAINTS_API virtual void SetGlobalTransform(const FTransform& InGlobal) const PURE_VIRTUAL(SetGlobalTransform, );
+	virtual void SetGlobalTransform(const FTransform& InGlobal) const PURE_VIRTUAL(SetGlobalTransform, );
 	/** Sets the local transform of the underlying transformable object in it's parent space. */
-	CONSTRAINTS_API virtual void SetLocalTransform(const FTransform& InLocal) const PURE_VIRTUAL(SetLocalTransform, );
+	virtual void SetLocalTransform(const FTransform& InLocal) const PURE_VIRTUAL(SetLocalTransform, );
 	/** Gets the global transform of the underlying transformable object. */
-	CONSTRAINTS_API virtual FTransform GetGlobalTransform() const PURE_VIRTUAL(GetGlobalTransform, return FTransform::Identity;);
+	virtual FTransform GetGlobalTransform() const PURE_VIRTUAL(GetGlobalTransform, return FTransform::Identity;);
 	/** Gets the local transform of the underlying transformable object in it's parent space. */
-	CONSTRAINTS_API virtual FTransform GetLocalTransform() const PURE_VIRTUAL(GetLocalTransform, return FTransform::Identity;);
+	virtual FTransform GetLocalTransform() const PURE_VIRTUAL(GetLocalTransform, return FTransform::Identity;);
 
 	/** If true it contains objects bound to an external system, like sequencer so we don't do certain things, like remove constraints when they don't resolve*/
 	CONSTRAINTS_API virtual bool HasBoundObjects() const;
 
 	/** Resolve the bound objects so that any object it references are resolved and correctly set up*/
-	CONSTRAINTS_API virtual void ResolveBoundObjects(FMovieSceneSequenceID LocalSequenceID, TSharedRef<UE::MovieScene::FSharedPlaybackState> SharedPlaybackState, UObject* SubObject = nullptr) PURE_VIRTUAL(ResolveBoundObjects);
+	virtual void ResolveBoundObjects(FMovieSceneSequenceID LocalSequenceID, TSharedRef<UE::MovieScene::FSharedPlaybackState> SharedPlaybackState, UObject* SubObject = nullptr) PURE_VIRTUAL(ResolveBoundObjects);
 
 	/** Make a duplicate of myself with this outer*/
-	CONSTRAINTS_API virtual UTransformableHandle* Duplicate(UObject* NewOuter) const PURE_VIRTUAL(Duplicate, return nullptr;);
+	virtual UTransformableHandle* Duplicate(UObject* NewOuter) const PURE_VIRTUAL(Duplicate, return nullptr;);
 
 
 	/** Fix up Binding in case it has changed*/
@@ -91,11 +91,11 @@ public:
 	virtual void PreEvaluate(const bool bTick = false) const;
 	
 	/** Get the array of float channels for the specified section*/
-	CONSTRAINTS_API virtual TArrayView<FMovieSceneFloatChannel*>  GetFloatChannels(const UMovieSceneSection* InSection) const PURE_VIRTUAL(GetFloatChannels, return TArrayView<FMovieSceneFloatChannel*>(); );
+	virtual TArrayView<FMovieSceneFloatChannel*>  GetFloatChannels(const UMovieSceneSection* InSection) const PURE_VIRTUAL(GetFloatChannels, return TArrayView<FMovieSceneFloatChannel*>(); );
 	/** Get the array of double channels for the specified section*/
-	CONSTRAINTS_API virtual TArrayView<FMovieSceneDoubleChannel*>  GetDoubleChannels(const UMovieSceneSection * InSection) const PURE_VIRTUAL(GetDoubleChannels, return TArrayView<FMovieSceneDoubleChannel*>(); );
+	virtual TArrayView<FMovieSceneDoubleChannel*>  GetDoubleChannels(const UMovieSceneSection * InSection) const PURE_VIRTUAL(GetDoubleChannels, return TArrayView<FMovieSceneDoubleChannel*>(); );
 	/** Add Transform Keys at the specified times*/
-	CONSTRAINTS_API virtual bool AddTransformKeys(const TArray<FFrameNumber>& InFrames,
+	virtual bool AddTransformKeys(const TArray<FFrameNumber>& InFrames,
 		const TArray<FTransform>& InTransforms,
 		const EMovieSceneTransformChannel& InChannels,
 		const FFrameRate& InTickResolution,
@@ -106,32 +106,32 @@ public:
 	 * Returns the target object containing the tick function (returned in GetTickFunction).
 	 * See FTickFunction::AddPrerequisite for details.
 	 **/
-	CONSTRAINTS_API virtual UObject* GetPrerequisiteObject() const PURE_VIRTUAL(GetPrerequisiteObject, return nullptr;);
+	virtual UObject* GetPrerequisiteObject() const PURE_VIRTUAL(GetPrerequisiteObject, return nullptr;);
 	/**
 	 * Returns the tick function of the underlying transformable object.
 	 * This is used to set dependencies with the constraint.
 	**/
-	CONSTRAINTS_API virtual FTickFunction* GetTickFunction() const PURE_VIRTUAL(GetTickFunction, return nullptr;);
+	virtual FTickFunction* GetTickFunction() const PURE_VIRTUAL(GetTickFunction, return nullptr;);
 
 	/** Generates a hash value of the underlying transformable object. */
-	CONSTRAINTS_API virtual uint32 GetHash() const PURE_VIRTUAL(GetHash, return 0;);
+	virtual uint32 GetHash() const PURE_VIRTUAL(GetHash, return 0;);
 
 	/** Returns the underlying targeted object. */
-	CONSTRAINTS_API virtual TWeakObjectPtr<UObject> GetTarget() const PURE_VIRTUAL(GetTarget, return nullptr;);
+	virtual TWeakObjectPtr<UObject> GetTarget() const PURE_VIRTUAL(GetTarget, return nullptr;);
 
 	/** Check for direct dependencies with InOther. */
-	CONSTRAINTS_API virtual bool HasDirectDependencyWith(const UTransformableHandle& InOther) const PURE_VIRTUAL(HasDirectDependencyWith, return false;);
+	virtual bool HasDirectDependencyWith(const UTransformableHandle& InOther) const PURE_VIRTUAL(HasDirectDependencyWith, return false;);
 
 	/** Look for a possible tick function that can be used as a prerequisite. */
-	CONSTRAINTS_API virtual FTickPrerequisite GetPrimaryPrerequisite(const bool bAllowThis = true) const PURE_VIRTUAL(GetPrimaryPrerequisite, return FTickPrerequisite(););
+	virtual FTickPrerequisite GetPrimaryPrerequisite(const bool bAllowThis = true) const PURE_VIRTUAL(GetPrimaryPrerequisite, return FTickPrerequisite(););
 	
 	CONSTRAINTS_API FHandleModifiedEvent& HandleModified();
 	CONSTRAINTS_API void Notify(EHandleEvent InEvent, const bool bPreTickTarget = false) const;
 	mutable bool bNotifying = false;
 
 #if WITH_EDITOR
-	CONSTRAINTS_API virtual FString GetLabel() const PURE_VIRTUAL(GetLabel, return FString(););
-	CONSTRAINTS_API virtual FString GetFullLabel() const PURE_VIRTUAL(GetFullLabel, return FString(););
+	virtual FString GetLabel() const PURE_VIRTUAL(GetLabel, return FString(););
+	virtual FString GetFullLabel() const PURE_VIRTUAL(GetFullLabel, return FString(););
 #endif
 
 	//possible bindingID
diff --git a/Engine/Source/Runtime/Experimental/Iris/Core/Public/Iris/DataStream/DataStream.h b/Engine/Source/Runtime/Experimental/Iris/Core/Public/Iris/DataStream/DataStream.h
--- a/Engine/Source/Runtime/Experimental/Iris/Core/Public/Iris/DataStream/DataStream.h
+++ b/Engine/Source/Runtime/Experimental/Iris/Core/Public/Iris/DataStream/DataStream.h
@@ -165,7 +165,7 @@ public:
 	 *        ProcessPacketDeliveryStatus will be called with the record in the same order as WriteData was called.
 	 * @return Whether there was data written or not and if the stream has more data to write if bandwidth settings allow it.
 	 */
-	IRISCORE_API virtual EWriteResult WriteData(UE::Net::FNetSerializationContext& Context, FDataStreamRecord const*& OutRecord) PURE_VIRTUAL(WriteData, return EWriteResult::NoData;);
+	virtual EWriteResult WriteData(UE::Net::FNetSerializationContext& Context, FDataStreamRecord const*& OutRecord) PURE_VIRTUAL(WriteData, return EWriteResult::NoData;);
 
 	/**
 	 * Called after the final call to WriteData this frame, allowing the DataStream to cleanup data setup during BeginWrite.
@@ -176,7 +176,7 @@ public:
 	 * Deserialize data that was written with WriteData.
 	 * @param Context The FNetSerializationContext which has accessors for the bitstream to read from among other things. 
 	 */
-	IRISCORE_API virtual void ReadData(UE::Net::FNetSerializationContext& Context) PURE_VIRTUAL(ReadData,);
+	virtual void ReadData(UE::Net::FNetSerializationContext& Context) PURE_VIRTUAL(ReadData,);
 
 	/**
 	 * For each packet into which we have written data we are guaranteed to get a call to ProcessPacketDeliveryStatus when
@@ -184,12 +184,12 @@ public:
 	 * @param Status Whether the packet was delivered or not or if the record should simply be discarded due to closing a connection.
 	 * @param Record The record which was set by this stream during a WriteData call.
 	 */
-	IRISCORE_API virtual void ProcessPacketDeliveryStatus(UE::Net::EPacketDeliveryStatus Status, FDataStreamRecord const* Record) PURE_VIRTUAL(ProcessPacketDeliveryStatus,);
+	virtual void ProcessPacketDeliveryStatus(UE::Net::EPacketDeliveryStatus Status, FDataStreamRecord const* Record) PURE_VIRTUAL(ProcessPacketDeliveryStatus,);
 
 	/**
 	 * @return true if the stream has no pending reliable data for which it is waiting on an acknowledgement.
 	 */
-	IRISCORE_API virtual bool HasAcknowledgedAllReliableData() const PURE_VIRTUAL(HasAcknowledgedAllReliableData, return true;);
+	virtual bool HasAcknowledgedAllReliableData() const PURE_VIRTUAL(HasAcknowledgedAllReliableData, return true;);
 
 	/**
 	 * Get name of DataStream
diff --git a/Engine/Source/Runtime/Experimental/Iris/Core/Public/Iris/ReplicationSystem/Filtering/NetObjectFilter.h b/Engine/Source/Runtime/Experimental/Iris/Core/Public/Iris/ReplicationSystem/Filtering/NetObjectFilter.h
--- a/Engine/Source/Runtime/Experimental/Iris/Core/Public/Iris/ReplicationSystem/Filtering/NetObjectFilter.h
+++ b/Engine/Source/Runtime/Experimental/Iris/Core/Public/Iris/ReplicationSystem/Filtering/NetObjectFilter.h
@@ -186,10 +186,10 @@ public:
 	IRISCORE_API virtual void RemoveConnection(uint32 ConnectionId);
 
 	/** A new object want to use this filter. Opportunity to cache some information for it. The info struct passed has been zeroed. Must be overriden. */
-	IRISCORE_API virtual bool AddObject(uint32 ObjectIndex, FNetObjectFilterAddObjectParams&) PURE_VIRTUAL(AddObject, return false;)
+	virtual bool AddObject(uint32 ObjectIndex, FNetObjectFilterAddObjectParams&) PURE_VIRTUAL(AddObject, return false;)
 
 	/** An object no longer wants to use this filter. */
-	IRISCORE_API virtual void RemoveObject(uint32 ObjectIndex, const FNetObjectFilteringInfo&) PURE_VIRTUAL(RemoveObject,)
+	virtual void RemoveObject(uint32 ObjectIndex, const FNetObjectFilteringInfo&) PURE_VIRTUAL(RemoveObject,)
 
 	/** A set of objects using this filter are dirty since the last update. An opportunity for the filter to update cached data. */
 	IRISCORE_API virtual void UpdateObjects(FNetObjectFilterUpdateParams&);
@@ -236,13 +236,13 @@ protected:
 	IRISCORE_API UNetObjectFilter();
 
 	/** Called right after constructor for enabled filters. Must be overriden. */
-	IRISCORE_API virtual void OnInit(const FNetObjectFilterInitParams&) PURE_VIRTUAL(OnInit, );
+	virtual void OnInit(const FNetObjectFilterInitParams&) PURE_VIRTUAL(OnInit, );
 
 	/** Called when the replication system is shutting down. Use this to remove references to other systems */
-	IRISCORE_API virtual void OnDeinit() PURE_VIRTUAL(OnDeinit);
+	virtual void OnDeinit() PURE_VIRTUAL(OnDeinit);
 
 	/** Called when the maximum InternalNetRefIndex increased and we need to realloc our lists */
-	IRISCORE_API virtual void OnMaxInternalNetRefIndexIncreased(uint32 NewMaxInternalIndex) PURE_VIRTUAL(OnMaxInternalNetRefIndexIncreased);
+	virtual void OnMaxInternalNetRefIndexIncreased(uint32 NewMaxInternalIndex) PURE_VIRTUAL(OnMaxInternalNetRefIndexIncreased);
 
 	/* Returns the filtering info for this object if it's handled by this filter, nullptr otherwise. */
 	IRISCORE_API FNetObjectFilteringInfo* GetFilteringInfo(uint32 ObjectIndex);
diff --git a/Engine/Source/Runtime/Experimental/Iris/Core/Public/Iris/ReplicationSystem/NetBlob/NetBlobHandler.h b/Engine/Source/Runtime/Experimental/Iris/Core/Public/Iris/ReplicationSystem/NetBlob/NetBlobHandler.h
--- a/Engine/Source/Runtime/Experimental/Iris/Core/Public/Iris/ReplicationSystem/NetBlob/NetBlobHandler.h
+++ b/Engine/Source/Runtime/Experimental/Iris/Core/Public/Iris/ReplicationSystem/NetBlob/NetBlobHandler.h
@@ -66,10 +66,10 @@ private:
 	// INetBlobReceiver
 
 	/** Override to create the NetBlob. This will be called when receiving a NetBlob of the type this handler is responsible for. */
-	IRISCORE_API virtual TRefCountPtr<FNetBlob> CreateNetBlob(const FNetBlobCreationInfo&) const override PURE_VIRTUAL(CreateNetBlob, return nullptr;);
+	virtual TRefCountPtr<FNetBlob> CreateNetBlob(const FNetBlobCreationInfo&) const override PURE_VIRTUAL(CreateNetBlob, return nullptr;);
 
 	/** Override to process the NetBlob when it's received. */
-	IRISCORE_API virtual void OnNetBlobReceived(UE::Net::FNetSerializationContext& Context, const TRefCountPtr<FNetBlob>&) override PURE_VIRTUAL(OnNetBlobReceived,);
+	virtual void OnNetBlobReceived(UE::Net::FNetSerializationContext& Context, const TRefCountPtr<FNetBlob>&) override PURE_VIRTUAL(OnNetBlobReceived,);
 
 private:
 	friend class UE::Net::Private::FNetBlobHandlerManager;
diff --git a/Engine/Source/Runtime/Experimental/Iris/Core/Public/Iris/ReplicationSystem/Prioritization/NetObjectPrioritizer.h b/Engine/Source/Runtime/Experimental/Iris/Core/Public/Iris/ReplicationSystem/Prioritization/NetObjectPrioritizer.h
--- a/Engine/Source/Runtime/Experimental/Iris/Core/Public/Iris/ReplicationSystem/Prioritization/NetObjectPrioritizer.h
+++ b/Engine/Source/Runtime/Experimental/Iris/Core/Public/Iris/ReplicationSystem/Prioritization/NetObjectPrioritizer.h
@@ -160,13 +160,13 @@ class UNetObjectPrioritizer : public UObject
 
 public:
 	/** Called once at init time before any other calls to the prioritizer. */
-	IRISCORE_API virtual void Init(FNetObjectPrioritizerInitParams& Params) PURE_VIRTUAL(Init,)
+	virtual void Init(FNetObjectPrioritizerInitParams& Params) PURE_VIRTUAL(Init,)
 
 	/** Called when the replication system is shutting down. Use this to remove references to other systems */
-	IRISCORE_API virtual void Deinit() PURE_VIRTUAL(Deinit)
+	virtual void Deinit() PURE_VIRTUAL(Deinit)
 
 	/** Called when the maximum InternalNetRefIndex increased and we need to realloc our lists */
-	IRISCORE_API virtual void OnMaxInternalNetRefIndexIncreased(uint32 NewMaxInternalIndex) PURE_VIRTUAL(OnMaxInternalNetRefIndexIncreased);
+	virtual void OnMaxInternalNetRefIndexIncreased(uint32 NewMaxInternalIndex) PURE_VIRTUAL(OnMaxInternalNetRefIndexIncreased);
 
 	/** A new connection has been added. An opportunity for the prioritizer to allocate per connection info. */
 	IRISCORE_API virtual void AddConnection(uint32 ConnectionId);
@@ -175,13 +175,13 @@ public:
 	IRISCORE_API virtual void RemoveConnection(uint32 ConnectionId);
 
 	/** A new object want to use this prioritizer. Opportunity to cache some information for it. The info struct passed has been zeroed. */
-	IRISCORE_API virtual bool AddObject(uint32 ObjectIndex, FNetObjectPrioritizerAddObjectParams& Params) PURE_VIRTUAL(AddObject, return false;)
+	virtual bool AddObject(uint32 ObjectIndex, FNetObjectPrioritizerAddObjectParams& Params) PURE_VIRTUAL(AddObject, return false;)
 
 	/** An object do no longer want to use this prioritizer. */
-	IRISCORE_API virtual void RemoveObject(uint32 ObjectIndex, const FNetObjectPrioritizationInfo& Info) PURE_VIRTUAL(RemoveObject,)
+	virtual void RemoveObject(uint32 ObjectIndex, const FNetObjectPrioritizationInfo& Info) PURE_VIRTUAL(RemoveObject,)
 
 	/** A set of objects used by this prioritizer have been updated. An opportunity for the prioritizer to update cached data. */
-	IRISCORE_API virtual void UpdateObjects(FNetObjectPrioritizerUpdateParams&) PURE_VIRTUAL(UpdateObjects,)
+	virtual void UpdateObjects(FNetObjectPrioritizerUpdateParams&) PURE_VIRTUAL(UpdateObjects,)
 
 	/**
 	 * If there are any connections being replicated and there's a chance Prioritize() will be called then PrePrioritize()
@@ -194,7 +194,7 @@ public:
 	 * Stored priorities are expected to use the maximum of the already stored priority and the prioritizer calculated one.
 	 * That allows multiple prioritizers to be used on the same object.
 	 */
-	IRISCORE_API virtual void Prioritize(FNetObjectPrioritizationParams&) PURE_VIRTUAL(Prioritize,);
+	virtual void Prioritize(FNetObjectPrioritizationParams&) PURE_VIRTUAL(Prioritize,);
 
 	/** If PrePrioritize() was called then PostPrioritize() will be called exactly once after all Prioritize() calls. */
 	IRISCORE_API virtual void PostPrioritize(FNetObjectPostPrioritizationParams&);
diff --git a/Engine/Source/Runtime/GameplayTasks/Classes/GameplayTaskOwnerInterface.h b/Engine/Source/Runtime/GameplayTasks/Classes/GameplayTaskOwnerInterface.h
--- a/Engine/Source/Runtime/GameplayTasks/Classes/GameplayTaskOwnerInterface.h
+++ b/Engine/Source/Runtime/GameplayTasks/Classes/GameplayTaskOwnerInterface.h
@@ -23,10 +23,10 @@ class IGameplayTaskOwnerInterface
 public:
 
 	/** Finds tasks component for given GameplayTask, Task.GetGameplayTasksComponent() may not be initialized at this point! */
-	GAMEPLAYTASKS_API virtual UGameplayTasksComponent* GetGameplayTasksComponent(const UGameplayTask& Task) const PURE_VIRTUAL(IGameplayTaskOwnerInterface::GetGameplayTasksComponent, return nullptr;);
+	virtual UGameplayTasksComponent* GetGameplayTasksComponent(const UGameplayTask& Task) const PURE_VIRTUAL(IGameplayTaskOwnerInterface::GetGameplayTasksComponent, return nullptr;);
 
 	/** Get owner of a task or default one when task is null */
-	GAMEPLAYTASKS_API virtual AActor* GetGameplayTaskOwner(const UGameplayTask* Task) const PURE_VIRTUAL(IGameplayTaskOwnerInterface::GetGameplayTaskOwner, return nullptr;);
+	virtual AActor* GetGameplayTaskOwner(const UGameplayTask* Task) const PURE_VIRTUAL(IGameplayTaskOwnerInterface::GetGameplayTaskOwner, return nullptr;);
 
 	/** Get "body" of task's owner / default, having location in world (e.g. Owner = AIController, Avatar = Pawn) */
 	virtual AActor* GetGameplayTaskAvatar(const UGameplayTask* Task) const { return GetGameplayTaskOwner(Task); }
diff --git a/Engine/Source/Runtime/InteractiveToolsFramework/Public/BaseGizmos/GizmoElementBase.h b/Engine/Source/Runtime/InteractiveToolsFramework/Public/BaseGizmos/GizmoElementBase.h
--- a/Engine/Source/Runtime/InteractiveToolsFramework/Public/BaseGizmos/GizmoElementBase.h
+++ b/Engine/Source/Runtime/InteractiveToolsFramework/Public/BaseGizmos/GizmoElementBase.h
@@ -114,13 +114,13 @@ public:
 public:
 
 	// Render enabled visible element.
-	INTERACTIVETOOLSFRAMEWORK_API virtual void Render(IToolsContextRenderAPI* RenderAPI, const FRenderTraversalState& RenderState) PURE_VIRTUAL(UGizmoElementBase::Render);
+	virtual void Render(IToolsContextRenderAPI* RenderAPI, const FRenderTraversalState& RenderState) PURE_VIRTUAL(UGizmoElementBase::Render);
 
 	// Do screen space drawing.
 	INTERACTIVETOOLSFRAMEWORK_API virtual void DrawHUD(FCanvas* Canvas, IToolsContextRenderAPI* RenderAPI, const FRenderTraversalState& RenderState);
 
 	// Line trace enabled hittable element.
-	INTERACTIVETOOLSFRAMEWORK_API virtual FInputRayHit LineTrace(const UGizmoViewContext* ViewContext, const FLineTraceTraversalState& LineTraceState, const FVector& RayOrigin, const FVector& RayDirection) PURE_VIRTUAL(UGizmoElementBase::LineTrace, return FInputRayHit(););
+	virtual FInputRayHit LineTrace(const UGizmoViewContext* ViewContext, const FLineTraceTraversalState& LineTraceState, const FVector& RayOrigin, const FVector& RayDirection) PURE_VIRTUAL(UGizmoElementBase::LineTrace, return FInputRayHit(););
 
 	// Set/get the visible bit in element state.
 	INTERACTIVETOOLSFRAMEWORK_API virtual void SetVisibleState(bool bVisible);
diff --git a/Engine/Source/Runtime/InteractiveToolsFramework/Public/ToolTargets/ToolTarget.h b/Engine/Source/Runtime/InteractiveToolsFramework/Public/ToolTargets/ToolTarget.h
--- a/Engine/Source/Runtime/InteractiveToolsFramework/Public/ToolTargets/ToolTarget.h
+++ b/Engine/Source/Runtime/InteractiveToolsFramework/Public/ToolTargets/ToolTarget.h
@@ -42,7 +42,7 @@ class UToolTarget : public UObject
 public:
 
 	/** @return true if target is still valid. May become invalid for various reasons (eg Component was deleted out from under us) */
-	INTERACTIVETOOLSFRAMEWORK_API virtual bool IsValid() const PURE_VIRTUAL(UToolTarget::IsValid, return false;);
+	virtual bool IsValid() const PURE_VIRTUAL(UToolTarget::IsValid, return false;);
 };
 
 
@@ -94,7 +94,7 @@ class UToolTargetFactory : public UObject
 
 public:
 
-	INTERACTIVETOOLSFRAMEWORK_API virtual bool CanBuildTarget(UObject* SourceObject, const FToolTargetTypeRequirements& TargetTypeInfo) const PURE_VIRTUAL(UToolTargetFactory::CanBuildTarget, return false;);
+	virtual bool CanBuildTarget(UObject* SourceObject, const FToolTargetTypeRequirements& TargetTypeInfo) const PURE_VIRTUAL(UToolTargetFactory::CanBuildTarget, return false;);
 
-	INTERACTIVETOOLSFRAMEWORK_API virtual UToolTarget* BuildTarget(UObject* SourceObject, const FToolTargetTypeRequirements& TargetTypeInfo) PURE_VIRTUAL(UToolTargetFactory::BuildTarget, return nullptr;);
+	virtual UToolTarget* BuildTarget(UObject* SourceObject, const FToolTargetTypeRequirements& TargetTypeInfo) PURE_VIRTUAL(UToolTargetFactory::BuildTarget, return nullptr;);
 };
diff --git a/Engine/Source/Runtime/Interchange/Core/Public/InterchangeTranslatorBase.h b/Engine/Source/Runtime/Interchange/Core/Public/InterchangeTranslatorBase.h
--- a/Engine/Source/Runtime/Interchange/Core/Public/InterchangeTranslatorBase.h
+++ b/Engine/Source/Runtime/Interchange/Core/Public/InterchangeTranslatorBase.h
@@ -85,11 +85,11 @@ public:
 	bool DoesSupportAssetType(EInterchangeTranslatorAssetType AssetType) const { return EnumHasAllFlags(GetSupportedAssetTypes(), AssetType); }
 
 	UFUNCTION(BlueprintCallable, Category = "Interchange | Translator")
-	INTERCHANGECORE_API virtual EInterchangeTranslatorAssetType GetSupportedAssetTypes() const PURE_VIRTUAL(UInterchangeTranslatorBase::GetSupportedAssetTypes, return EInterchangeTranslatorAssetType::None;);
+	virtual EInterchangeTranslatorAssetType GetSupportedAssetTypes() const PURE_VIRTUAL(UInterchangeTranslatorBase::GetSupportedAssetTypes, return EInterchangeTranslatorAssetType::None;);
 
 	/** List of formats supported by the translator. Each entry is of the form "ext;Description" where ext is the file extension. */
 	UFUNCTION(BlueprintCallable, Category = "Interchange | Translator")
-	INTERCHANGECORE_API virtual TArray<FString> GetSupportedFormats() const PURE_VIRTUAL(UInterchangeTranslatorBase::GetSupportedExtensions, return TArray<FString>(););
+	virtual TArray<FString> GetSupportedFormats() const PURE_VIRTUAL(UInterchangeTranslatorBase::GetSupportedExtensions, return TArray<FString>(););
 
 	/**
 	 * Translate the associated source data into node(s) that are hold in the specified nodes container.
diff --git a/Engine/Source/Runtime/Landscape/Classes/ILandscapeSplineInterface.h b/Engine/Source/Runtime/Landscape/Classes/ILandscapeSplineInterface.h
--- a/Engine/Source/Runtime/Landscape/Classes/ILandscapeSplineInterface.h
+++ b/Engine/Source/Runtime/Landscape/Classes/ILandscapeSplineInterface.h
@@ -20,16 +20,16 @@ class ILandscapeSplineInterface
 {
 	GENERATED_IINTERFACE_BODY()
 public:
-	LANDSCAPE_API virtual FGuid GetLandscapeGuid() const PURE_VIRTUAL(ILandscapeSplineInterface::GetLandscapeGuid, return FGuid(););
-	LANDSCAPE_API virtual ULandscapeInfo* GetLandscapeInfo() const PURE_VIRTUAL(ILandscapeSplineInterface::GetLandscapeInfo, return nullptr;);
-	LANDSCAPE_API virtual FTransform LandscapeActorToWorld() const PURE_VIRTUAL(ILandscapeSplineInterface::LandscapeActorToWorld, return FTransform::Identity;);
-	LANDSCAPE_API virtual ULandscapeSplinesComponent* GetSplinesComponent() const PURE_VIRTUAL(ILandscapeSplineInterface::GetSplineComponent, return nullptr;);
-	LANDSCAPE_API virtual void UpdateSharedProperties(ULandscapeInfo* InLandscapeInfo) PURE_VIRTUAL(ILandscapeSplineInterface::UpdateSharedProperties);
-	LANDSCAPE_API virtual bool IsSplineOwnerValid() const PURE_VIRTUAL(ILandscapeSplineInterface::IsSplineOwnerValid, return true;);
+	virtual FGuid GetLandscapeGuid() const PURE_VIRTUAL(ILandscapeSplineInterface::GetLandscapeGuid, return FGuid(););
+	virtual ULandscapeInfo* GetLandscapeInfo() const PURE_VIRTUAL(ILandscapeSplineInterface::GetLandscapeInfo, return nullptr;);
+	virtual FTransform LandscapeActorToWorld() const PURE_VIRTUAL(ILandscapeSplineInterface::LandscapeActorToWorld, return FTransform::Identity;);
+	virtual ULandscapeSplinesComponent* GetSplinesComponent() const PURE_VIRTUAL(ILandscapeSplineInterface::GetSplineComponent, return nullptr;);
+	virtual void UpdateSharedProperties(ULandscapeInfo* InLandscapeInfo) PURE_VIRTUAL(ILandscapeSplineInterface::UpdateSharedProperties);
+	virtual bool IsSplineOwnerValid() const PURE_VIRTUAL(ILandscapeSplineInterface::IsSplineOwnerValid, return true;);
 	 
 #if WITH_EDITOR
-	LANDSCAPE_API virtual bool SupportsForeignSplineMesh() const PURE_VIRTUAL(ILandscapeSplineInterface::SupportsForeignSplineMesh, return false;);
-	LANDSCAPE_API virtual void CreateSplineComponent() PURE_VIRTUAL(ILandscapeSplineInterface::CreateSplineComponent);
-	LANDSCAPE_API virtual void CreateSplineComponent(const FVector& Scale3D) PURE_VIRTUAL(ILandscapeSplineInterface::CreateSplineComponent);
+	virtual bool SupportsForeignSplineMesh() const PURE_VIRTUAL(ILandscapeSplineInterface::SupportsForeignSplineMesh, return false;);
+	virtual void CreateSplineComponent() PURE_VIRTUAL(ILandscapeSplineInterface::CreateSplineComponent);
+	virtual void CreateSplineComponent(const FVector& Scale3D) PURE_VIRTUAL(ILandscapeSplineInterface::CreateSplineComponent);
 #endif
 };
diff --git a/Engine/Source/Runtime/Landscape/Public/LandscapeEditLayerRenderer.h b/Engine/Source/Runtime/Landscape/Public/LandscapeEditLayerRenderer.h
--- a/Engine/Source/Runtime/Landscape/Public/LandscapeEditLayerRenderer.h
+++ b/Engine/Source/Runtime/Landscape/Public/LandscapeEditLayerRenderer.h
@@ -230,7 +230,7 @@ public:
 	 * @return a list of renderer states (i.e. a ILandscapeEditLayerRenderer and its current state) to be processed in that order by the merge operation
 	 */
 	LANDSCAPE_EDIT_LAYERS_BATCHED_MERGE_EXPERIMENTAL
-		LANDSCAPE_API virtual TArray<FEditLayerRendererState> GetEditLayerRendererStates(const FMergeContext* InMergeContext) PURE_VIRTUAL(IEditLayerRendererProvider::GetEditLayerRendererStates, return {}; );
+		virtual TArray<FEditLayerRendererState> GetEditLayerRendererStates(const FMergeContext* InMergeContext) PURE_VIRTUAL(IEditLayerRendererProvider::GetEditLayerRendererStates, return {}; );
 };
 
 #endif // WITH_EDITOR
@@ -285,7 +285,7 @@ public:
 	 *  and if this renderer doesn't act on one of these layers, it will simply do nothing with it in its RenderLayer function
 	 */
 	LANDSCAPE_EDIT_LAYERS_BATCHED_MERGE_EXPERIMENTAL
-	LANDSCAPE_API virtual void GetRendererStateInfo(const UE::Landscape::EditLayers::FMergeContext* InMergeContext,
+	virtual void GetRendererStateInfo(const UE::Landscape::EditLayers::FMergeContext* InMergeContext,
 		UE::Landscape::EditLayers::FEditLayerTargetTypeState& OutSupportedTargetTypeState, UE::Landscape::EditLayers::FEditLayerTargetTypeState& OutEnabledTargetTypeState, TArray<TBitArray<>>& OutTargetLayerGroups) const
 		PURE_VIRTUAL(ILandscapeEditLayerRenderer::GetRendererStateInfo, );
 
@@ -293,7 +293,7 @@ public:
 	* @return the a debug name for this renderer
 	*/
 	LANDSCAPE_EDIT_LAYERS_BATCHED_MERGE_EXPERIMENTAL
-	LANDSCAPE_API virtual FString GetEditLayerRendererDebugName() const 
+	virtual FString GetEditLayerRendererDebugName() const 
 		PURE_VIRTUAL(ILandscapeEditLayerRenderer::GetEditLayerRendererDebugName, return TEXT(""); );
 
 	/**
@@ -304,7 +304,7 @@ public:
 	 * @return list of all render items that affect this renderer
 	*/
 	LANDSCAPE_EDIT_LAYERS_BATCHED_MERGE_EXPERIMENTAL
-	LANDSCAPE_API virtual TArray<UE::Landscape::EditLayers::FEditLayerRenderItem> GetRenderItems(const UE::Landscape::EditLayers::FMergeContext* InMergeContext) const
+	virtual TArray<UE::Landscape::EditLayers::FEditLayerRenderItem> GetRenderItems(const UE::Landscape::EditLayers::FMergeContext* InMergeContext) const
 		PURE_VIRTUAL(ILandscapeEditLayerRenderer::GetRenderItems, return { }; );
 
 	/**
@@ -312,7 +312,7 @@ public:
 	 * 
 	 * @return details about how this renderer's render method is implemented 
 	*/
-	LANDSCAPE_API virtual UE::Landscape::EditLayers::ERenderFlags GetRenderFlags(const UE::Landscape::EditLayers::FMergeContext* InMergeContext) const
+	virtual UE::Landscape::EditLayers::ERenderFlags GetRenderFlags(const UE::Landscape::EditLayers::FMergeContext* InMergeContext) const
 		PURE_VIRTUAL(ILandscapeEditLayerRenderer::GetRenderFlags, return UE::Landscape::EditLayers::ERenderFlags::None; );
 
 	/**
@@ -336,7 +336,7 @@ public:
 	 * @return true if anything was rendered
 	*/
 	LANDSCAPE_EDIT_LAYERS_BATCHED_MERGE_EXPERIMENTAL
-	LANDSCAPE_API virtual bool RenderLayer(UE::Landscape::EditLayers::FRenderParams& RenderParams, UE::Landscape::FRDGBuilderRecorder& RDGBuilderRecorder)
+	virtual bool RenderLayer(UE::Landscape::EditLayers::FRenderParams& RenderParams, UE::Landscape::FRDGBuilderRecorder& RDGBuilderRecorder)
 		PURE_VIRTUAL(ILandscapeEditLayerRenderer::RenderLayer, return false; );
 
 	/**
diff --git a/Engine/Source/Runtime/LiveLinkInterface/Public/LiveLinkController.h b/Engine/Source/Runtime/LiveLinkInterface/Public/LiveLinkController.h
--- a/Engine/Source/Runtime/LiveLinkInterface/Public/LiveLinkController.h
+++ b/Engine/Source/Runtime/LiveLinkInterface/Public/LiveLinkController.h
@@ -23,7 +23,7 @@ class ULiveLinkController : public UObject
 	GENERATED_BODY()
 
 public:
-	LIVELINKINTERFACE_API virtual TSubclassOf<ULiveLinkRole> GetRole() const PURE_VIRTUAL(ULiveLinkController::GetRole, return TSubclassOf<ULiveLinkRole>(););
+	virtual TSubclassOf<ULiveLinkRole> GetRole() const PURE_VIRTUAL(ULiveLinkController::GetRole, return TSubclassOf<ULiveLinkRole>(););
 
 	/** Initialize the controller at the first tick of its owner component. */
 	virtual void OnRegistered() { }
diff --git a/Engine/Source/Runtime/LiveLinkInterface/Public/LiveLinkFrameInterpolationProcessor.h b/Engine/Source/Runtime/LiveLinkInterface/Public/LiveLinkFrameInterpolationProcessor.h
--- a/Engine/Source/Runtime/LiveLinkInterface/Public/LiveLinkFrameInterpolationProcessor.h
+++ b/Engine/Source/Runtime/LiveLinkInterface/Public/LiveLinkFrameInterpolationProcessor.h
@@ -65,6 +65,6 @@ class ULiveLinkFrameInterpolationProcessor : public UObject
 public:
 	using FWorkerSharedPtr = TSharedPtr<ILiveLinkFrameInterpolationProcessorWorker, ESPMode::ThreadSafe>;
 
-	LIVELINKINTERFACE_API virtual TSubclassOf<ULiveLinkRole> GetRole() const PURE_VIRTUAL(ULiveLinkFrameInterpolationProcessor::GetFromRole, return TSubclassOf<ULiveLinkRole>(););
-	LIVELINKINTERFACE_API virtual FWorkerSharedPtr FetchWorker() PURE_VIRTUAL(ULiveLinkFrameInterpolationProcessor::FetchWorker, return FWorkerSharedPtr(););
+	virtual TSubclassOf<ULiveLinkRole> GetRole() const PURE_VIRTUAL(ULiveLinkFrameInterpolationProcessor::GetFromRole, return TSubclassOf<ULiveLinkRole>(););
+	virtual FWorkerSharedPtr FetchWorker() PURE_VIRTUAL(ULiveLinkFrameInterpolationProcessor::FetchWorker, return FWorkerSharedPtr(););
 };
diff --git a/Engine/Source/Runtime/LiveLinkInterface/Public/LiveLinkFramePreProcessor.h b/Engine/Source/Runtime/LiveLinkInterface/Public/LiveLinkFramePreProcessor.h
--- a/Engine/Source/Runtime/LiveLinkInterface/Public/LiveLinkFramePreProcessor.h
+++ b/Engine/Source/Runtime/LiveLinkInterface/Public/LiveLinkFramePreProcessor.h
@@ -45,6 +45,6 @@ class ULiveLinkFramePreProcessor : public UObject
 public:
 	using FWorkerSharedPtr = TSharedPtr<ILiveLinkFramePreProcessorWorker, ESPMode::ThreadSafe>;
 
-	LIVELINKINTERFACE_API virtual TSubclassOf<ULiveLinkRole> GetRole() const PURE_VIRTUAL(ULiveLinkFramePreProcessor::GetFromRole, return TSubclassOf<ULiveLinkRole>(););
-	LIVELINKINTERFACE_API virtual FWorkerSharedPtr FetchWorker() PURE_VIRTUAL(ULiveLinkFramePreProcessor::FetchWorker, return FWorkerSharedPtr(););
+	virtual TSubclassOf<ULiveLinkRole> GetRole() const PURE_VIRTUAL(ULiveLinkFramePreProcessor::GetFromRole, return TSubclassOf<ULiveLinkRole>(););
+	virtual FWorkerSharedPtr FetchWorker() PURE_VIRTUAL(ULiveLinkFramePreProcessor::FetchWorker, return FWorkerSharedPtr(););
 };
diff --git a/Engine/Source/Runtime/LiveLinkInterface/Public/LiveLinkFrameTranslator.h b/Engine/Source/Runtime/LiveLinkInterface/Public/LiveLinkFrameTranslator.h
--- a/Engine/Source/Runtime/LiveLinkInterface/Public/LiveLinkFrameTranslator.h
+++ b/Engine/Source/Runtime/LiveLinkInterface/Public/LiveLinkFrameTranslator.h
@@ -46,7 +46,7 @@ public:
 
 	LIVELINKINTERFACE_API bool CanTranslate(TSubclassOf<ULiveLinkRole> ToRole) const;
 
-	LIVELINKINTERFACE_API virtual TSubclassOf<ULiveLinkRole> GetFromRole() const PURE_VIRTUAL(ULiveLinkFrameTranslator::GetFromRole, return TSubclassOf<ULiveLinkRole>(););
-	LIVELINKINTERFACE_API virtual TSubclassOf<ULiveLinkRole> GetToRole() const PURE_VIRTUAL(ULiveLinkFrameTranslator::GetToRole, return TSubclassOf<ULiveLinkRole>(););
-	LIVELINKINTERFACE_API virtual FWorkerSharedPtr FetchWorker() PURE_VIRTUAL(ULiveLinkFrameTranslator::FetchWorker, return FWorkerSharedPtr(););
+	virtual TSubclassOf<ULiveLinkRole> GetFromRole() const PURE_VIRTUAL(ULiveLinkFrameTranslator::GetFromRole, return TSubclassOf<ULiveLinkRole>(););
+	virtual TSubclassOf<ULiveLinkRole> GetToRole() const PURE_VIRTUAL(ULiveLinkFrameTranslator::GetToRole, return TSubclassOf<ULiveLinkRole>(););
+	virtual FWorkerSharedPtr FetchWorker() PURE_VIRTUAL(ULiveLinkFrameTranslator::FetchWorker, return FWorkerSharedPtr(););
 };
