Moved some PPV related code from UnrealClient.cpp to relevant files.
- also fixed comment in PostProcessVisualizeHDR.usf

diff --git a/Engine/Shaders/Private/PostProcessVisualizeHDR.usf b/Engine/Shaders/Private/PostProcessVisualizeHDR.usf
--- a/Engine/Shaders/Private/PostProcessVisualizeHDR.usf
+++ b/Engine/Shaders/Private/PostProcessVisualizeHDR.usf
@@ -543,7 +543,7 @@ void MainPS(noperspective float4 UVAndScreenPos : TEXCOORD0, float4 SvPosition :
 
 		if(abs(InsetPx.x - ValuePx) < 2 && PixelPos.y > HistogramLeftTop.y + 9)
 		{
-			// white line to show the smoothed exposure
+			// purple line to show the smoothed exposure
 			OutColor = lerp(OutColor, float4(0.5, 0, .5, 0), 1.0f);
 		}
 	}
diff --git a/Engine/Source/Runtime/Engine/Private/Components/PostProcessComponent.cpp b/Engine/Source/Runtime/Engine/Private/Components/PostProcessComponent.cpp
--- a/Engine/Source/Runtime/Engine/Private/Components/PostProcessComponent.cpp
+++ b/Engine/Source/Runtime/Engine/Private/Components/PostProcessComponent.cpp
@@ -2,6 +2,7 @@
 
 #include "Components/PostProcessComponent.h"
 #include "Components/SphereComponent.h"
+#include "Engine/World.h"
 
 #include UE_INLINE_GENERATED_CPP_BY_NAME(PostProcessComponent)
 
@@ -15,6 +16,30 @@ UPostProcessComponent::UPostProcessComponent(const FObjectInitializer& ObjectIni
 	bUnbound = 1;
 }
 
+void UPostProcessComponent::OnRegister()
+{
+	Super::OnRegister();
+	GetWorld()->InsertPostProcessVolume(this);
+}
+
+void UPostProcessComponent::OnUnregister()
+{
+	Super::OnUnregister();
+	GetWorld()->RemovePostProcessVolume(this);
+}
+
+void UPostProcessComponent::Serialize(FArchive& Ar)
+{
+	Super::Serialize(Ar);
+
+	if (Ar.IsLoading())
+	{
+#if WITH_EDITORONLY_DATA
+		Settings.OnAfterLoad();
+#endif
+	}
+}
+
 bool UPostProcessComponent::EncompassesPoint(FVector Point, float SphereRadius/*=0.f*/, float* OutDistanceToPoint)
 {
 	UShapeComponent* ParentShape = Cast<UShapeComponent>(GetAttachParent());
diff --git a/Engine/Source/Runtime/Engine/Private/PostProcessVolume.cpp b/Engine/Source/Runtime/Engine/Private/PostProcessVolume.cpp
--- a/Engine/Source/Runtime/Engine/Private/PostProcessVolume.cpp
+++ b/Engine/Source/Runtime/Engine/Private/PostProcessVolume.cpp
@@ -22,6 +22,25 @@ APostProcessVolume::APostProcessVolume(const FObjectInitializer& ObjectInitializ
 	BlendWeight = 1.0f;
 }
 
+void APostProcessVolume::PostUnregisterAllComponents()
+{
+	// Route clear to super first.
+	Super::PostUnregisterAllComponents();
+	// World will be NULL during exit purge.
+	if (GetWorld())
+	{
+		GetWorld()->RemovePostProcessVolume(this);
+		GetWorld()->PostProcessVolumes.RemoveSingle(this);
+	}
+}
+
+void APostProcessVolume::PostRegisterAllComponents()
+{
+	// Route update to super first.
+	Super::PostRegisterAllComponents();
+	GetWorld()->InsertPostProcessVolume(this);
+}
+
 bool APostProcessVolume::IsPPVEnabled() const
 {
 #if WITH_EDITOR
diff --git a/Engine/Source/Runtime/Engine/Private/UnrealClient.cpp b/Engine/Source/Runtime/Engine/Private/UnrealClient.cpp
--- a/Engine/Source/Runtime/Engine/Private/UnrealClient.cpp
+++ b/Engine/Source/Runtime/Engine/Private/UnrealClient.cpp
@@ -1655,49 +1655,6 @@ FRHIGPUMask FViewport::GetGPUMask(FRHICommandListImmediate& RHICmdList) const
 	return FRHIGPUMask::FromIndex(RHIGetViewportNextPresentGPUIndex(GetViewportRHI()));
 }
 
-void APostProcessVolume::PostUnregisterAllComponents()
-{
-	// Route clear to super first.
-	Super::PostUnregisterAllComponents();
-	// World will be NULL during exit purge.
-	if (GetWorld())
-	{
-		GetWorld()->RemovePostProcessVolume(this);
-		GetWorld()->PostProcessVolumes.RemoveSingle(this);
-	}
-}
-
-void APostProcessVolume::PostRegisterAllComponents()
-{
-	// Route update to super first.
-	Super::PostRegisterAllComponents();
-	GetWorld()->InsertPostProcessVolume(this);
-}
-
-void UPostProcessComponent::OnRegister()
-{
-	Super::OnRegister();
-	GetWorld()->InsertPostProcessVolume(this);
-}
-
-void UPostProcessComponent::OnUnregister()
-{
-	Super::OnUnregister();
-	GetWorld()->RemovePostProcessVolume(this);
-}
-
-void UPostProcessComponent::Serialize(FArchive& Ar)
-{
-	Super::Serialize(Ar);
-
-	if(Ar.IsLoading())
-	{
-#if WITH_EDITORONLY_DATA
-		Settings.OnAfterLoad();
-#endif
-	}
-}
-
 /**
 *	Starts a new rendering frame. Called from the game thread thread.
 */
