Single frame blendspaces:
- Adds the ability to set a _specific_ frame-index to use when sampling Animation Sequence for BlendSpace sampling
- Currently does not support Animation(BlendSpace) Analysis

diff --git a/Engine/Plugins/Experimental/AnimNextAnimGraph/Source/AnimNextAnimGraph/Private/Traits/BlendSpacePlayer.cpp b/Engine/Plugins/Experimental/AnimNextAnimGraph/Source/AnimNextAnimGraph/Private/Traits/BlendSpacePlayer.cpp
--- a/Engine/Plugins/Experimental/AnimNextAnimGraph/Source/AnimNextAnimGraph/Private/Traits/BlendSpacePlayer.cpp
+++ b/Engine/Plugins/Experimental/AnimNextAnimGraph/Source/AnimNextAnimGraph/Private/Traits/BlendSpacePlayer.cpp
@@ -128,7 +128,7 @@ namespace UE::AnimNext
 					NewNativeInterface.AnimSequence = CurSequence;
 
 					// Assume start position is normalized.
-					NewNativeInterface.StartPosition = SharedData->GetStartPosition(Binding) * CurSequence->GetPlayLength();
+					NewNativeInterface.StartPosition = SharedData->GetStartPosition(Binding) * BlendSamples[CurSampleIndex].GetSamplePlayLength();
 					NewNativeInterface.Loop = SharedData->GetbLoop(Binding);
 
 					GraphPayload.AddNative(MoveTemp(NewPayload));
diff --git a/Engine/Source/Editor/Persona/Private/Customization/BlendSampleDetails.cpp b/Engine/Source/Editor/Persona/Private/Customization/BlendSampleDetails.cpp
--- a/Engine/Source/Editor/Persona/Private/Customization/BlendSampleDetails.cpp
+++ b/Engine/Source/Editor/Persona/Private/Customization/BlendSampleDetails.cpp
@@ -193,6 +193,12 @@ void FBlendSampleDetails::CustomizeDetails(class IDetailLayoutBuilder& DetailBui
 			GET_MEMBER_NAME_CHECKED(FBlendSample, RateScale), FBlendSample::StaticStruct());
 		CategoryBuilder.AddProperty(RateScaleProperty);
 
+		TSharedPtr<IPropertyHandle> UseSingleFrameBlendingProperty = DetailBuilder.GetProperty(
+			GET_MEMBER_NAME_CHECKED(FBlendSample, FrameIndexToSample), FBlendSample::StaticStruct());
+			 
+		IDetailPropertyRow& UseSingleFrameBlendingPropertyRow = CategoryBuilder.AddProperty(UseSingleFrameBlendingProperty);
+		FBlendSampleDetails::GenerateSingleFrameBlendingWidget(UseSingleFrameBlendingPropertyRow, UseSingleFrameBlendingProperty, AnimationProperty);
+
 		TSharedPtr<IPropertyHandle> IncludeInAnalyseAllProperty = DetailBuilder.GetProperty(
 			GET_MEMBER_NAME_CHECKED(FBlendSample, bIncludeInAnalyseAll), FBlendSample::StaticStruct());
 		CategoryBuilder.AddProperty(IncludeInAnalyseAllProperty);
@@ -482,6 +488,75 @@ void FBlendSampleDetails::GenerateAnimationWidget(IDetailPropertyRow& PropertyRo
 		];
 }
 
+void FBlendSampleDetails::GenerateSingleFrameBlendingWidget(IDetailPropertyRow& PropertyRow, TSharedPtr<IPropertyHandle> SingleFrameIndexProperty, TSharedPtr<IPropertyHandle> AnimationProperty)
+{
+	auto ValueChangedLambda = [SingleFrameIndexProperty](uint32 NewValue, bool bIsInteractive)
+	{
+		if (SingleFrameIndexProperty.IsValid())
+		{			
+			SingleFrameIndexProperty->SetValue(NewValue, bIsInteractive ? EPropertyValueSetFlags::InteractiveChange : EPropertyValueSetFlags::DefaultFlags);
+		}
+	};
+	
+	auto MaxValueLambda = [AnimationProperty]() -> TOptional<uint32>
+	{
+		if (AnimationProperty.IsValid())
+		{
+			FAssetData AssetData;
+			AnimationProperty->GetValue(AssetData);
+
+			if (UAnimSequence* BlendSampleAnimation = Cast<UAnimSequence>(AssetData.GetAsset()))
+			{
+				return BlendSampleAnimation->GetNumberOfSampledKeys() - 1;
+			}
+		}
+		
+		return TOptional<uint32>();
+	};
+
+	TSharedPtr<SWidget> DefaultNameWidget;
+	TSharedPtr<SWidget> DefaultValueWidget;
+	PropertyRow.GetDefaultWidgets(DefaultNameWidget, DefaultValueWidget);
+	
+	PropertyRow.CustomWidget()
+	.NameContent()
+	[
+		DefaultNameWidget.ToSharedRef()
+	]
+	.ValueContent()
+	[
+		SNew(SNumericEntryBox<uint32>)
+		.Font(FAppStyle::GetFontStyle("CurveEd.InfoFont"))
+		.Value_Lambda([SingleFrameIndexProperty]() -> TOptional<uint32>
+		{
+			if (SingleFrameIndexProperty.IsValid())
+			{
+				uint32 Value = 0;
+				if (SingleFrameIndexProperty->GetValue(Value) == FPropertyAccess::Success)
+				{
+					return Value;
+				}
+			}
+			
+			return TOptional<uint32>();
+		})
+		.OnValueCommitted_Lambda([ValueChangedLambda](uint32 NewValue, ETextCommit::Type CommitType)
+		{
+			ValueChangedLambda(NewValue, false);
+		})
+		.OnValueChanged_Lambda([ValueChangedLambda](uint32 NewValue)
+		{
+			ValueChangedLambda(NewValue, true);
+		})
+		.LabelVAlign(VAlign_Center)
+		.AllowSpin(true)
+		.MinValue(0)
+		.MaxValue_Lambda(MaxValueLambda)
+		.MinSliderValue(0)
+		.MaxSliderValue_Lambda(MaxValueLambda)
+	];
+}
+
 void FBlendSampleDetails::GenerateSampleGraphWidget(FDetailWidgetRow& InRow, UAnimGraphNode_BlendSpaceGraphBase* InBlendSpaceNode, int32 InSampleIndex)
 {
 	InRow
diff --git a/Engine/Source/Editor/Persona/Private/Customization/BlendSampleDetails.h b/Engine/Source/Editor/Persona/Private/Customization/BlendSampleDetails.h
--- a/Engine/Source/Editor/Persona/Private/Customization/BlendSampleDetails.h
+++ b/Engine/Source/Editor/Persona/Private/Customization/BlendSampleDetails.h
@@ -32,6 +32,8 @@ public:
 
 	static void GenerateAnimationWidget(IDetailPropertyRow& PropertyRow, const UBlendSpace* BlendSpace, TSharedPtr<IPropertyHandle> AnimationProperty);
 
+	static void GenerateSingleFrameBlendingWidget(IDetailPropertyRow& PropertyRow, TSharedPtr<IPropertyHandle> SingleFrameIndexProperty, TSharedPtr<IPropertyHandle> AnimationProperty);
+
 	static void GenerateSampleGraphWidget(FDetailWidgetRow& Row, UAnimGraphNode_BlendSpaceGraphBase* BlendSpaceNode, int32 SampleIndex);
 
 	static bool ShouldFilterAssetStatic(const FAssetData& AssetData, const UBlendSpace* BlendSpaceBase);
diff --git a/Engine/Source/Editor/Persona/Private/Customization/BlendSpaceDetails.cpp b/Engine/Source/Editor/Persona/Private/Customization/BlendSpaceDetails.cpp
--- a/Engine/Source/Editor/Persona/Private/Customization/BlendSpaceDetails.cpp
+++ b/Engine/Source/Editor/Persona/Private/Customization/BlendSpaceDetails.cpp
@@ -23,6 +23,7 @@
 #include "PersonaBlendSpaceAnalysis.h"
 
 #include "DetailBuilderTypes.h"
+#include "SWarningOrErrorBox.h"
 
 #define LOCTEXT_NAMESPACE "BlendSpaceDetails"
 
@@ -269,6 +270,19 @@ void FBlendSpaceDetails::CustomizeDetails(class IDetailLayoutBuilder& DetailBuil
 		if (BlendSpace->IsAsset())
 		{
 			IDetailCategoryBuilder& DetailCategoryBuilder = DetailBuilder.EditCategory(FName("Analysis"));
+
+			TWeakObjectPtr<const UBlendSpace> WeakBlendspace = BlendSpace;
+			auto ContainsSingleFrameBlendSampleVisibility = [WeakBlendspace](bool bExpectedValueForVisibility) -> EVisibility
+			{
+				if (const UBlendSpace* PinnedBlendspace = WeakBlendspace.Get())
+				{
+					const bool bValue = PinnedBlendspace->GetBlendSamples().ContainsByPredicate([](const FBlendSample& BlendSample) { return BlendSample.bUseSingleFrameForBlending; } );
+					return bExpectedValueForVisibility == bValue ? EVisibility::Visible : EVisibility::Collapsed;
+				}
+
+				return EVisibility::Collapsed;
+			};
+					
 			TSharedPtr<IPropertyHandle> AnalysisPropertiesHandle = DetailBuilder.GetProperty(
 				GET_MEMBER_NAME_CHECKED(UBlendSpace, AnalysisProperties), UBlendSpace::StaticClass());
 
@@ -394,13 +408,26 @@ void FBlendSpaceDetails::CustomizeDetails(class IDetailLayoutBuilder& DetailBuil
 									this->HandleAnalysisFunctionChanged(AxisIndex, NewFunctionName);
 								}
 							})
-					];
+					]
+					.Visibility(TAttribute<EVisibility>::CreateLambda(ContainsSingleFrameBlendSampleVisibility, false));
 				}
 				else
 				{
 					DetailBuilder.HideProperty(AnalysisProperty);
 				}
 			}
+
+			// show incompatibility warning
+			const FText WarningText = LOCTEXT("AnalysisAndSingleFrameIncompatibilityWarning", "Analysis feature is not compatible with single frame blending. "
+																					  "\n\nPlease disable single frame blending in each blend sample to use analysis.");
+			FDetailWidgetRow& IncompatibilityWarningRow = DetailCategoryBuilder.AddCustomRow(WarningText)
+				.WholeRowContent()
+				[
+					SNew(SWarningOrErrorBox)
+						.MessageStyle(EMessageStyle::Warning)
+						.Message(WarningText)
+				]
+				.Visibility(TAttribute<EVisibility>::CreateLambda(ContainsSingleFrameBlendSampleVisibility, true));
 		}
 		
 		//==============================================================================================================
@@ -451,6 +478,7 @@ void FBlendSpaceDetails::CustomizeDetails(class IDetailLayoutBuilder& DetailBuil
 				TSharedPtr<IPropertyHandle> SampleValueProperty = BlendSampleProperty->GetChildHandle(GET_MEMBER_NAME_CHECKED(FBlendSample, SampleValue));
 				TSharedPtr<IPropertyHandle> RateScaleProperty = BlendSampleProperty->GetChildHandle(GET_MEMBER_NAME_CHECKED(FBlendSample, RateScale));
 				TSharedPtr<IPropertyHandle> IncludeInAnalyseAllProperty = BlendSampleProperty->GetChildHandle(GET_MEMBER_NAME_CHECKED(FBlendSample, bIncludeInAnalyseAll));
+				TSharedPtr<IPropertyHandle> UseSingleFrameBlendingProperty = BlendSampleProperty->GetChildHandle(GET_MEMBER_NAME_CHECKED(FBlendSample, FrameIndexToSample));
 
 				IDetailGroup& Group = DetailCategoryBuilder.AddGroup(FName("BlendSamples_Samples"), FText::GetEmpty());
 				Group.HeaderRow()
@@ -489,31 +517,35 @@ void FBlendSpaceDetails::CustomizeDetails(class IDetailLayoutBuilder& DetailBuil
 
 				FBlendSampleDetails::GenerateBlendSampleWidget(
 					[&Group, &SampleValueProperty]() -> IDetailPropertyRow& 
-					{ 
-						return Group.AddPropertyRow(SampleValueProperty.ToSharedRef()); 
-					}, FOnSampleMoved::CreateLambda(
+						{ 
+							return Group.AddPropertyRow(SampleValueProperty.ToSharedRef()); 
+						}, 
+					FOnSampleMoved::CreateLambda(
 						[this](const uint32 Index, const FVector& SampleValue, bool bIsInteractive) 
-				{
-					if (BlendSpace->IsValidBlendSampleIndex(Index) && BlendSpace->GetBlendSample(Index).SampleValue !=
-						SampleValue && !BlendSpace->IsTooCloseToExistingSamplePoint(SampleValue, Index))
-					{
-						BlendSpace->Modify();
-						bool bMoveSuccesful = BlendSpace->EditSampleValue(Index, SampleValue);
-						if (bMoveSuccesful)
 						{
-							BlendSpace->ValidateSampleData();
-							FPropertyChangedEvent ChangedEvent(
-								nullptr, bIsInteractive ? EPropertyChangeType::Interactive : EPropertyChangeType::ValueSet);
-							BlendSpace->PostEditChangeProperty(ChangedEvent);
-						}
-					}
-				}), BlendSpace, SampleIndex, false);
+							if (BlendSpace->IsValidBlendSampleIndex(Index) && BlendSpace->GetBlendSample(Index).SampleValue !=
+								SampleValue && !BlendSpace->IsTooCloseToExistingSamplePoint(SampleValue, Index))
+							{
+								BlendSpace->Modify();
+								bool bMoveSuccesful = BlendSpace->EditSampleValue(Index, SampleValue);
+								if (bMoveSuccesful)
+								{
+									BlendSpace->ValidateSampleData();
+									FPropertyChangedEvent ChangedEvent(
+										nullptr, bIsInteractive ? EPropertyChangeType::Interactive : EPropertyChangeType::ValueSet);
+									BlendSpace->PostEditChangeProperty(ChangedEvent);
+								}
+							}
+						}), 
+					BlendSpace, SampleIndex, false);
 
 				if(BlendSpace->IsAsset())
 				{
 					IDetailPropertyRow& AnimationRow = Group.AddPropertyRow(AnimationProperty.ToSharedRef());
 					FBlendSampleDetails::GenerateAnimationWidget(AnimationRow, BlendSpace, AnimationProperty);
 					Group.AddPropertyRow(RateScaleProperty.ToSharedRef());
+					IDetailPropertyRow& UseSingleFrameBlendingPropertyRow = Group.AddPropertyRow(UseSingleFrameBlendingProperty.ToSharedRef());
+					FBlendSampleDetails::GenerateSingleFrameBlendingWidget(UseSingleFrameBlendingPropertyRow, UseSingleFrameBlendingProperty.ToSharedRef(), AnimationProperty.ToSharedRef());
 					Group.AddPropertyRow(IncludeInAnalyseAllProperty.ToSharedRef());
 				}
 				else if(BlendSpaceNode.Get())
diff --git a/Engine/Source/Editor/Persona/Private/SAnimationBlendSpaceGridWidget.cpp b/Engine/Source/Editor/Persona/Private/SAnimationBlendSpaceGridWidget.cpp
--- a/Engine/Source/Editor/Persona/Private/SAnimationBlendSpaceGridWidget.cpp
+++ b/Engine/Source/Editor/Persona/Private/SAnimationBlendSpaceGridWidget.cpp
@@ -1086,7 +1086,10 @@ void SBlendSpaceGridWidget::PaintAnimationNames(const FGeometry& AllottedGeometr
 		{
 			const FBlendSample& Sample = Samples[SampleIndex];
 
-			const FText Name = FText::Format(LOCTEXT("SampleNameFormat", "{0} ({1})"), GetSampleName(Sample, SampleIndex), FText::AsNumber(SampleIndex));
+			const FText Name = !Sample.bUseSingleFrameForBlending ? 
+				FText::Format(LOCTEXT("SampleNameFormat", "{0} ({1})"), GetSampleName(Sample, SampleIndex), FText::AsNumber(SampleIndex)) : 
+				FText::Format(LOCTEXT("FrameIndexSampleNameFormat", "{0} ({1}) (Frame {2})"), GetSampleName(Sample, SampleIndex), FText::AsNumber(SampleIndex), FText::AsNumber(Sample.FrameIndexToSample));
+
 			const FVector2D TextSize = FontMeasure->Measure(Name, FontInfo);
 			const FVector2D Padding = FVector2D(12.0f, 4.0f);
 
@@ -1940,7 +1943,9 @@ FText SBlendSpaceGridWidget::GetToolTipAnimationName() const
 					else if (HighlightedSampleIndex != INDEX_NONE && BlendSpace->IsValidBlendSampleIndex(HighlightedSampleIndex))
 					{
 						const FBlendSample& BlendSample = BlendSpace->GetBlendSample(HighlightedSampleIndex);
-						ToolTipText = GetSampleName(BlendSample, HighlightedSampleIndex);
+						ToolTipText = !BlendSample.bUseSingleFrameForBlending ? 
+							FText::Format(LOCTEXT("SampleNameFormat", "{0} ({1})"), GetSampleName(BlendSample, HighlightedSampleIndex), FText::AsNumber(HighlightedSampleIndex)) : 
+							FText::Format(LOCTEXT("FrameIndexSampleNameFormat", "{0} ({1}) (Frame {2})"), GetSampleName(BlendSample, HighlightedSampleIndex), FText::AsNumber(HighlightedSampleIndex), FText::AsNumber(BlendSample.FrameIndexToSample));
 					}
 					else if(TargetPosition.IsSet())
 					{
@@ -1960,7 +1965,9 @@ FText SBlendSpaceGridWidget::GetToolTipAnimationName() const
 					if (BlendSpace->IsValidBlendSampleIndex(DraggedSampleIndex))
 					{
 						const FBlendSample& BlendSample = BlendSpace->GetBlendSample(DraggedSampleIndex);
-						ToolTipText = GetSampleName(BlendSample, DraggedSampleIndex);
+						ToolTipText = !BlendSample.bUseSingleFrameForBlending ? 
+							FText::Format(LOCTEXT("SampleNameFormat", "{0} ({1})"), GetSampleName(BlendSample, DraggedSampleIndex), FText::AsNumber(DraggedSampleIndex)) : 
+							FText::Format(LOCTEXT("FrameIndexSampleNameFormat", "{0} ({1}) (Frame {2})"), GetSampleName(BlendSample, DraggedSampleIndex), FText::AsNumber(DraggedSampleIndex), FText::AsNumber(BlendSample.FrameIndexToSample));
 					}
 			
 					break;
@@ -2021,8 +2028,10 @@ FText SBlendSpaceGridWidget::GetToolTipSampleValue() const
 					if(BlendSpace->IsValidBlendSampleIndex(SampleData.SampleDataIndex))
 					{
 						const FBlendSample& BlendSample = BlendSpace->GetBlendSample(SampleData.SampleDataIndex);
-						static const FTextFormat SampleFormat = LOCTEXT("SampleFormat", "{0}: {1}");
-						TextBuilder.AppendLine(FText::Format(SampleFormat, GetSampleName(BlendSample, SampleData.SampleDataIndex), FText::AsNumber(SampleData.TotalWeight)));
+						const FText SampleNameWithWeight = !BlendSample.bUseSingleFrameForBlending ? 
+							FText::Format(LOCTEXT("SampleFormat", "{0} ({1}): {2}"), GetSampleName(BlendSample, SampleData.SampleDataIndex), FText::AsNumber(SampleData.SampleDataIndex),FText::AsNumber(SampleData.TotalWeight)) : 
+							FText::Format(LOCTEXT("FrameIndexSampleFormat", "{0} ({1}) (Frame {2}): {3}"), GetSampleName(BlendSample, SampleData.SampleDataIndex), FText::AsNumber(SampleData.SampleDataIndex), FText::AsNumber(BlendSample.FrameIndexToSample), FText::AsNumber(SampleData.TotalWeight));
+						TextBuilder.AppendLine(SampleNameWithWeight);
 					}
 				}
 			}
@@ -2182,6 +2191,14 @@ FText SBlendSpaceGridWidget::GetSampleErrorMessage(const FBlendSample &BlendSamp
 		static const FText SampleNotAtGridPoint = LOCTEXT("SampleNotAtGridPointErrorText", "Sample is not on a valid Grid Point");
 		return SampleNotAtGridPoint;
 	}
+	else if (const UBlendSpace* BlendSpaceBaseObject = BlendSpaceBase.Get())
+	{
+		if (!BlendSpaceBaseObject->IsSingleFrameBlendingIndexInBounds(BlendSample))
+		{
+			static const FText SamplingIndexOutOfBounds = LOCTEXT("SingleFrameBlendingIndexOutOfBounds", "Single frame blending index is out of bounds.");
+			return SamplingIndexOutOfBounds;
+		}
+	}
 
 	static const FText UnknownError = LOCTEXT("UnknownErrorText", "Sample is invalid for an Unknown Reason");
 	return UnknownError;
diff --git a/Engine/Source/Runtime/AnimGraphRuntime/Private/AnimNodes/AnimNode_BlendSpacePlayer.cpp b/Engine/Source/Runtime/AnimGraphRuntime/Private/AnimNodes/AnimNode_BlendSpacePlayer.cpp
--- a/Engine/Source/Runtime/AnimGraphRuntime/Private/AnimNodes/AnimNode_BlendSpacePlayer.cpp
+++ b/Engine/Source/Runtime/AnimGraphRuntime/Private/AnimNodes/AnimNode_BlendSpacePlayer.cpp
@@ -44,10 +44,7 @@ float FAnimNode_BlendSpacePlayerBase::GetCurrentAssetLength() const
 		if (CurrentBlendSpace != nullptr)
 		{
 			const FBlendSample& Sample = CurrentBlendSpace->GetBlendSample(HighestWeightedSample->SampleDataIndex);
-			if (ensure(Sample.Animation != nullptr))
-			{
-				return Sample.Animation->GetPlayLength();
-			}
+			return Sample.GetSamplePlayLength();
 		}
 	}
 
diff --git a/Engine/Source/Runtime/Engine/Classes/Animation/BlendSpace.h b/Engine/Source/Runtime/Engine/Classes/Animation/BlendSpace.h
--- a/Engine/Source/Runtime/Engine/Classes/Animation/BlendSpace.h
+++ b/Engine/Source/Runtime/Engine/Classes/Animation/BlendSpace.h
@@ -176,6 +176,13 @@ struct FBlendSample
 	
 	UPROPERTY(EditAnywhere, Category = BlendSample, meta=(UIMin="0.01", UIMax="2.0", ClampMin="0.01", ClampMax="64.0"))
 	float RateScale = 1.0f;
+	
+	UPROPERTY(EditAnywhere, Category = BlendSample, meta=(InlineEditConditionToggle))
+	bool bUseSingleFrameForBlending = false;
+	
+	// Single Frame Blending: If selected, instead of using the entire animation sequence, only the specified single frame is used for sampling the animation
+	UPROPERTY(EditAnywhere, Category = BlendSample, meta=(EditCondition="bUseSingleFrameForBlending", DisplayName="Use Single Frame", UIMin="0"))
+	uint32 FrameIndexToSample = 0;
 
 #if WITH_EDITORONLY_DATA
 	// Whether or not this sample will be moved when the "analyse all" button is used. Note that, even if disabled,
@@ -200,8 +207,8 @@ struct FBlendSample
 		, bIsValid(false)
 		, CachedMarkerDataUpdateCounter(INDEX_NONE)
 #endif // WITH_EDITORONLY_DATA
-{		
-}
+	{		
+	}
 	
 	FBlendSample(class UAnimSequence* InAnim, FVector InValue, bool bInIsSnapped, bool bInIsValid) 
 		: Animation(InAnim)
@@ -217,8 +224,10 @@ struct FBlendSample
 	
 	bool operator==( const FBlendSample& Other ) const 
 	{
-		return (Other.Animation == Animation && Other.SampleValue == SampleValue && FMath::IsNearlyEqual(Other.RateScale, RateScale));
+		return (Other.Animation == Animation && Other.SampleValue == SampleValue && Other.bUseSingleFrameForBlending == bUseSingleFrameForBlending && Other.FrameIndexToSample == FrameIndexToSample && FMath::IsNearlyEqual(Other.RateScale, RateScale));
 	}
+
+	ENGINE_API float GetSamplePlayLength() const;
 };
 
 /**
@@ -463,6 +472,7 @@ public:
 	friend class FBlendSpaceDetails;
 	friend class FBlendSampleDetails;
 	friend class UAnimGraphNode_BlendSpaceGraphBase;
+	friend class SBlendSpaceGridWidget;
 
 	//~ Begin UObject Interface
 	virtual void PostLoad() override;
@@ -729,13 +739,14 @@ protected:
 	/** If around border, snap to the border to avoid empty hole of data that is not valid **/
 	virtual void SnapSamplesToClosestGridPoint();
 #endif // WITH_EDITOR
-	
+
+	ENGINE_API static bool IsSingleFrameBlendingIndexInBounds(const FBlendSample& BlendSample);
 private:
 	// Internal helper function for GetAnimationPose variants
 	void GetAnimationPose_Internal(TArray<FBlendSampleData>& BlendSampleDataCache, TArrayView<FPoseLink> InPoseLinks, FAnimInstanceProxy* InProxy, bool bInExpectsAdditivePose, const FAnimExtractContext& ExtractionContext, /*out*/ FAnimationPoseData& OutAnimationPoseData) const;
 
 	// Internal helper function for UpdateBlendSamples and TickAssetPlayer
-	bool UpdateBlendSamples_Internal(const FVector& InBlendSpacePosition, float InDeltaTime, TArray<FBlendSampleData>& InOutOldSampleDataList, TArray<FBlendSampleData>& InOutSampleDataCache, int32& InOutCachedTriangulationIndex) const;
+	bool UpdateBlendSamples_Internal(const FVector& InBlendSpacePosition, float InDeltaTime, TArray<FBlendSampleData>& InOutOldSampleDataList, TArray<FBlendSampleData>& InOutSampleDataCache, int32& InOutCachedTriangulationIndex) const;	
 
 public:
 #if WITH_EDITORONLY_DATA
diff --git a/Engine/Source/Runtime/Engine/Private/Animation/AnimSequence.cpp b/Engine/Source/Runtime/Engine/Private/Animation/AnimSequence.cpp
--- a/Engine/Source/Runtime/Engine/Private/Animation/AnimSequence.cpp
+++ b/Engine/Source/Runtime/Engine/Private/Animation/AnimSequence.cpp
@@ -3311,11 +3311,17 @@ void UAnimSequence::RefreshSyncMarkerDataFromAuthored()
 	check(IsInGameThread());
 
 	// Update blend spaces that may be referencing us
+	TArray<UAnimationAsset*> ReferredAssets;
 	for(TObjectIterator<UBlendSpace> It; It; ++It)
 	{
 		if(!It->HasAnyFlags(RF_NeedLoad | RF_NeedPostLoad))
 		{
-			It->RuntimeValidateMarkerData();
+			ReferredAssets.Reset();
+			It->GetAllAnimationSequencesReferred(ReferredAssets, false);
+			if (ReferredAssets.Contains(this))
+			{
+				It->RuntimeValidateMarkerData();
+			}
 		}
 	}
 #endif
diff --git a/Engine/Source/Runtime/Engine/Private/Animation/BlendSpace.cpp b/Engine/Source/Runtime/Engine/Private/Animation/BlendSpace.cpp
--- a/Engine/Source/Runtime/Engine/Private/Animation/BlendSpace.cpp
+++ b/Engine/Source/Runtime/Engine/Private/Animation/BlendSpace.cpp
@@ -20,6 +20,7 @@
 #include "UObject/UObjectIterator.h"
 #include "Logging/MessageLog.h"
 #include "UObject/UnrealType.h"
+#include "UObject/UObjectThreadContext.h"
 #if WITH_EDITOR
 #endif
 
@@ -320,8 +321,8 @@ void UBlendSpace::ResetBlendSamples(TArray<FBlendSampleData>& InOutSampleDataCac
 		
 		// Reset leader sample to match requested time.
 		LeaderSampleData.MarkerTickRecord.Reset();
-		LeaderSampleData.PreviousTime = InNormalizedCurrentTime * LeaderSampleData.Animation->GetPlayLength();
-		LeaderSampleData.Time = InNormalizedCurrentTime * LeaderSampleData.Animation->GetPlayLength();
+		LeaderSampleData.PreviousTime = InNormalizedCurrentTime * LeaderSample.GetSamplePlayLength();
+		LeaderSampleData.Time = InNormalizedCurrentTime * LeaderSample.GetSamplePlayLength();
 
 		// Query valid marker position for leader.
 		LeaderSample.Animation->GetMarkerIndicesForTime(LeaderSampleData.Time, bLooping, Context.MarkerTickContext.GetValidMarkerNames(), LeaderSampleData.MarkerTickRecord.PreviousMarker, LeaderSampleData.MarkerTickRecord.NextMarker);
@@ -389,8 +390,8 @@ void UBlendSpace::ResetBlendSamples(TArray<FBlendSampleData>& InOutSampleDataCac
 				{
 					// Fallback to length based syncing so it matches default behaviour when ticking a blend space.
 					SampleDataItem.MarkerTickRecord.Reset();
-					SampleDataItem.PreviousTime = InNormalizedCurrentTime * Sample.Animation->GetPlayLength();
-					SampleDataItem.Time = InNormalizedCurrentTime * Sample.Animation->GetPlayLength();
+					SampleDataItem.PreviousTime = InNormalizedCurrentTime * Sample.GetSamplePlayLength();
+					SampleDataItem.Time = InNormalizedCurrentTime * Sample.GetSamplePlayLength();
 				}
 			}
 		}
@@ -403,8 +404,8 @@ void UBlendSpace::ResetBlendSamples(TArray<FBlendSampleData>& InOutSampleDataCac
 			FBlendSampleData& SampleDataItem = InOutSampleDataCache[SampleIndex];
 			
 			SampleDataItem.MarkerTickRecord.Reset();
-			SampleDataItem.PreviousTime = InNormalizedCurrentTime * SampleDataItem.Animation->GetPlayLength();
-			SampleDataItem.Time = InNormalizedCurrentTime * SampleDataItem.Animation->GetPlayLength();
+			SampleDataItem.PreviousTime = InNormalizedCurrentTime * SampleData[SampleDataItem.SampleDataIndex].GetSamplePlayLength();
+			SampleDataItem.Time = InNormalizedCurrentTime * SampleData[SampleDataItem.SampleDataIndex].GetSamplePlayLength();
 		}
 	}
 }
@@ -661,7 +662,7 @@ void UBlendSpace::TickAssetPlayer(FAnimTickRecord& Instance, struct FAnimNotifyQ
 						if (!Instance.MarkerTickRecord->IsValid(Instance.bLooping))
 						{
 							LeaderSampleData.MarkerTickRecord.Reset();
-							LeaderSampleData.Time = NormalizedCurrentTime * LeaderSample.Animation->GetPlayLength();
+							LeaderSampleData.Time = NormalizedCurrentTime * LeaderSample.GetSamplePlayLength();
 							bResetMarkerDataOnFollowers = true;
 						}
 						// Re-compute marker indices since the leader sample's tick record is invalid. Get previous and next markers.
@@ -699,7 +700,7 @@ void UBlendSpace::TickAssetPlayer(FAnimTickRecord& Instance, struct FAnimNotifyQ
 							TickFollowerSamples(SampleDataList, HighestMarkerSyncWeightIndex, Context, true, Instance.bLooping, Instance.MirrorDataTable);
 						}
 						
-						NormalizedCurrentTime = LeaderSampleData.Time / LeaderSample.Animation->GetPlayLength();
+						NormalizedCurrentTime = LeaderSampleData.Time / LeaderSample.GetSamplePlayLength();
 						*Instance.MarkerTickRecord = LeaderSampleData.MarkerTickRecord;
 					}
 				}
@@ -740,7 +741,7 @@ void UBlendSpace::TickAssetPlayer(FAnimTickRecord& Instance, struct FAnimNotifyQ
 					{
 						if (!Instance.MarkerTickRecord->IsValid(Instance.bLooping))
 						{
-							SampleDataItem.Time = NormalizedCurrentTime * Sample.Animation->GetPlayLength();
+							SampleDataItem.Time = NormalizedCurrentTime * Sample.GetSamplePlayLength();
 						}
 
 						// Tick all samples as followers
@@ -748,7 +749,7 @@ void UBlendSpace::TickAssetPlayer(FAnimTickRecord& Instance, struct FAnimNotifyQ
 					}
 						
 					*Instance.MarkerTickRecord = SampleDataItem.MarkerTickRecord;
-					NormalizedCurrentTime = SampleDataItem.Time / Sample.Animation->GetPlayLength();
+					NormalizedCurrentTime = SampleDataItem.Time / Sample.GetSamplePlayLength();
 				}
 			}
 			else
@@ -808,8 +809,15 @@ void UBlendSpace::TickAssetPlayer(FAnimTickRecord& Instance, struct FAnimNotifyQ
 						{
 							const float SampleNormalizedPreviousTime = MultipliedSampleRateScale >= 0.f ? ClampedNormalizedPreviousTime : 1.f - ClampedNormalizedPreviousTime;
 							const float SampleNormalizedCurrentTime = MultipliedSampleRateScale >= 0.f ? ClampedNormalizedCurrentTime : 1.f - ClampedNormalizedCurrentTime;
-							PrevSampleDataTime = SampleNormalizedPreviousTime * Sample.Animation->GetPlayLength();
-							CurrentSampleDataTime = SampleNormalizedCurrentTime * Sample.Animation->GetPlayLength();
+							if (Sample.bUseSingleFrameForBlending)
+							{
+								PrevSampleDataTime = CurrentSampleDataTime = Sample.Animation->GetTimeAtFrame(Sample.FrameIndexToSample);
+							}
+							else
+							{
+								PrevSampleDataTime = SampleNormalizedPreviousTime * Sample.GetSamplePlayLength();
+								CurrentSampleDataTime = SampleNormalizedCurrentTime * Sample.GetSamplePlayLength();
+							}
 						}
 						else
 						{
@@ -823,7 +831,7 @@ void UBlendSpace::TickAssetPlayer(FAnimTickRecord& Instance, struct FAnimNotifyQ
 						// if we went against play rate, then loop around.
 						if ((SampleMoveDelta * DeltaTimePosition) < 0.f)
 						{
-							DeltaTimePosition += FMath::Sign<float>(SampleMoveDelta) * Sample.Animation->GetPlayLength();
+							DeltaTimePosition += FMath::Sign<float>(SampleMoveDelta) * Sample.GetSamplePlayLength();
 						}
 
 						if (bGenerateNotifies && (!bTriggerNotifyHighestWeightedAnim || (I == HighestWeightIndex)))
@@ -1204,7 +1212,7 @@ void UBlendSpace::GetAnimationPose_Internal(TArray<FBlendSampleData>& BlendSampl
 			{
 				if (Sample.Animation && Sample.Animation->GetSkeleton() != nullptr)
 				{
-					const float Time = FMath::Clamp<float>(BlendSampleDataCache[I].Time, 0.f, Sample.Animation->GetPlayLength());
+					const float Time = FMath::Clamp<float>(BlendSampleDataCache[I].Time, 0.f, Sample.GetSamplePlayLength());
 
 					FAnimationPoseData ChildAnimationPoseData = { Pose, ChildrenCurves[I], ChildrenAttributes[I] };
 					// first one always fills up the source one
@@ -1346,9 +1354,13 @@ bool UBlendSpace::GetSamplesFromBlendInput(
 			for (int32 Index2 = Index1 + 1; Index2 < OutSampleDataList.Num(); ++Index2)
 			{
 				FBlendSampleData* SecondSample = &OutSampleDataList[Index2];
+				
+				const bool bBothUseSingleFrameForBlending = SampleData[FirstSample->SampleDataIndex].bUseSingleFrameForBlending && SampleData[SecondSample->SampleDataIndex].bUseSingleFrameForBlending;
+				const bool bFrameIndicesMatch = SampleData[FirstSample->SampleDataIndex].FrameIndexToSample == SampleData[SecondSample->SampleDataIndex].FrameIndexToSample;				
 				// if they have same sample, remove the Index2, and get out
 				if (FirstSample->SampleDataIndex == SecondSample->SampleDataIndex || // Shouldn't happen
-					(FirstSample->Animation != nullptr && FirstSample->Animation == SecondSample->Animation))
+					(FirstSample->Animation != nullptr && FirstSample->Animation == SecondSample->Animation && 
+					bBothUseSingleFrameForBlending && bFrameIndicesMatch))
 				{
 					//Calc New Sample Playrate
 					const float TotalWeight = FirstSample->GetClampedWeight() + SecondSample->GetClampedWeight();
@@ -1469,18 +1481,23 @@ void UBlendSpace::ValidateSampleData()
 
 		if (IsAsset())
 		{
-			bool bAnimationExists = Sample.Animation != nullptr;
-			bool bSampleInBounds = bAnimationExists ? IsSampleWithinBounds(Sample.SampleValue) : true;
-			bool bSampleIsUnique = bAnimationExists ? !IsTooCloseToExistingSamplePoint(Sample.SampleValue, SampleIndex) : true;
-
-			Sample.bIsValid = bAnimationExists && bSampleInBounds && bSampleIsUnique;
+			const bool bAnimationExists = Sample.Animation != nullptr;
+			const bool bSampleInBounds = bAnimationExists ? IsSampleWithinBounds(Sample.SampleValue) : true;
+			const bool bSampleIsUnique = bAnimationExists ? !IsTooCloseToExistingSamplePoint(Sample.SampleValue, SampleIndex) : true;
+			const bool bShouldForceConditionalPostLoad = bAnimationExists && FUObjectThreadContext::Get().IsRoutingPostLoad && !Sample.Animation->IsDataModelValid();
+			if (bShouldForceConditionalPostLoad)
+			{
+				Sample.Animation->ConditionalPostLoad();
+			}
+			const bool bSingleFrameBlendingIndexInBounds = IsSingleFrameBlendingIndexInBounds(Sample);
+			Sample.bIsValid = bAnimationExists && bSampleInBounds && bSampleIsUnique && bSingleFrameBlendingIndexInBounds;
 
 			if (Sample.bIsValid)
 			{
-				if (Sample.Animation->GetPlayLength() > AnimLength)
+				if (Sample.GetSamplePlayLength() > AnimLength)
 				{
 					// @todo : should apply scale? If so, we'll need to apply also when blend
-					AnimLength = Sample.Animation->GetPlayLength();
+					AnimLength = Sample.GetSamplePlayLength();
 				}
 
 				Sample.CachedMarkerDataUpdateCounter = Sample.Animation->GetMarkerUpdateCounter();
@@ -1547,6 +1564,14 @@ void UBlendSpace::ValidateSampleData()
 						Message->AddToken(FAssetNameToken::Create(Sample.Animation->GetPathName(), FText::FromString(Sample.Animation->GetName())));
 						Message->AddToken(FTextToken::Create(LOCTEXT("EmptyAnimationData_NotUnique", "that is not unique. Recommend adjusting it.")));
 					}
+					else if (!bSingleFrameBlendingIndexInBounds)
+					{
+						Message->AddToken(FTextToken::Create(HasSampleText));
+						Message->AddToken(FAssetNameToken::Create(Sample.Animation->GetPathName(), FText::FromString(Sample.Animation->GetName())));
+						Message->AddToken(FTextToken::Create(LOCTEXT("EmptyAnimationData_SingleFrameBlendingNotInBounds", "single frame blending index is out of bounds. Recommend adjusting it.")));
+						
+						UE_LOG(LogAnimation, Warning, TEXT("[%s : %d] - single frame blending index: [%d] is out of bounds"), *GetFullName(), SampleIndex + 1, Sample.FrameIndexToSample);
+					} 
 					else
 					{
 						// Shouldn't get here
@@ -1693,6 +1718,26 @@ bool UBlendSpace::IsValidBlendSampleIndex(const int32 SampleIndex) const
 	return SampleData.IsValidIndex(SampleIndex);
 }
 
+bool UBlendSpace::IsSingleFrameBlendingIndexInBounds(const FBlendSample& BlendSample)
+{
+	if (BlendSample.bUseSingleFrameForBlending)
+	{
+		if (BlendSample.Animation && BlendSample.Animation->IsDataModelValid())
+		{
+			const uint32 NumberOfFrames = static_cast<uint32>(BlendSample.Animation->GetDataModelInterface()->GetNumberOfFrames());		
+			if (BlendSample.FrameIndexToSample >= 0 && BlendSample.FrameIndexToSample <= NumberOfFrames)
+			{
+				return true;
+			}
+		}
+
+		// either we are not using single frame blending or index is outside of bounds (or data model is invalid, in which case can't report this error properly)
+		return false;
+	}
+	
+	return true;
+}
+
 const TArray<FEditorElement>& UBlendSpace::GetGridSamples() const
 {
 	return GridSamples;
@@ -1800,6 +1845,7 @@ bool UBlendSpace::ValidateSampleValue(const FVector& SampleValue, int32 Original
 	bool bValid = true;
 	bValid &= IsSampleWithinBounds(SampleValue);
 	bValid &= !IsTooCloseToExistingSamplePoint(SampleValue, OriginalIndex);
+	bValid &= IsSingleFrameBlendingIndexInBounds(SampleData[OriginalIndex]);
 	return bValid;
 }
 
@@ -1937,7 +1983,8 @@ float UBlendSpace::GetAnimationLengthFromSampleData(const TArray<FBlendSampleDat
 				//Multiple samples contribution which we would otherwise lose
 				const float MultipliedSampleRateScale = Sample.Animation->RateScale * SampleDataList[I].SamplePlayRate;
 				// apply rate scale to get actual playback time
-				BlendAnimLength += (Sample.Animation->GetPlayLength() / ((MultipliedSampleRateScale) != 0.0f ? FMath::Abs(MultipliedSampleRateScale) : 1.0f)) * SampleDataList[I].GetClampedWeight();
+				BlendAnimLength += (Sample.GetSamplePlayLength() / ((MultipliedSampleRateScale) != 0.0f ? FMath::Abs(MultipliedSampleRateScale) : 1.0f)) * SampleDataList[I].GetClampedWeight();
+
 				UE_LOG(LogAnimation, Verbose, TEXT("[%d] - Sample Animation(%s) : Weight(%0.5f) "), I + 1, *Sample.Animation->GetName(), SampleDataList[I].GetClampedWeight());
 			}
 		}
@@ -3011,6 +3058,21 @@ void FBlendSpaceData::GetSamples2D(
 	UE_LOG(LogAnimation, Warning, TEXT("Unable to find BlendSpace triangle"));
 }
 
+float FBlendSample::GetSamplePlayLength() const
+{
+	if (Animation)
+	{
+		if (bUseSingleFrameForBlending)
+		{
+			return Animation->GetSamplingFrameRate().AsInterval();
+		}
+
+		return Animation->GetPlayLength();
+	}
+
+	return 0.f;
+}
+
 void FBlendSpaceData::GetSamples(
 	TArray<FWeightedBlendSample>& OutWeightedSamples,
 	const TArray<int32>&          InDimensionIndices,
