[AutoRTFM] Simplify Tickable test.

This now uses the [SupportsAutoRTFM] flag instead of manually testing commit and abort behaviors.

diff --git a/Engine/Source/Runtime/Engine/Private/Tests/AutoRTFM/AutoRTFMTickableTests.cpp b/Engine/Source/Runtime/Engine/Private/Tests/AutoRTFM/AutoRTFMTickableTests.cpp
--- a/Engine/Source/Runtime/Engine/Private/Tests/AutoRTFM/AutoRTFMTickableTests.cpp
+++ b/Engine/Source/Runtime/Engine/Private/Tests/AutoRTFM/AutoRTFMTickableTests.cpp
@@ -1,43 +1,21 @@
 // Copyright Epic Games, Inc. All Rights Reserved.
 
-#include "Misc/AutomationTest.h"
+#include "Tests/TestHarnessAdapter.h"
 #include "AutoRTFM.h"
 #include "Tickable.h"
 
 #if WITH_DEV_AUTOMATION_TESTS
 
-IMPLEMENT_SIMPLE_AUTOMATION_TEST(FAutoRTFMTickableTests, "AutoRTFM + FTickableObject", EAutomationTestFlags::EngineFilter | EAutomationTestFlags::ClientContext | EAutomationTestFlags::ServerContext | EAutomationTestFlags::CommandletContext)
-
-bool FAutoRTFMTickableTests::RunTest(const FString & Parameters)
+TEST_CASE_NAMED(FAutoRTFMTickableTests, "AutoRTFM.FTickableObject", "[EngineFilter][ClientContext][ServerContext][CommandletContext][SupportsAutoRTFM]")
 {
-	if (!AutoRTFM::ForTheRuntime::IsAutoRTFMRuntimeEnabled())
-	{
-		ExecutionInfo.AddEvent(FAutomationEvent(EAutomationEventType::Info, TEXT("SKIPPED 'FAutoRTFMTickableTests' test. AutoRTFM disabled.")));
-		return true;
-	}
-
 	struct FMyTickableGameObject final : FTickableGameObject
 	{
 		void Tick(float) override {}
 		TStatId GetStatId() const override { return TStatId(); }
 	};
 
-	AutoRTFM::ETransactionResult Result = AutoRTFM::Transact([&]
-	{
-		FMyTickableGameObject Tickable;
-		Tickable.SetTickableTickType(ETickableTickType::Always);
-		AutoRTFM::AbortTransaction();
-	});
-	TestTrueExpr(AutoRTFM::ETransactionResult::AbortedByRequest == Result);
-
-	Result = AutoRTFM::Transact([&]
-	{
-		FMyTickableGameObject Tickable;
-		Tickable.SetTickableTickType(ETickableTickType::Always);
-	});
-	TestTrueExpr(AutoRTFM::ETransactionResult::Committed == Result);
-
-	return true;
+	FMyTickableGameObject Tickable;
+	Tickable.SetTickableTickType(ETickableTickType::Always);
 }
 
 #endif //WITH_DEV_AUTOMATION_TESTS
