Texture.cpp nop comment only

diff --git a/Engine/Source/Runtime/Engine/Private/Texture.cpp b/Engine/Source/Runtime/Engine/Private/Texture.cpp
--- a/Engine/Source/Runtime/Engine/Private/Texture.cpp
+++ b/Engine/Source/Runtime/Engine/Private/Texture.cpp
@@ -1871,6 +1871,10 @@ bool UTexture::IsPossibleToStream() const
 	{
 		// turn off streaming for non-pow2 sources :
 
+		// note: this replicates the legacy logic, but it is not actually correct
+		//	for example it does not account for options like ResizeToSpecific
+		// more correct would be to use GetBuiltTextureSize here.
+
 		// AreAllBlocksPowerOfTwo only checks XY :
 		bool bIsPowerOfTwo = Source.AreAllBlocksPowerOfTwo();
 		if ( ! FMath::IsPowerOfTwo(Source.GetVolumeSizeZ()) )
@@ -1880,6 +1884,7 @@ bool UTexture::IsPossibleToStream() const
 		if ( PowerOfTwoMode != ETexturePowerOfTwoSetting::None )
 		{
 			bIsPowerOfTwo = true;
+			// @@ not correct for ETexturePowerOfTwoSetting::ResizeToSpecificResolution
 		}
 		if ( Source.IsLongLatCubemap() )
 		{
