Adding a speculative fix for a crash with an invalid body instance.
The working theory is that an object is detached from a skeletal mesh but the body is still in the list and is still welded.

diff --git a/Engine/Source/Runtime/Engine/Private/SkeletalMeshComponentPhysics.cpp b/Engine/Source/Runtime/Engine/Private/SkeletalMeshComponentPhysics.cpp
--- a/Engine/Source/Runtime/Engine/Private/SkeletalMeshComponentPhysics.cpp
+++ b/Engine/Source/Runtime/Engine/Private/SkeletalMeshComponentPhysics.cpp
@@ -1022,6 +1022,11 @@ void USkeletalMeshComponent::TermArticulated()
 	{
 		if (ensure(Bodies[i]))
 		{
+			// There seems to be scenarios where there is a body that still has an active weld parent.
+			if (!ensure(Bodies[i]->WeldParent == nullptr))
+			{
+				Bodies[i]->WeldParent->UnWeld(Bodies[i]);
+			}
 			Bodies[i]->TermBody();
 			delete Bodies[i];
 		}
