New Material Translator:  support for atan2 expression.

#rnx

diff --git a/Engine/Source/Runtime/Engine/Private/Materials/MaterialExpressionsIR.cpp b/Engine/Source/Runtime/Engine/Private/Materials/MaterialExpressionsIR.cpp
--- a/Engine/Source/Runtime/Engine/Private/Materials/MaterialExpressionsIR.cpp
+++ b/Engine/Source/Runtime/Engine/Private/Materials/MaterialExpressionsIR.cpp
@@ -673,7 +673,7 @@ void UMaterialExpressionDistance::Build(MIR::FEmitter& Em)
 }
 
 void UMaterialExpressionFmod::Build(MIR::FEmitter& Em)
-{ 
+{
 	Em.Output(0, Em.Fmod(Em.Input(&A), Em.Input(&B)));
 }
 
@@ -719,6 +719,11 @@ void UMaterialExpressionCrossProduct::Build(MIR::FEmitter& Em)
 	Em.Output(0, Em.Cross(Em.Input(&A), Em.Input(&B)));
 }
 
+void UMaterialExpressionArctangent2::Build(MIR::FEmitter& Em)
+{
+	Em.Output(0, Em.Atan2(Em.Input(&Y), Em.Input(&X)));
+}
+
 void UMaterialExpressionEyeAdaptationInverse::Build(MIR::FEmitter& Em)
 {
 	check(ExternalCodeIdentifiers.Num() == 1);
diff --git a/Engine/Source/Runtime/Engine/Private/Materials/MaterialIR.cpp b/Engine/Source/Runtime/Engine/Private/Materials/MaterialIR.cpp
--- a/Engine/Source/Runtime/Engine/Private/Materials/MaterialIR.cpp
+++ b/Engine/Source/Runtime/Engine/Private/Materials/MaterialIR.cpp
@@ -650,6 +650,7 @@ const TCHAR* LexToString(EOperator Op)
 		/* Binary operators */
 		case BO_Add: return TEXT("Add");
 		case BO_And: return TEXT("And");
+		case BO_Atan2: return TEXT("Atan2");
 		case BO_BitShiftLeft: return TEXT("BitShiftLeft");
 		case BO_BitShiftRight: return TEXT("BitShiftRight");
 		case BO_BitwiseAnd: return TEXT("BitwiseAnd");
diff --git a/Engine/Source/Runtime/Engine/Private/Materials/MaterialIREmitter.cpp b/Engine/Source/Runtime/Engine/Private/Materials/MaterialIREmitter.cpp
--- a/Engine/Source/Runtime/Engine/Private/Materials/MaterialIREmitter.cpp
+++ b/Engine/Source/Runtime/Engine/Private/Materials/MaterialIREmitter.cpp
@@ -928,6 +928,7 @@ static T FoldScalarOperator(FEmitter& Emitter, EOperator Operator, T A, T B, T C
 				{
 					return FGenericPlatformMath::TruncToDouble(A);
 				}
+			case BO_Atan2: return FGenericPlatformMath::Atan2(A, B);
 			case BO_Fmod: return FGenericPlatformMath::Fmod(A, B);
 			case BO_Pow: return FGenericPlatformMath::Pow(A, B);
 			case TO_Lerp: return FMath::Lerp<T>(A, B, C);
@@ -1322,6 +1323,7 @@ static const FOperatorSignature* GetOperatorSignature(EOperator Op)
 		S[BO_BitShiftLeft]				= BinaryInteger;
 		S[BO_BitShiftRight]				= BinaryInteger;
 
+		S[BO_Atan2]						= BinaryFloat;
 		S[BO_Cross] 					= { { OPTF_CheckIsArithmetic | OPTF_CheckIsVector3, OPTF_CastToFirstArgumentType }, ORT_FirstArgumentType };
 		S[BO_Distance]					= BinaryFloat;
 		S[BO_Dot] 						= { { OPTF_CheckIsArithmetic | OPTF_CheckIsNotMatrix, OPTF_CastToFirstArgumentType }, ORT_FirstArgumentTypeToScalar };
@@ -2573,6 +2575,13 @@ static FValue* DifferentiateOperator(FEmitter& E, FOperator* Op, EDerivativeAxis
 			return E.Multiply(Op, E.Add(Term1, Term2));
 		}
 
+		// d/dx atan2(f(x), g(x)) = g(x) / (f(x)^2 + g(x)^2) * f'(x)  -  f(x) / (f(x)^2 + g(x)^2) * g'(x)
+		case BO_Atan2:
+		{
+			FValue* Magnitude = E.Divide(One(), E.Add(E.Multiply(F, F), E.Multiply(G, G)));		// 1 / (f(x)^2 + g(x)^2)
+			return E.Subtract(E.Multiply(E.Multiply(G, Magnitude), dF), E.Multiply(E.Multiply(F, Magnitude), dG));
+		}
+
 		// The multiplication rule applies for the dot product too.
 		// d/dx (f(x) • g(x)) = f'(x) • g(x) + f(x) • g'(x)
 		case BO_Dot:
diff --git a/Engine/Source/Runtime/Engine/Private/Materials/MaterialIRToHLSLTranslator.cpp b/Engine/Source/Runtime/Engine/Private/Materials/MaterialIRToHLSLTranslator.cpp
--- a/Engine/Source/Runtime/Engine/Private/Materials/MaterialIRToHLSLTranslator.cpp
+++ b/Engine/Source/Runtime/Engine/Private/Materials/MaterialIRToHLSLTranslator.cpp
@@ -895,6 +895,7 @@ struct FTranslator : FMaterialIRToHLSLTranslation
 				case MIR::UO_Truncate: OpString = TEXT("trunc"); break;
 
 				case MIR::BO_And: OpString = TEXT("and"); break;
+				case MIR::BO_Atan2: OpString = TEXT("atan2"); break;
 				case MIR::BO_Cross: OpString = TEXT("cross"); break;
 				case MIR::BO_Distance: OpString = TEXT("distance"); break;
 				case MIR::BO_Dot: OpString = TEXT("dot"); break;
diff --git a/Engine/Source/Runtime/Engine/Public/Materials/MaterialExpressionArctangent2.h b/Engine/Source/Runtime/Engine/Public/Materials/MaterialExpressionArctangent2.h
--- a/Engine/Source/Runtime/Engine/Public/Materials/MaterialExpressionArctangent2.h
+++ b/Engine/Source/Runtime/Engine/Public/Materials/MaterialExpressionArctangent2.h
@@ -23,6 +23,7 @@ public:
 
 	//~ Begin UMaterialExpression Interface
 #if WITH_EDITOR
+	virtual void Build(MIR::FEmitter& Emitter) override;
 	virtual int32 Compile(class FMaterialCompiler* Compiler, int32 OutputIndex) override;
 	virtual void GetCaption(TArray<FString>& OutCaptions) const override;
 	virtual void GetExpressionToolTip(TArray<FString>& OutToolTip) override;
diff --git a/Engine/Source/Runtime/Engine/Public/Materials/MaterialExpressionOperator.h b/Engine/Source/Runtime/Engine/Public/Materials/MaterialExpressionOperator.h
--- a/Engine/Source/Runtime/Engine/Public/Materials/MaterialExpressionOperator.h
+++ b/Engine/Source/Runtime/Engine/Public/Materials/MaterialExpressionOperator.h
@@ -65,6 +65,7 @@ enum class EMaterialExpressionOperatorKind
 	BitwiseOr,
 	BitShiftLeft,
 	BitShiftRight,
+	Atan2,
 	Cross,
 	Distance,
 	Dot,
diff --git a/Engine/Source/Runtime/Engine/Public/Materials/MaterialIR.h b/Engine/Source/Runtime/Engine/Public/Materials/MaterialIR.h
--- a/Engine/Source/Runtime/Engine/Public/Materials/MaterialIR.h
+++ b/Engine/Source/Runtime/Engine/Public/Materials/MaterialIR.h
@@ -628,6 +628,7 @@ enum EOperator
 	BO_BitShiftRight,
 	
 	// Binary intrinsics
+	BO_Atan2,
 	BO_Cross,
 	BO_Distance,
 	BO_Dot,
diff --git a/Engine/Source/Runtime/Engine/Public/Materials/MaterialIREmitter.h b/Engine/Source/Runtime/Engine/Public/Materials/MaterialIREmitter.h
--- a/Engine/Source/Runtime/Engine/Public/Materials/MaterialIREmitter.h
+++ b/Engine/Source/Runtime/Engine/Public/Materials/MaterialIREmitter.h
@@ -504,6 +504,7 @@ public:
 	FValueRef Min(FValueRef A, FValueRef B) { return Operator(BO_Min, A, B); }
 	FValueRef Dot(FValueRef A, FValueRef B) { return Operator(BO_Dot, A, B); }
 	FValueRef Cross(FValueRef A, FValueRef B) { return Operator(BO_Cross, A, B); }
+	FValueRef Atan2(FValueRef A, FValueRef B) { return Operator(BO_Atan2, A, B); }
 
 	/*-------------------------- Ternary operators shortcuts ---------------------------*/
 
