UE: Fix bug in logic that detects if a child component is culled via detail mode when the parent is not  that was introduced with recent logging improvements

diff --git a/Engine/Source/Runtime/Engine/Private/Components/SceneComponent.cpp b/Engine/Source/Runtime/Engine/Private/Components/SceneComponent.cpp
--- a/Engine/Source/Runtime/Engine/Private/Components/SceneComponent.cpp
+++ b/Engine/Source/Runtime/Engine/Private/Components/SceneComponent.cpp
@@ -484,7 +484,10 @@ static bool SceneComponentNeedsLoadForTarget(USceneComponent const* SceneCompone
 					"SceneComponent {ChildComponent} will not be culled because its detail mode {DetailMode} is less than the maximum detail mode for the target platform {PlatformMaxDetailMode}, "
 					"but its parent {ParentComponent} will be culled. Make sure that attached components have a detail mode equal or lower to their parent.", 
 					FSoftObjectPath(SceneComponentObject), SceneComponentObject->DetailMode, CVarDetailMode, FSoftObjectPath(ParentThatWillBeCulled));
-				return bKeep;
+				if (!bKeep)
+				{
+					return false;
+				}					
 			}
 		}
 	}
