[StandaloneHLOD] Add editor experimental setting, controlling visibility of Standalone HLOD WP setting and creation of Standalone HLOD Subsystem. Disabled by default.

diff --git a/Engine/Source/Editor/UnrealEd/Classes/Settings/EditorExperimentalSettings.h b/Engine/Source/Editor/UnrealEd/Classes/Settings/EditorExperimentalSettings.h
--- a/Engine/Source/Editor/UnrealEd/Classes/Settings/EditorExperimentalSettings.h
+++ b/Engine/Source/Editor/UnrealEd/Classes/Settings/EditorExperimentalSettings.h
@@ -150,6 +150,9 @@ public:
 	UPROPERTY(EditAnywhere, config, Category = WorldPartition)
 	bool bEnableWorldPartitionExternalDataLayers;
 
+	UPROPERTY(EditAnywhere, config, Category = WorldPartition)
+	bool bEnableStandaloneHLOD;
+
 	UPROPERTY(EditAnywhere, config, Category = LevelInstance)
 	bool bEnableLevelInstancePropertyOverrides;
 
diff --git a/Engine/Source/Runtime/Engine/Private/WorldPartition/HLOD/StandaloneHLODSubsystem.cpp b/Engine/Source/Runtime/Engine/Private/WorldPartition/HLOD/StandaloneHLODSubsystem.cpp
--- a/Engine/Source/Runtime/Engine/Private/WorldPartition/HLOD/StandaloneHLODSubsystem.cpp
+++ b/Engine/Source/Runtime/Engine/Private/WorldPartition/HLOD/StandaloneHLODSubsystem.cpp
@@ -7,6 +7,7 @@
 #if WITH_EDITOR
 
 #include "LevelInstance/LevelInstanceActor.h"
+#include "Settings/EditorExperimentalSettings.h"
 #include "WorldPartition/ActorDescContainerSubsystem.h"
 #include "WorldPartition/HLOD/StandaloneHLODActor.h"
 #include "WorldPartition/LevelInstance/LevelInstanceActorDesc.h"
@@ -52,6 +53,11 @@ bool UWorldPartitionStandaloneHLODSubsystem::DoesSupportWorldType(const EWorldTy
 	return WorldType == EWorldType::Editor || WorldType == EWorldType::Inactive;
 }
 
+bool UWorldPartitionStandaloneHLODSubsystem::ShouldCreateSubsystem(UObject* Outer) const
+{
+	return Super::ShouldCreateSubsystem(Outer) && GetDefault<UEditorExperimentalSettings>()->bEnableStandaloneHLOD;
+}
+
 void UWorldPartitionStandaloneHLODSubsystem::OnWorldPartitionInitialized(UWorldPartition* InWorldPartition)
 {
 	InWorldPartition->OnActorDescContainerInstanceRegistered.AddUObject(this, &UWorldPartitionStandaloneHLODSubsystem::OnActorDescContainerInstanceRegistered);
diff --git a/Engine/Source/Runtime/Engine/Private/WorldPartition/WorldPartition.cpp b/Engine/Source/Runtime/Engine/Private/WorldPartition/WorldPartition.cpp
--- a/Engine/Source/Runtime/Engine/Private/WorldPartition/WorldPartition.cpp
+++ b/Engine/Source/Runtime/Engine/Private/WorldPartition/WorldPartition.cpp
@@ -41,6 +41,7 @@
 #include "LocationVolume.h"
 #include "Engine/ExternalObjectAndActorDependencyGatherer.h"
 #include "Engine/LevelScriptBlueprint.h"
+#include "Settings/EditorExperimentalSettings.h"
 #include "WorldPartition/DataLayer/WorldDataLayersActorDesc.h"
 #include "WorldPartition/IWorldPartitionEditorModule.h"
 #include "WorldPartition/WorldPartitionLevelHelper.h"
@@ -2397,6 +2398,11 @@ void UWorldPartition::OnLoaderAdapterStateChanged(IWorldPartitionActorLoaderInte
 	LoaderAdapterStateChanged.Broadcast(InLoaderAdapter);
 }
 
+bool UWorldPartition::IsStandaloneHLODAllowed() const
+{
+	return bEnableStreaming && GetDefault<UEditorExperimentalSettings>()->bEnableStandaloneHLOD;
+}
+
 void UWorldPartition::OnWorldRenamed(UWorld* RenamedWorld)
 {
 	if (GetWorld() == RenamedWorld)
diff --git a/Engine/Source/Runtime/Engine/Public/WorldPartition/HLOD/StandaloneHLODSubsystem.h b/Engine/Source/Runtime/Engine/Public/WorldPartition/HLOD/StandaloneHLODSubsystem.h
--- a/Engine/Source/Runtime/Engine/Public/WorldPartition/HLOD/StandaloneHLODSubsystem.h
+++ b/Engine/Source/Runtime/Engine/Public/WorldPartition/HLOD/StandaloneHLODSubsystem.h
@@ -35,6 +35,7 @@ public:
 
 	//~ Begin UWorldSubsystem Interface.
 	ENGINE_API virtual bool DoesSupportWorldType(const EWorldType::Type WorldType) const override;
+	ENGINE_API virtual bool ShouldCreateSubsystem(UObject* Outer) const override;
 	//~ End UWorldSubsystem Interface.
 
 	ENGINE_API void UpdateStandaloneHLODActors(FStandaloneHLODActorParams InStandaloneHLODActorParams);
diff --git a/Engine/Source/Runtime/Engine/Public/WorldPartition/WorldPartition.h b/Engine/Source/Runtime/Engine/Public/WorldPartition/WorldPartition.h
--- a/Engine/Source/Runtime/Engine/Public/WorldPartition/WorldPartition.h
+++ b/Engine/Source/Runtime/Engine/Public/WorldPartition/WorldPartition.h
@@ -472,6 +472,9 @@ public:
 
 	bool IsHLODsInEditorAllowed() const { return bAllowShowingHLODsInEditor; }
 
+	UFUNCTION()
+	ENGINE_API bool IsStandaloneHLODAllowed() const;
+
 	void SetIsStandaloneHLODWorld(bool bInIsStandaloneHLODWorld) { bIsStandaloneHLODWorld = bInIsStandaloneHLODWorld; }
 #endif
 
@@ -614,7 +617,7 @@ private:
 	uint8 bAllowShowingHLODsInEditor : 1;
 #endif
 
-	UPROPERTY(EditAnywhere, Category = WorldPartitionSetup, meta = (DisplayName = "Build Standalone HLOD", EditConditionHides, EditCondition = "bEnableStreaming", HideEditConditionToggle))
+	UPROPERTY(EditAnywhere, Category = WorldPartitionSetup, meta = (DisplayName = "Build Standalone HLOD", EditConditionHides, EditCondition = "IsStandaloneHLODAllowed()", HideEditConditionToggle))
 	uint8 bHasStandaloneHLOD : 1;
 
 	UPROPERTY()
