Add memory estimate for mesh card generation

#rnx

diff --git a/Engine/Source/Developer/MeshUtilities/Private/MeshCardRepresentationUtilities.cpp b/Engine/Source/Developer/MeshUtilities/Private/MeshCardRepresentationUtilities.cpp
--- a/Engine/Source/Developer/MeshUtilities/Private/MeshCardRepresentationUtilities.cpp
+++ b/Engine/Source/Developer/MeshUtilities/Private/MeshCardRepresentationUtilities.cpp
@@ -1066,6 +1066,12 @@ void BuildMeshCards(const FBox& MeshBounds, const FGenerateCardMeshContext& Cont
 
 #endif // #if USE_EMBREE
 
+int64 FMeshUtilities::MemoryForGenerateMeshCardRepresentation(uint64_t IndexCount)
+{
+	// 1.5 was a derived from insights + LLM captures across a few different projects.
+	return 1.5f * MeshRepresentation::MemoryEstimateForEmbreeScene(IndexCount); 
+}
+
 bool FMeshUtilities::GenerateCardRepresentationData(
 	FString MeshName,
 	const FMeshDataForDerivedDataTask& MeshData,
diff --git a/Engine/Source/Developer/MeshUtilities/Private/MeshUtilitiesPrivate.h b/Engine/Source/Developer/MeshUtilities/Private/MeshUtilitiesPrivate.h
--- a/Engine/Source/Developer/MeshUtilities/Private/MeshUtilitiesPrivate.h
+++ b/Engine/Source/Developer/MeshUtilities/Private/MeshUtilitiesPrivate.h
@@ -69,6 +69,8 @@ private:
 		bool bGenerateAsIfTwoSided,
 		class FDistanceFieldVolumeData& OutData) override;
 	
+	virtual int64 MemoryForGenerateMeshCardRepresentation(uint64_t IndexCount);
+
 	virtual bool GenerateCardRepresentationData(
 		FString MeshName,
 		const FMeshDataForDerivedDataTask& MeshData,
diff --git a/Engine/Source/Developer/MeshUtilities/Public/MeshUtilities.h b/Engine/Source/Developer/MeshUtilities/Public/MeshUtilities.h
--- a/Engine/Source/Developer/MeshUtilities/Public/MeshUtilities.h
+++ b/Engine/Source/Developer/MeshUtilities/Public/MeshUtilities.h
@@ -164,6 +164,8 @@ public:
 		GenerateSignedDistanceFieldVolumeData(MeshName, MeshData, DistanceFieldResolutionScale, bGenerateAsIfTwoSided, OutData);
 	}
 
+	virtual int64 MemoryForGenerateMeshCardRepresentation(uint64_t IndexCount) = 0;
+
 	virtual bool GenerateCardRepresentationData(
 		FString MeshName,
 		const FMeshDataForDerivedDataTask& MeshData,
diff --git a/Engine/Source/Runtime/Engine/Private/MeshCardRepresentation.cpp b/Engine/Source/Runtime/Engine/Private/MeshCardRepresentation.cpp
--- a/Engine/Source/Runtime/Engine/Private/MeshCardRepresentation.cpp
+++ b/Engine/Source/Runtime/Engine/Private/MeshCardRepresentation.cpp
@@ -23,6 +23,7 @@
 #include "DerivedDataCacheInterface.h"
 #include "DerivedDataCacheKey.h"
 #include "StaticMeshCompiler.h"
+#include "MeshUtilities.h"
 #endif
 
 #define LOCTEXT_NAMESPACE "MeshCardRepresentation"
@@ -72,6 +73,12 @@ static TAutoConsoleVariable<int32> CVarCardRepresentationDebugSurfelDirection(
 	TEXT("Generate cards for only surfels pointing in a specific direction."),
 	ECVF_Default);
 
+static TAutoConsoleVariable<bool> CVarMeshCardRepresentationEstimate(
+	TEXT("r.CookMemory.MeshCardRepresentationEstimate"), true,
+	TEXT("Enable memory estimates for mesh card representation asset cooks"),
+	ECVF_ReadOnly
+);
+
 FAsyncCardRepresentationTask* CreateCardRepresentationTask(
 	const FString& InDDCKey,
 	UStaticMesh* Mesh,
@@ -507,8 +514,17 @@ void FCardRepresentationAsyncQueue::CancelAndDeleteTask(const TSet<FAsyncCardRep
 void FCardRepresentationAsyncQueue::StartBackgroundTask(FAsyncCardRepresentationTask* Task)
 {
 	check(Task->AsyncTask == nullptr);
+	
+	int64 RequiredMemory = -1;
+
+#if WITH_EDITOR
+	if (CVarMeshCardRepresentationEstimate.GetValueOnGameThread()) {
+		check(MeshUtilities); // Loaded in AddTask
+		RequiredMemory = MeshUtilities->MemoryForGenerateMeshCardRepresentation(Task->SourceMeshData.GetNumIndices());
+	}	
+#endif
+
 	Task->AsyncTask = MakeUnique<FAsyncTask<FAsyncCardRepresentationTaskWorker>>(*Task);
-	int64 RequiredMemory = -1; // @todo RequiredMemory
 	Task->AsyncTask->StartBackgroundTask(ThreadPool.Get(), EQueuedWorkPriority::Lowest, EQueuedWorkFlags::DoNotRunInsideBusyWait, RequiredMemory, TEXT("Card"));
 }
 
