[ISKM] Added shadow settings and fixed min draw distance on descriptor.

#rnx

diff --git a/Engine/Source/Runtime/Engine/Classes/Animation/AnimBank.h b/Engine/Source/Runtime/Engine/Classes/Animation/AnimBank.h
--- a/Engine/Source/Runtime/Engine/Classes/Animation/AnimBank.h
+++ b/Engine/Source/Runtime/Engine/Classes/Animation/AnimBank.h
@@ -388,6 +388,33 @@ public:
 	UPROPERTY(EditAnywhere, Category = "Component Settings")
 	uint8 bCastShadow : 1;
 
+	UPROPERTY(EditAnywhere, Category = "Component Settings", meta = (EditCondition = "bCastShadow"))
+	uint8 bCastDynamicShadow : 1;
+
+	UPROPERTY(EditAnywhere, Category = "Component Settings", meta = (EditCondition = "bCastShadow"))
+	uint8 bCastStaticShadow : 1;
+
+	UPROPERTY(EditAnywhere, Category = "Component Settings", meta = (EditCondition = "bCastShadow"))
+	uint8 bCastVolumetricTranslucentShadow : 1;
+
+	UPROPERTY(EditAnywhere, Category = "Component Settings", meta = (EditCondition = "bCastShadow"))
+	uint8 bCastContactShadow : 1;
+
+	UPROPERTY(EditAnywhere, Category = "Component Settings", meta = (EditCondition = "bCastShadow"))
+	uint8 bSelfShadowOnly : 1;
+
+	UPROPERTY(EditAnywhere, Category = "Component Settings", meta = (EditCondition = "bCastShadow"))
+	uint8 bCastFarShadow : 1;
+
+	UPROPERTY(EditAnywhere, Category = "Component Settings", meta = (EditCondition = "bCastShadow"))
+	uint8 bCastInsetShadow : 1;
+
+	UPROPERTY(EditAnywhere, Category = "Component Settings", meta = (EditCondition = "bCastShadow"))
+	uint8 bCastCinematicShadow : 1;
+
+	UPROPERTY(EditAnywhere, Category = "Component Settings", meta = (EditCondition = "bCastShadow"))
+	uint8 bCastShadowAsTwoSided : 1;
+
 	UPROPERTY(EditAnywhere, Category = "Component Settings")
 	bool bVisibleInRayTracing = true;
 
diff --git a/Engine/Source/Runtime/Engine/Private/Animation/AnimBank.cpp b/Engine/Source/Runtime/Engine/Private/Animation/AnimBank.cpp
--- a/Engine/Source/Runtime/Engine/Private/Animation/AnimBank.cpp
+++ b/Engine/Source/Runtime/Engine/Private/Animation/AnimBank.cpp
@@ -1420,10 +1420,19 @@ void FSkinnedMeshComponentDescriptorBase::InitFrom(const UInstancedSkinnedMeshCo
 	check(Template);
 
 	Mobility = Template->Mobility;
-	InstanceMinDrawDistance = Template->GetMinDrawDistance();
+	InstanceMinDrawDistance = Template->InstanceMinDrawDistance;
 	Template->GetCullDistances(InstanceStartCullDistance, InstanceEndCullDistance);
 	ComponentClass = Template->GetClass();
 	bCastShadow = Template->CastShadow;
+	bCastDynamicShadow = Template->bCastDynamicShadow;
+	bCastStaticShadow = Template->bCastStaticShadow;
+	bCastVolumetricTranslucentShadow = Template->bCastVolumetricTranslucentShadow;
+	bCastContactShadow = Template->bCastContactShadow;
+	bSelfShadowOnly = Template->bSelfShadowOnly;
+	bCastFarShadow = Template->bCastFarShadow;
+	bCastInsetShadow = Template->bCastInsetShadow;
+	bCastCinematicShadow = Template->bCastCinematicShadow;
+	bCastShadowAsTwoSided = Template->bCastShadowAsTwoSided;
 	bVisibleInRayTracing = Template->bVisibleInRayTracing;
 	bAffectDynamicIndirectLighting = Template->bAffectDynamicIndirectLighting;
 	bAffectDistanceFieldLighting = Template->bAffectDistanceFieldLighting;
@@ -1444,9 +1453,18 @@ void FSkinnedMeshComponentDescriptorBase::InitFrom(const UInstancedSkinnedMeshCo
 void FSkinnedMeshComponentDescriptorBase::InitComponent(UInstancedSkinnedMeshComponent* Component) const
 {
 	Component->Mobility = Mobility;
-	Component->MinDrawDistance = InstanceMinDrawDistance;
+	Component->InstanceMinDrawDistance = InstanceMinDrawDistance;
 	Component->SetCullDistances(InstanceStartCullDistance, InstanceEndCullDistance);
 	Component->CastShadow = bCastShadow;
+	Component->bCastDynamicShadow = bCastDynamicShadow;
+	Component->bCastStaticShadow = bCastStaticShadow;
+	Component->bCastVolumetricTranslucentShadow = bCastVolumetricTranslucentShadow;
+	Component->bCastContactShadow = bCastContactShadow;
+	Component->bSelfShadowOnly = bSelfShadowOnly;
+	Component->bCastFarShadow = bCastFarShadow;
+	Component->bCastInsetShadow = bCastInsetShadow;
+	Component->bCastCinematicShadow = bCastCinematicShadow;
+	Component->bCastShadowAsTwoSided = bCastShadowAsTwoSided;
 	Component->bVisibleInRayTracing = bVisibleInRayTracing;
 	Component->bAffectDynamicIndirectLighting = bAffectDynamicIndirectLighting;
 	Component->bAffectDistanceFieldLighting = bAffectDistanceFieldLighting;
@@ -1482,6 +1500,15 @@ bool FSkinnedMeshComponentDescriptorBase::operator==(const FSkinnedMeshComponent
 		bIncludeInHLOD == Other.bIncludeInHLOD &&
 #endif 
 		bCastShadow == Other.bCastShadow &&
+		bCastDynamicShadow == Other.bCastDynamicShadow &&
+		bCastStaticShadow == Other.bCastStaticShadow &&
+		bCastVolumetricTranslucentShadow == Other.bCastVolumetricTranslucentShadow &&
+		bCastContactShadow == Other.bCastContactShadow &&
+		bSelfShadowOnly == Other.bSelfShadowOnly &&
+		bCastFarShadow == Other.bCastFarShadow &&
+		bCastInsetShadow == Other.bCastInsetShadow &&
+		bCastCinematicShadow == Other.bCastCinematicShadow &&
+		bCastShadowAsTwoSided == Other.bCastShadowAsTwoSided &&
 		bIsInstanceDataGPUOnly == Other.bIsInstanceDataGPUOnly &&
 		PrimitiveBoundsOverride == Other.PrimitiveBoundsOverride &&
 		NumInstancesGPUOnly == Other.NumInstancesGPUOnly &&
