Add more information to assert in FMaterialPSORequestManager::MarkCompilationComplete

Also clean up task graph usage in PSOPrecacheMaterial.cpp, but this isn't expected to fix the assert. DontCompleteUntil is the wrong function to call for unlaunched tasks, but this is a task event, so it was internally calling AddPrerequisites anyway.

diff --git a/Engine/Source/Runtime/Engine/Private/PSOPrecacheMaterial.cpp b/Engine/Source/Runtime/Engine/Private/PSOPrecacheMaterial.cpp
--- a/Engine/Source/Runtime/Engine/Private/PSOPrecacheMaterial.cpp
+++ b/Engine/Source/Runtime/Engine/Private/PSOPrecacheMaterial.cpp
@@ -568,7 +568,10 @@ private:
 		FPrecacheData* FindResult = MaterialPSORequestData.Find(Params);
 		if (FindResult && RequestLifecycleID == LifecycleID)
 		{
-			verify(!CheckCompilingPSOs(*FindResult, false /*bBoostPriority*/));
+			verifyf(!CheckCompilingPSOs(*FindResult, false /*bBoostPriority*/),
+				TEXT("CheckCompilingPSOs should not be active: EState=%d, ActivePSOPrecacheRequests.Num()=%d"),
+				(int32)FindResult->State,
+				FindResult->ActivePSOPrecacheRequests.Num());
 		}
 	}
 
@@ -596,7 +599,7 @@ void FMaterialPSOPrecacheCollectionTask::DoTask(ENamedThreads::Type CurrentThrea
 	// Make sure task is still relevant
 	if (RequestLifecycleID != GMaterialPSORequestManager.GetLifecycleID())
 	{
-		CollectionGraphEvent->DispatchSubsequents();
+		CollectionGraphEvent->Unlock();
 		MaterialInterface->Reset();
 		return;
 	}
@@ -625,11 +628,11 @@ void FMaterialPSOPrecacheCollectionTask::DoTask(ENamedThreads::Type CurrentThrea
 		for (FPSOPrecacheRequestResult& Result : PrecacheResults)
 		{
 			check(Result.IsValid());
-			CollectionGraphEvent->DontCompleteUntil(Result.AsyncCompileEvent);
+			CollectionGraphEvent->AddPrerequisites(Result.AsyncCompileEvent);
 		}
 	}
-	
-	CollectionGraphEvent->DispatchSubsequents();
+
+	CollectionGraphEvent->Unlock();
 }
 
 void FShaderMapPreloadTask::DoTask(ENamedThreads::Type CurrentThread, const FGraphEventRef& MyCompletionGraphEvent)
@@ -646,13 +649,10 @@ void FShaderMapPreloadTask::DoTask(ENamedThreads::Type CurrentThread, const FGra
 	{
 		FGraphEventArray OutCompletionEvents;
 		MaterialShaderMap->GetResource()->PreloadShaderMap(OutCompletionEvents);
-		for (FGraphEventRef Event : OutCompletionEvents)
-		{
-			ShaderPreloadEvents->DontCompleteUntil(Event);
-		}
+		ShaderPreloadEvents->AddPrerequisites(OutCompletionEvents);
 	}
 
-	ShaderPreloadEvents->DispatchSubsequents();
+	ShaderPreloadEvents->Unlock();
 }
 
 void FShaderPreloadCollectionTask::DoTask(ENamedThreads::Type CurrentThread, const FGraphEventRef& MyCompletionGraphEvent)
@@ -696,13 +696,10 @@ void FShaderPreloadCollectionTask::DoTask(ENamedThreads::Type CurrentThread, con
 		// preloaded yet.
 		FGraphEventArray OutCompletionEvents;
 		Shader.GetResource()->PreloadShader(Shader->GetResourceIndex(), OutCompletionEvents);
-		for (FGraphEventRef Event : OutCompletionEvents)
-		{
-			ShaderPreloadEvents->DontCompleteUntil(Event);
-		}
+		ShaderPreloadEvents->AddPrerequisites(OutCompletionEvents);
 	}
-	
-	ShaderPreloadEvents->DispatchSubsequents();
+
+	ShaderPreloadEvents->Unlock();
 }
 
 void PrecacheMaterialPSOs(const FMaterialInterfacePSOPrecacheParamsList& PSOPrecacheParamsList, TArray<FMaterialPSOPrecacheRequestID>& OutMaterialPSOPrecacheRequestIDs, FGraphEventArray& OutGraphEvents)
