Turn off Filters on Standalone Level Instances
==============
Added ILevelInstanceInterface::SetDesiredRuntimeBehavior virtual method
Made ALevelInstance::DesiredRuntimeBehavior private and added SetDesiredRuntimeBehavior setter to ensure that it stays synced with ULevelInstanceComponent::ShadowDesiredRuntimeBehavior

diff --git a/Engine/Source/Editor/UnrealEd/Private/Commandlets/WorldPartitionConvertCommandlet.cpp b/Engine/Source/Editor/UnrealEd/Private/Commandlets/WorldPartitionConvertCommandlet.cpp
--- a/Engine/Source/Editor/UnrealEd/Private/Commandlets/WorldPartitionConvertCommandlet.cpp
+++ b/Engine/Source/Editor/UnrealEd/Private/Commandlets/WorldPartitionConvertCommandlet.cpp
@@ -1263,7 +1263,7 @@ int32 UWorldPartitionConvertCommandlet::Main(const FString& Params)
 				LevelTransform = SubLevelStreaming->LevelTransform;
 			}
 
-			LevelInstanceActor->DesiredRuntimeBehavior = ELevelInstanceRuntimeBehavior::LevelStreaming;
+			LevelInstanceActor->SetDesiredRuntimeBehavior(ELevelInstanceRuntimeBehavior::LevelStreaming);
 			LevelInstanceActor->SetActorTransform(LevelTransform);
 			LevelInstanceActor->SetWorldAsset(SubLevelWorld);
 			LevelInstanceActor->SetActorLabel(SubLevelWorld->GetName());
diff --git a/Engine/Source/Runtime/Engine/Private/LevelInstance/LevelInstanceActor.cpp b/Engine/Source/Runtime/Engine/Private/LevelInstance/LevelInstanceActor.cpp
--- a/Engine/Source/Runtime/Engine/Private/LevelInstance/LevelInstanceActor.cpp
+++ b/Engine/Source/Runtime/Engine/Private/LevelInstance/LevelInstanceActor.cpp
@@ -30,7 +30,7 @@ ALevelInstance::ALevelInstance()
 	RootComponent->Mobility = EComponentMobility::Static;
 
 #if WITH_EDITORONLY_DATA
-	DesiredRuntimeBehavior = ELevelInstanceRuntimeBehavior::Partitioned;
+	SetDesiredRuntimeBehavior(ELevelInstanceRuntimeBehavior::Partitioned);
 #endif
 }
 
@@ -270,6 +270,16 @@ void ALevelInstance::PostEditChangeProperty(FPropertyChangedEvent& PropertyChang
 	const bool bWorldAssetChange = PropertyChangedEvent.Property && PropertyChangedEvent.Property->GetFName() == GET_MEMBER_NAME_CHECKED(ALevelInstance, WorldAsset);
 
 	LevelInstanceActorImpl.PostEditChangeProperty(PropertyChangedEvent, bWorldAssetChange, [this](FPropertyChangedEvent& Event) { Super::PostEditChangeProperty(Event); });
+
+#if WITH_EDITORONLY_DATA
+	const bool bLIBehaviorChange = PropertyChangedEvent.Property && PropertyChangedEvent.Property->GetFName() == GET_MEMBER_NAME_CHECKED(ALevelInstance, DesiredRuntimeBehavior);
+
+	if (bLIBehaviorChange)
+	{
+		LevelInstanceComponent->ShadowDesiredRuntimeBehavior = DesiredRuntimeBehavior;
+	}
+#endif
+
 }
 
 bool ALevelInstance::CanEditChange(const FProperty* Property) const
@@ -430,6 +440,12 @@ bool ALevelInstance::HasHLODRelevantComponents() const
 	return true;
 }
 
+void ALevelInstance::SetDesiredRuntimeBehavior(ELevelInstanceRuntimeBehavior NewBehavior)
+{
+	DesiredRuntimeBehavior = NewBehavior;
+	LevelInstanceComponent->ShadowDesiredRuntimeBehavior = NewBehavior;
+}
+
 #endif // WITH_EDITOR
 
 #undef LOCTEXT_NAMESPACE
diff --git a/Engine/Source/Runtime/Engine/Private/LevelInstance/LevelInstanceComponent.cpp b/Engine/Source/Runtime/Engine/Private/LevelInstance/LevelInstanceComponent.cpp
--- a/Engine/Source/Runtime/Engine/Private/LevelInstance/LevelInstanceComponent.cpp
+++ b/Engine/Source/Runtime/Engine/Private/LevelInstance/LevelInstanceComponent.cpp
@@ -68,6 +68,12 @@ void ULevelInstanceComponent::OnRegister()
 			SpriteComponent->RegisterComponent();
 		}
 	}
+
+	//Initialize the Shadow variable to the parents value. This is for the cases where the ILevelInstanceInterface child classes use a constant behavior
+	if (ILevelInstanceInterface* LevelInstance = Cast<ILevelInstanceInterface>(GetOwner()))
+	{
+		ShadowDesiredRuntimeBehavior = LevelInstance->GetDesiredRuntimeBehavior();
+	}
 #endif //WITH_EDITORONLY_DATA
 }
 
diff --git a/Engine/Source/Runtime/Engine/Private/PackedLevelActor/PackedLevelActor.cpp b/Engine/Source/Runtime/Engine/Private/PackedLevelActor/PackedLevelActor.cpp
--- a/Engine/Source/Runtime/Engine/Private/PackedLevelActor/PackedLevelActor.cpp
+++ b/Engine/Source/Runtime/Engine/Private/PackedLevelActor/PackedLevelActor.cpp
@@ -21,7 +21,7 @@ APackedLevelActor::APackedLevelActor()
 {
 #if WITH_EDITORONLY_DATA
 	// Packed Level Instances don't support level streaming or sub actors
-	DesiredRuntimeBehavior = ELevelInstanceRuntimeBehavior::None;
+	SetDesiredRuntimeBehavior(ELevelInstanceRuntimeBehavior::None);
 #endif
 }
 
diff --git a/Engine/Source/Runtime/Engine/Public/LevelInstance/LevelInstanceActor.h b/Engine/Source/Runtime/Engine/Public/LevelInstance/LevelInstanceActor.h
--- a/Engine/Source/Runtime/Engine/Public/LevelInstance/LevelInstanceActor.h
+++ b/Engine/Source/Runtime/Engine/Public/LevelInstance/LevelInstanceActor.h
@@ -44,12 +44,14 @@ protected:
 	FLevelInstanceActorGuid LevelInstanceActorGuid;
 	FLevelInstanceActorImpl LevelInstanceActorImpl;
 
-public:
 #if WITH_EDITORONLY_DATA
+//Force the usage of SetDesiredRuntimeBehavior to change the behavior
+private:
 	UPROPERTY(EditAnywhere, Category = Level, AdvancedDisplay, Meta = (DisplayName="Level Behavior", DisableLevelInstancePropertyOverride))
 	ELevelInstanceRuntimeBehavior DesiredRuntimeBehavior;
 #endif
 
+public:
 	// Begin ILevelInstanceInterface
 	ENGINE_API const FLevelInstanceID& GetLevelInstanceID() const override;
 	ENGINE_API bool HasValidLevelInstanceID() const override;
@@ -71,6 +73,7 @@ public:
 	ENGINE_API virtual ULevelInstanceComponent* GetLevelInstanceComponent() const override;
 	ENGINE_API virtual bool SetWorldAsset(TSoftObjectPtr<UWorld> WorldAsset) override;
 	virtual ELevelInstanceRuntimeBehavior GetDesiredRuntimeBehavior() const override { return DesiredRuntimeBehavior; }
+	ENGINE_API void SetDesiredRuntimeBehavior(ELevelInstanceRuntimeBehavior NewBehavior) override;
 	virtual ELevelInstanceRuntimeBehavior GetDefaultRuntimeBehavior() const override { return ELevelInstanceRuntimeBehavior::Partitioned; }
 	ENGINE_API virtual TSubclassOf<AActor> GetEditorPivotClass() const override;
 	ENGINE_API virtual bool SupportsPartialEditorLoading() const override;
diff --git a/Engine/Source/Runtime/Engine/Public/LevelInstance/LevelInstanceComponent.h b/Engine/Source/Runtime/Engine/Public/LevelInstance/LevelInstanceComponent.h
--- a/Engine/Source/Runtime/Engine/Public/LevelInstance/LevelInstanceComponent.h
+++ b/Engine/Source/Runtime/Engine/Public/LevelInstance/LevelInstanceComponent.h
@@ -6,6 +6,7 @@
 #include "Components/SceneComponent.h"
 #include "WorldPartition/Filter/WorldPartitionActorFilter.h"
 #include "WorldPartition/WorldPartitionActorContainerID.h"
+#include "LevelInstance/LevelInstanceTypes.h"
 
 #include "LevelInstanceComponent.generated.h"
 
@@ -59,7 +60,14 @@ private:
 #endif
 
 #if WITH_EDITORONLY_DATA
-	UPROPERTY(EditAnywhere, Category = Filter, meta=(LevelInstanceFilter, DisableLevelInstancePropertyOverride))
+public:
+	//Filters are only allowed on Embedded Level Instances Editor\Private LI per design.
+	UPROPERTY()
+	ELevelInstanceRuntimeBehavior ShadowDesiredRuntimeBehavior;
+
+private:
+	//, EditCondition="bFiltersAllowed"))
+	UPROPERTY(EditAnywhere, Category = Filter, meta=(LevelInstanceFilter, DisableLevelInstancePropertyOverride, EditCondition = "ShadowDesiredRuntimeBehavior==ELevelInstanceRuntimeBehavior::Partitioned"))
 	FWorldPartitionActorFilter Filter;
 
 	UPROPERTY(EditAnywhere, Transient, Category = Filter, meta=(LevelInstanceEditFilter))
diff --git a/Engine/Source/Runtime/Engine/Public/LevelInstance/LevelInstanceInterface.h b/Engine/Source/Runtime/Engine/Public/LevelInstance/LevelInstanceInterface.h
--- a/Engine/Source/Runtime/Engine/Public/LevelInstance/LevelInstanceInterface.h
+++ b/Engine/Source/Runtime/Engine/Public/LevelInstance/LevelInstanceInterface.h
@@ -57,6 +57,9 @@ class ILevelInstanceInterface
 
 	virtual ELevelInstanceRuntimeBehavior GetDefaultRuntimeBehavior() const = 0;
 
+	//Default assumes it's a constant Behavior. This is needed to allow\deny the editor of actor layers at the component level. You must implement this method if your LevelInstance child support multiple runtime behavior. 
+	virtual void SetDesiredRuntimeBehavior(ELevelInstanceRuntimeBehavior NewBehavior) { checkf(false, TEXT("SetDesiredRuntimeBehavior should not be called on this child of ILevelInstanceInterface. Please implement an override if required.")); }
+
 	virtual TSubclassOf<AActor> GetEditorPivotClass() const { return nullptr; }
 	
 	virtual void AdjustPivotOnCreation() { }
