RuntimeVirtualTexture holds a delegate to call if it is Initialized twice with no Release.
This can happen in a couple of cases:
- If two RVT volumes in the same scene point at the same RVT asset
- When we start PIE so that a volume is duplicated for the PIE scene
In these cases we want the original "owner" of the asset to receive a call to notify that it is no longer the owner and to clear its cached render thread state.
#rnx

diff --git a/Engine/Source/Runtime/Engine/Classes/VT/RuntimeVirtualTexture.h b/Engine/Source/Runtime/Engine/Classes/VT/RuntimeVirtualTexture.h
--- a/Engine/Source/Runtime/Engine/Classes/VT/RuntimeVirtualTexture.h
+++ b/Engine/Source/Runtime/Engine/Classes/VT/RuntimeVirtualTexture.h
@@ -164,7 +164,16 @@ public:
 	/** Return true if the virtual texture layer should be sampled as YCoCg */
 	ENGINE_API bool IsLayerYCoCg(int32 LayerIndex) const;
 
+	/**
+	 * Declare a delegate funtion signiture for a function that is registered in Initialize() to be called if we call Initialize() without calling Release().
+	 * This can happen for the case where the RuntimeVirtualTexture is registered twice. For example in PIE and the main editor.
+	 * In that case we want to notify the first caller.
+	 */
+	DECLARE_DELEGATE(FOnReInitDelegate);
+
 	/** (Re)Initialize this object. Call this whenever we modify the producer or transform. */
+	ENGINE_API void Initialize(IVirtualTexture* InProducer, FVTProducerDescription const& InProducerDesc, FTransform const& InVolumeToWorld, FBox const& InWorldBounds, FOnReInitDelegate InDelegate);
+	UE_DEPRECATED(5.6, "Use the version that takes a FOnReInitDelegate instead.")
 	ENGINE_API void Initialize(IVirtualTexture* InProducer, FVTProducerDescription const& InProducerDesc, FTransform const& InVolumeToWorld, FBox const& InWorldBounds);
 
 	/** Release the resources for this object This will need to be called if our producer becomes stale and we aren't doing a full reinit with a new producer. */
@@ -206,12 +215,14 @@ protected:
 private:
 	/** Render thread resource container. */
 	class FRuntimeVirtualTextureRenderResource* Resource;
-	int32 ResourceRefCount = 0;
 
 	/** Material uniform parameters to support transform from world to UV coordinates. */
 	FVector4 WorldToUVTransformParameters[3];
 	/** Material uniform parameter used to pack world height. */
 	FVector4 WorldHeightUnpackParameter;
+
+	/** Delegate to be called on ReInitialize(). */
+	FOnReInitDelegate OnReInitDelegate; 
 };
 
 class UVirtualTexture2D;
diff --git a/Engine/Source/Runtime/Engine/Private/VT/RuntimeVirtualTexture.cpp b/Engine/Source/Runtime/Engine/Private/VT/RuntimeVirtualTexture.cpp
--- a/Engine/Source/Runtime/Engine/Private/VT/RuntimeVirtualTexture.cpp
+++ b/Engine/Source/Runtime/Engine/Private/VT/RuntimeVirtualTexture.cpp
@@ -603,35 +603,31 @@ EVTProducerPriority URuntimeVirtualTexture::GetDefaultPriority() const
 	return EVTProducerPriority::Normal;
 }
 
-void URuntimeVirtualTexture::Initialize(IVirtualTexture* InProducer, FVTProducerDescription const& InProducerDesc, FTransform const& InVolumeToWorld, FBox const& InWorldBounds)
+void URuntimeVirtualTexture::Initialize(IVirtualTexture* InProducer, FVTProducerDescription const& InProducerDesc, FTransform const& InVolumeToWorld, FBox const& InWorldBounds, FOnReInitDelegate InDelegate)
 {
-	// TODO: Look into multi-scene support if needed. The ref count prevents resource being reallocated and invalidating
-	// cached pointers. However, it assumes that the RVT is initialized using the same transform, bounds, and producer
-	// all scenes which seems to be the case for PIE. So this currently won't work if the RVT needs to be initialized
-	// differently in different scenes
-	if (ResourceRefCount == 0)
-	{
-		WorldToUVTransformParameters[0] = InVolumeToWorld.GetTranslation();
-		WorldToUVTransformParameters[1] = InVolumeToWorld.GetUnitAxis(EAxis::X) * 1.f / InVolumeToWorld.GetScale3D().X;
-		WorldToUVTransformParameters[2] = InVolumeToWorld.GetUnitAxis(EAxis::Y) * 1.f / InVolumeToWorld.GetScale3D().Y;
+	OnReInitDelegate.ExecuteIfBound();
+	OnReInitDelegate = InDelegate;
 
-		const FVector::FReal HeightRange = FMath::Max<FVector::FReal>(InWorldBounds.Max.Z - InWorldBounds.Min.Z, 1.f);
-		WorldHeightUnpackParameter = FVector4(HeightRange, InWorldBounds.Min.Z, 0.f, 0.f);
+	WorldToUVTransformParameters[0] = InVolumeToWorld.GetTranslation();
+	WorldToUVTransformParameters[1] = InVolumeToWorld.GetUnitAxis(EAxis::X) * 1.f / InVolumeToWorld.GetScale3D().X;
+	WorldToUVTransformParameters[2] = InVolumeToWorld.GetUnitAxis(EAxis::Y) * 1.f / InVolumeToWorld.GetScale3D().Y;
 
-		InitResource(InProducer, InProducerDesc);
-	}
+	const FVector::FReal HeightRange = FMath::Max<FVector::FReal>(InWorldBounds.Max.Z - InWorldBounds.Min.Z, 1.f);
+	WorldHeightUnpackParameter = FVector4(HeightRange, InWorldBounds.Min.Z, 0.f, 0.f);
+
+	InitResource(InProducer, InProducerDesc);
+}
 
-	++ResourceRefCount;
+void URuntimeVirtualTexture::Initialize(IVirtualTexture* InProducer, FVTProducerDescription const& InProducerDesc, FTransform const& InVolumeToWorld, FBox const& InWorldBounds)
+{
+	Initialize(InProducer, InProducerDesc, InVolumeToWorld, InWorldBounds, {});
 }
 
 void URuntimeVirtualTexture::Release()
 {
-	--ResourceRefCount;
+	OnReInitDelegate = {};
 
-	if (ResourceRefCount == 0)
-	{
-		InitNullResource();
-	}
+	InitNullResource();
 }
 
 void URuntimeVirtualTexture::InitResource(IVirtualTexture* InProducer, FVTProducerDescription const& InProducerDesc)
diff --git a/Engine/Source/Runtime/Renderer/Private/RendererScene.cpp b/Engine/Source/Runtime/Renderer/Private/RendererScene.cpp
--- a/Engine/Source/Runtime/Renderer/Private/RendererScene.cpp
+++ b/Engine/Source/Runtime/Renderer/Private/RendererScene.cpp
@@ -3151,20 +3151,6 @@ static void UpdateRuntimeVirtualTextureHidePrimitiveMasks(TSparseArray<FRuntimeV
 	}
 }
 
-static void FixMultipleRuntimeVirtualTextureUsage(TSparseArray<FRuntimeVirtualTextureSceneProxy*> const& InProxies, FRuntimeVirtualTextureSceneProxy const* SceneProxy)
-{
-	// If we add a runtime virtual texture component to a scene which uses the same asset as another, then we need to reset the first one.
-	// This should only happen for bad data setups, so warn here.
-	for (TSparseArray<FRuntimeVirtualTextureSceneProxy*>::TConstIterator It(InProxies); It; ++It)
-	{
-		if (*It != SceneProxy && (*It)->RuntimeVirtualTextureId == SceneProxy->RuntimeVirtualTextureId)
-		{
-			UE_LOG(LogRenderer, Warning, TEXT("Found two RVT components in a scene using the same RVT asset."));
-			(*It)->MarkUnused();
-		}
-	}
-}
-
 void FScene::AddRuntimeVirtualTexture(URuntimeVirtualTextureComponent* Component)
 {
 	if (Component->SceneProxy == nullptr)
@@ -3179,8 +3165,6 @@ void FScene::AddRuntimeVirtualTexture(URuntimeVirtualTextureComponent* Component
 		{
 			SceneProxy->SceneIndex = Scene->RuntimeVirtualTextures.Add(SceneProxy);
 			
-			FixMultipleRuntimeVirtualTextureUsage(Scene->RuntimeVirtualTextures, SceneProxy);
-
 			UpdateRuntimeVirtualTextureHidePrimitiveMasks(Scene->RuntimeVirtualTextures, Scene->bRuntimeVirtualTexturePrimitiveHideEditor, Scene->bRuntimeVirtualTexturePrimitiveHideGame);
 		});
 	}
diff --git a/Engine/Source/Runtime/Renderer/Private/VT/RuntimeVirtualTextureSceneProxy.cpp b/Engine/Source/Runtime/Renderer/Private/VT/RuntimeVirtualTextureSceneProxy.cpp
--- a/Engine/Source/Runtime/Renderer/Private/VT/RuntimeVirtualTextureSceneProxy.cpp
+++ b/Engine/Source/Runtime/Renderer/Private/VT/RuntimeVirtualTextureSceneProxy.cpp
@@ -121,7 +121,8 @@ FRuntimeVirtualTextureSceneProxy::FRuntimeVirtualTextureSceneProxy(URuntimeVirtu
 			VirtualTextureSize = FIntPoint(ProducerDesc.BlockWidthInTiles * ProducerDesc.TileSize, ProducerDesc.BlockHeightInTiles * ProducerDesc.TileSize);
 
 			// The Initialize() call will allocate the virtual texture by spawning work on the render thread.
-			VirtualTexture->Initialize(Producer, ProducerDesc, Transform, Bounds);
+			VirtualTexture->Initialize(Producer, ProducerDesc, Transform, Bounds,
+				URuntimeVirtualTexture::FOnReInitDelegate::CreateRaw(this, &FRuntimeVirtualTextureSceneProxy::MarkUnused));
 
 			bAdaptive = VirtualTexture->GetAdaptivePageTable();
 
@@ -166,10 +167,14 @@ void FRuntimeVirtualTextureSceneProxy::Release()
 
 void FRuntimeVirtualTextureSceneProxy::MarkUnused()
 {
-	VirtualTexture = nullptr;
-	ProducerHandle = {};
-	AllocatedVirtualTexture = nullptr;
-	SpaceID = -1;
+	ENQUEUE_RENDER_COMMAND(GetProducerHandle)(
+		[this](FRHICommandList& RHICmdList)
+	{
+		VirtualTexture = nullptr;
+		ProducerHandle = {};
+		AllocatedVirtualTexture = nullptr;
+		SpaceID = -1;
+	});
 }
 
 /** Transform world bounds into Virtual Texture UV space. */
