Waveform Editor: Fix exported soundwaves not supporting undo transactions

diff --git a/Engine/Source/Runtime/Engine/Private/SampleBufferIO.cpp b/Engine/Source/Runtime/Engine/Private/SampleBufferIO.cpp
--- a/Engine/Source/Runtime/Engine/Private/SampleBufferIO.cpp
+++ b/Engine/Source/Runtime/Engine/Private/SampleBufferIO.cpp
@@ -196,9 +196,14 @@ namespace Audio
 		CurrentPackage = CreatePackage( *AbsoluteFilePath);
 
 		// Create a new USoundWave.
-		CurrentSoundWave = NewObject<USoundWave>(CurrentPackage, *FileName, RF_Public | RF_Standalone);
+		CurrentSoundWave = NewObject<USoundWave>(CurrentPackage, *FileName, RF_Public | RF_Standalone | RF_Transactional, InSoundWaveTransformationTarget);
 		CurrentSoundWave->AddToRoot();
 
+#if WITH_EDITORONLY_DATA
+		// Clears transformations in the exported soundwave as they are now baked into the sound wave cue points
+		CurrentSoundWave->Transformations.Empty();
+#endif
+
 #if WITH_EDITOR
 		if (InSoundWaveTransformationTarget != nullptr)
 		{
