[HWRT] Split logic from FRayTracingGeometryManager::Tick(...) into separate functions.
- also modified some counters to be int64 instead of uint64 so verification logic that checks the numbers are never negative work correctly.
- and some small cleanup in a couple of functions.

diff --git a/Engine/Source/Runtime/Engine/Private/Rendering/RayTracingGeometryManager.cpp b/Engine/Source/Runtime/Engine/Private/Rendering/RayTracingGeometryManager.cpp
--- a/Engine/Source/Runtime/Engine/Private/Rendering/RayTracingGeometryManager.cpp
+++ b/Engine/Source/Runtime/Engine/Private/Rendering/RayTracingGeometryManager.cpp
@@ -253,13 +253,13 @@ void FRayTracingGeometryManager::ReleaseRayTracingGeometryGroup(RayTracing::FGeo
 {
 	FScopeLock ScopeLock(&MainCS);
 
-	check(RegisteredGroups.IsValidIndex(Handle));
-
 	ReleaseRayTracingGeometryGroupReference(Handle);
 }
 
 void FRayTracingGeometryManager::ReleaseRayTracingGeometryGroupReference(RayTracing::FGeometryGroupHandle Handle)
 {
+	check(RegisteredGroups.IsValidIndex(Handle));
+
 	FRayTracingGeometryGroup& Group = RegisteredGroups[Handle];
 
 	--Group.NumReferences;
@@ -293,16 +293,19 @@ FRayTracingGeometryManager::FGeometryHandle FRayTracingGeometryManager::Register
 	RegisteredGeometry.Geometry = InGeometry;
 	RegisteredGeometry.LastReferencedFrame = 0;
 
-	if (InGeometry->GroupHandle != INDEX_NONE)
+	const RayTracing::FGeometryGroupHandle GroupHandle = InGeometry->GroupHandle;
+	const int8 LODIndex = InGeometry->LODIndex;
+
+	if (GroupHandle != INDEX_NONE)
 	{
-		checkf(RegisteredGroups.IsValidIndex(InGeometry->GroupHandle), TEXT("FRayTracingGeometry.GroupHandle must be valid"));
+		checkf(RegisteredGroups.IsValidIndex(GroupHandle), TEXT("FRayTracingGeometry.GroupHandle must be valid"));
 
-		FRayTracingGeometryGroup& Group = RegisteredGroups[InGeometry->GroupHandle];
+		FRayTracingGeometryGroup& Group = RegisteredGroups[GroupHandle];
 
-		checkf(InGeometry->LODIndex >= 0 && InGeometry->LODIndex < Group.GeometryHandles.Num(), TEXT("FRayTracingGeometry assigned to a group must have a valid LODIndex"));
-		checkf(Group.GeometryHandles[InGeometry->LODIndex] == INDEX_NONE, TEXT("Each LOD inside a FRayTracingGeometryGroup can only be associated with a single FRayTracingGeometry"));
+		checkf(LODIndex >= 0 && LODIndex < Group.GeometryHandles.Num(), TEXT("FRayTracingGeometry assigned to a group must have a valid LODIndex"));
+		checkf(Group.GeometryHandles[LODIndex] == INDEX_NONE, TEXT("Each LOD inside a FRayTracingGeometryGroup can only be associated with a single FRayTracingGeometry"));
 
-		Group.GeometryHandles[InGeometry->LODIndex] = Handle;
+		Group.GeometryHandles[LODIndex] = Handle;
 		++Group.NumReferences;
 
 		RegisteredGeometry.bAlwaysResident = IsAlwaysResidentGeometry(InGeometry, Group);
@@ -312,7 +315,7 @@ FRayTracingGeometryManager::FGeometryHandle FRayTracingGeometryManager::Register
 			AlwaysResidentGeometries.Add(Handle);
 		}
 
-		if (IsRayTracingEnabled() && InGeometry->LODIndex >= Group.CurrentFirstLODIdx && (!IsRayTracingUsingReferenceBasedResidency() || RegisteredGeometry.bAlwaysResident))
+		if (IsRayTracingEnabled() && LODIndex >= Group.CurrentFirstLODIdx && (!IsRayTracingUsingReferenceBasedResidency() || RegisteredGeometry.bAlwaysResident))
 		{
 			PendingStreamingRequests.Add(Handle);
 			INC_DWORD_STAT(STAT_RayTracingPendingStreamingRequests);
@@ -321,7 +324,7 @@ FRayTracingGeometryManager::FGeometryHandle FRayTracingGeometryManager::Register
 		
 	INC_DWORD_STAT(STAT_RayTracingGeometryCount);
 
-	GRayTracingGeometryManager->RefreshRegisteredGeometry(Handle);
+	RefreshRegisteredGeometry(Handle);
 
 	return Handle;
 }
@@ -376,7 +379,7 @@ void FRayTracingGeometryManager::ReleaseRayTracingGeometryHandle(FGeometryHandle
 	EvictableGeometries.Remove(Handle);
 
 	RegisteredGeometries.RemoveAt(Handle);
-	ReferencedGeometryHandles.Remove(Handle);
+	ReferencedGeometries.Remove(Handle);
 	if (PendingStreamingRequests.Remove(Handle) > 0)
 	{
 		DEC_DWORD_STAT(STAT_RayTracingPendingStreamingRequests);
@@ -389,9 +392,11 @@ void FRayTracingGeometryManager::SetRayTracingGeometryStreamingData(const FRayTr
 {
 	FScopeLock ScopeLock(&MainCS);
 
-	checkf(RegisteredGeometries.IsValidIndex(Geometry->RayTracingGeometryHandle), TEXT("SetRayTracingGeometryStreamingData(...) can only be used with FRayTracingGeometry that has been registered with FRayTracingGeometryManager."));
+	const FGeometryHandle GeometryHandle = Geometry->RayTracingGeometryHandle;
+
+	checkf(RegisteredGeometries.IsValidIndex(GeometryHandle), TEXT("SetRayTracingGeometryStreamingData(...) can only be used with FRayTracingGeometry that has been registered with FRayTracingGeometryManager."));
 
-	FRegisteredGeometry& RegisteredGeometry = RegisteredGeometries[Geometry->RayTracingGeometryHandle];
+	FRegisteredGeometry& RegisteredGeometry = RegisteredGeometries[GeometryHandle];
 	RegisteredGeometry.StreamableAsset = &StreamableAsset;
 	RegisteredGeometry.StreamableBVHSize = StreamableAsset.GetRequestSizeBVH();
 	RegisteredGeometry.StreamableBuffersSize = StreamableAsset.GetRequestSizeBuffers();
@@ -399,7 +404,6 @@ void FRayTracingGeometryManager::SetRayTracingGeometryStreamingData(const FRayTr
 	checkf(RegisteredGeometry.StreamableBVHSize > 0 || RegisteredGeometry.StreamableBuffersSize > 0, TEXT("FRayTracingStreamableAsset should have data to stream."));
 }
 
-
 void FRayTracingGeometryManager::SetRayTracingGeometryGroupCurrentFirstLODIndex(FRHICommandListBase& RHICmdList, RayTracing::FGeometryGroupHandle Handle, uint8 NewCurrentFirstLODIdx)
 {
 	FScopeLock ScopeLock(&MainCS);
@@ -407,7 +411,7 @@ void FRayTracingGeometryManager::SetRayTracingGeometryGroupCurrentFirstLODIndex(
 	FRayTracingGeometryGroup& Group = RegisteredGroups[Handle];
 
 	// immediately release streamed out LODs
-	if(NewCurrentFirstLODIdx > Group.CurrentFirstLODIdx)
+	if (NewCurrentFirstLODIdx > Group.CurrentFirstLODIdx)
 	{
 		FRHIResourceReplaceBatcher Batcher(RHICmdList, NewCurrentFirstLODIdx - Group.CurrentFirstLODIdx);
 		for (int32 LODIdx = Group.CurrentFirstLODIdx; LODIdx < NewCurrentFirstLODIdx; ++LODIdx)
@@ -432,7 +436,7 @@ void FRayTracingGeometryManager::SetRayTracingGeometryGroupCurrentFirstLODIndex(
 			}
 		}
 	}
-	else if(IsRayTracingEnabled() && !IsRayTracingUsingReferenceBasedResidency())
+	else if (IsRayTracingEnabled() && !IsRayTracingUsingReferenceBasedResidency())
 	{
 		for (int32 LODIdx = NewCurrentFirstLODIdx; LODIdx < Group.CurrentFirstLODIdx; ++LODIdx)
 		{
@@ -553,112 +557,28 @@ void FRayTracingGeometryManager::Tick(FRHICommandList& RHICmdList)
 	PreviousFrameCounter = GFrameCounterRenderThread;
 #endif
 
+	const bool bIsRayTracingEnabled = IsRayTracingEnabled();
 	const bool bUsingReferenceBasedResidency = IsRayTracingUsingReferenceBasedResidency();
 
-	checkf(bUsingReferenceBasedResidency || (ReferencedGeometryHandles.IsEmpty() && ReferencedGeometryGroups.IsEmpty() && ReferencedGeometryGroupsForDynamicUpdate.IsEmpty()),
+	checkf(bUsingReferenceBasedResidency || (ReferencedGeometries.IsEmpty() && ReferencedGeometryGroups.IsEmpty() && ReferencedGeometryGroupsForDynamicUpdate.IsEmpty()),
 		TEXT("ReferencedGeometryHandles, ReferencedGeometryGroups and ReferencedGeometryGroupsForDynamicUpdate are expected to be empty when not using reference based residency"));
 
-	if (bRefreshAlwaysResidentRayTracingGeometries)
-	{
-		bRefreshAlwaysResidentRayTracingGeometries = false;
-
-		AlwaysResidentGeometries.Empty();
-		TotalAlwaysResidentSize = 0;
-
-		for (FRegisteredGeometry& RegisteredGeometry : RegisteredGeometries)
-		{
-			const RayTracing::FGeometryGroupHandle GroupHandle = RegisteredGeometry.Geometry->GroupHandle;
-
-			if (GroupHandle == INDEX_NONE)
-			{
-				RegisteredGeometry.bAlwaysResident = false;
-				continue;
-			}
-
-			const FRayTracingGeometryGroup& Group = RegisteredGroups[GroupHandle];
-
-			RegisteredGeometry.bAlwaysResident = IsAlwaysResidentGeometry(RegisteredGeometry.Geometry, Group);
-
-			if (RegisteredGeometry.bAlwaysResident)
-			{
-				AlwaysResidentGeometries.Add(RegisteredGeometry.Geometry->RayTracingGeometryHandle);
-				TotalAlwaysResidentSize += RegisteredGeometry.Size;
-
-				if (RegisteredGeometry.Geometry->IsEvicted())
-				{
-					MakeGeometryResident(RHICmdList, RegisteredGeometry);
-				}
-
-				if (!RequestRayTracingGeometryStreamIn(RHICmdList, RegisteredGeometry.Geometry->RayTracingGeometryHandle))
-				{
-					PendingStreamingRequests.Add(RegisteredGeometry.Geometry->RayTracingGeometryHandle);
-
-					INC_DWORD_STAT(STAT_RayTracingPendingStreamingRequests);
-				}
-
-				EvictableGeometries.Remove(RegisteredGeometry.Geometry->RayTracingGeometryHandle);
-			}
-			else if (bUsingReferenceBasedResidency && RegisteredGeometry.Geometry->GetRHI() != nullptr)
-			{
-				EvictGeometry(RHICmdList, RegisteredGeometry);
-			}
-		}
-	}
+	RefreshAlwaysResidentRayTracingGeometries(RHICmdList);
 
 #if DO_CHECK
-	if (GRayTracingTestCheckIntegrity)
-	{
-		for (FRegisteredGeometry& RegisteredGeometry : RegisteredGeometries)
-		{
-			const FGeometryHandle GeometryHandle = RegisteredGeometry.Geometry->RayTracingGeometryHandle;
-			const RayTracing::FGeometryGroupHandle GroupHandle = RegisteredGeometry.Geometry->GroupHandle;
-
-			bool bAlwaysResident = false;
-
-			if (GroupHandle != INDEX_NONE)
-			{
-				const FRayTracingGeometryGroup& Group = RegisteredGroups[GroupHandle];
-				bAlwaysResident = IsAlwaysResidentGeometry(RegisteredGeometry.Geometry, Group);
-			}
-			else
-			{
-				// geometries not assigned to a group (eg: dynamic geometry) are always evictable
-			}
-
-			checkf(RegisteredGeometry.bEvicted == RegisteredGeometry.Geometry->IsEvicted(), TEXT("Cached bEvicted flag in FRegisteredGeometry is stale"));
-			checkf(RegisteredGeometry.bAlwaysResident == bAlwaysResident, TEXT("Cached bAlwaysResident flag in FRegisteredGeometry is stale"));
-			checkf(RegisteredGeometry.bAlwaysResident == AlwaysResidentGeometries.Contains(GeometryHandle), TEXT("Geometry with bAlwaysResident flag set should be in the AlwaysResidentGeometries set."));
-		}
-	}
+	CheckIntegrity();
 #endif
 
-	if (!IsRayTracingEnabled())
+	if (!bIsRayTracingEnabled)
 	{
 		if (bHasRayTracingEnableChanged)
 		{
-			// evict all geometries
-			for (FRegisteredGeometry& RegisteredGeometry : RegisteredGeometries)
-			{
-				if (RegisteredGeometry.Geometry->GetRHI() != nullptr)
-				{
-					EvictGeometry(RHICmdList, RegisteredGeometry);
-				}
-			}
+			EvictAllGeometries(RHICmdList);
 
 			PendingStreamingRequests.Empty();
 
 			SET_DWORD_STAT(STAT_RayTracingPendingStreamingRequests, 0);
 		}
-		else
-		{
-#if DO_CHECK
-			// otherwise just check that everything is evicted
-			for (FRegisteredGeometry& RegisteredGeometry : RegisteredGeometries)
-			{
-				checkf(RegisteredGeometry.Geometry->IsEvicted() || RegisteredGeometry.Geometry->GetRHI() == nullptr, TEXT("Ray tracing geometry should be evicted when ray tracing is disabled."));
-			}
-#endif
-		}
 
 		checkf(TotalResidentSize == 0,
 			TEXT("TotalResidentSize should be 0 when ray tracing is disabled but is currently %lld.\n")
@@ -675,11 +595,11 @@ void FRayTracingGeometryManager::Tick(FRHICommandList& RHICmdList)
 	}
 	else if (bUsingReferenceBasedResidency)
 	{
-		check(IsRayTracingEnabled());
+		check(bIsRayTracingEnabled);
 
 		if (!bRenderedFrame)
 		{
-			ensureMsgf(ReferencedGeometryHandles.IsEmpty() && ReferencedGeometryGroups.IsEmpty() && ReferencedGeometryGroupsForDynamicUpdate.IsEmpty(),
+			ensureMsgf(ReferencedGeometries.IsEmpty() && ReferencedGeometryGroups.IsEmpty() && ReferencedGeometryGroupsForDynamicUpdate.IsEmpty(),
 				TEXT("Unexpected entries in ReferencedGeometryHandles/ReferencedGeometryGroups/ReferencedGeometryGroupsForDynamicUpdate. ")
 				TEXT("Missing a call to PreRender() or didn't clear the arrays in the last frame?"));
 			return;
@@ -708,197 +628,299 @@ void FRayTracingGeometryManager::Tick(FRHICommandList& RHICmdList)
 			}
 		}
 
-		TSet<FGeometryHandle> NotReferencedResidentGeometries = EvictableGeometries;
-
-		TArray<FGeometryHandle> ReferencedGeometries;
-
-		uint64 ReferencedSize = 0;
-		uint64 RequestedSize = TotalAlwaysResidentSize;
-		uint64 RequestedButEvictedSize = 0;
+		UpdateReferenceBasedResidency(RHICmdList);
+	}
+	else
+	{
+		check(bIsRayTracingEnabled);
 
-		// Step 1
-		// - update LastReferencedFrame of referenced geometries and calculate memory required to make evicted geometries resident
-		for (FGeometryHandle GeometryHandle : ReferencedGeometryHandles)
+		if (bHasRayTracingEnableChanged)
 		{
-			FRegisteredGeometry& RegisteredGeometry = RegisteredGeometries[GeometryHandle];
-			RegisteredGeometry.LastReferencedFrame = GFrameCounterRenderThread;
+			MakeResidentAllGeometries(RHICmdList);
+		}
 
-			ReferencedGeometries.Add(GeometryHandle);
-			NotReferencedResidentGeometries.Remove(GeometryHandle);
+		SET_MEMORY_STAT(STAT_RayTracingGeometryReferencedMemory, 0);
+		SET_MEMORY_STAT(STAT_RayTracingGeometryRequestedMemory, TotalResidentSize);
+		CSV_CUSTOM_STAT(RayTracingGeometry, ReferencedSizeMB, 0, ECsvCustomStatOp::Set);
+		CSV_CUSTOM_STAT(RayTracingGeometry, RequestedSizeMB, TotalResidentSize / 1024.0f / 1024.0f, ECsvCustomStatOp::Set);
+	}
 
-			ReferencedSize += RegisteredGeometry.Size;
+	{
+		TSet<FGeometryHandle> CurrentPendingStreamingRequests;
+		Swap(CurrentPendingStreamingRequests, PendingStreamingRequests);
+		PendingStreamingRequests.Reserve(CurrentPendingStreamingRequests.Num());
 
-			if (!RegisteredGeometry.bAlwaysResident)
+		for (FGeometryHandle GeometryHandle : CurrentPendingStreamingRequests)
+		{
+			if (!RequestRayTracingGeometryStreamIn(RHICmdList, GeometryHandle))
 			{
-				RequestedSize += RegisteredGeometry.Size;
+				PendingStreamingRequests.Add(GeometryHandle);
 			}
+		}
+	}
 
-			if (RegisteredGeometry.bEvicted)
-			{
-				RequestedButEvictedSize += RegisteredGeometry.Size;
-			}
+	SET_DWORD_STAT(STAT_RayTracingPendingStreamingRequests, PendingStreamingRequests.Num());
+
+	ProcessCompletedStreamingRequests(RHICmdList);
+
+	ReferencedGeometries.Reset();
+	ReferencedGeometryGroups.Reset();
+	ReferencedGeometryGroupsForDynamicUpdate.Reset();
+
+	bHasRayTracingEnableChanged = false;
+
+	SET_MEMORY_STAT(STAT_RayTracingGeometryResidentMemory, TotalResidentSize);
+	SET_MEMORY_STAT(STAT_RayTracingGeometryAlwaysResidentMemory, TotalAlwaysResidentSize);
+
+	CSV_CUSTOM_STAT(RayTracingGeometry, TotalResidentSizeMB, TotalResidentSize / 1024.0f / 1024.0f, ECsvCustomStatOp::Set);
+	CSV_CUSTOM_STAT(RayTracingGeometry, TotalAlwaysResidentSizeMB, TotalAlwaysResidentSize / 1024.0f / 1024.0f, ECsvCustomStatOp::Set);
+}
+
+void FRayTracingGeometryManager::UpdateReferenceBasedResidency(FRHICommandList& RHICmdList)
+{
+	TSet<FGeometryHandle> NotReferencedResidentGeometries = EvictableGeometries;
+
+	TArray<FGeometryHandle> LocalReferencedGeometries;
+	LocalReferencedGeometries.Reserve(RegisteredGeometries.Num());
+
+	int64 ReferencedSize = 0;
+	int64 RequestedSize = TotalAlwaysResidentSize;
+	int64 RequestedButEvictedSize = 0;
+
+	// Step 1
+	// - update LastReferencedFrame of referenced geometries and calculate memory required to make evicted geometries resident
+	for (FGeometryHandle GeometryHandle : ReferencedGeometries)
+	{
+		FRegisteredGeometry& RegisteredGeometry = RegisteredGeometries[GeometryHandle];
+		RegisteredGeometry.LastReferencedFrame = GFrameCounterRenderThread;
+
+		LocalReferencedGeometries.Add(GeometryHandle);
+		NotReferencedResidentGeometries.Remove(GeometryHandle);
+
+		ReferencedSize += RegisteredGeometry.Size;
+
+		if (!RegisteredGeometry.bAlwaysResident)
+		{
+			RequestedSize += RegisteredGeometry.Size;
 		}
 
-		// Step 2
-		// - add all geometries in referenced groups to ReferencedGeometries
-		//		- need to make all geometries in group resident otherwise might not have valid geometry when reducing LOD
-		//		- TODO: Could track TargetLOD and only make [TargetLOD ... LastLOD] range resident
-		// - also update LastReferencedFrame and calculate memory required to make evicted geometries resident
-		for (RayTracing::FGeometryGroupHandle GroupHandle : ReferencedGeometryGroups)
+		if (RegisteredGeometry.bEvicted)
 		{
-			checkf(RegisteredGroups.IsValidIndex(GroupHandle), TEXT("RayTracingGeometryGroupHandle must be valid"));
+			RequestedButEvictedSize += RegisteredGeometry.Size;
+		}
+	}
 
-			const FRayTracingGeometryGroup& Group = RegisteredGroups[GroupHandle];
-			const int32 NumGeometriesInGroup = Group.GeometryHandles.Num();
+	// Step 2
+	// - add all geometries in referenced groups to LocalReferencedGeometries
+	//		- need to make all geometries in group resident otherwise might not have valid geometry when reducing LOD
+	//		- TODO: Could track TargetLOD and only make [TargetLOD ... LastLOD] range resident
+	// - also update LastReferencedFrame and calculate memory required to make evicted geometries resident
+	for (RayTracing::FGeometryGroupHandle GroupHandle : ReferencedGeometryGroups)
+	{
+		checkf(RegisteredGroups.IsValidIndex(GroupHandle), TEXT("RayTracingGeometryGroupHandle must be valid"));
 
-			for (uint8 LODIndex = Group.CurrentFirstLODIdx; LODIndex < NumGeometriesInGroup; ++LODIndex)
+		const FRayTracingGeometryGroup& Group = RegisteredGroups[GroupHandle];
+		const int32 NumGeometriesInGroup = Group.GeometryHandles.Num();
+
+		for (uint8 LODIndex = Group.CurrentFirstLODIdx; LODIndex < NumGeometriesInGroup; ++LODIndex)
+		{
+			FGeometryHandle GeometryHandle = Group.GeometryHandles[LODIndex];
+
+			if (GeometryHandle != INDEX_NONE) // some LODs might be stripped during cook
 			{
-				FGeometryHandle GeometryHandle = Group.GeometryHandles[LODIndex];
+				FRegisteredGeometry& RegisteredGeometry = RegisteredGeometries[GeometryHandle];
+				RegisteredGeometry.LastReferencedFrame = GFrameCounterRenderThread;
+
+				ReferencedSize += RegisteredGeometry.Size;
 
-				if (GeometryHandle != INDEX_NONE) // some LODs might be stripped during cook
+				if (RegisteredGeometry.bAlwaysResident)
 				{
-					FRegisteredGeometry& RegisteredGeometry = RegisteredGeometries[GeometryHandle];
-					RegisteredGeometry.LastReferencedFrame = GFrameCounterRenderThread;
+					checkf(!RegisteredGeometry.bEvicted, TEXT("Always resident ray tracing geometry was unexpectely evicted."));
+				}
+				else
+				{
+					LocalReferencedGeometries.Add(GeometryHandle);
+					NotReferencedResidentGeometries.Remove(GeometryHandle);
 
-					ReferencedSize += RegisteredGeometry.Size;
+					RequestedSize += RegisteredGeometry.Size;
 
-					if (RegisteredGeometry.bAlwaysResident)
-					{
-						checkf(!RegisteredGeometry.bEvicted, TEXT("Always resident ray tracing geometry was unexpectely evicted."));
-					}
-					else
+					if (RegisteredGeometry.bEvicted)
 					{
-						ReferencedGeometries.Add(GeometryHandle);
-						NotReferencedResidentGeometries.Remove(GeometryHandle);
-
-						RequestedSize += RegisteredGeometry.Size;
-
-						if (RegisteredGeometry.bEvicted)
-						{
-							RequestedButEvictedSize += RegisteredGeometry.Size;
-						}
+						RequestedButEvictedSize += RegisteredGeometry.Size;
 					}
 				}
 			}
 		}
+	}
 
 #if DO_CHECK
-		// ensure(ReferencedGeometries.Num() == TSet(ReferencedGeometries).Num());
+	// ensure(LocalReferencedGeometries.Num() == TSet(LocalReferencedGeometries).Num());
 #endif
 
-		const uint64 ResidentGeometryMemoryPoolSize = FUnitConversion::Convert(GRayTracingResidentGeometryMemoryPoolSizeInMB, EUnit::Megabytes, EUnit::Bytes);
+	const int64 ResidentGeometryMemoryPoolSize = FUnitConversion::Convert(GRayTracingResidentGeometryMemoryPoolSizeInMB, EUnit::Megabytes, EUnit::Bytes);
 
-		// Step 3
-		// - if making requested geometries resident will put us over budget -> evict some geometry not referenced by TLAS
-		if (TotalResidentSize + TotalStreamingSize + RequestedButEvictedSize > ResidentGeometryMemoryPoolSize)
-		{
-			TArray<FGeometryHandle> NotReferencedResidentGeometriesArray = NotReferencedResidentGeometries.Array();
+	// Step 3
+	// - if making requested geometries resident will put us over budget -> evict some geometry not referenced by TLAS
+	if (TotalResidentSize + TotalStreamingSize + RequestedButEvictedSize > ResidentGeometryMemoryPoolSize)
+	{
+		TArray<FGeometryHandle> NotReferencedResidentGeometriesArray = NotReferencedResidentGeometries.Array();
 
-			// Step 3.1
-			// - sort to evict geometries in the following order:
-			//		- least recently used
-			//		- largest geometries
-			Algo::Sort(NotReferencedResidentGeometriesArray, [this](FGeometryHandle& LHSHandle, FGeometryHandle& RHSHandle)
-				{
-					FRegisteredGeometry& LHS = RegisteredGeometries[LHSHandle];
-					FRegisteredGeometry& RHS = RegisteredGeometries[RHSHandle];
+		// Step 3.1
+		// - sort to evict geometries in the following order:
+		//		- least recently used
+		//		- largest geometries
+		Algo::Sort(NotReferencedResidentGeometriesArray, [this](FGeometryHandle& LHSHandle, FGeometryHandle& RHSHandle)
+			{
+				FRegisteredGeometry& LHS = RegisteredGeometries[LHSHandle];
+				FRegisteredGeometry& RHS = RegisteredGeometries[RHSHandle];
 
-					// TODO: evict unreferenced dynamic geometries using shared buffers first since they need to be rebuild anyway
-					// (and then dynamic geometries requiring update?
+				// TODO: evict unreferenced dynamic geometries using shared buffers first since they need to be rebuild anyway
+				// (and then dynamic geometries requiring update?
 
-					// 1st - last referenced frame
-					if (LHS.LastReferencedFrame != RHS.LastReferencedFrame)
-					{
-						return LHS.LastReferencedFrame < RHS.LastReferencedFrame;
-					}
+				// 1st - last referenced frame
+				if (LHS.LastReferencedFrame != RHS.LastReferencedFrame)
+				{
+					return LHS.LastReferencedFrame < RHS.LastReferencedFrame;
+				}
 
-					// 2nd - size
-					return LHS.Size > RHS.Size;
-				});
+				// 2nd - size
+				return LHS.Size > RHS.Size;
+			});
 
-			// Step 3.2
-			// - evict geometries until we are in budget
-			int32 Index = 0;
-			while (TotalResidentSize + TotalStreamingSize + RequestedButEvictedSize > ResidentGeometryMemoryPoolSize && Index < NotReferencedResidentGeometriesArray.Num())
-			{
-				FGeometryHandle GeometryHandle = NotReferencedResidentGeometriesArray[Index];
-				FRegisteredGeometry& RegisteredGeometry = RegisteredGeometries[GeometryHandle];
+		// Step 3.2
+		// - evict geometries until we are in budget
+		int32 Index = 0;
+		while (TotalResidentSize + TotalStreamingSize + RequestedButEvictedSize > ResidentGeometryMemoryPoolSize && Index < NotReferencedResidentGeometriesArray.Num())
+		{
+			FGeometryHandle GeometryHandle = NotReferencedResidentGeometriesArray[Index];
+			FRegisteredGeometry& RegisteredGeometry = RegisteredGeometries[GeometryHandle];
 
-				check(RegisteredGeometry.Geometry->IsValid() && !RegisteredGeometry.Geometry->IsEvicted());
+			check(RegisteredGeometry.Geometry->IsValid() && !RegisteredGeometry.Geometry->IsEvicted());
 
-				EvictGeometry(RHICmdList, RegisteredGeometry);
+			EvictGeometry(RHICmdList, RegisteredGeometry);
 
-				++Index;
-			}
+			++Index;
 		}
+	}
 
-		// Step 4
-		// - make referenced geometries resident until we go over budget
-		if (TotalResidentSize + TotalStreamingSize < ResidentGeometryMemoryPoolSize)
-		{
-			// Step 4.1
-			//	- sort by size to prioritize smaller geometries
-			Algo::Sort(ReferencedGeometries, [this](FGeometryHandle& LHSHandle, FGeometryHandle& RHSHandle)
-				{
-					FRegisteredGeometry& LHS = RegisteredGeometries[LHSHandle];
-					FRegisteredGeometry& RHS = RegisteredGeometries[RHSHandle];
-
-					return LHS.Size < RHS.Size;
-				});
-
-			// Step 4.2
-			// - make geometries resident until we go over budget
-			int32 Index = 0;
-			while (TotalResidentSize + TotalStreamingSize < ResidentGeometryMemoryPoolSize && Index < ReferencedGeometries.Num())
+	// Step 4
+	// - make referenced geometries resident until we go over budget
+	if (TotalResidentSize + TotalStreamingSize < ResidentGeometryMemoryPoolSize)
+	{
+		// Step 4.1
+		//	- sort by size to prioritize smaller geometries
+		Algo::Sort(LocalReferencedGeometries, [this](FGeometryHandle& LHSHandle, FGeometryHandle& RHSHandle)
 			{
-				FGeometryHandle GeometryHandle = ReferencedGeometries[Index];
-				FRegisteredGeometry& RegisteredGeometry = RegisteredGeometries[GeometryHandle];
+				FRegisteredGeometry& LHS = RegisteredGeometries[LHSHandle];
+				FRegisteredGeometry& RHS = RegisteredGeometries[RHSHandle];
 
-				if (RegisteredGeometry.Geometry->IsEvicted())
-				{
-					MakeGeometryResident(RHICmdList, RegisteredGeometry);
-				}
+				return LHS.Size < RHS.Size;
+			});
 
-				RequestRayTracingGeometryStreamIn(RHICmdList, GeometryHandle);
+		// Step 4.2
+		// - make geometries resident until we go over budget
+		int32 Index = 0;
+		while (TotalResidentSize + TotalStreamingSize < ResidentGeometryMemoryPoolSize && Index < LocalReferencedGeometries.Num())
+		{
+			FGeometryHandle GeometryHandle = LocalReferencedGeometries[Index];
+			FRegisteredGeometry& RegisteredGeometry = RegisteredGeometries[GeometryHandle];
 
-				++Index;
+			if (RegisteredGeometry.Geometry->IsEvicted())
+			{
+				MakeGeometryResident(RHICmdList, RegisteredGeometry);
 			}
+
+			RequestRayTracingGeometryStreamIn(RHICmdList, GeometryHandle);
+
+			++Index;
 		}
+	}
 
-		SET_MEMORY_STAT(STAT_RayTracingGeometryReferencedMemory, ReferencedSize);
-		SET_MEMORY_STAT(STAT_RayTracingGeometryRequestedMemory, RequestedSize);
-		CSV_CUSTOM_STAT(RayTracingGeometry, ReferencedSizeMB, ReferencedSize / 1024.0f / 1024.0f, ECsvCustomStatOp::Set);
-		CSV_CUSTOM_STAT(RayTracingGeometry, RequestedSizeMB, RequestedSize / 1024.0f / 1024.0f, ECsvCustomStatOp::Set);
+	SET_MEMORY_STAT(STAT_RayTracingGeometryReferencedMemory, ReferencedSize);
+	SET_MEMORY_STAT(STAT_RayTracingGeometryRequestedMemory, RequestedSize);
+	CSV_CUSTOM_STAT(RayTracingGeometry, ReferencedSizeMB, ReferencedSize / 1024.0f / 1024.0f, ECsvCustomStatOp::Set);
+	CSV_CUSTOM_STAT(RayTracingGeometry, RequestedSizeMB, RequestedSize / 1024.0f / 1024.0f, ECsvCustomStatOp::Set);
 
 #if !UE_BUILD_SHIPPING
-		if (GRayTracingShowOnScreenWarnings)
+	if (GRayTracingShowOnScreenWarnings)
+	{
+		if (TotalAlwaysResidentSize > ResidentGeometryMemoryPoolSize)
 		{
-			if (TotalAlwaysResidentSize > ResidentGeometryMemoryPoolSize)
-			{
-				GEngine->AddOnScreenDebugMessage(uint64(this), 1.f, FColor::Red,
-					*FString::Printf(TEXT("RAY TRACING GEOMETRY - ALWAYS RESIDENT MEMORY OVER BUDGET %s / %s"),
-						*FText::AsMemory(TotalAlwaysResidentSize).ToString(),
-						*FText::AsMemory(ResidentGeometryMemoryPoolSize).ToString()));
-			}
-			else if (RequestedSize > ResidentGeometryMemoryPoolSize)
-			{
-				GEngine->AddOnScreenDebugMessage(uint64(this), 1.f, FColor::Yellow,
-					*FString::Printf(TEXT("RAY TRACING GEOMETRY - REQUESTED MEMORY OVER BUDGET %s / %s"),
-						*FText::AsMemory(RequestedSize).ToString(),
-						*FText::AsMemory(ResidentGeometryMemoryPoolSize).ToString()));
-			}
+			GEngine->AddOnScreenDebugMessage(uint64(this), 1.f, FColor::Red,
+				*FString::Printf(TEXT("RAY TRACING GEOMETRY - ALWAYS RESIDENT MEMORY OVER BUDGET %s / %s"),
+					*FText::AsMemory(TotalAlwaysResidentSize).ToString(),
+					*FText::AsMemory(ResidentGeometryMemoryPoolSize).ToString()));
 		}
+		else if (RequestedSize > ResidentGeometryMemoryPoolSize)
+		{
+			GEngine->AddOnScreenDebugMessage(uint64(this), 1.f, FColor::Yellow,
+				*FString::Printf(TEXT("RAY TRACING GEOMETRY - REQUESTED MEMORY OVER BUDGET %s / %s"),
+					*FText::AsMemory(RequestedSize).ToString(),
+					*FText::AsMemory(ResidentGeometryMemoryPoolSize).ToString()));
+		}
+	}
 #endif
+}
+
+void FRayTracingGeometryManager::MakeResidentAllGeometries(FRHICommandList& RHICmdList)
+{
+	for (FRegisteredGeometry& RegisteredGeometry : RegisteredGeometries)
+	{
+		if (RegisteredGeometry.Geometry->IsEvicted())
+		{
+			MakeGeometryResident(RHICmdList, RegisteredGeometry);
+		}
+
+		if (!RequestRayTracingGeometryStreamIn(RHICmdList, RegisteredGeometry.Geometry->RayTracingGeometryHandle))
+		{
+			PendingStreamingRequests.Add(RegisteredGeometry.Geometry->RayTracingGeometryHandle);
+
+			INC_DWORD_STAT(STAT_RayTracingPendingStreamingRequests);
+		}
 	}
-	else
+}
+
+void FRayTracingGeometryManager::EvictAllGeometries(FRHICommandList& RHICmdList)
+{
+	for (FRegisteredGeometry& RegisteredGeometry : RegisteredGeometries)
 	{
-		check(IsRayTracingEnabled());
+		if (RegisteredGeometry.Geometry->GetRHI() != nullptr)
+		{
+			EvictGeometry(RHICmdList, RegisteredGeometry);
+		}
+	}
+}
 
-		if (bHasRayTracingEnableChanged)
+void FRayTracingGeometryManager::RefreshAlwaysResidentRayTracingGeometries(FRHICommandList& RHICmdList)
+{
+	const bool bUsingReferenceBasedResidency = IsRayTracingUsingReferenceBasedResidency();
+
+	if (bRefreshAlwaysResidentRayTracingGeometries)
+	{
+		bRefreshAlwaysResidentRayTracingGeometries = false;
+
+		AlwaysResidentGeometries.Empty();
+		TotalAlwaysResidentSize = 0;
+
+		for (FRegisteredGeometry& RegisteredGeometry : RegisteredGeometries)
 		{
-			// make all geometries resident
-			for (FRegisteredGeometry& RegisteredGeometry : RegisteredGeometries)
+			const RayTracing::FGeometryGroupHandle GroupHandle = RegisteredGeometry.Geometry->GroupHandle;
+
+			if (GroupHandle == INDEX_NONE)
+			{
+				RegisteredGeometry.bAlwaysResident = false;
+				continue;
+			}
+
+			const FRayTracingGeometryGroup& Group = RegisteredGroups[GroupHandle];
+
+			RegisteredGeometry.bAlwaysResident = IsAlwaysResidentGeometry(RegisteredGeometry.Geometry, Group);
+
+			if (RegisteredGeometry.bAlwaysResident)
 			{
+				AlwaysResidentGeometries.Add(RegisteredGeometry.Geometry->RayTracingGeometryHandle);
+				TotalAlwaysResidentSize += RegisteredGeometry.Size;
+
 				if (RegisteredGeometry.Geometry->IsEvicted())
 				{
 					MakeGeometryResident(RHICmdList, RegisteredGeometry);
@@ -910,55 +932,66 @@ void FRayTracingGeometryManager::Tick(FRHICommandList& RHICmdList)
 
 					INC_DWORD_STAT(STAT_RayTracingPendingStreamingRequests);
 				}
+
+				EvictableGeometries.Remove(RegisteredGeometry.Geometry->RayTracingGeometryHandle);
 			}
-		}
-		else
-		{
-#if DO_CHECK
-			// otherwise just check that all geometries are resident
-			for (FRegisteredGeometry& RegisteredGeometry : RegisteredGeometries)
+			else if (bUsingReferenceBasedResidency && RegisteredGeometry.Geometry->GetRHI() != nullptr)
 			{
-				checkf(!RegisteredGeometry.Geometry->IsEvicted(), TEXT("Ray tracing geometry should not be evicted when ray tracing is enabled."));
+				EvictGeometry(RHICmdList, RegisteredGeometry);
 			}
-#endif
 		}
-
-		SET_MEMORY_STAT(STAT_RayTracingGeometryReferencedMemory, 0);
-		SET_MEMORY_STAT(STAT_RayTracingGeometryRequestedMemory, TotalResidentSize);
-		CSV_CUSTOM_STAT(RayTracingGeometry, ReferencedSizeMB, 0, ECsvCustomStatOp::Set);
-		CSV_CUSTOM_STAT(RayTracingGeometry, RequestedSizeMB, TotalResidentSize / 1024.0f / 1024.0f, ECsvCustomStatOp::Set);
 	}
+}
+
+#if DO_CHECK
+void FRayTracingGeometryManager::CheckIntegrity()
+{
+	const bool bIsRayTracingEnabled = IsRayTracingEnabled();
+	const bool bUsingReferenceBasedResidency = IsRayTracingUsingReferenceBasedResidency();
 
+	if (GRayTracingTestCheckIntegrity)
 	{
-		TSet<FGeometryHandle> CurrentPendingStreamingRequests;
-		Swap(CurrentPendingStreamingRequests, PendingStreamingRequests);
-		PendingStreamingRequests.Reserve(CurrentPendingStreamingRequests.Num());
+		for (FRegisteredGeometry& RegisteredGeometry : RegisteredGeometries)
+		{
+			const FGeometryHandle GeometryHandle = RegisteredGeometry.Geometry->RayTracingGeometryHandle;
+			const RayTracing::FGeometryGroupHandle GroupHandle = RegisteredGeometry.Geometry->GroupHandle;
 
-		for (FGeometryHandle GeometryHandle : CurrentPendingStreamingRequests)
+			bool bAlwaysResident = false;
+
+			if (GroupHandle != INDEX_NONE)
+			{
+				const FRayTracingGeometryGroup& Group = RegisteredGroups[GroupHandle];
+				bAlwaysResident = IsAlwaysResidentGeometry(RegisteredGeometry.Geometry, Group);
+			}
+			else
+			{
+				// geometries not assigned to a group (eg: dynamic geometry) are always evictable
+			}
+
+			checkf(RegisteredGeometry.bEvicted == RegisteredGeometry.Geometry->IsEvicted(), TEXT("Cached bEvicted flag in FRegisteredGeometry is stale"));
+			checkf(RegisteredGeometry.bAlwaysResident == bAlwaysResident, TEXT("Cached bAlwaysResident flag in FRegisteredGeometry is stale"));
+			checkf(RegisteredGeometry.bAlwaysResident == AlwaysResidentGeometries.Contains(GeometryHandle), TEXT("Geometry with bAlwaysResident flag set should be in the AlwaysResidentGeometries set."));
+		}
+
+		if (!bIsRayTracingEnabled)
 		{
-			if (!RequestRayTracingGeometryStreamIn(RHICmdList, GeometryHandle))
+			// check that everything is evicted
+			for (FRegisteredGeometry& RegisteredGeometry : RegisteredGeometries)
 			{
-				PendingStreamingRequests.Add(GeometryHandle);
+				checkf(RegisteredGeometry.Geometry->IsEvicted() || RegisteredGeometry.Geometry->GetRHI() == nullptr, TEXT("Ray tracing geometry should be evicted when ray tracing is disabled."));
+			}
+		}
+		else if (!bUsingReferenceBasedResidency)
+		{
+			// check that all geometries are resident
+			for (FRegisteredGeometry& RegisteredGeometry : RegisteredGeometries)
+			{
+				checkf(!RegisteredGeometry.Geometry->IsEvicted(), TEXT("Ray tracing geometry should not be evicted when ray tracing is enabled."));
 			}
 		}
 	}
-
-	SET_DWORD_STAT(STAT_RayTracingPendingStreamingRequests, PendingStreamingRequests.Num());
-
-	ProcessCompletedStreamingRequests(RHICmdList);
-
-	ReferencedGeometryHandles.Reset();
-	ReferencedGeometryGroups.Reset();
-	ReferencedGeometryGroupsForDynamicUpdate.Reset();
-
-	bHasRayTracingEnableChanged = false;
-
-	SET_MEMORY_STAT(STAT_RayTracingGeometryResidentMemory, TotalResidentSize);
-	SET_MEMORY_STAT(STAT_RayTracingGeometryAlwaysResidentMemory, TotalAlwaysResidentSize);
-
-	CSV_CUSTOM_STAT(RayTracingGeometry, TotalResidentSizeMB, TotalResidentSize / 1024.0f / 1024.0f, ECsvCustomStatOp::Set);
-	CSV_CUSTOM_STAT(RayTracingGeometry, TotalAlwaysResidentSizeMB, TotalAlwaysResidentSize / 1024.0f / 1024.0f, ECsvCustomStatOp::Set);
 }
+#endif // DO_CHECK
 
 bool FRayTracingGeometryManager::RequestRayTracingGeometryStreamIn(FRHICommandList& RHICmdList, FGeometryHandle GeometryHandle)
 {
@@ -1417,9 +1450,13 @@ void FRayTracingGeometryManager::AddReferencedGeometry(const FRayTracingGeometry
 	{
 		if (RegisteredGeometries.IsValidIndex(Geometry->RayTracingGeometryHandle))
 		{
-			ReferencedGeometryHandles.Add(Geometry->RayTracingGeometryHandle);
+			ReferencedGeometries.Add(Geometry->RayTracingGeometryHandle);
 		}
 	}
+	else
+	{
+		ensureMsgf(ReferencedGeometries.IsEmpty(), TEXT("Should only track ReferencedGeometries when using using reference based residency"));
+	}
 }
 
 void FRayTracingGeometryManager::AddReferencedGeometryGroups(const TSet<RayTracing::FGeometryGroupHandle>& GeometryGroups)
@@ -1470,7 +1507,7 @@ void FRayTracingGeometryManager::ResetVisibleGeometries()
 #if DO_CHECK
 bool FRayTracingGeometryManager::IsGeometryReferenced(const FRayTracingGeometry* Geometry) const
 {
-	return ReferencedGeometryHandles.Contains(Geometry->RayTracingGeometryHandle);
+	return ReferencedGeometries.Contains(Geometry->RayTracingGeometryHandle);
 }
 
 bool FRayTracingGeometryManager::IsGeometryGroupReferenced(RayTracing::FGeometryGroupHandle GeometryGroup) const
diff --git a/Engine/Source/Runtime/Engine/Public/Rendering/RayTracingGeometryManager.h b/Engine/Source/Runtime/Engine/Public/Rendering/RayTracingGeometryManager.h
--- a/Engine/Source/Runtime/Engine/Public/Rendering/RayTracingGeometryManager.h
+++ b/Engine/Source/Runtime/Engine/Public/Rendering/RayTracingGeometryManager.h
@@ -36,8 +36,8 @@ public:
 
 	ENGINE_API virtual void RefreshRegisteredGeometry(FGeometryHandle Handle) override;
 
-	void SetRayTracingGeometryStreamingData(const FRayTracingGeometry* Geometry, FRayTracingStreamableAsset& StreamableAsset);
-	void SetRayTracingGeometryGroupCurrentFirstLODIndex(FRHICommandListBase& RHICmdList, RayTracing::FGeometryGroupHandle Handle, uint8 CurrentFirstLODIdx);
+	ENGINE_API void SetRayTracingGeometryStreamingData(const FRayTracingGeometry* Geometry, FRayTracingStreamableAsset& StreamableAsset);
+	ENGINE_API void SetRayTracingGeometryGroupCurrentFirstLODIndex(FRHICommandListBase& RHICmdList, RayTracing::FGeometryGroupHandle Handle, uint8 CurrentFirstLODIdx);
 
 	ENGINE_API virtual void PreRender() override;
 	ENGINE_API virtual void Tick(FRHICommandList& RHICmdList) override;
@@ -82,6 +82,8 @@ private:
 	bool RequestRayTracingGeometryStreamIn(FRHICommandList& RHICmdList, FGeometryHandle GeometryHandle);
 	void ProcessCompletedStreamingRequests(FRHICommandList& RHICmdList);
 
+	void UpdateReferenceBasedResidency(FRHICommandList& RHICmdList);
+
 	FCriticalSection RequestCS;
 
 	TSparseArray<FBuildRequest> GeometryBuildRequests;
@@ -184,6 +186,15 @@ private:
 	// Must call this instead of FRayTracingGeometry::Evict() directly
 	void EvictGeometry(FRHICommandListBase& RHICmdList, FRegisteredGeometry& RegisteredGeometry);
 
+	void MakeResidentAllGeometries(FRHICommandList& RHICmdList);
+	void EvictAllGeometries(FRHICommandList& RHICmdList);
+
+	void RefreshAlwaysResidentRayTracingGeometries(FRHICommandList& RHICmdList);
+
+#if DO_CHECK
+	void CheckIntegrity();
+#endif
+
 	static bool IsAlwaysResidentGeometry(const FRayTracingGeometry* InGeometry, const FRayTracingGeometryGroup& Group);
 
 	TSparseArray<FRayTracingGeometryGroup> RegisteredGroups;
@@ -192,14 +203,14 @@ private:
 	TSparseArray<FRegisteredGeometry> RegisteredGeometries;
 
 	TSet<FGeometryHandle> ResidentGeometries;
-	uint64 TotalResidentSize = 0;
+	int64 TotalResidentSize = 0;
 
 	TSet<FGeometryHandle> AlwaysResidentGeometries;
-	uint64 TotalAlwaysResidentSize = 0;
+	int64 TotalAlwaysResidentSize = 0;
 
 	TSet<FGeometryHandle> EvictableGeometries;
 
-	TSet<FGeometryHandle> ReferencedGeometryHandles;
+	TSet<FGeometryHandle> ReferencedGeometries;
 	TSet<RayTracing::FGeometryGroupHandle> ReferencedGeometryGroups;
 
 	TSet<RayTracing::FGeometryGroupHandle> ReferencedGeometryGroupsForDynamicUpdate;
@@ -212,7 +223,7 @@ private:
 
 	// Total size of the BLASes currently being streamed.
 	// Used to keep track of how much TotalResidentSize will increase when the requests complete.
-	uint64 TotalStreamingSize = 0;
+	int64 TotalStreamingSize = 0;
 
 	// Feedback
 	TSet<int32> VisibleGeometryHandles;
