Don't allow GPU skin cache on a component if it doesn't have access to a GPU skin cache instance.

diff --git a/Engine/Source/Runtime/Engine/Private/Components/SkinnedMeshComponent.cpp b/Engine/Source/Runtime/Engine/Private/Components/SkinnedMeshComponent.cpp
--- a/Engine/Source/Runtime/Engine/Private/Components/SkinnedMeshComponent.cpp
+++ b/Engine/Source/Runtime/Engine/Private/Components/SkinnedMeshComponent.cpp
@@ -2329,7 +2329,7 @@ bool USkinnedMeshComponent::GetTwistAndSwingAngleOfDeltaRotationFromRefPose(FNam
 
 bool USkinnedMeshComponent::IsSkinCacheAllowed(int32 LodIdx) const
 {
-	if (!GEnableGPUSkinCache)
+	if (!GEnableGPUSkinCache || !GetScene() || !GetScene()->GetGPUSkinCache())
 	{
 		return false;
 	}
diff --git a/Engine/Source/Runtime/Engine/Private/SkinnedMeshSceneProxyDesc.cpp b/Engine/Source/Runtime/Engine/Private/SkinnedMeshSceneProxyDesc.cpp
--- a/Engine/Source/Runtime/Engine/Private/SkinnedMeshSceneProxyDesc.cpp
+++ b/Engine/Source/Runtime/Engine/Private/SkinnedMeshSceneProxyDesc.cpp
@@ -298,7 +298,7 @@ bool FSkinnedMeshSceneProxyDesc::IsSkinCacheAllowed(int32 LodIdx) const
 {
 	const bool bGlobalDefault = GetSkinCacheDefaultBehavior() == ESkinCacheDefaultBehavior::Inclusive;
 
-	if (!GEnableGPUSkinCache)
+	if (!GEnableGPUSkinCache || !Scene || !Scene->GetGPUSkinCache())
 	{
 		return false;
 	}
diff --git a/Engine/Source/Runtime/Engine/Public/SkinnedMeshSceneProxyDesc.h b/Engine/Source/Runtime/Engine/Public/SkinnedMeshSceneProxyDesc.h
--- a/Engine/Source/Runtime/Engine/Public/SkinnedMeshSceneProxyDesc.h
+++ b/Engine/Source/Runtime/Engine/Public/SkinnedMeshSceneProxyDesc.h
@@ -117,7 +117,6 @@ struct FSkinnedMeshSceneProxyDynamicData
 	ENGINE_API FSkinnedMeshSceneProxyDynamicData();
 	ENGINE_API bool IsValidExternalMorphSetLODIndex(uint32 InLODIndex) const;
 	ENGINE_API const FExternalMorphSets& GetExternalMorphSets(uint32 InLODIndex) const;
-	ENGINE_API bool IsSkinCacheAllowed(int32 LodIdx, const USkinnedAsset* InSkinnedAsset) const;
 	ENGINE_API int32 GetMeshDeformerMaxLOD() const;
 	ENGINE_API UMeshDeformerInstance* GetMeshDeformerInstanceForLOD(int32 LODIndex) const;
 
@@ -166,4 +165,7 @@ struct FSkinnedMeshSceneProxyDynamicData
 	uint8 bDrawInGame : 1 = true;
 	uint8 bCastsHiddenShadow : 1 = false;
 	uint8 bAffectIndirectLightingWhileHidden : 1 = false;
+
+	UE_DEPRECATED(5.7, TEXT("IsSkinCacheAllowed is deprecated"))
+	ENGINE_API bool IsSkinCacheAllowed(int32 LodIdx, const USkinnedAsset* InSkinnedAsset) const;
 };
\ No newline at end of file
