Fixed build error due to malformed deprecation tag.

diff --git a/Engine/Source/Runtime/Engine/Public/SkinnedMeshSceneProxyDesc.h b/Engine/Source/Runtime/Engine/Public/SkinnedMeshSceneProxyDesc.h
--- a/Engine/Source/Runtime/Engine/Public/SkinnedMeshSceneProxyDesc.h
+++ b/Engine/Source/Runtime/Engine/Public/SkinnedMeshSceneProxyDesc.h
@@ -166,6 +166,6 @@ struct FSkinnedMeshSceneProxyDynamicData
 	uint8 bCastsHiddenShadow : 1 = false;
 	uint8 bAffectIndirectLightingWhileHidden : 1 = false;
 
-	UE_DEPRECATED(5.7, TEXT("IsSkinCacheAllowed is deprecated"))
+	UE_DEPRECATED(5.7, "IsSkinCacheAllowed is deprecated")
 	ENGINE_API bool IsSkinCacheAllowed(int32 LodIdx, const USkinnedAsset* InSkinnedAsset) const;
 };
\ No newline at end of file
