Level Instance: fix missing RF_Standalone flag reset when unloading worlds.

diff --git a/Engine/Source/Runtime/Engine/Private/WorldPartition/WorldPartitionPackageHelper.cpp b/Engine/Source/Runtime/Engine/Private/WorldPartition/WorldPartitionPackageHelper.cpp
--- a/Engine/Source/Runtime/Engine/Private/WorldPartition/WorldPartitionPackageHelper.cpp
+++ b/Engine/Source/Runtime/Engine/Private/WorldPartition/WorldPartitionPackageHelper.cpp
@@ -21,21 +21,23 @@ void FWorldPartitionPackageHelper::UnloadPackage(UPackage* InPackage)
 	// World specific
 	if (UWorld* PackageWorld = UWorld::FindWorldInPackage(InPackage))
 	{
+		PackageWorld->ClearFlags(RF_Standalone);
+
 		if (PackageWorld->PersistentLevel)
 		{
 			// Manual cleanup of level since world was not initialized
 			PackageWorld->PersistentLevel->CleanupLevel(/*bCleanupResources*/ true, /*bUnloadFromEditor*/true);
-		}
 
-		if (PackageWorld->PersistentLevel->IsUsingExternalObjects())
-		{
-			ForEachObjectWithOuter(PackageWorld->PersistentLevel, [&TrashPackage](UObject* InObject)
+			if (PackageWorld->PersistentLevel->IsUsingExternalObjects())
 			{
-				if (UPackage* ExternalPackage = InObject->GetExternalPackage())
+				ForEachObjectWithOuter(PackageWorld->PersistentLevel, [&TrashPackage](UObject* InObject)
 				{
-					TrashPackage(ExternalPackage);
-				}
-			}, /*bIncludeNestedObjects*/ true);
+					if (UPackage* ExternalPackage = InObject->GetExternalPackage())
+					{
+						TrashPackage(ExternalPackage);
+					}
+				}, /*bIncludeNestedObjects*/ true);
+			}
 		}
 	}
 }
