FSimpleStreamableRenderAssetManager - ensuring we have only one bounds per registered proxy as there is one bounds provided by [REDACTED_PERSON]
Removing the code that assumes there can be more than one per proxy bounds
Adding TSimpleSparseArray that will use TBitSet to insert/remove elements in empty spaces for cheap add/remove but allow fast TArray iteration over elements
Adding possibility to control TSimpleSparseArray growth speed
Allow optionally sort and make unique list of RenderAssets
- it yields benefits only if we use a lot of sections with same material
or if different materials reference same textures within single proxy

diff --git a/Engine/Source/Runtime/Engine/Private/Streaming/SimpleStreamableAssetManager.cpp b/Engine/Source/Runtime/Engine/Private/Streaming/SimpleStreamableAssetManager.cpp
--- a/Engine/Source/Runtime/Engine/Private/Streaming/SimpleStreamableAssetManager.cpp
+++ b/Engine/Source/Runtime/Engine/Private/Streaming/SimpleStreamableAssetManager.cpp
@@ -25,6 +25,8 @@ TRACE_DECLARE_INT_COUNTER(UpdatedObjects, TEXT("StreamableAssets/SimpleStreamabl
 
 FSimpleStreamableAssetManager* FSimpleStreamableAssetManager::Instance = nullptr;
 int32 FSimpleStreamableAssetManager::GUseSimpleStreamableAssetManager = 0;
+int32 FSimpleStreamableAssetManager::GSimpleStreamableAssetManagerSparseArrayGrowSize = 64;
+int32 FSimpleStreamableAssetManager::GSimpleStreamableAssetManagerEnsureAssetUniqueOnRegistration = 1;
 
 FAutoConsoleVariableRef FSimpleStreamableAssetManager::CVarUseSimpleStreamableAssetManager(
 TEXT("s.StreamableAssets.UseSimpleStreamableAssetManager"),
@@ -35,6 +37,21 @@ TEXT("If 1, The FSimpleStreamableAssetManager is Enabled and works by integratin
 ECVF_SetByGameSetting | ECVF_ReadOnly
 );
 
+FAutoConsoleVariableRef FSimpleStreamableAssetManager::CVarUseSimpleStreamableAssetManagerSparseArrayGrowSize(
+TEXT("s.StreamableAssets.SimpleStreamableAssetManager.SparseArrayGrowSize"),
+FSimpleStreamableAssetManager::GSimpleStreamableAssetManagerSparseArrayGrowSize,
+TEXT("The growth size of SparseArray used for tracking objects pointing specific assets"),
+ECVF_SetByGameSetting
+);
+FAutoConsoleVariableRef FSimpleStreamableAssetManager::CVarUseSimpleStreamableAssetManagerEnsureAssetUniqueOnRegistration(
+TEXT("s.StreamableAssets.SimpleStreamableAssetManager.SortAssetsOnRegistration"),
+FSimpleStreamableAssetManager::GSimpleStreamableAssetManagerEnsureAssetUniqueOnRegistration,
+TEXT("If true when object will be added, referenced assets will be sorted \n")
+     TEXT("we will make sure we register asset only once with highest Texel Factor \n")
+     TEXT("It will be beneficial only if multiple materials use same texture"),
+ECVF_SetByGameSetting
+);
+
 
 void FSimpleStreamableAssetManager::Init()
 {
@@ -148,16 +165,18 @@ void FSimpleStreamableAssetManager::GetAssetReferenceBounds_Impl(const UStreamab
 
 	if (AssetIndex != INDEX_NONE)
 	{
-		const TSet<FAssetBoundElement>& AssetElements = AssetIndexToBounds4Index[AssetIndex];
+		const TSimpleSparseArray<FAssetBoundElement>& AssetElements = AssetIndexToBounds4Index[AssetIndex];
 
 		AssetBoxes.Reserve(AssetElements.Num());
 		
-		for (const FAssetBoundElement& Element : AssetElements)
+		for (const FAssetBoundElement& Element : AssetElements.GetSparseView())
 		{
-			const int32 BoundIndex = Element.ElementIndex;
-			const FBoxSphereBounds Bounds = GetBounds(BoundIndex);
-
-			AssetBoxes.Add(Bounds.GetBox());
+			const int32 ObjectRegistrationIndex = Element.ObjectRegistrationIndex;
+			if (ObjectRegistrationIndex != INDEX_NONE)
+			{
+				const FBoxSphereBounds Bounds = GetBounds(ObjectRegistrationIndex);
+				AssetBoxes.Add(Bounds.GetBox());
+			}
 		}
 	}
 }
@@ -171,17 +190,12 @@ uint32 FSimpleStreamableAssetManager::GetAllocatedSize_Impl() const
 	{
 		AllocatedSize += Assets.GetAllocatedSize();
 	}
-	for (const TArray<int32>& BoundsIndices : ObjectRegistrationIndexToBounds4Index)
-	{
-		AllocatedSize += BoundsIndices.GetAllocatedSize();
-	}
-
-	AllocatedSize += ObjectBounds4UsedIndices.GetAllocatedSize();
+	
 	AllocatedSize += ObjectBounds4.GetAllocatedSize();
 
 	AllocatedSize += AssetUsedIndices.GetAllocatedSize();
 
-	for (const TSet<FAssetBoundElement>& AssetBoundElements : AssetIndexToBounds4Index)
+	for (const TSimpleSparseArray<FAssetBoundElement>& AssetBoundElements : AssetIndexToBounds4Index)
 	{
 		AllocatedSize += AssetBoundElements.GetAllocatedSize();
 	}
@@ -230,8 +244,15 @@ void FSimpleStreamableAssetManager::UpdateTask_Async()
 				SCOPED_NAMED_EVENT(FSimpleStreamableAssetManager_Resize, FColor::Silver);
 				ObjectUsedIndices.Add(false, NeedToReserve);
 				ObjectRegistrationIndexToAssetProperty.AddZeroed(NeedToReserve);
-				ObjectRegistrationIndexToBounds4Index.AddZeroed(NeedToReserve);
 				MaxObjects += NeedToReserve;
+
+				// We store 4 Object bounds in one entry
+				const int32 MaxObject4BoundsNeeded = MaxObjects/4 + int32(MaxObjects % 4 != 0);
+				const int32 Object4BoundsNeedToReserve = MaxObject4BoundsNeeded - ObjectBounds4.Num();
+				if (Object4BoundsNeedToReserve > 0)
+				{
+					ObjectBounds4.AddDefaulted(Object4BoundsNeedToReserve);
+				}
 			}
 
 			// Cheap index assignment for proxy
@@ -404,46 +425,54 @@ void FSimpleStreamableAssetManager::RegisterRecord(FRegister& Record)
 		}
 	}
 
-	// Sort by PackedRelativeBox, to identical bounds identical entries together
-	// Sort by Texture to merge duplicate texture entries.
-	// Then sort by TexelFactor, to merge negative entries together.
-	Assets.Sort([](const FStreamingRenderAssetPrimitiveInfo& Lhs, const FStreamingRenderAssetPrimitiveInfo& Rhs)
-	{
-		if (Lhs.PackedRelativeBox == Rhs.PackedRelativeBox)
-		{
-			if (Lhs.RenderAsset == Rhs.RenderAsset)
-			{
-				return Lhs.TexelFactor < Rhs.TexelFactor;
-			}
-			return Lhs.RenderAsset < Rhs.RenderAsset;
-		}
-		return Lhs.PackedRelativeBox < Rhs.PackedRelativeBox;
-	});
-
-	float MinDistanceSq = 0, MinRangeSq = 0, MaxRangeSq = FLT_MAX;
+	int32 AssetCount = Assets.Num();
 
-	// Loop for each bound.
-	if (Assets.Num() > 0)
+	if (GSimpleStreamableAssetManagerEnsureAssetUniqueOnRegistration && AssetCount > 0)
 	{
-		const bool bForceLOD = Record.bForceMipStreaming;
-		for (int32 InfoIndex = 0; InfoIndex < Assets.Num();)
+		// Sort by Texture to merge duplicate texture entries.
+		// Then sort by TexelFactor
+		Assets.Sort([](const FStreamingRenderAssetPrimitiveInfo& Lhs, const FStreamingRenderAssetPrimitiveInfo& Rhs)
 		{
-			const FStreamingRenderAssetPrimitiveInfo& Info = Assets[InfoIndex];
+				if (Lhs.RenderAsset == Rhs.RenderAsset)
+				{
+					return Lhs.TexelFactor > Rhs.TexelFactor;
+				}
+				return Lhs.RenderAsset < Rhs.RenderAsset;
+		});
 
-			int32 NumOfBoundReferences = 1;
-			for (int32 NextInfoIndex = InfoIndex + 1; NextInfoIndex < Assets.Num() && Assets[NextInfoIndex].Bounds == Info.Bounds; ++NextInfoIndex)
+		int32 ProcessedIndex = 0;
+		int32 EmplaceIndex = 1;
+		for (int32 Index = 1; Index < Assets.Num(); Index++)
+		{
+			// We found new asset 
+			if (Assets[ProcessedIndex].RenderAsset != Assets[Index].RenderAsset)
+			{
+				
+				if (EmplaceIndex != Index) // We need to fill the gap in array
+				{
+					Assets[EmplaceIndex] = MoveTemp(Assets[Index]);
+				}
+				++ProcessedIndex;
+				++EmplaceIndex;
+			}
+			// For landscape entries negative TexelFactor values are used we need to ensure we get min while sorted max 
+			else if (Assets[ProcessedIndex].TexelFactor < 0 && Assets[ProcessedIndex].RenderAsset == Assets[Index].RenderAsset)
 			{
-				++NumOfBoundReferences;
+				Assets[ProcessedIndex].TexelFactor = FMath::Min(Assets[ProcessedIndex].TexelFactor, Assets[Index].TexelFactor);
 			}
+		}
 
-			GetDistanceAndRange(Record, MinDistanceSq, MinRangeSq, MaxRangeSq);
-			
-			const int32 BoundsIndex = AddBounds(Record.ObjectBounds, PackedRelativeBox_Identity, Record.LastRenderedTime, Record.ObjectBounds.Origin, MinDistanceSq, MinRangeSq, MaxRangeSq);
+		AssetCount = EmplaceIndex;
+	}
 		
-			AddRenderAssetElements(TArrayView<FStreamingRenderAssetPrimitiveInfo>(Assets.GetData() + InfoIndex, NumOfBoundReferences), BoundsIndex, ObjectIndex, bForceLOD);
-
-			InfoIndex += NumOfBoundReferences;
-		}
+	if (AssetCount > 0)
+	{
+		float MinDistanceSq = 0, MinRangeSq = 0, MaxRangeSq = FLT_MAX;
+        GetDistanceAndRange(Record, MinDistanceSq, MinRangeSq, MaxRangeSq);
+        SetBounds(ObjectIndex, Record.ObjectBounds, PackedRelativeBox_Identity, Record.LastRenderedTime,Record.ObjectBounds.Origin, MinDistanceSq, MinRangeSq, MaxRangeSq);
+       	
+		const bool bForceLOD = Record.bForceMipStreaming;
+		AddRenderAssetElements(TArrayView<FStreamingRenderAssetPrimitiveInfo>(Assets.GetData(), AssetCount), ObjectIndex, bForceLOD);
 	}
 }
 
@@ -452,16 +481,9 @@ void FSimpleStreamableAssetManager::UpdateRecord(FUpdate& Record)
 	const int32 ObjectIndex = *Record.ObjectRegistrationIndex;
 	if (ObjectIndex != INDEX_NONE) // Filter out updates for not registered proxies
 	{
-		const TArray<int32>& ObjectBounds4Indices = ObjectRegistrationIndexToBounds4Index[ObjectIndex];
-
 		float MinDistanceSq = 0, MinRangeSq = 0, MaxRangeSq = FLT_MAX;
-
 		GetDistanceAndRange(Record, MinDistanceSq, MinRangeSq, MaxRangeSq);
-		
-		for (int32 BoundsIndex : ObjectBounds4Indices)
-		{
-			SetBounds(BoundsIndex, Record.ObjectBounds, PackedRelativeBox_Identity, Record.LastRenderedTime,Record.ObjectBounds.Origin, MinDistanceSq, MinRangeSq, MaxRangeSq);
-		}
+		SetBounds(ObjectIndex, Record.ObjectBounds, PackedRelativeBox_Identity, Record.LastRenderedTime,Record.ObjectBounds.Origin, MinDistanceSq, MinRangeSq, MaxRangeSq);
 	}
 }
 
@@ -475,23 +497,7 @@ void FSimpleStreamableAssetManager::UnregisterRecord(FUnregister& Record)
 	}
 }
 
-int32 FSimpleStreamableAssetManager::AddBounds(const FBoxSphereBounds& Bounds, uint32 PackedRelativeBox, float LastRenderTime, const FVector4& RangeOrigin, float MinDistanceSq, float MinRangeSq, float MaxRangeSq)
-{
-	// check if we need to grow
-	if (UsedObjectBounds4Indices == ObjectBounds4UsedIndices.Num())
-	{
-		ObjectBounds4UsedIndices.Add(false, 4);			
-	}
-
-	const int32 BoundsIndex = ObjectBounds4UsedIndices.FindAndSetFirstZeroBit(FreeObjectBounds4IndexHint);
-	check(BoundsIndex >= 0);
-	FreeObjectBounds4IndexHint = BoundsIndex + 1;
-	++UsedObjectBounds4Indices;
 
-	SetBounds(BoundsIndex, Bounds, PackedRelativeBox, LastRenderTime, RangeOrigin, MinDistanceSq, MinRangeSq, MaxRangeSq);
-	
-	return BoundsIndex;
-}
 
 void FSimpleStreamableAssetManager::SetBounds(int32 BoundsIndex, const FBoxSphereBounds& Bounds, uint32 PackedRelativeBox, float LastRenderTime, const FVector4& RangeOrigin, float MinDistanceSq, float MinRangeSq, float MaxRangeSq)
 {
@@ -506,20 +512,6 @@ void FSimpleStreamableAssetManager::SetBounds(int32 BoundsIndex, const FBoxSpher
 	ObjectBounds4[ObjectBounds4Index].Set(BoundsIndex % 4, Bounds, PackedRelativeBox, LastRenderTime, RangeOrigin, MinDistanceSq, MinRangeSq, MaxRangeSq);
 }
 
-void FSimpleStreamableAssetManager::RemoveBounds(int32 BoundsIndex)
-{
-	check(ObjectBounds4.Num() > (BoundsIndex / 4));
-	check(ObjectBounds4UsedIndices[BoundsIndex]);
-	ObjectBounds4UsedIndices[BoundsIndex] = false;
-	
-	// We store 4 Objects in one entry - but there is no real value in clearing the entry
-	// ObjectBounds4[BoundsIndex / 4].Clear(BoundsIndex % 4);
-		
-	--UsedObjectBounds4Indices;
-	FreeObjectBounds4IndexHint = FMath::Min(FreeObjectBounds4IndexHint, BoundsIndex);
-	check(!ObjectBounds4UsedIndices[BoundsIndex]);
-}
-
 FBoxSphereBounds FSimpleStreamableAssetManager::GetBounds(int32 BoundsIndex) const
 {
 	FBoxSphereBounds Bounds(ForceInitToZero);
@@ -543,9 +535,8 @@ FBoxSphereBounds FSimpleStreamableAssetManager::GetBounds(int32 BoundsIndex) con
 	return Bounds;
 }
 
-void FSimpleStreamableAssetManager::AddRenderAssetElements(const TArrayView<FStreamingRenderAssetPrimitiveInfo>& RenderAssetInstanceInfos, int32 BoundsIndex, int32 ObjectRegistrationIndex, bool bForceMipStreaming)
+void FSimpleStreamableAssetManager::AddRenderAssetElements(const TArrayView<FStreamingRenderAssetPrimitiveInfo>& RenderAssetInstanceInfos, int32 ObjectRegistrationIndex, bool bForceMipStreaming)
 {
-	const int32 BoundsOffset = ObjectRegistrationIndexToBounds4Index[ObjectRegistrationIndex].Add(BoundsIndex);
 	TSet<FAssetRecord> ObjectAssets;
 	ObjectAssets.Reserve(RenderAssetInstanceInfos.Num());
 	for (const FStreamingRenderAssetPrimitiveInfo& AssetInfo : RenderAssetInstanceInfos)
@@ -581,22 +572,19 @@ void FSimpleStreamableAssetManager::AddRenderAssetElements(const TArrayView<FStr
 
 		check(AssetIndex != INDEX_NONE);
 
-		ObjectAssets.Add(FAssetRecord{ .AssetRegistrationIndex = AssetIndex
-#if SIMPLE_STREAMABLE_ASSET_MANAGER_ALLOW_DEBUG_POINTER
-			, .StreamableRenderAsset_ForDebug = RenderAsset
-#endif
-			});			
-
 		if (AssetIndex >= AssetIndexToBounds4Index.Num())
 		{
 			check(AssetIndex == AssetIndexToBounds4Index.Num());
 			AssetIndexToBounds4Index.Add({});
 		}
+		
+		const int32 AssetElementIndex = AssetIndexToBounds4Index[AssetIndex].Add(FAssetBoundElement{ .ObjectRegistrationIndex = ObjectRegistrationIndex, .TexelFactor =  AssetInfo.TexelFactor, .bForceLOD = bForceMipStreaming});
 
-		check(BoundsOffset < 0x7FF); // 11-Bits
-		check(ObjectRegistrationIndex < 0xFFFFF); // 20-Bits
-
-		AssetIndexToBounds4Index[AssetIndex].Add(FAssetBoundElement{.ElementIndex = uint32(BoundsIndex), .ObjectRegistrationIndex = uint32(ObjectRegistrationIndex), .ElementOffset =  uint32(BoundsOffset), .bForceLOD = bForceMipStreaming, .TexelFactor =  AssetInfo.TexelFactor});
+		ObjectAssets.Add(FAssetRecord{ .AssetRegistrationIndex = AssetIndex, .AssetElementIndex = AssetElementIndex
+#if SIMPLE_STREAMABLE_ASSET_MANAGER_ALLOW_DEBUG_POINTER
+			, .StreamableRenderAsset_ForDebug = RenderAsset
+#endif
+			});			
 	}
 
 	ObjectRegistrationIndexToAssetProperty[ObjectRegistrationIndex] = ObjectAssets.Array();
@@ -608,33 +596,22 @@ void FSimpleStreamableAssetManager::RemoveRenderAssetElements(int32 ObjectRegist
 	--RegisteredObjectCount;
 	ObjectUsedIndices[ObjectRegistrationIndex] = false;			
 	FreeObjectIndexHint = FMath::Min(FreeObjectIndexHint, ObjectRegistrationIndex);
-
-	TArray<int32>& ObjectBoundsIndices = ObjectRegistrationIndexToBounds4Index[ObjectRegistrationIndex];
-	const int32 BoundCount = ObjectBoundsIndices.Num();
 	
 	TArray<FAssetRecord>& ObjectAssets = ObjectRegistrationIndexToAssetProperty[ObjectRegistrationIndex];
 	
 	for (const FAssetRecord& Asset : ObjectAssets)
 	{
 		const int32& AssetIndex = Asset.AssetRegistrationIndex;
+		const int32& AssetElementIndex = Asset.AssetElementIndex;
 		check(AssetIndex != INDEX_NONE);
-		TSet<FAssetBoundElement>& AssetBondElements = AssetIndexToBounds4Index[AssetIndex];
+		check(AssetElementIndex != INDEX_NONE);
 
-		if (AssetBondElements.Num() > 0)
-		{
-			for(int32 BoundOffset = 0; BoundOffset < BoundCount; ++BoundOffset)
-			{
-				AssetBondElements.Remove(FAssetBoundElement{.ObjectRegistrationIndex = static_cast<uint32>(ObjectRegistrationIndex), .ElementOffset = static_cast<uint32>(BoundOffset)});
-			}
-		}
+		AssetIndexToBounds4Index[AssetIndex].Reset(AssetElementIndex);
 	}
 	ObjectAssets.Reset();
-	
-	for (int32 ObjectBoundIndex : ObjectBoundsIndices)
-	{
-		RemoveBounds(ObjectBoundIndex);
-	}
-	ObjectBoundsIndices.Reset();
+
+	// We store 4 Objects in one Bounds entry - but there is no real value in clearing the entry
+	// ObjectBounds4[ObjectRegistrationIndex / 4].Clear(ObjectRegistrationIndex % 4);
 }
 
 void FSimpleStreamableAssetManager::GetRenderAssetScreenSize_Impl(
@@ -653,21 +630,24 @@ void FSimpleStreamableAssetManager::GetRenderAssetScreenSize_Impl(
 		// Asset might be registered, but it might not be yet used by any proxy
 		if (AssetIndex != INDEX_NONE)
 		{
-			const TSet<FAssetBoundElement>& AssetBoundElements = AssetIndexToBounds4Index[AssetIndex];
+			const TSimpleSparseArray<FAssetBoundElement>& AssetBoundElements = AssetIndexToBounds4Index[AssetIndex];
 						
-			for (const FAssetBoundElement& AssetBoundElement : AssetBoundElements)
+			for (const FAssetBoundElement& AssetBoundElement : AssetBoundElements.GetSparseView())	
 			{
-				const int32 AssetBoundsIndex =  AssetBoundElement.ElementIndex; 
-				
-				const FBoundsViewInfo& BoundsViewInfo = BoundsViewInfos[AssetBoundsIndex];
-			
-				const float TexelFactor = AssetBoundElement.TexelFactor;
-				const bool bForcedLODs = AssetBoundElement.bForceLOD;
-				FRenderAssetInstanceAsyncView::ProcessElement(AssetType, BoundsViewInfo, AssetType != EStreamableRenderAssetType::Texture ? TexelFactor : TexelFactor * BoundsViewInfo.ComponentScale, bForcedLODs, MaxSize, MaxSize_VisibleOnly, MaxNumForcedLODs);
+				const int32 ObjectRegistrationIndex =  AssetBoundElement.ObjectRegistrationIndex; 
 
-				if (MaxSize_VisibleOnly >= MaxAssetSize || MaxNumForcedLODs >= MaxAllowedMip)
+				if (ObjectRegistrationIndex != INDEX_NONE)
 				{
-					return;
+					const FBoundsViewInfo& BoundsViewInfo = BoundsViewInfos[ObjectRegistrationIndex];
+			
+					const float TexelFactor = AssetBoundElement.TexelFactor;
+					const bool bForcedLODs = AssetBoundElement.bForceLOD;
+					FRenderAssetInstanceAsyncView::ProcessElement(AssetType, BoundsViewInfo, AssetType != EStreamableRenderAssetType::Texture ? TexelFactor : TexelFactor * BoundsViewInfo.ComponentScale, bForcedLODs, MaxSize, MaxSize_VisibleOnly, MaxNumForcedLODs);
+
+					if (MaxSize_VisibleOnly >= MaxAssetSize || MaxNumForcedLODs >= MaxAllowedMip)
+					{
+						return;
+					}
 				}
 			}
 		}
diff --git a/Engine/Source/Runtime/Engine/Public/Streaming/SimpleStreamableAssetManager.h b/Engine/Source/Runtime/Engine/Public/Streaming/SimpleStreamableAssetManager.h
--- a/Engine/Source/Runtime/Engine/Public/Streaming/SimpleStreamableAssetManager.h
+++ b/Engine/Source/Runtime/Engine/Public/Streaming/SimpleStreamableAssetManager.h
@@ -234,37 +234,97 @@ public:
 
 	struct FAssetBoundElement
 	{
-		uint32 ElementIndex = 0xFFFFFFFF;
-
-		union 
-		{
-			struct 
-			{
-				uint32 ObjectRegistrationIndex : 20;
-				uint32 ElementOffset : 11;
-			};
-			uint32 Raw = 0xFFFFFFFF;
-		};
-		
-		uint32 bForceLOD : 1 = true;
+		int32 ObjectRegistrationIndex = INDEX_NONE;
 		float TexelFactor = 0.0f;
-		
-		friend bool operator==(const FAssetBoundElement& A, const FAssetBoundElement& B) { return A.Raw == B.Raw; }
+		uint32 bForceLOD : 1 = 0;
 	};
 
 	struct FAssetRecord
 	{
-		int32 AssetRegistrationIndex = 0;
+		int32 AssetRegistrationIndex = INDEX_NONE;
+		int32 AssetElementIndex = INDEX_NONE; 
 #if SIMPLE_STREAMABLE_ASSET_MANAGER_ALLOW_DEBUG_POINTER
 		const UStreamableRenderAsset* StreamableRenderAsset_ForDebug = nullptr;
 #endif
 		friend bool operator==(const FAssetRecord& A, const FAssetRecord& B) { return A.AssetRegistrationIndex == B.AssetRegistrationIndex; }
 	};
+
+	struct FObjectBoundsRecord
+	{
+		int32 BoundsIndex = INDEX_NONE;
+	};
 private:
-	
-	
+	template <typename T>
+	struct TSimpleSparseArray
+	{
+		int32 Add(T&& InElement)
+		{
+			if (UsedElementsCount == FreeElementIndexHint)
+			{
+				UsedElements.Add(false, GSimpleStreamableAssetManagerSparseArrayGrowSize);
+				Elements.AddDefaulted(GSimpleStreamableAssetManagerSparseArrayGrowSize);
+			}
+			const int32 Index = UsedElements.FindAndSetFirstZeroBit(FreeElementIndexHint);
+			check(Index != INDEX_NONE);
+			FreeElementIndexHint = Index + 1;
+			++UsedElementsCount;
+			Elements[Index] = MoveTemp(InElement);
+			return Index;
+		}
+		
+		void Reset(int32 Index)
+		{
+			if (UsedElements[Index])
+			{
+				UsedElements[Index] = false;
+				FreeElementIndexHint = FMath::Min(FreeElementIndexHint, Index);
+				Elements[Index] = T{};
+				--UsedElementsCount;
+			}
+		}
+		
+		void Empty()
+		{
+			FreeElementIndexHint = 0;
+			UsedElementsCount = 0;
+			UsedElements.Empty();
+			Elements.Empty();
+		}
+
+		int32 Num() const { return UsedElementsCount; }
+
+		SIZE_T GetAllocatedSize(void) const
+		{
+			return sizeof(TSimpleSparseArray)
+				+ UsedElements.GetAllocatedSize()
+				+ Elements.GetAllocatedSize();
+		}
+
+		TConstArrayView<T> GetSparseView() const
+		{
+			const int32 Max = UsedElements.FindLast(true);
+			if (Max == INDEX_NONE)
+			{
+				return TConstArrayView<T>{ Elements.GetData(), Max + 1};
+			}
+			return {};
+		}
+		private:
+		int32 FreeElementIndexHint = 0;
+		int32 UsedElementsCount = 0;
+		TBitArray<> UsedElements;
+		TArray<T> Elements;
+	};
+
 	static FAutoConsoleVariableRef CVarUseSimpleStreamableAssetManager;
 	ENGINE_API static int32 GUseSimpleStreamableAssetManager;
+
+	static FAutoConsoleVariableRef CVarUseSimpleStreamableAssetManagerSparseArrayGrowSize;
+	ENGINE_API static int32 GSimpleStreamableAssetManagerSparseArrayGrowSize;
+
+	static FAutoConsoleVariableRef CVarUseSimpleStreamableAssetManagerEnsureAssetUniqueOnRegistration;
+	ENGINE_API static int32 GSimpleStreamableAssetManagerEnsureAssetUniqueOnRegistration;
+	
 	static FSimpleStreamableAssetManager* Instance;
 	FCriticalSection CriticalSection;
 
@@ -280,28 +340,22 @@ private:
 	int32 FreeObjectIndexHint = 0;
 	TBitArray<> ObjectUsedIndices;
 	TArray<TArray<FAssetRecord>> ObjectRegistrationIndexToAssetProperty;
-	TArray<TArray<int32>> ObjectRegistrationIndexToBounds4Index;
+	TArray<FBounds4> ObjectBounds4;
 
 	void RegisterRecord(FRegister& Record);
 	void UpdateRecord(FUpdate& Record);
 	void UnregisterRecord(FUnregister& Record);
 
 	// ** Variables to manage Bounds registration ** //
-	int32 FreeObjectBounds4IndexHint = 0;
-	int32 UsedObjectBounds4Indices = 0;
-	TBitArray<> ObjectBounds4UsedIndices;
-	TArray<FBounds4> ObjectBounds4;
-	int32 AddBounds(const FBoxSphereBounds& Bounds, uint32 PackedRelativeBox, float LastRenderTime, const FVector4& RangeOrigin, float MinDistanceSq, float MinRangeSq, float MaxRangeSq);
 	void SetBounds(int32 BoundsIndex, const FBoxSphereBounds& Bounds, uint32 PackedRelativeBox, float LastRenderTime, const FVector4& RangeOrigin, float MinDistanceSq, float MinRangeSq, float MaxRangeSq);
-	void RemoveBounds(int32 BoundsIndex);
 	FBoxSphereBounds GetBounds(int32 BoundsIndex) const;
 
 	// ** Variables to manage Assets registration ** //
 	int32 FreeAssetIndexHint = 0;
 	int32 UsedAssetIndices = 0;
 	TBitArray<> AssetUsedIndices;
-	TArray<TSet<FAssetBoundElement>> AssetIndexToBounds4Index;
-	void AddRenderAssetElements(const TArrayView<FStreamingRenderAssetPrimitiveInfo>& RenderAssetInstanceInfos, int32 BoundsIndex, int32 ObjectRegistrationIndex, bool bForceMipStreaming);
+	TArray<TSimpleSparseArray<FAssetBoundElement>> AssetIndexToBounds4Index;
+	void AddRenderAssetElements(const TArrayView<FStreamingRenderAssetPrimitiveInfo>& RenderAssetInstanceInfos, int32 ObjectRegistrationIndex, bool bForceMipStreaming);
 	void RemoveRenderAssetElements(int32 ObjectRegistrationIndex);
 
 	// ** Background task data** //
@@ -393,5 +447,4 @@ FSimpleStreamableAssetManager::FRegister::FRegister(const TObject* InObject, con
 	}
 }
 
-inline int32 GetTypeHash(const FSimpleStreamableAssetManager::FAssetBoundElement& Object) { return Object.Raw; }
-inline int32 GetTypeHash(const FSimpleStreamableAssetManager::FAssetRecord& Object) { return Object.AssetRegistrationIndex; }
\ No newline at end of file
+inline int32 GetTypeHash(const FSimpleStreamableAssetManager::FAssetRecord& Object) { return Object.AssetRegistrationIndex; }
