FSimpleStreamableManager Fixing the view condition calculation, adding debug code for reseting bounds if checks are enabled

diff --git a/Engine/Source/Runtime/Engine/Private/Streaming/SimpleStreamableAssetManager.cpp b/Engine/Source/Runtime/Engine/Private/Streaming/SimpleStreamableAssetManager.cpp
--- a/Engine/Source/Runtime/Engine/Private/Streaming/SimpleStreamableAssetManager.cpp
+++ b/Engine/Source/Runtime/Engine/Private/Streaming/SimpleStreamableAssetManager.cpp
@@ -609,9 +609,11 @@ void FSimpleStreamableAssetManager::RemoveRenderAssetElements(int32 ObjectRegist
 		AssetIndexToBounds4Index[AssetIndex].Reset(AssetElementIndex);
 	}
 	ObjectAssets.Reset();
-
-	// We store 4 Objects in one Bounds entry - but there is no real value in clearing the entry
-	// ObjectBounds4[ObjectRegistrationIndex / 4].Clear(ObjectRegistrationIndex % 4);
+#if DO_CHECK
+	// We store 4 Objects in one Bounds entry, there is no real need for clearing the entry
+	// we will clean it only for ease of debugging if checks enabled 
+	ObjectBounds4[ObjectRegistrationIndex / 4].Clear(ObjectRegistrationIndex % 4);
+#endif
 }
 
 void FSimpleStreamableAssetManager::GetRenderAssetScreenSize_Impl(
diff --git a/Engine/Source/Runtime/Engine/Public/Streaming/SimpleStreamableAssetManager.h b/Engine/Source/Runtime/Engine/Public/Streaming/SimpleStreamableAssetManager.h
--- a/Engine/Source/Runtime/Engine/Public/Streaming/SimpleStreamableAssetManager.h
+++ b/Engine/Source/Runtime/Engine/Public/Streaming/SimpleStreamableAssetManager.h
@@ -303,11 +303,7 @@ private:
 		TConstArrayView<T> GetSparseView() const
 		{
 			const int32 Max = UsedElements.FindLast(true);
-			if (Max == INDEX_NONE)
-			{
-				return TConstArrayView<T>{ Elements.GetData(), Max + 1};
-			}
-			return {};
+			return TConstArrayView<T>{ Elements.GetData(), Max + 1};
 		}
 		private:
 		int32 FreeElementIndexHint = 0;
