Bumped async deferred update parallel for to foreground to get it to finish faster.

diff --git a/Engine/Source/Runtime/Engine/Private/LevelTick.cpp b/Engine/Source/Runtime/Engine/Private/LevelTick.cpp
--- a/Engine/Source/Runtime/Engine/Private/LevelTick.cpp
+++ b/Engine/Source/Runtime/Engine/Private/LevelTick.cpp
@@ -1219,7 +1219,7 @@ void UWorld::SendAllEndOfFrameUpdatesInternal(EEndOfFrameUpdateLocation Location
 		Context->UpdateComponent_Parallel(InRenderCommandList, Index);
 	};
 
-	const EParallelForFlags ParallelForFlags = EParallelForFlags::Unbalanced | (bAsyncInTick ? EParallelForFlags::BackgroundPriority : EParallelForFlags::None);
+	const EParallelForFlags ParallelForFlags = EParallelForFlags::Unbalanced;
 
 	FRenderCommandList* RootRenderCommandList = FRenderCommandList::Create(ERenderCommandListFlags::CloseOnSubmit);
 	FRenderCommandList::FRecordScope RecordScope(RootRenderCommandList, FRenderCommandList::EStopRecordingAction::Submit);
