[Audio Insights] Support all sound asset types that inherit from one of the main Sound Dashboard categories to avoid incorrectly categorizing sounds as "other"

diff --git a/Engine/Plugins/AudioInsights/Source/AudioInsights/Private/Messages/SoundTraceMessages.cpp b/Engine/Plugins/AudioInsights/Source/AudioInsights/Private/Messages/SoundTraceMessages.cpp
--- a/Engine/Plugins/AudioInsights/Source/AudioInsights/Private/Messages/SoundTraceMessages.cpp
+++ b/Engine/Plugins/AudioInsights/Source/AudioInsights/Private/Messages/SoundTraceMessages.cpp
@@ -28,10 +28,8 @@ namespace UE::Audio::Insights
 		EventData.GetString("Name", Name);
 
 		FString SoundClassName;
-		FString SoundSuperClassName;
 
 		EventData.GetString("SoundClassName", SoundClassName);
-		EventData.GetString("SoundSuperClassName", SoundSuperClassName);
 
 		if (SoundClassName == SoundClassNames::SoundCue)
 		{
@@ -45,13 +43,11 @@ namespace UE::Audio::Insights
 		{
 			EntryType = ESoundDashboardEntryType::MetaSound;
 		}
-		else if (SoundClassName == SoundClassNames::SoundWaveProcedural
-			|| SoundSuperClassName == SoundClassNames::SoundWaveProcedural)
+		else if (SoundClassName == SoundClassNames::SoundWaveProcedural)
 		{
 			EntryType = ESoundDashboardEntryType::ProceduralSource;
 		}
-		else if (SoundClassName == SoundClassNames::SoundCueTemplate
-			|| SoundSuperClassName == SoundClassNames::SoundCueTemplate)
+		else if (SoundClassName == SoundClassNames::SoundCueTemplate)
 		{
 			EntryType = ESoundDashboardEntryType::SoundCueTemplate;
 		}
diff --git a/Engine/Source/Runtime/Engine/Private/AudioDevice.cpp b/Engine/Source/Runtime/Engine/Private/AudioDevice.cpp
--- a/Engine/Source/Runtime/Engine/Private/AudioDevice.cpp
+++ b/Engine/Source/Runtime/Engine/Private/AudioDevice.cpp
@@ -5,7 +5,6 @@
 #include "AudioCompressionSettingsUtils.h"
 #include "AudioDecompress.h"
 #include "AudioEffect.h"
-#include "AudioMixerTrace.h"
 #include "AudioPluginUtilities.h"
 #include "Audio/AudioDebug.h"
 #include "GameFramework/GameUserSettings.h"
@@ -246,7 +245,6 @@ UE_TRACE_EVENT_BEGIN(Audio, SoundStart)
 	UE_TRACE_EVENT_FIELD(uint32, PlayOrder)
 	UE_TRACE_EVENT_FIELD(UE::Trace::WideString, Name)
 	UE_TRACE_EVENT_FIELD(UE::Trace::WideString, SoundClassName)
-	UE_TRACE_EVENT_FIELD(UE::Trace::WideString, SoundSuperClassName)
 UE_TRACE_EVENT_END()
 
 UE_TRACE_EVENT_BEGIN(Audio, SoundStop)
@@ -5033,6 +5031,42 @@ void FAudioDevice::UpdateAudioEngineSubsystems()
 }
 
 #if UE_AUDIO_PROFILERTRACE_ENABLED
+namespace AudioProfilerTrace
+{
+	const FString MetaSoundSource = TEXT("MetaSoundSource");
+	const FString SoundWaveProcedural = TEXT("SoundWaveProcedural");
+	const FString SoundCue = TEXT("SoundCue");
+	const FString SoundWave = TEXT("SoundWave");
+	const FString SoundCueTemplate = TEXT("SoundCueTemplate");
+
+	bool NameMatchesBaseAudioAsset(const FString& SoundClassName)
+	{
+		return SoundClassName == MetaSoundSource
+			|| SoundClassName == SoundWaveProcedural
+			|| SoundClassName == SoundCue
+			|| SoundClassName == SoundWave
+			|| SoundClassName == SoundCueTemplate;
+	}
+
+	FString GetSoundBaseAssetName(const TObjectPtr<UClass> SoundClass)
+	{
+		if (SoundClass == nullptr)
+		{
+			return FString();
+		}
+
+		TObjectPtr<UClass> SoundSuperClass = SoundClass->GetSuperClass();
+		FString SoundClassName = SoundClass->GetName();
+		while (!NameMatchesBaseAudioAsset(SoundClassName) && !SoundClassName.IsEmpty())
+		{
+			SoundClassName = SoundSuperClass ? SoundSuperClass->GetName() : FString();
+			SoundSuperClass = SoundSuperClass ? SoundSuperClass->GetSuperClass() : nullptr;
+		}
+
+		return SoundClassName;
+	}
+}
+
 void FAudioDevice::SendSoundStartedTraceMessage(const FActiveSound& InActiveSound) const
 {
 	const bool bChannelEnabled = UE_TRACE_CHANNELEXPR_IS_ENABLED(AudioMixerChannel);
@@ -5046,18 +5080,14 @@ void FAudioDevice::SendSoundStartedTraceMessage(const FActiveSound& InActiveSoun
 			return;
 		}
 
-		const FString SoundClassName = SoundObjectPtr->GetClass()->GetName();
-
-		const TObjectPtr<UClass> SoundSuperClass = SoundObjectPtr->GetClass()->GetSuperClass();
-		const FString SuperClassName = SoundSuperClass ? SoundSuperClass->GetName() : FString();
+		const FString SoundBaseAssetName = AudioProfilerTrace::GetSoundBaseAssetName(SoundObjectPtr->GetClass());
 
 		UE_TRACE_LOG(Audio, SoundStart, AudioMixerChannel)
 			<< SoundStart.DeviceId(InActiveSound.AudioDevice->DeviceID)
 			<< SoundStart.Timestamp(FPlatformTime::Cycles64())
 			<< SoundStart.PlayOrder(InActiveSound.PlayOrder)
 			<< SoundStart.Name(*SoundObjectPtr->GetPathName())
-			<< SoundStart.SoundClassName(*SoundClassName)
-			<< SoundStart.SoundSuperClassName(*SuperClassName);
+			<< SoundStart.SoundClassName(*SoundBaseAssetName);
 	}
 }
 #endif // UE_AUDIO_PROFILERTRACE_ENABLED
diff --git a/Engine/Source/Runtime/Engine/Private/SoundWave.cpp b/Engine/Source/Runtime/Engine/Private/SoundWave.cpp
--- a/Engine/Source/Runtime/Engine/Private/SoundWave.cpp
+++ b/Engine/Source/Runtime/Engine/Private/SoundWave.cpp
@@ -116,7 +116,6 @@ UE_TRACE_EVENT_BEGIN(Audio, SoundWaveStart)
 	UE_TRACE_EVENT_FIELD(uint32, ActiveSoundPlayOrder)
 	UE_TRACE_EVENT_FIELD(UE::Trace::WideString, Name)
 	UE_TRACE_EVENT_FIELD(UE::Trace::WideString, SoundClassName)
-	UE_TRACE_EVENT_FIELD(UE::Trace::WideString, SoundSuperClassName)
 UE_TRACE_EVENT_END()
 
 UE_TRACE_EVENT_BEGIN(Audio, SoundPriority)
@@ -3878,10 +3877,7 @@ void USoundWave::Parse(FAudioDevice* AudioDevice, const UPTRINT NodeWaveInstance
 			const TObjectPtr<UObject> SoundWaveObjectPtr = FSoftObjectPath(WaveInstance->WaveData->GetPathName()).ResolveObject();
 			if (SoundWaveObjectPtr)
 			{
-				const FString SoundWaveClassName = SoundWaveObjectPtr->GetClass()->GetName();
-
-				const TObjectPtr<UClass> SoundWaveSuperClass = SoundWaveObjectPtr->GetClass()->GetSuperClass();
-				const FString SoundWaveSuperClassName = SoundWaveSuperClass ? SoundWaveSuperClass->GetName() : FString();
+				const FString SoundWaveBaseAssetName = AudioProfilerTrace::GetSoundBaseAssetName(SoundWaveObjectPtr->GetClass());
 
 				UE_TRACE_LOG(Audio, SoundWaveStart, AudioMixerChannel)
 					<< SoundWaveStart.DeviceId(AudioDevice->DeviceID)
@@ -3889,8 +3885,7 @@ void USoundWave::Parse(FAudioDevice* AudioDevice, const UPTRINT NodeWaveInstance
 					<< SoundWaveStart.PlayOrder(WaveInstance->GetPlayOrder())
 					<< SoundWaveStart.ActiveSoundPlayOrder(ActiveSound.GetPlayOrder())
 					<< SoundWaveStart.Name(*SoundWaveObjectPtr->GetPathName())
-					<< SoundWaveStart.SoundClassName(*SoundWaveClassName)
-					<< SoundWaveStart.SoundSuperClassName(*SoundWaveSuperClassName);
+					<< SoundWaveStart.SoundClassName(*SoundWaveBaseAssetName);
 			}
 		}
 
diff --git a/Engine/Source/Runtime/Engine/Public/AudioDevice.h b/Engine/Source/Runtime/Engine/Public/AudioDevice.h
--- a/Engine/Source/Runtime/Engine/Public/AudioDevice.h
+++ b/Engine/Source/Runtime/Engine/Public/AudioDevice.h
@@ -4,6 +4,7 @@
 
 #include "AudioDeviceHandle.h"
 #include "AudioDynamicParameter.h"
+#include "AudioMixerTrace.h"
 #include "AudioVirtualLoop.h"
 #include "Components/AudioComponent.h"
 #include "HAL/LowLevelMemStats.h"
@@ -148,6 +149,13 @@ namespace ESortedActiveWaveGetType
 	};
 }
 
+#if UE_AUDIO_PROFILERTRACE_ENABLED
+namespace AudioProfilerTrace
+{
+	FString GetSoundBaseAssetName(const TObjectPtr<UClass> SoundClass);
+}
+#endif // #if UE_AUDIO_PROFILERTRACE_ENABLED
+
 /**
  * Defines the properties of the listener
  */
