Added a check for FSkinWeightProfileManager requests so the manager can sleep while all remained mesh LODs aren't loaded and not scheduled for streaming

#rnx

diff --git a/Engine/Source/Runtime/Engine/Private/Animation/SkinWeightProfileManager.cpp b/Engine/Source/Runtime/Engine/Private/Animation/SkinWeightProfileManager.cpp
--- a/Engine/Source/Runtime/Engine/Private/Animation/SkinWeightProfileManager.cpp
+++ b/Engine/Source/Runtime/Engine/Private/Animation/SkinWeightProfileManager.cpp
@@ -42,6 +42,14 @@ static TAutoConsoleVariable<int32> CVarForceSkinWeightProfileManagerAsyncTaskOnG
 	ECVF_Default | ECVF_RenderThreadSafe
 );
 
+static int32 GSkipLODWithoutDataOrStreamingRequest = 0;
+static FAutoConsoleVariableRef CVarSkipLODWithoutDataOrStreamingRequest(
+	TEXT("SkinWeightProfileManager.SkipLODWithoutDataOrStreamingRequest"),
+	GSkipLODWithoutDataOrStreamingRequest,
+	TEXT(""),
+	ECVF_Default | ECVF_RenderThreadSafe
+);
+
 FSkinWeightProfileManager::FSkinWeightProfileManager(UWorld* InWorld) : LastGamethreadProfileIndex(INDEX_NONE), WeakWorld(nullptr)
 {
 	if (InWorld)
@@ -215,11 +223,12 @@ void FSkinWeightProfileManager::DoTick(float DeltaTime, ENamedThreads::Type Curr
 			AsyncTask.SafeRelease();
 		}
 
+		bool bInProgressRequestsExist = false;
 		for (int32 RequestIndex = 0; RequestIndex < PendingSetProfileRequests.Num(); ++RequestIndex)
 		{
 			QUICK_SCOPE_CYCLE_COUNTER(STAT_SkinWeightProfileManager_TickMesh);
 
-			const FSetProfileRequest& Request = PendingSetProfileRequests[RequestIndex];
+			FSetProfileRequest& Request = PendingSetProfileRequests[RequestIndex];
 			
 			// Remove if canceled or objects have gone stale 
 			bool bRemove = Request.IdentifyingObject.IsStale() || Request.WeakSkeletalMesh.IsStale() ||
@@ -236,15 +245,38 @@ void FSkinWeightProfileManager::DoTick(float DeltaTime, ENamedThreads::Type Curr
 			if (!bRemove)
 			{	
 				bool bAllBuffersReady = true;
-				for (const int32 LODIndex : Request.LODIndices)
+				int32 LODIndexToWaitStreaming = INDEX_NONE;
+				for (int32 LODProfileArrayIndex = 0; LODProfileArrayIndex < Request.LODIndices.Num(); ++LODProfileArrayIndex)
 				{
-					if (!RenderData->LODRenderData[LODIndex].SkinWeightProfilesData.ContainsOverrideBuffer(Request.ProfileStack))
+					int32 LODRenderDataArrayIndex = Request.LODIndices[LODProfileArrayIndex];
+					const FSkinWeightProfilesData& SkinWeightProfilesData = RenderData->LODRenderData[LODRenderDataArrayIndex].SkinWeightProfilesData;
+					if (!SkinWeightProfilesData.ContainsOverrideBuffer(Request.ProfileStack))
 					{
 						bAllBuffersReady = false;
-						break;
+						if (ShouldIgnoreLOD(SkinWeightProfilesData, Request.WeakSkeletalMesh))
+						{
+							// This LOD isn't streamed in and has no requests for that. There may be several such LODs for the current model.
+							// We will create a callback to be notified when any of these LODs will be streamed in,
+							// and for that we need to know the index of the least detailed LOD among them.
+							LODIndexToWaitStreaming = FMath::Max(LODIndexToWaitStreaming, LODRenderDataArrayIndex);
+						}
+						else
+						{
+							bInProgressRequestsExist = true;
+							if (!GSkipLODWithoutDataOrStreamingRequest)
+							{
+								break;
+							}
+						}
 					}
 				}
 
+				if ((LODIndexToWaitStreaming != INDEX_NONE) &&
+					(LODIndexToWaitStreaming != Request.StreamingCallbackLODIndex))
+				{
+					CreateLODStreamingCallback(Request, LODIndexToWaitStreaming);
+				}
+
 				if (bAllBuffersReady)
 				{
 					if (Request.IdentifyingObject.IsValid())
@@ -265,13 +297,22 @@ void FSkinWeightProfileManager::DoTick(float DeltaTime, ENamedThreads::Type Curr
 			}
 		}
 
-		// This is the last index that we will process asynchronously, while the task is running we do not remove request from the array so access should be safe
-		LastGamethreadProfileIndex = PendingSetProfileRequests.Num() - 1;
+		if (bInProgressRequestsExist || !GSkipLODWithoutDataOrStreamingRequest)
+		{
+			// This is the last index that we will process asynchronously, while the task is running we do not remove request from the array so access should be safe
+			LastGamethreadProfileIndex = PendingSetProfileRequests.Num() - 1;
 
-		if (MyCompletionGraphEvent)
+			if (MyCompletionGraphEvent)
+			{
+				AsyncTask = TGraphTask<FSkinWeightProfileManagerAsyncTask>::CreateTask(nullptr, CurrentThread).ConstructAndDispatchWhenReady(this);
+				TickFunction.GetCompletionHandle()->DontCompleteUntil(AsyncTask);
+			}
+		}
+		else
 		{
-			AsyncTask = TGraphTask<FSkinWeightProfileManagerAsyncTask>::CreateTask(nullptr, CurrentThread).ConstructAndDispatchWhenReady(this);
-			TickFunction.GetCompletionHandle()->DontCompleteUntil(AsyncTask);
+			// This is called when we still have pending requests, but all their unprocessed LODs have not been streamed in.
+			// TickFunction will be re-enabled once the streaming is complete or a new request is received.
+			TickFunction.SetTickFunctionEnable(false);
 		}
 	}
 	else
@@ -280,6 +321,34 @@ void FSkinWeightProfileManager::DoTick(float DeltaTime, ENamedThreads::Type Curr
 	}
 }
 
+void FSkinWeightProfileManager::OnRequestedMeshesStreamingUpdated(const UPrimitiveComponent* PrimitiveComponent, const UStreamableRenderAsset* StreamableRenderAsset, const FSkinWeightProfileStack& ProfileStack)
+{
+	if (!PrimitiveComponent || !StreamableRenderAsset)
+	{
+		return;
+	}
+
+	const USkeletalMesh* SkinnedMeshComponent = Cast<USkeletalMesh>(StreamableRenderAsset);
+	if (!SkinnedMeshComponent)
+	{
+		return;
+	}
+
+	// Find the request for which streaming was updated to reset its callback index
+	for (FSetProfileRequest& Request : PendingSetProfileRequests)
+	{
+		if ((Request.ProfileStack == ProfileStack) &&
+			(Request.IdentifyingObject.Get() == PrimitiveComponent) &&
+			(Request.WeakSkeletalMesh.Get() == SkinnedMeshComponent))
+		{
+			Request.StreamingCallbackLODIndex = INDEX_NONE;
+			break;
+		}
+	}
+
+	TickFunction.SetTickFunctionEnable(true);
+}
+
 void FSkinWeightProfileManager::Tick(float DeltaTime)
 {
 	DoTick(DeltaTime, ENamedThreads::GameThread, nullptr);
@@ -325,6 +394,60 @@ void FSkinWeightProfileManager::CleanupRequest(const FSetProfileRequest& Request
 	
 }
 
+bool FSkinWeightProfileManager::ShouldIgnoreLOD(const FSkinWeightProfilesData& SkinWeightProfilesData, const TWeakObjectPtr<USkeletalMesh>& WeakSkeletalMesh) const
+{
+	if (!GSkipLODWithoutDataOrStreamingRequest)
+	{
+		return false;
+	}
+
+	// Are LOD's vertices loaded? 
+	FSkinWeightVertexBuffer* BaseBuffer = SkinWeightProfilesData.BaseBuffer;
+	if (BaseBuffer && BaseBuffer->GetDataVertexBuffer()->VertexBufferRHI->GetSize())
+	{
+		return false;
+	}
+
+	// Are they scheduled to be streamed in?
+	const USkeletalMesh* SkeletalMesh = WeakSkeletalMesh.Get();
+	if (SkeletalMesh && SkeletalMesh->HasPendingInitOrStreaming())
+	{
+		return false;
+	}
+
+	// This mesh LOD isn't loaded and not scheduled for streaming, we can ignore it
+	return true;
+}
+
+void FSkinWeightProfileManager::CreateLODStreamingCallback(FSetProfileRequest& Request, int32 LODIndexToWaitStreaming)
+{
+	UObject* IdentifyingObject = Request.IdentifyingObject.Get();
+	if (!IdentifyingObject)
+	{
+		return;
+	}
+
+	USkinnedMeshComponent* Component = Cast<USkinnedMeshComponent>(IdentifyingObject);
+	if (!Component)
+	{
+		return;
+	}
+
+	constexpr float CallbackTimeoutSecs = 60.0f;
+	Component->RegisterLODStreamingCallback(
+		[ProfileStack=Request.ProfileStack](UPrimitiveComponent* PrimitiveComponent, UStreamableRenderAsset* StreamableRenderAsset, ELODStreamingCallbackResult LODStreamingCallbackResult)
+		{
+			FSkinWeightProfileManager* Manager = FSkinWeightProfileManager::Get(PrimitiveComponent->GetWorld());
+			if (Manager)
+			{
+				Manager->OnRequestedMeshesStreamingUpdated(PrimitiveComponent, StreamableRenderAsset, ProfileStack);
+			}
+		},
+		LODIndexToWaitStreaming, CallbackTimeoutSecs, true);
+
+	Request.StreamingCallbackLODIndex = LODIndexToWaitStreaming;
+}
+
 bool FSkinWeightProfileManager::IsTickableWhenPaused() const
 {
 	return false;
diff --git a/Engine/Source/Runtime/Engine/Public/Animation/SkinWeightProfileManager.h b/Engine/Source/Runtime/Engine/Public/Animation/SkinWeightProfileManager.h
--- a/Engine/Source/Runtime/Engine/Public/Animation/SkinWeightProfileManager.h
+++ b/Engine/Source/Runtime/Engine/Public/Animation/SkinWeightProfileManager.h
@@ -35,6 +35,9 @@ struct FSetProfileRequest
 	/** Weak skeletal mesh for which the skin weight profile is to be loaded */
 	TWeakObjectPtr<USkeletalMesh> WeakSkeletalMesh;
 
+	/** An LOD index for which we have an active streaming callback, INDEX_NONE if none exists */
+	int32 StreamingCallbackLODIndex = INDEX_NONE;
+
 	friend bool operator==(const FSetProfileRequest& A, const FSetProfileRequest& B)
 	{
 		return A.ProfileStack == B.ProfileStack && A.WeakSkeletalMesh == B.WeakSkeletalMesh && A.IdentifyingObject == B.IdentifyingObject;
@@ -115,8 +118,12 @@ public:
 	void CancelSkinWeightProfileRequest(UObject* Requester);
 	
 	void DoTick(float DeltaTime, ENamedThreads::Type CurrentThread, const FGraphEventRef& MyCompletionGraphEvent);
+
+	void OnRequestedMeshesStreamingUpdated(const UPrimitiveComponent* PrimitiveComponent, const UStreamableRenderAsset* StreamableRenderAsset, const FSkinWeightProfileStack& ProfileStack);
 protected:	
 	void CleanupRequest(const FSetProfileRequest& Request);
+	bool ShouldIgnoreLOD(const FSkinWeightProfilesData& SkinWeightProfilesData, const TWeakObjectPtr<USkeletalMesh>& WeakSkeletalMesh) const;
+	void CreateLODStreamingCallback(FSetProfileRequest& Request, int32 LODIndexToWaitStreaming);
 protected:
 	TArray<FSetProfileRequest, TInlineAllocator<4>> CanceledRequest;
 	TArray<FSetProfileRequest> PendingSetProfileRequests;
