Fix invalid assignation of bIsMapBuildDataOwner in some non world partitionned levels.
Fix reuse of current World's PersistentLevel MapBuildData to only occur when it's not be stripped of it's RF_Standalone flag

diff --git a/Engine/Source/Runtime/Engine/Private/Level.cpp b/Engine/Source/Runtime/Engine/Private/Level.cpp
--- a/Engine/Source/Runtime/Engine/Private/Level.cpp
+++ b/Engine/Source/Runtime/Engine/Private/Level.cpp
@@ -3952,10 +3952,11 @@ UMapBuildDataRegistry* ULevel::GetOrCreateMapBuildData()
 		// assign to all transient Levels inside the World
 		for (ULevel* Level : GetWorld()->GetLevels())
 		{
-			if (ULevelInstanceSubsystem::GetOwningLevel(Level, true) == this)
+			ULevel* OwningLevel = ULevelInstanceSubsystem::GetOwningLevel(Level, true);
+			if (OwningLevel == this)
 			{
 				Level->MapBuildData = MapBuildData;
-				Level->bIsMapBuildDataOwner = false;
+				Level->bIsMapBuildDataOwner = (Level == OwningLevel); // MapBuilData is owned by the OwningLevel, assign flag accordingly
 			}
 		}
 		MarkPackageDirty();
diff --git a/Engine/Source/Runtime/Engine/Private/StaticLightingBuildContext.cpp b/Engine/Source/Runtime/Engine/Private/StaticLightingBuildContext.cpp
--- a/Engine/Source/Runtime/Engine/Private/StaticLightingBuildContext.cpp
+++ b/Engine/Source/Runtime/Engine/Private/StaticLightingBuildContext.cpp
@@ -21,13 +21,14 @@ FStaticLightingBuildContext::FStaticLightingBuildContext(UWorld* InWorld, ULevel
 	World = InWorld;
 	LightingScenario = InLightingScenario;
 
-	if (LightingScenario)
-	{
-		MapBuildDataRegistry = InLightingScenario->MapBuildData;
-	}
-	else
+	MapBuildDataRegistry = nullptr;
+	ULevel* GlobalRegistryLevel = (InLightingScenario) ? InLightingScenario : InWorld->PersistentLevel.Get();
+	
+	// If the Level already has MapBuildData only reuse it if it still has it's standlone flag 
+	// this ensures calls to GetOrCreateMapBuildData() on the Level will not create a new MapBuildData
+	if (GlobalRegistryLevel->MapBuildData && GlobalRegistryLevel->MapBuildData->HasAllFlags(RF_Standalone|RF_Public))
 	{
-		MapBuildDataRegistry = InWorld->PersistentLevel->MapBuildData;
+		MapBuildDataRegistry = GlobalRegistryLevel->MapBuildData;
 	}
 
 	for (int32 LevelIndex = 0; LevelIndex < World->GetNumLevels(); LevelIndex++)
@@ -116,11 +117,6 @@ bool FStaticLightingBuildContext::ShouldIncludeLevel(ULevel* Level) const
 	return IncludeLevel;
 }
 
-UMapBuildDataRegistry* FStaticLightingBuildContext::GetGlobalRegistry() const
-{
-	return MapBuildDataRegistry;
-}
-
 UMapBuildDataRegistry* FStaticLightingBuildContext::GetOrCreateGlobalRegistry() const
 {
 	if (!MapBuildDataRegistry)
diff --git a/Engine/Source/Runtime/Engine/Public/StaticLightingBuildContext.h b/Engine/Source/Runtime/Engine/Public/StaticLightingBuildContext.h
--- a/Engine/Source/Runtime/Engine/Public/StaticLightingBuildContext.h
+++ b/Engine/Source/Runtime/Engine/Public/StaticLightingBuildContext.h
@@ -45,7 +45,6 @@ public:
 	ENGINE_API bool ShouldIncludeActor(AActor* Actor) const;
 	ENGINE_API bool ShouldIncludeLevel(ULevel* Level) const;
 
-	ENGINE_API UMapBuildDataRegistry* GetGlobalRegistry() const;
 	ENGINE_API UMapBuildDataRegistry* GetOrCreateGlobalRegistry() const;
 	
 	ENGINE_API UMapBuildDataRegistry* GetOrCreateRegistryForLevelGuid(const FGuid& Guid) const;
