[EDL] Fix handling of EDL injection of client-only cells in cases when it happens after BeginPlay
Removes irrelevant cells from cells lists during the injection.

diff --git a/Engine/Source/Runtime/Engine/Private/WorldPartition/RuntimeHashSet/WorldPartitionRuntimeHashSet.cpp b/Engine/Source/Runtime/Engine/Private/WorldPartition/RuntimeHashSet/WorldPartitionRuntimeHashSet.cpp
--- a/Engine/Source/Runtime/Engine/Private/WorldPartition/RuntimeHashSet/WorldPartitionRuntimeHashSet.cpp
+++ b/Engine/Source/Runtime/Engine/Private/WorldPartition/RuntimeHashSet/WorldPartitionRuntimeHashSet.cpp
@@ -242,8 +242,7 @@ void UWorldPartitionRuntimeHashSet::OnBeginPlay()
 		{
 #if !WITH_EDITOR
 			FRuntimePartitionStreamingData& NonConstStreamingData = (const_cast<FRuntimePartitionStreamingData&>(StreamingData));		
-			NonConstStreamingData.SpatiallyLoadedCells.SetNum(Algo::RemoveIf(NonConstStreamingData.SpatiallyLoadedCells, [this](UWorldPartitionRuntimeCell* Cell) { return !IsCellRelevantFor(Cell->GetClientOnlyVisible()); }));
-			NonConstStreamingData.NonSpatiallyLoadedCells.SetNum(Algo::RemoveIf(NonConstStreamingData.NonSpatiallyLoadedCells, [this](UWorldPartitionRuntimeCell* Cell) { return !IsCellRelevantFor(Cell->GetClientOnlyVisible()); }));
+			RemoveIrrelevantCells(NonConstStreamingData);
 #endif
 			StreamingData.CreatePartitionsSpatialIndex();
 			return true;
@@ -505,6 +504,17 @@ bool UWorldPartitionRuntimeHashSet::InjectExternalStreamingObject(URuntimeHashEx
 	if (Super::InjectExternalStreamingObject(ExternalStreamingObject))
 	{
 		URuntimeHashSetExternalStreamingObject* HashSetExternalStreamingObject = CastChecked<URuntimeHashSetExternalStreamingObject>(ExternalStreamingObject);
+
+#if !WITH_EDITOR
+		if (UWorld* OuterWorld = GetTypedOuter<UWorld>(); OuterWorld && OuterWorld->HasBegunPlay())
+		{
+			for (FRuntimePartitionStreamingData& StreamingData : HashSetExternalStreamingObject->RuntimeStreamingData)
+			{
+				RemoveIrrelevantCells(StreamingData);
+			}
+		}
+#endif
+
 		HashSetExternalStreamingObject->CreatePartitionsSpatialIndex();
 		UpdateRuntimeDataGridMap();
 		return true;
@@ -1285,6 +1295,12 @@ void UWorldPartitionRuntimeHashSet::ForEachStreamingData(TFunctionRef<bool(const
 	}
 }
 
+void UWorldPartitionRuntimeHashSet::RemoveIrrelevantCells(FRuntimePartitionStreamingData& StreamingData)
+{
+	StreamingData.SpatiallyLoadedCells.SetNum(Algo::RemoveIf(StreamingData.SpatiallyLoadedCells, [this](UWorldPartitionRuntimeCell* Cell) { return !IsCellRelevantFor(Cell->GetClientOnlyVisible()); }));
+	StreamingData.NonSpatiallyLoadedCells.SetNum(Algo::RemoveIf(StreamingData.NonSpatiallyLoadedCells, [this](UWorldPartitionRuntimeCell* Cell) { return !IsCellRelevantFor(Cell->GetClientOnlyVisible()); }));
+}
+
 void UWorldPartitionRuntimeHashSet::UpdateRuntimeDataGridMap()
 {
 	RuntimeSpatiallyLoadedDataGridMap.Reset();
diff --git a/Engine/Source/Runtime/Engine/Public/WorldPartition/RuntimeHashSet/WorldPartitionRuntimeHashSet.h b/Engine/Source/Runtime/Engine/Public/WorldPartition/RuntimeHashSet/WorldPartitionRuntimeHashSet.h
--- a/Engine/Source/Runtime/Engine/Public/WorldPartition/RuntimeHashSet/WorldPartitionRuntimeHashSet.h
+++ b/Engine/Source/Runtime/Engine/Public/WorldPartition/RuntimeHashSet/WorldPartitionRuntimeHashSet.h
@@ -235,6 +235,8 @@ private:
 
 	void FixupHLODSetup(FRuntimePartitionDesc& RuntimePartition);
 
+	void RemoveIrrelevantCells(FRuntimePartitionStreamingData& StreamingData);
+
 	/** Array of runtime partition descriptors */
 	UPROPERTY(EditAnywhere, Category = RuntimeSettings, Meta = (TitleProperty = "Name"))
 	TArray<FRuntimePartitionDesc> RuntimePartitions;
