Added native override point executed when all Linked Animation Layers are initialized

diff --git a/Engine/Source/Runtime/Engine/Classes/Animation/AnimInstance.h b/Engine/Source/Runtime/Engine/Classes/Animation/AnimInstance.h
--- a/Engine/Source/Runtime/Engine/Classes/Animation/AnimInstance.h
+++ b/Engine/Source/Runtime/Engine/Classes/Animation/AnimInstance.h
@@ -1380,6 +1380,8 @@ public:
 	ENGINE_API virtual void NativePostEvaluateAnimation();
 	// Native Uninitialize override point
 	ENGINE_API virtual void NativeUninitializeAnimation();
+	// Native override point executed when all Linked Animation Layers are initialized
+	ENGINE_API virtual void NativeLinkedAnimationLayersInitialized();
 
 	// Enable / Disable animation update. This is provided as an optimization to disable linked instances that aren't relevant. Disabling an instance whose graph is still evaluated will assert and causes issues.
 	void EnableUpdateAnimation(bool bEnable) { bUpdateAnimationEnabled = bEnable; }
diff --git a/Engine/Source/Runtime/Engine/Private/Animation/AnimInstance.cpp b/Engine/Source/Runtime/Engine/Private/Animation/AnimInstance.cpp
--- a/Engine/Source/Runtime/Engine/Private/Animation/AnimInstance.cpp
+++ b/Engine/Source/Runtime/Engine/Private/Animation/AnimInstance.cpp
@@ -324,6 +324,7 @@ void UAnimInstance::InitializeAnimation(bool bInDeferRootNodeInitialization)
 
 	InitializeGroupedLayers(bInDeferRootNodeInitialization);
 
+	NativeLinkedAnimationLayersInitialized();
 	BlueprintLinkedAnimationLayersInitialized();
 }
 
@@ -934,6 +935,10 @@ void UAnimInstance::NativeInitializeAnimation()
 {
 }
 
+void UAnimInstance::NativeLinkedAnimationLayersInitialized()
+{
+}
+
 void UAnimInstance::NativeUpdateAnimation(float DeltaSeconds)
 {
 }
