[NewTranslator] - Implement MaterialExpressionPower

- The Build function for Power implements the native HLSL function `PositiveClampedPow` directly in the IR instead of inlining an HLSL call.
- Minor cleanup to `UMaterialExpressionComponentMask::Build` to use `Em` instead of `Emitter.`
- Handle case with a constant exponent of 1 (this handles all the use cases where the old translator would try to simplify and avoid a call to pow).
- Improve the HLSL lowering code to emit 9 decimal places, and also use `%g` which will print out variable length floating point values.  This results in printing scientific notation for floats w/ many decimal places and results in much cleaner generated code)
- For example it will output 0 instead of 0.00000000

#rnx

diff --git a/Engine/Source/Runtime/Engine/Private/Materials/MaterialExpressionsIR.cpp b/Engine/Source/Runtime/Engine/Private/Materials/MaterialExpressionsIR.cpp
--- a/Engine/Source/Runtime/Engine/Private/Materials/MaterialExpressionsIR.cpp
+++ b/Engine/Source/Runtime/Engine/Private/Materials/MaterialExpressionsIR.cpp
@@ -562,9 +562,9 @@ void UMaterialExpressionArctangent::Build(MIR::FEmitter& Em)
 	Em.Output(0, Em.ATan(Em.Input(&Input)));
 }
 
-void UMaterialExpressionComponentMask::Build(MIR::FEmitter& Emitter)
+void UMaterialExpressionComponentMask::Build(MIR::FEmitter& Em)
 {
-	FValueRef Value = Emitter.Input(&Input);
+	FValueRef Value = Em.Input(&Input);
 
 	MIR::FSwizzleMask Mask;
 	if (R)
@@ -584,9 +584,30 @@ void UMaterialExpressionComponentMask::Build(MIR::FEmitter& Emitter)
 		Mask.Components[Mask.NumComponents++] = MIR::EVectorComponent::W;
 	}
 
-	Emitter.Output(0, Emitter.Swizzle(Value, Mask));
+	Em.Output(0, Em.Swizzle(Value, Mask));
 }
 
+static FValueRef PositiveClampedPow(MIR::FEmitter& Em, FValueRef Base, FValueRef Exponent)
+{
+	FValueRef PrimitiveBase = Em.CheckIsPrimitive(Base);
+	if (!PrimitiveBase.IsValid())
+	{
+		return PrimitiveBase.ToPoison();
+	}
+
+	const MIR::FPrimitiveType* ValuePrimitiveType = Base->Type->AsPrimitive();
+	return Em.Select(
+				Em.LessThanOrEquals(Base, Em.ConstantFloat(2.980233e-8f)),
+				Em.ConstantZero(ValuePrimitiveType->ScalarKind),
+				Em.Pow(Base, Exponent));
+}
+
+void UMaterialExpressionPower::Build(MIR::FEmitter& Em)
+{
+	Em.Output(0, PositiveClampedPow(Em, 
+					 				Em.Input(&Base),
+					 				Em.InputDefaultFloat(&Exponent, ConstExponent)));
+}
 
 static FValueRef GetTrigonometricInputWithPeriod(MIR::FEmitter& Em, const FExpressionInput* Input, float Period)
 {
diff --git a/Engine/Source/Runtime/Engine/Private/Materials/MaterialIREmitter.cpp b/Engine/Source/Runtime/Engine/Private/Materials/MaterialIREmitter.cpp
--- a/Engine/Source/Runtime/Engine/Private/Materials/MaterialIREmitter.cpp
+++ b/Engine/Source/Runtime/Engine/Private/Materials/MaterialIREmitter.cpp
@@ -1089,12 +1089,19 @@ static FValueRef TrySimplifyOperator(FEmitter& Emitter, EOperator& Op, FValueRef
 		case BO_Pow:
 			if (A->AreAllNearlyZero())
 			{
+				// If the base is 0.
 				return A;
 			}
 			else if (B->AreAllNearlyZero())
 			{
+				// If the exponent is 0.
 				return Emitter.ConstantOne(A->Type->AsPrimitive()->ScalarKind);
 			}
+			else if (B->AreAllExactlyOne())
+			{
+				// If the exponent is 1.
+				return A;
+			}
 			break;
 
 		case TO_Clamp:
diff --git a/Engine/Source/Runtime/Engine/Private/Materials/MaterialIRToHLSLTranslator.cpp b/Engine/Source/Runtime/Engine/Private/Materials/MaterialIRToHLSLTranslator.cpp
--- a/Engine/Source/Runtime/Engine/Private/Materials/MaterialIRToHLSLTranslator.cpp
+++ b/Engine/Source/Runtime/Engine/Private/Materials/MaterialIRToHLSLTranslator.cpp
@@ -429,7 +429,7 @@ struct FTranslator : FMaterialIRToHLSLTranslation
 				}
 				else
 				{
-					Printer.Buffer.Appendf(TEXT("%.8f"), Constant->Float);
+					Printer.Buffer.Appendf(TEXT("%.9g"), Constant->Float);
 				}
 				break;
 		}
diff --git a/Engine/Source/Runtime/Engine/Public/Materials/MaterialExpressionPower.h b/Engine/Source/Runtime/Engine/Public/Materials/MaterialExpressionPower.h
--- a/Engine/Source/Runtime/Engine/Public/Materials/MaterialExpressionPower.h
+++ b/Engine/Source/Runtime/Engine/Public/Materials/MaterialExpressionPower.h
@@ -28,6 +28,7 @@ public:
 
 	//~ Begin UMaterialExpression Interface
 #if WITH_EDITOR
+	virtual void Build(MIR::FEmitter& Emitter) override;
 	virtual int32 Compile(class FMaterialCompiler* Compiler, int32 OutputIndex) override;
 	virtual void GetCaption(TArray<FString>& OutCaptions) const override;
 	virtual void GetExpressionToolTip(TArray<FString>& OutToolTip) override;
diff --git a/Engine/Source/Runtime/Engine/Public/Materials/MaterialIREmitter.h b/Engine/Source/Runtime/Engine/Public/Materials/MaterialIREmitter.h
--- a/Engine/Source/Runtime/Engine/Public/Materials/MaterialIREmitter.h
+++ b/Engine/Source/Runtime/Engine/Public/Materials/MaterialIREmitter.h
@@ -505,6 +505,7 @@ public:
 	FValueRef Dot(FValueRef A, FValueRef B) { return Operator(BO_Dot, A, B); }
 	FValueRef Cross(FValueRef A, FValueRef B) { return Operator(BO_Cross, A, B); }
 	FValueRef Atan2(FValueRef A, FValueRef B) { return Operator(BO_Atan2, A, B); }
+	FValueRef Pow(FValueRef A, FValueRef B) { return Operator(BO_Pow, A, B); }
 
 	/*-------------------------- Ternary operators shortcuts ---------------------------*/
 
