- Nanite Renderer enable motion vectors to ISM path

diff --git a/Engine/Source/Runtime/Engine/Classes/Engine/InstancedStaticMesh.h b/Engine/Source/Runtime/Engine/Classes/Engine/InstancedStaticMesh.h
--- a/Engine/Source/Runtime/Engine/Classes/Engine/InstancedStaticMesh.h
+++ b/Engine/Source/Runtime/Engine/Classes/Engine/InstancedStaticMesh.h
@@ -519,11 +519,11 @@ public:
 
 	ENGINE_API virtual float GetLodScreenSizeScale() const override;
 	ENGINE_API virtual float GetGpuLodInstanceRadius() const override;
-	virtual FInstanceDataUpdateTaskInfo *GetInstanceDataUpdateTaskInfo() const override;
+	ENGINE_API virtual FInstanceDataUpdateTaskInfo *GetInstanceDataUpdateTaskInfo() const override;
 
 	virtual bool IsDetailMesh() const override { return true; }
 
-	virtual void SetInstanceCullDistance_RenderThread(float StartCullDistance, float EndCullDistance) override;
+	ENGINE_API virtual void SetInstanceCullDistance_RenderThread(float StartCullDistance, float EndCullDistance) override;
 
 protected:
 	ENGINE_API FInstancedStaticMeshVFLooseUniformShaderParametersRef CreateLooseUniformBuffer(const FSceneView* View, const FInstancingUserData* InstancingUserData, uint32 InstancedLODRange, uint32 InstancedLODIndex, EUniformBufferUsage UniformBufferUsage) const;
