New Material Translator:  Support for thumbnail preview.  Also fix for bug in IsOperatorInfix that produced an assert during testing (min / max wrongly classified as infix, and modulo wrongly not classified).

#rnx

diff --git a/Engine/Source/Runtime/Engine/Private/Materials/MaterialIRModuleBuilder.cpp b/Engine/Source/Runtime/Engine/Private/Materials/MaterialIRModuleBuilder.cpp
--- a/Engine/Source/Runtime/Engine/Private/Materials/MaterialIRModuleBuilder.cpp
+++ b/Engine/Source/Runtime/Engine/Private/Materials/MaterialIRModuleBuilder.cpp
@@ -84,6 +84,7 @@ struct FMaterialIRModuleBuilderImpl
 	MIR::FEmitter* Emitter;
 	TArray<FAnalysisContext> AnalysisContextStack;
 	FMaterialIRValueAnalyzer ValueAnalyzer;
+	FColorMaterialInput PreviewInput;
 
 	void Step_Initialize()
 	{
@@ -92,6 +93,8 @@ struct FMaterialIRModuleBuilderImpl
 		
 		Emitter->Initialize();
 		AnalysisContextStack.Emplace();
+
+		PreviewInput.Expression = Builder->PreviewExpression;
 	}
 
 	void Step_GenerateOutputInstructions()
@@ -117,10 +120,24 @@ struct FMaterialIRModuleBuilderImpl
 		}
 	}
 
+	bool GetExpressionInputDescription(EMaterialProperty Property, FMaterialInputDescription& Input)
+	{
+		if (PreviewInput.IsConnected() && Property == MP_EmissiveColor)
+		{
+			Input.Type = UE::Shader::EValueType::Float3;
+			Input.Input = &PreviewInput;
+			return true;
+		}
+		else
+		{
+			return Builder->Material->GetExpressionInputDescription(Property, Input);
+		}
+	}
+
 	void PrepareSingleMaterialAttribute(EMaterialProperty Property)
 	{
 		FMaterialInputDescription Input;
-		Builder->Material->GetExpressionInputDescription(Property, Input);
+		GetExpressionInputDescription(Property, Input);
 
 		MIR::FSetMaterialOutput* Output = Emitter->SetMaterialOutput(Property, nullptr);
 
@@ -328,7 +345,7 @@ struct FMaterialIRModuleBuilderImpl
 			for (MIR::FSetMaterialOutput* Output : Module->Outputs[StageIndex])
 			{
 				FMaterialInputDescription Input;
-				ensure(Builder->Material->GetExpressionInputDescription(Output->Property, Input));
+				ensure(GetExpressionInputDescription(Output->Property, Input));
 
 				// Connect output properties with substrate parameter if a FrontMaterial property is provided
 				if (Output->Property == MP_FrontMaterial && Output->Arg)
@@ -346,11 +363,26 @@ struct FMaterialIRModuleBuilderImpl
 				{
 					MIR::FValue** ValuePtr = Context.OutputValues.Find(Input.Input->GetConnectedOutput());
 					check(ValuePtr && *ValuePtr);
+					MIR::FValue* Value = *ValuePtr;
+
+					if (PreviewInput.IsConnected() && Output->Property == MP_EmissiveColor)
+					{
+						MIR::FValue* Zero = Emitter->ConstantZero(MIR::EScalarKind::ScalarKind_Float);
+
+						// Get preview expression back into gamma corrected space, as DrawTile does not do this adjustment.
+						Value = Emitter->Operator(MIR::EOperator::BO_Pow, Emitter->Max(Value, Zero), Emitter->ConstantFloat(1.f / 2.2f));
 
-					MIR::Internal::BindValueToExpressionInput(this, Input.Input, *ValuePtr);
+						// Preview should display scalars as red, so if this is a scalar, create a vector padded with zeroes
+						if (Value->IsScalar())
+						{
+							Value = Emitter->Vector3(Value, Zero, Zero);
+						}
+					}
+					
+					MIR::Internal::BindValueToExpressionInput(this, Input.Input, Value);
 
 					const MIR::FType* OutputArgType = MIR::FType::FromShaderType(Input.Type);
-					Output->Arg = Emitter->Cast(*ValuePtr, OutputArgType);
+					Output->Arg = Emitter->Cast(Value, OutputArgType);
 				}
 
 				// Push this connection insight
diff --git a/Engine/Source/Runtime/Engine/Private/Materials/MaterialIRToHLSLTranslator.cpp b/Engine/Source/Runtime/Engine/Private/Materials/MaterialIRToHLSLTranslator.cpp
--- a/Engine/Source/Runtime/Engine/Private/Materials/MaterialIRToHLSLTranslator.cpp
+++ b/Engine/Source/Runtime/Engine/Private/Materials/MaterialIRToHLSLTranslator.cpp
@@ -217,8 +217,7 @@ static bool IsOperatorInfix(MIR::EOperator Op)
 		case MIR::BO_Subtract:
 		case MIR::BO_Multiply:
 		case MIR::BO_Divide:
-		case MIR::BO_Max:
-		case MIR::BO_Min:
+		case MIR::BO_Modulo:
 			return true;
 		default: return false;
 	}
diff --git a/Engine/Source/Runtime/Engine/Private/Materials/MaterialShared.cpp b/Engine/Source/Runtime/Engine/Private/Materials/MaterialShared.cpp
--- a/Engine/Source/Runtime/Engine/Private/Materials/MaterialShared.cpp
+++ b/Engine/Source/Runtime/Engine/Private/Materials/MaterialShared.cpp
@@ -3343,6 +3343,7 @@ bool FMaterial::Translate_New(const FMaterialShaderMapId& InShaderMapId,
 		.TargetPlatform = InTargetPlatform,
 		.StaticParameters = InStaticParameters,
 		.TargetInsights = GetMaterialInterface()->MaterialInsight.Get(),
+		.PreviewExpression = GetMaterialGraphNodePreviewExpression()
 	};
 
 	if (!Builder.Build(&Module))
diff --git a/Engine/Source/Runtime/Engine/Public/Materials/MaterialIRModuleBuilder.h b/Engine/Source/Runtime/Engine/Public/Materials/MaterialIRModuleBuilder.h
--- a/Engine/Source/Runtime/Engine/Public/Materials/MaterialIRModuleBuilder.h
+++ b/Engine/Source/Runtime/Engine/Public/Materials/MaterialIRModuleBuilder.h
@@ -13,6 +13,7 @@ struct FMaterialIRModuleBuilder
 	const ITargetPlatform* TargetPlatform;
 	const FStaticParameterSet& StaticParameters;
 	FMaterialInsights* TargetInsights{};
+	UMaterialExpression* PreviewExpression;
 
 	bool Build(FMaterialIRModule* TargetModule);
 };
