ISMComponentBatcher - Do not batch ISMC/ISKMC instances when GetInstanceTransform() returns false

diff --git a/Engine/Source/Runtime/Engine/Private/ISMPartition/ISMComponentBatcher.cpp b/Engine/Source/Runtime/Engine/Private/ISMPartition/ISMComponentBatcher.cpp
--- a/Engine/Source/Runtime/Engine/Private/ISMPartition/ISMComponentBatcher.cpp
+++ b/Engine/Source/Runtime/Engine/Private/ISMPartition/ISMComponentBatcher.cpp
@@ -98,37 +98,39 @@ void FISMComponentBatcher::AddInternal(const UActorComponent* InComponent, TOpti
 			}
 			
 			FTransform InstanceTransformWS;
-			ISMC->GetInstanceTransform(InstanceIdx, InstanceTransformWS, /*bWorldSpace*/ true);
-			if (InTransformFunc.IsSet())
+			if (ISMC->GetInstanceTransform(InstanceIdx, InstanceTransformWS, /*bWorldSpace*/ true))
 			{
-				InstanceTransformWS = InTransformFunc.GetValue()(InstanceTransformWS);
-			}
+				if (InTransformFunc.IsSet())
+				{
+					InstanceTransformWS = InTransformFunc.GetValue()(InstanceTransformWS);
+				}
 
-			if (InFilterFunc.IsSet())
-			{
-				bool bFilterPass = InFilterFunc.GetValue()(StaticMeshBox.TransformBy(InstanceTransformWS));
-				if (!bFilterPass)
+				if (InFilterFunc.IsSet())
 				{
-					RemappedInstanceIdx--;
-					continue;
+					bool bFilterPass = InFilterFunc.GetValue()(StaticMeshBox.TransformBy(InstanceTransformWS));
+					if (!bFilterPass)
+					{
+						RemappedInstanceIdx--;
+						continue;
+					}
 				}
-			}
 
-			NumInstances++;
+				NumInstances++;
 
-			// Add instance transform
-			InstancesTransformsWS.Add(InstanceTransformWS);
+				// Add instance transform
+				InstancesTransformsWS.Add(InstanceTransformWS);
 
-			// Add per instance custom data, if any
-			if (NumCustomDataFloats > 0)
-			{
-				if (ISMC->NumCustomDataFloats > 0)
+				// Add per instance custom data, if any
+				if (NumCustomDataFloats > 0)
 				{
-					TConstArrayView<float> InstanceCustomData(&ISMC->PerInstanceSMCustomData[InstanceIdx * ISMC->NumCustomDataFloats], ISMC->NumCustomDataFloats);
-					InstancesCustomData.Append(InstanceCustomData);
-				}
+					if (ISMC->NumCustomDataFloats > 0)
+					{
+						TConstArrayView<float> InstanceCustomData(&ISMC->PerInstanceSMCustomData[InstanceIdx * ISMC->NumCustomDataFloats], ISMC->NumCustomDataFloats);
+						InstancesCustomData.Append(InstanceCustomData);
+					}
 
-				InstancesCustomData.AddDefaulted(NumCustomDataFloats - ISMC->NumCustomDataFloats);
+					InstancesCustomData.AddDefaulted(NumCustomDataFloats - ISMC->NumCustomDataFloats);
+				}
 			}
 		}
 	}
@@ -171,24 +173,26 @@ void FISMComponentBatcher::AddInternal(const UActorComponent* InComponent, TOpti
 			// Add instance transform
 			FTransform InstanceTransformWS;
 
-			ISKMC->GetInstanceTransform(InstanceId, InstanceTransformWS, /*bWorldSpace*/ true);
-			if (InTransformFunc.IsSet())
+			if (ISKMC->GetInstanceTransform(InstanceId, InstanceTransformWS, /*bWorldSpace*/ true))
 			{
-				InstanceTransformWS = InTransformFunc.GetValue()(InstanceTransformWS);
-			}
-			InstancesTransformsWS.Add(InstanceTransformWS);
+				if (InTransformFunc.IsSet())
+				{
+					InstanceTransformWS = InTransformFunc.GetValue()(InstanceTransformWS);
+				}
+				InstancesTransformsWS.Add(InstanceTransformWS);
 
-			ISKMC->GetInstanceBankIndex(InstanceId, BankIndices.AddZeroed_GetRef());
-			
-			// Add per instance custom data, if any
-			if (NumCustomDataFloats > 0)
-			{
-				int32 StartIndex = InstancesCustomData.Num(); // Get the current end index
-				InstancesCustomData.AddDefaulted(NumCustomDataFloats); // Append new elements
+				ISKMC->GetInstanceBankIndex(InstanceId, BankIndices.AddZeroed_GetRef());
 
-				if (ISKMC->GetNumCustomDataFloats() > 0)
+				// Add per instance custom data, if any
+				if (NumCustomDataFloats > 0)
 				{
-					ISKMC->GetCustomData(InstanceId, TArrayView<float>(InstancesCustomData.GetData() + StartIndex, NumCustomDataFloats));
+					int32 StartIndex = InstancesCustomData.Num(); // Get the current end index
+					InstancesCustomData.AddDefaulted(NumCustomDataFloats); // Append new elements
+
+					if (ISKMC->GetNumCustomDataFloats() > 0)
+					{
+						ISKMC->GetCustomData(InstanceId, TArrayView<float>(InstancesCustomData.GetData() + StartIndex, NumCustomDataFloats));
+					}
 				}
 			}
 		}
