[RewindDebugger] minor fixes for previous change to add support for child element.
- missing rename for RouteIds to match event names
- INDEX_NONE was still used in some code path instead of '0'
- TRACE_CHILD_ELEMENT macros should reuse the provided UObject as the outer (not its outer)

diff --git a/Engine/Plugins/Animation/GameplayInsights/Source/GameplayInsights/Private/GameplayAnalyzer.cpp b/Engine/Plugins/Animation/GameplayInsights/Source/GameplayInsights/Private/GameplayAnalyzer.cpp
--- a/Engine/Plugins/Animation/GameplayInsights/Source/GameplayInsights/Private/GameplayAnalyzer.cpp
+++ b/Engine/Plugins/Animation/GameplayInsights/Source/GameplayInsights/Private/GameplayAnalyzer.cpp
@@ -35,10 +35,10 @@ void FGameplayAnalyzer::OnAnalysisBegin(const FOnAnalysisContext& Context)
 	Builder.RouteEvent(RouteId_PropertiesEnd, "Object", "PropertiesEnd");
 	Builder.RouteEvent(RouteId_PropertyValue2, "Object", "PropertyValue2");
 	Builder.RouteEvent(RouteId_PropertyValue3, "Object", "PropertyValue3");
-	Builder.RouteEvent(RouteId_ChildElement, "Object", "ObjectChildElement");
-	Builder.RouteEvent(RouteId_ChildElementLifetimeBegin, "Object", "ObjectChildElementLifetimeBegin");
-	Builder.RouteEvent(RouteId_ChildElementLifetimeEnd, "Object", "ObjectChildElementLifetimeEnd");
-	Builder.RouteEvent(RouteId_ChildElementEvent, "Object", "ObjectChildElementEvent");
+	Builder.RouteEvent(RouteId_ChildElement, "Object", "ChildElement");
+	Builder.RouteEvent(RouteId_ChildElementLifetimeBegin, "Object", "ChildElementLifetimeBegin");
+	Builder.RouteEvent(RouteId_ChildElementLifetimeEnd, "Object", "ChildElementLifetimeEnd");
+	Builder.RouteEvent(RouteId_ChildElementEvent, "Object", "ChildElementEvent");
 }
 
 bool FGameplayAnalyzer::OnEvent(uint16 RouteId, EStyle Style, const FOnEventContext& Context)
diff --git a/Engine/Plugins/Runtime/StateTree/Source/StateTreeModule/Private/Debugger/StateTreeTrace.cpp b/Engine/Plugins/Runtime/StateTree/Source/StateTreeModule/Private/Debugger/StateTreeTrace.cpp
--- a/Engine/Plugins/Runtime/StateTree/Source/StateTreeModule/Private/Debugger/StateTreeTrace.cpp
+++ b/Engine/Plugins/Runtime/StateTree/Source/StateTreeModule/Private/Debugger/StateTreeTrace.cpp
@@ -758,7 +758,7 @@ void OutputInstanceLifetimeEvent(
 			const uint64 ChildId(InstanceId.ToUint64());
 			if (EventType == EStateTreeTraceEventType::Push)
 			{
-				TRACE_CHILD_ELEMENT_WITH_OUTER_LIFETIME_BEGIN(Owner, ChildId, Owner, INDEX_NONE);
+				TRACE_CHILD_ELEMENT_LIFETIME_BEGIN(Owner, ChildId);
 			}
 			else if (EventType == EStateTreeTraceEventType::Pop)
 			{
diff --git a/Engine/Source/Runtime/Engine/Private/ObjectTrace.cpp b/Engine/Source/Runtime/Engine/Private/ObjectTrace.cpp
--- a/Engine/Source/Runtime/Engine/Private/ObjectTrace.cpp
+++ b/Engine/Source/Runtime/Engine/Private/ObjectTrace.cpp
@@ -601,7 +601,7 @@ void FObjectTrace::OutputChildElement(const UObject* InObject, const uint64 InCh
 		return;
 	}
 
-	if (InOuterChildElementId == INDEX_NONE)
+	if (InOuterChildElementId == InvalidObjectId)
 	{
 		OutputObject(InOuterObject);
 	}
diff --git a/Engine/Source/Runtime/Engine/Public/ObjectTrace.h b/Engine/Source/Runtime/Engine/Public/ObjectTrace.h
--- a/Engine/Source/Runtime/Engine/Public/ObjectTrace.h
+++ b/Engine/Source/Runtime/Engine/Public/ObjectTrace.h
@@ -45,6 +45,8 @@ class FSceneView;
 
 struct FObjectTrace
 {
+	static constexpr uint64 InvalidObjectId = 0;
+
 	/** Initialize object tracing */
 	ENGINE_API static void Init();
 
@@ -109,7 +111,7 @@ struct FObjectTrace
 	ENGINE_API static void OutputWorld(const UWorld* InWorld);
 
 	/** Helper function to output a child element of a UObject */
-	ENGINE_API static void OutputChildElement(const UObject* InObject, uint64 InChildElementId, const UObject* InOuterObject = nullptr , uint64 InOuterChildElementId = INDEX_NONE);
+	ENGINE_API static void OutputChildElement(const UObject* InObject, uint64 InChildElementId, const UObject* InOuterObject = nullptr, uint64 InOuterChildElementId = InvalidObjectId);
 
 	/** Helper function to output destruction event of a child element */
 	ENGINE_API static void OutputChildElementLifetimeEnd(const UObject* InObject, uint64 InChildElementId);
@@ -145,13 +147,13 @@ struct FObjectTrace
 	if (CAN_TRACE_OBJECT(Controller) && (Pawn==nullptr || (CAN_TRACE_OBJECT(Pawn)))) { FObjectTrace::OutputPawnPossess(Controller, Pawn); }
 
 #define TRACE_CHILD_ELEMENT(Object, ChildElementId) \
-	if (CAN_TRACE_OBJECT(Object)) { FObjectTrace::OutputChildElement(Object, ChildElementId, (Object ? Object->GetOuter() : nullptr), INDEX_NONE); }
+	if (CAN_TRACE_OBJECT(Object)) { FObjectTrace::OutputChildElement(Object, ChildElementId, Object, FObjectTrace::InvalidObjectId); }
 
 #define TRACE_CHILD_ELEMENT_WITH_OUTER(Object, ChildElementId, OuterObject, OuterChildElementId) \
 	if (CAN_TRACE_OBJECT(Object)) { FObjectTrace::OutputChildElement(Object, ChildElementId, OuterObject, OuterChildElementId); }
 
 #define TRACE_CHILD_ELEMENT_LIFETIME_BEGIN(Object, ChildElementId) \
-	if (CAN_TRACE_OBJECT(Object)) { FObjectTrace::OutputChildElement(Object, ChildElementId, (Object ? Object->GetOuter() : nullptr), INDEX_NONE); }
+	if (CAN_TRACE_OBJECT(Object)) { FObjectTrace::OutputChildElement(Object, ChildElementId, Object, FObjectTrace::InvalidObjectId); }
 
 #define TRACE_CHILD_ELEMENT_WITH_OUTER_LIFETIME_BEGIN(Object, ChildElementId, OuterObject, OuterChildElementId) \
 	if (CAN_TRACE_OBJECT(Object)) { FObjectTrace::OutputChildElement(Object, ChildElementId, OuterObject, OuterChildElementId); }
@@ -180,13 +182,13 @@ struct FObjectTrace
 	FObjectTrace::OutputPawnPossess(Controller, Pawn);
 
 #define TRACE_CHILD_ELEMENT(Object, ChildElementId) \
-	FObjectTrace::OutputChildElement(Object, ChildElementId, (Object ? Object->GetOuter() : nullptr), INDEX_NONE);
+	FObjectTrace::OutputChildElement(Object, ChildElementId, Object, FObjectTrace::InvalidObjectId);
 
 #define TRACE_CHILD_ELEMENT_WITH_OUTER(Object, ChildElementId, OuterObject, OuterChildElementId) \
 	FObjectTrace::OutputChildElement(Object, ChildElementId, OuterObject, OuterChildElementId);
 
 #define TRACE_CHILD_ELEMENT_LIFETIME_BEGIN(Object, ChildElementId) \
-	FObjectTrace::OutputChildElement(Object, ChildElementId, (Object ? Object->GetOuter() : nullptr), INDEX_NONE);
+	FObjectTrace::OutputChildElement(Object, ChildElementId, Object, FObjectTrace::InvalidObjectId);
 
 #define TRACE_CHILD_ELEMENT_WITH_OUTER_LIFETIME_BEGIN(Object, ChildElementId, OuterObject, OuterChildElementId) \
 	FObjectTrace::OutputChildElement(Object, ChildElementId, OuterObject, OuterChildElementId);
