Async physics state creation/destruction now supports to set a maximum time budget per frame for the asynchronous task using p.Chaos.AsyncPhysicsStateTask.TimeBudgetMS (no limit by default).

diff --git a/Engine/Plugins/Experimental/FastGeoStreaming/Source/FastGeoStreaming/Private/FastGeoContainer.cpp b/Engine/Plugins/Experimental/FastGeoStreaming/Source/FastGeoStreaming/Private/FastGeoContainer.cpp
--- a/Engine/Plugins/Experimental/FastGeoStreaming/Source/FastGeoStreaming/Private/FastGeoContainer.cpp
+++ b/Engine/Plugins/Experimental/FastGeoStreaming/Source/FastGeoStreaming/Private/FastGeoContainer.cpp
@@ -291,12 +291,14 @@ bool UFastGeoContainer::HasAnyPendingTasks() const
 
 bool UFastGeoContainer::HasAnyPendingCreateTasks() const
 {
-	return PendingCreate.HasAnyPendingState();
+	// Physics state can have no PendingState but OnAsyncCreatePhysicsStateEnd_GameThread has not been called yet
+	return (PhysicsStateCreation == EPhysicsStateCreation::Creating) || PendingCreate.HasAnyPendingState();
 }
 
 bool UFastGeoContainer::HasAnyPendingDestroyTasks() const
 {
-	return PendingDestroy.HasAnyPendingState();
+	// Physics state can have no PendingState but OnAsyncDestroyPhysicsStateEnd_GameThread has not been called yet
+	return (PhysicsStateCreation == EPhysicsStateCreation::Destroying) || PendingDestroy.HasAnyPendingState();
 }
 
 FFastGeoComponentCluster* UFastGeoContainer::GetComponentCluster(uint32 InComponentClusterTypeID, int32 InComponentClusterIndex)
@@ -735,22 +737,22 @@ void UFastGeoContainer::OnCreatePhysicsStateBegin_GameThread()
 {
 	TRACE_CPUPROFILER_EVENT_SCOPE(UFastGeoContainer::OnCreatePhysicsStateBegin_GameThread);
 
+	PendingCreate.PhysicsState.ComponentsToProcess = &CollisionComponents;
+	check(PendingCreate.PhysicsState.TotalNumProcessed == 0);
+
 	UWorld* World = GetWorld();
 	check(World);
 	FPhysScene* PhysScene = World->GetPhysicsScene();
 	check(PhysScene);
 	PhysScene->PushAsyncCreatePhysicsStateJob(this);
-
-	check(!PendingCreate.bPhysicsState);
-	PendingCreate.bPhysicsState = true;
 }
 
 void UFastGeoContainer::OnDestroyPhysicsStateBegin_GameThread()
 {
 	TRACE_CPUPROFILER_EVENT_SCOPE(UFastGeoContainer::OnDestroyPhysicsStateBegin_GameThread);
 
-	check(!PendingDestroy.bPhysicsState);
-	PendingDestroy.bPhysicsState = true;
+	PendingDestroy.PhysicsState.ComponentsToProcess = &CollisionComponents;
+	check(PendingDestroy.PhysicsState.TotalNumProcessed == 0);
 
 	UWorld* World = GetWorld();
 	check(World);
@@ -781,27 +783,36 @@ UObject* UFastGeoContainer::GetAsyncPhysicsStateObject() const
 	return const_cast<UFastGeoContainer*>(this);
 }
 
-void UFastGeoContainer::OnAsyncCreatePhysicsState()
+bool UFastGeoContainer::OnAsyncCreatePhysicsState(const UE::FTimeout& Timeout)
 {
 	TRACE_CPUPROFILER_EVENT_SCOPE(UFastGeoContainer::OnAsyncCreatePhysicsState);
-	check(PhysicsStateCreation == EPhysicsStateCreation::NotCreated);
+	check((PhysicsStateCreation == EPhysicsStateCreation::NotCreated) || 
+		  ((PhysicsStateCreation == EPhysicsStateCreation::Creating) && (PendingCreate.PhysicsState.TotalNumProcessed > 0) && !PendingCreate.PhysicsState.IsCompleted()));
 	PhysicsStateCreation = EPhysicsStateCreation::Creating;
 
-	for (FFastGeoComponent* Component : CollisionComponents)
+	check(PendingCreate.PhysicsState.ComponentsToProcess);
+	for (int Index = PendingCreate.PhysicsState.TotalNumProcessed; Index < PendingCreate.PhysicsState.ComponentsToProcess->Num(); ++Index)
 	{
+		FFastGeoComponent* Component = (*PendingCreate.PhysicsState.ComponentsToProcess)[Index];
 		Component->OnAsyncCreatePhysicsState();
+		++PendingCreate.PhysicsState.TotalNumProcessed;
+		if (!PendingCreate.PhysicsState.IsCompleted() && Timeout.IsExpired())
+		{
+			return false;
+		}
 	}
+
+	return true;
 }
 
 void UFastGeoContainer::OnAsyncCreatePhysicsStateEnd_GameThread()
 {
 	TRACE_CPUPROFILER_EVENT_SCOPE(UFastGeoContainer::OnAsyncCreatePhysicsStateEnd_GameThread);
+	check(PendingCreate.PhysicsState.ComponentsToProcess && PendingCreate.PhysicsState.IsCompleted());
+	PendingCreate.PhysicsState.Reset();
 	check(PhysicsStateCreation == EPhysicsStateCreation::Creating);
 	PhysicsStateCreation = EPhysicsStateCreation::Created;
-
-	check(PendingCreate.bPhysicsState);
-	PendingCreate.bPhysicsState = false;
-
+	
 	// This loop could be removed if component's PhysicsStateCreation was removed
 	for (FFastGeoComponent* Component : CollisionComponents)
 	{
@@ -821,24 +832,33 @@ void UFastGeoContainer::OnAsyncDestroyPhysicsStateBegin_GameThread()
 	}
 }
 
-void UFastGeoContainer::OnAsyncDestroyPhysicsState()
+bool UFastGeoContainer::OnAsyncDestroyPhysicsState(const UE::FTimeout& Timeout)
 {
 	TRACE_CPUPROFILER_EVENT_SCOPE(UFastGeoContainer::OnAsyncDestroyPhysicsState);
-	for (FFastGeoComponent* Component : CollisionComponents)
+
+	//for (FFastGeoComponent* Component : CollisionComponents)
+	check(PendingDestroy.PhysicsState.ComponentsToProcess);
+	for (int Index = PendingDestroy.PhysicsState.TotalNumProcessed; Index < PendingDestroy.PhysicsState.ComponentsToProcess->Num(); ++Index)
 	{
+		FFastGeoComponent* Component = (*PendingDestroy.PhysicsState.ComponentsToProcess)[Index];
 		Component->OnAsyncDestroyPhysicsState();
+		++PendingDestroy.PhysicsState.TotalNumProcessed;
+		if (!PendingDestroy.PhysicsState.IsCompleted() && Timeout.IsExpired())
+		{
+			return false;
+		}
 	}
+	return true;
 }
 
 void UFastGeoContainer::OnAsyncDestroyPhysicsStateEnd_GameThread()
 {
 	TRACE_CPUPROFILER_EVENT_SCOPE(UFastGeoContainer::OnAsyncDestroyPhysicsStateEnd_GameThread);
+	check(PendingDestroy.PhysicsState.ComponentsToProcess && PendingDestroy.PhysicsState.IsCompleted());
+	PendingDestroy.PhysicsState.Reset();
 	check(PhysicsStateCreation == EPhysicsStateCreation::Destroying);
 	PhysicsStateCreation = EPhysicsStateCreation::NotCreated;
 
-	check(PendingDestroy.bPhysicsState);
-	PendingDestroy.bPhysicsState = false;
-
 	// This loop could be removed if component's PhysicsStateCreation was removed
 	for (FFastGeoComponent* Component : CollisionComponents)
 	{
diff --git a/Engine/Plugins/Experimental/FastGeoStreaming/Source/FastGeoStreaming/Private/FastGeoContainer.h b/Engine/Plugins/Experimental/FastGeoStreaming/Source/FastGeoStreaming/Private/FastGeoContainer.h
--- a/Engine/Plugins/Experimental/FastGeoStreaming/Source/FastGeoStreaming/Private/FastGeoContainer.h
+++ b/Engine/Plugins/Experimental/FastGeoStreaming/Source/FastGeoStreaming/Private/FastGeoContainer.h
@@ -29,6 +29,24 @@ struct FFastGeoComponentClusterPendingState
 {
 	struct FRenderState
 	{
+		FRenderState()
+		{
+			Reset();
+		}
+
+		void Reset()
+		{
+			ComponentsToProcess.Reset();
+			NumToProcess = 0;
+			NumProcessed = 0;
+			TotalNumProcessed = 0;
+		}
+
+		bool IsCompleted() const
+		{
+			return TotalNumProcessed >= ComponentsToProcess.Num();
+		}
+
 		TArray<FFastGeoPrimitiveComponent*> ComponentsToProcess;
 
 		int32 NumToProcess;
@@ -36,7 +54,27 @@ struct FFastGeoComponentClusterPendingState
 		int32 TotalNumProcessed;
 	} RenderState;
 
-	bool bPhysicsState;
+	struct FPhysicsState
+	{
+		FPhysicsState()
+		{
+			Reset();
+		}
+
+		void Reset()
+		{
+			ComponentsToProcess = nullptr;
+			TotalNumProcessed = 0;
+		}
+
+		bool IsCompleted() const
+		{
+			return !ComponentsToProcess || (TotalNumProcessed >= ComponentsToProcess->Num());
+		}
+
+		TArray<FFastGeoComponent*>* ComponentsToProcess;
+		std::atomic<int32> TotalNumProcessed;
+	} PhysicsState;
 
 	bool bIsInBlockingWait;
 	
@@ -47,15 +85,15 @@ struct FFastGeoComponentClusterPendingState
 
 	void Reset()
 	{
-		RenderState.ComponentsToProcess.Reset();
-		RenderState.NumToProcess = 0;
-		RenderState.NumProcessed = 0;
-		RenderState.TotalNumProcessed = 0;
+		RenderState.Reset();
+		PhysicsState.Reset();
 		bIsInBlockingWait = false;
-		bPhysicsState = false;
 	}
 
-	bool HasAnyPendingState() const { return RenderState.TotalNumProcessed < RenderState.ComponentsToProcess.Num() || bPhysicsState; }
+	bool HasAnyPendingState() const
+	{ 
+		return !RenderState.IsCompleted() || !PhysicsState.IsCompleted();
+	}
 };
 
 UCLASS(Within = Level)
@@ -111,10 +149,10 @@ public:
 	virtual bool AllowsAsyncPhysicsStateDestruction() const override final;
 	virtual bool IsAsyncPhysicsStateCreated() const override final;
 	virtual UObject* GetAsyncPhysicsStateObject() const override final;
-	virtual void OnAsyncCreatePhysicsState() override;
+	virtual bool OnAsyncCreatePhysicsState(const UE::FTimeout& Timeout) override;
 	virtual void OnAsyncCreatePhysicsStateEnd_GameThread() override;
 	virtual void OnAsyncDestroyPhysicsStateBegin_GameThread() override;
-	virtual void OnAsyncDestroyPhysicsState() override;
+	virtual bool OnAsyncDestroyPhysicsState(const UE::FTimeout& Timeout) override;
 	virtual void OnAsyncDestroyPhysicsStateEnd_GameThread() override;
 	virtual void CollectBodySetupsWithPhysicsMeshesToCreate(TSet<UBodySetup*>& OutBodySetups) const override;
 	//~ End IAsyncPhysicsStateProcessor
diff --git a/Engine/Plugins/Experimental/FastGeoStreaming/Source/FastGeoStreaming/Private/FastGeoInstancedStaticMeshComponent.cpp b/Engine/Plugins/Experimental/FastGeoStreaming/Source/FastGeoStreaming/Private/FastGeoInstancedStaticMeshComponent.cpp
--- a/Engine/Plugins/Experimental/FastGeoStreaming/Source/FastGeoStreaming/Private/FastGeoInstancedStaticMeshComponent.cpp
+++ b/Engine/Plugins/Experimental/FastGeoStreaming/Source/FastGeoStreaming/Private/FastGeoInstancedStaticMeshComponent.cpp
@@ -197,8 +197,6 @@ void FFastGeoInstancedStaticMeshComponent::OnAsyncCreatePhysicsState()
 
 void FFastGeoInstancedStaticMeshComponent::OnAsyncDestroyPhysicsStateBegin_GameThread()
 {
-	TRACE_CPUPROFILER_EVENT_SCOPE(FFastGeoInstancedStaticMeshComponent::OnAsyncDestroyPhysicsStateBegin_GameThread);
-
 	FFastGeoComponent::OnAsyncDestroyPhysicsStateBegin_GameThread();
 
 	// Move InstanceBodies in AsyncDestroyPhysicsStatePayload
diff --git a/Engine/Plugins/Experimental/FastGeoStreaming/Source/FastGeoStreaming/Private/FastGeoPrimitiveComponent.cpp b/Engine/Plugins/Experimental/FastGeoStreaming/Source/FastGeoStreaming/Private/FastGeoPrimitiveComponent.cpp
--- a/Engine/Plugins/Experimental/FastGeoStreaming/Source/FastGeoStreaming/Private/FastGeoPrimitiveComponent.cpp
+++ b/Engine/Plugins/Experimental/FastGeoStreaming/Source/FastGeoStreaming/Private/FastGeoPrimitiveComponent.cpp
@@ -450,8 +450,6 @@ void FFastGeoPrimitiveComponent::OnAsyncCreatePhysicsState()
 
 void FFastGeoPrimitiveComponent::OnAsyncDestroyPhysicsStateBegin_GameThread()
 {
-	TRACE_CPUPROFILER_EVENT_SCOPE(FFastGeoPrimitiveComponent::OnAsyncDestroyPhysicsStateBegin_GameThread);
-
 	check(!AsyncTermBodyPayload.IsSet());
 	AsyncTermBodyPayload = BodyInstance.StartAsyncTermBody_GameThread();
 	check(!BodyInstance.IsValidBodyInstance());
diff --git a/Engine/Source/Runtime/Engine/Classes/Components/ActorComponent.h b/Engine/Source/Runtime/Engine/Classes/Components/ActorComponent.h
--- a/Engine/Source/Runtime/Engine/Classes/Components/ActorComponent.h
+++ b/Engine/Source/Runtime/Engine/Classes/Components/ActorComponent.h
@@ -851,9 +851,9 @@ public:
 	}
 	ENGINE_API virtual bool IsAsyncPhysicsStateCreated() const override final;
 	ENGINE_API virtual UObject* GetAsyncPhysicsStateObject() const override final;
-	ENGINE_API virtual void OnAsyncCreatePhysicsState() override;
+	ENGINE_API virtual bool OnAsyncCreatePhysicsState(const UE::FTimeout& Timeout) override;
 	ENGINE_API virtual void OnAsyncCreatePhysicsStateEnd_GameThread() override;
-	virtual void OnAsyncDestroyPhysicsState() override {}
+	virtual bool OnAsyncDestroyPhysicsState(const UE::FTimeout& Timeout) override { return true; }
 	ENGINE_API virtual void OnAsyncDestroyPhysicsStateBegin_GameThread() override;
 	ENGINE_API virtual void OnAsyncDestroyPhysicsStateEnd_GameThread() override;
 	//~ End IAsyncPhysicsStateProcessor
diff --git a/Engine/Source/Runtime/Engine/Classes/Components/InstancedStaticMeshComponent.h b/Engine/Source/Runtime/Engine/Classes/Components/InstancedStaticMeshComponent.h
--- a/Engine/Source/Runtime/Engine/Classes/Components/InstancedStaticMeshComponent.h
+++ b/Engine/Source/Runtime/Engine/Classes/Components/InstancedStaticMeshComponent.h
@@ -522,7 +522,7 @@ public:
 protected:
 	ENGINE_API virtual void OnCreatePhysicsState() override;
 	ENGINE_API virtual void OnDestroyPhysicsState() override;
-	ENGINE_API virtual void OnAsyncDestroyPhysicsState() override;
+	ENGINE_API virtual bool OnAsyncDestroyPhysicsState(const UE::FTimeout& Timeout) override;
 	ENGINE_API virtual void OnAsyncDestroyPhysicsStateBegin_GameThread() override;
 
 public:
diff --git a/Engine/Source/Runtime/Engine/Classes/Components/PrimitiveComponent.h b/Engine/Source/Runtime/Engine/Classes/Components/PrimitiveComponent.h
--- a/Engine/Source/Runtime/Engine/Classes/Components/PrimitiveComponent.h
+++ b/Engine/Source/Runtime/Engine/Classes/Components/PrimitiveComponent.h
@@ -2474,7 +2474,7 @@ public:
 
 	//~ Begin IAsyncPhysicsStateProcessor
 	ENGINE_API virtual void OnAsyncCreatePhysicsStateEnd_GameThread() override;
-	ENGINE_API virtual void OnAsyncDestroyPhysicsState() override;
+	ENGINE_API virtual bool OnAsyncDestroyPhysicsState(const UE::FTimeout& Timeout) override;
 	ENGINE_API virtual void OnAsyncDestroyPhysicsStateBegin_GameThread() override;
 	ENGINE_API virtual void OnAsyncDestroyPhysicsStateEnd_GameThread() override;
 	ENGINE_API virtual void CollectBodySetupsWithPhysicsMeshesToCreate(TSet<UBodySetup*>& OutBodySetups) const override;
diff --git a/Engine/Source/Runtime/Engine/Private/Components/ActorComponent.cpp b/Engine/Source/Runtime/Engine/Private/Components/ActorComponent.cpp
--- a/Engine/Source/Runtime/Engine/Private/Components/ActorComponent.cpp
+++ b/Engine/Source/Runtime/Engine/Private/Components/ActorComponent.cpp
@@ -2200,11 +2200,12 @@ UObject* UActorComponent::GetAsyncPhysicsStateObject() const
 	return const_cast<UActorComponent*>(this);
 }
 
-void UActorComponent::OnAsyncCreatePhysicsState()
+bool UActorComponent::OnAsyncCreatePhysicsState(const UE::FTimeout& Timeout)
 {
 	TRACE_CPUPROFILER_EVENT_SCOPE(UActorComponent::OnAsyncCreatePhysicsState);
 	check(AllowsAsyncPhysicsStateCreation());
 	OnCreatePhysicsState();
+	return true;
 }
 
 void UActorComponent::OnAsyncCreatePhysicsStateEnd_GameThread()
@@ -2225,7 +2226,6 @@ void UActorComponent::OnAsyncCreatePhysicsStateEnd_GameThread()
 
 void UActorComponent::OnAsyncDestroyPhysicsStateBegin_GameThread()
 {
-	TRACE_CPUPROFILER_EVENT_SCOPE(UActorComponent::OnAsyncDestroyPhysicsStateBegin_GameThread);
 	GlobalDestroyPhysicsDelegate.Broadcast(this);
 	ensureMsgf(bRegistered, TEXT("Component has physics state when not registered (%s)"), *GetFullName()); // should not have physics state unless we are registered
 }
diff --git a/Engine/Source/Runtime/Engine/Private/Components/PrimitiveComponent.cpp b/Engine/Source/Runtime/Engine/Private/Components/PrimitiveComponent.cpp
--- a/Engine/Source/Runtime/Engine/Private/Components/PrimitiveComponent.cpp
+++ b/Engine/Source/Runtime/Engine/Private/Components/PrimitiveComponent.cpp
@@ -1044,7 +1044,6 @@ void UPrimitiveComponent::SendPhysicsTransform(ETeleportType Teleport)
 
 void UPrimitiveComponent::OnAsyncDestroyPhysicsStateBegin_GameThread()
 {
-	TRACE_CPUPROFILER_EVENT_SCOPE(UPrimitiveComponent::OnAsyncDestroyPhysicsStateBegin_GameThread);
 	check(AllowsAsyncPhysicsStateDestruction());
 
 	// we remove welding related to this component
@@ -1062,7 +1061,7 @@ void UPrimitiveComponent::OnAsyncDestroyPhysicsStateBegin_GameThread()
 	Super::OnAsyncDestroyPhysicsStateBegin_GameThread();
 }
 
-void UPrimitiveComponent::OnAsyncDestroyPhysicsState()
+bool UPrimitiveComponent::OnAsyncDestroyPhysicsState(const UE::FTimeout& Timeout)
 {
 	TRACE_CPUPROFILER_EVENT_SCOPE(UPrimitiveComponent::OnAsyncDestroyPhysicsState);
 	
@@ -1073,7 +1072,7 @@ void UPrimitiveComponent::OnAsyncDestroyPhysicsState()
 		AsyncTermBodyPayload.Reset();
 	}
 
-	Super::OnAsyncDestroyPhysicsState();
+	return Super::OnAsyncDestroyPhysicsState(Timeout);
 }
 
 void UPrimitiveComponent::OnAsyncDestroyPhysicsStateEnd_GameThread()
diff --git a/Engine/Source/Runtime/Engine/Private/Components/StaticMeshComponent.cpp b/Engine/Source/Runtime/Engine/Private/Components/StaticMeshComponent.cpp
--- a/Engine/Source/Runtime/Engine/Private/Components/StaticMeshComponent.cpp
+++ b/Engine/Source/Runtime/Engine/Private/Components/StaticMeshComponent.cpp
@@ -971,7 +971,6 @@ void UStaticMeshComponent::OnDestroyPhysicsState()
 
 void UStaticMeshComponent::OnAsyncDestroyPhysicsStateEnd_GameThread()
 {
-	TRACE_CPUPROFILER_EVENT_SCOPE(UStaticMeshComponent::OnAsyncDestroyPhysicsStateEnd_GameThread);
 	Super::OnAsyncDestroyPhysicsStateEnd_GameThread();
 
 	bNavigationRelevant = IsNavigationRelevant();
diff --git a/Engine/Source/Runtime/Engine/Private/InstancedStaticMesh.cpp b/Engine/Source/Runtime/Engine/Private/InstancedStaticMesh.cpp
--- a/Engine/Source/Runtime/Engine/Private/InstancedStaticMesh.cpp
+++ b/Engine/Source/Runtime/Engine/Private/InstancedStaticMesh.cpp
@@ -2788,7 +2788,6 @@ void UInstancedStaticMeshComponent::OnDestroyPhysicsState()
 
 void UInstancedStaticMeshComponent::OnAsyncDestroyPhysicsStateBegin_GameThread()
 {
-	TRACE_CPUPROFILER_EVENT_SCOPE(UInstancedStaticMeshComponent::OnAsyncDestroyPhysicsStateBegin_GameThread);
 	USceneComponent::OnAsyncDestroyPhysicsStateBegin_GameThread();
 
 	// Move InstanceBodies in AsyncDestroyPhysicsStatePayload
@@ -2798,10 +2797,10 @@ void UInstancedStaticMeshComponent::OnAsyncDestroyPhysicsStateBegin_GameThread()
 	bNavigationRelevant = IsNavigationRelevant();
 }
 
-void UInstancedStaticMeshComponent::OnAsyncDestroyPhysicsState()
+bool UInstancedStaticMeshComponent::OnAsyncDestroyPhysicsState(const UE::FTimeout& Timeout)
 {
 	TRACE_CPUPROFILER_EVENT_SCOPE(UInstancedStaticMeshComponent::OnAsyncDestroyPhysicsState);
-	USceneComponent::OnAsyncDestroyPhysicsState();
+	USceneComponent::OnAsyncDestroyPhysicsState(Timeout);
 
 	check(InstanceBodies.IsEmpty());
 	for (FBodyInstance*& Instance : AsyncDestroyPhysicsStatePayload)
@@ -2813,6 +2812,7 @@ void UInstancedStaticMeshComponent::OnAsyncDestroyPhysicsState()
 		}
 	}
 	AsyncDestroyPhysicsStatePayload.Empty();
+	return true;
 }
 
 Chaos::FPhysicsObject* UInstancedStaticMeshComponent::GetPhysicsObjectById(Chaos::FPhysicsObjectId Id) const
diff --git a/Engine/Source/Runtime/Engine/Private/PhysicsEngine/Experimental/PhysScene_AsyncPhysicsStateJobQueue.cpp b/Engine/Source/Runtime/Engine/Private/PhysicsEngine/Experimental/PhysScene_AsyncPhysicsStateJobQueue.cpp
--- a/Engine/Source/Runtime/Engine/Private/PhysicsEngine/Experimental/PhysScene_AsyncPhysicsStateJobQueue.cpp
+++ b/Engine/Source/Runtime/Engine/Private/PhysicsEngine/Experimental/PhysScene_AsyncPhysicsStateJobQueue.cpp
@@ -5,6 +5,16 @@
 #include "PhysicsEngine/BodySetup.h"
 #include "Engine/World.h"
 
+namespace Chaos
+{
+	static float GAsyncPhysicsStateTaskTimeBudgetMS = 0.0f;
+	static FAutoConsoleVariableRef CVarAsyncPhysicsStateTaskTimeBudgetMS(
+		TEXT("p.Chaos.AsyncPhysicsStateTask.TimeBudgetMS"),
+		GAsyncPhysicsStateTaskTimeBudgetMS,
+		TEXT("Maximum time budget in milliseconds for the async physics state task (0 = no time limit)"),
+		ECVF_Default);
+}
+
 FPhysScene_AsyncPhysicsStateJobQueue::FPhysScene_AsyncPhysicsStateJobQueue(FPhysScene* InPhysicScene)
 	: PhysScene(InPhysicScene)
 {
@@ -14,9 +24,11 @@ FPhysScene_AsyncPhysicsStateJobQueue::FPhysScene_AsyncPhysicsStateJobQueue(FPhys
 FPhysScene_AsyncPhysicsStateJobQueue::~FPhysScene_AsyncPhysicsStateJobQueue()
 {
 	PhysScene->GetOwningWorld()->OnAllLevelsChanged().RemoveAll(this);
-	Launch();
-	WaitForAsyncTasksExecution();
-	OnAsyncTasksExecuted();
+
+	if (!IsCompleted())
+	{
+		Tick(true);
+	}
 	check(IsCompleted());
 }
 
@@ -25,11 +37,9 @@ void FPhysScene_AsyncPhysicsStateJobQueue::AddJob(const FJob& Job)
 	TRACE_CPUPROFILER_EVENT_SCOPE(FPhysScene_AsyncPhysicsStateJobQueue::AddJob);
 	check(IsInGameThread());
 
-	if (!PendingJobs.IsValid())
-	{
-		PendingJobs = MakeUnique<FJobSet>();
-	}
-	PendingJobs->Add(Job);
+	Job.OnPreExecute_GameThread();
+	FWriteScopeLock Lock(JobsLock);
+	JobsToExecute.Add(Job);
 }
 
 void FPhysScene_AsyncPhysicsStateJobQueue::RemoveJob(const FJob& Job)
@@ -38,166 +48,192 @@ void FPhysScene_AsyncPhysicsStateJobQueue::RemoveJob(const FJob& Job)
 	check(IsInGameThread());
 	check(Job.IsValid());
 
-	if (PendingJobs.IsValid() && PendingJobs->Remove(Job))
+	bool bWaitForJobToComplete = false;
+	bool bFoundJob = false;
 	{
-		return;
+		FWriteScopeLock Lock(JobsLock);
+		// First test in CompletedJobs, if found, leave JobsToExecute 
+		// as it will be cleaned up by the async task.
+		if (CompletedJobs.Remove(Job))
+		{
+			bFoundJob = true;
+		}
+		// If the job is executing, wait for it to complete
+		else if (ExecutingJob == Job)
+		{
+			bFoundJob = true;
+			bWaitForJobToComplete = true;
+		}
+		// Else it's safe to remove it from JobsToExecute
+		else if (JobsToExecute.Remove(Job))
+		{
+			bFoundJob = true;
+		}
 	}
 
-	WaitForAsyncTasksExecution();
-	
-	for (FJobs& Jobs : PipedJobs)
+	if (bFoundJob)
 	{
-		if (Jobs->Remove(Job))
+		if (bWaitForJobToComplete)
 		{
-			Job.OnPostExecute_GameThread();
-			break;
+			// Wait for the async task
+			AsyncJobTask.Wait();
+			FWriteScopeLock Lock(JobsLock);
+			// If job is completed, needs to be removed from CompletedJobs
+			if (!CompletedJobs.Remove(Job))
+			{
+				// Else if it's still the executing job, finish executing it
+				if (ExecutingJob == Job)
+				{
+					UE::FTimeout NeverTimeout(UE::FTimeout::Never());
+					ExecutingJob->Execute(NeverTimeout);
+					ExecutingJob.Reset();
+					JobsToExecute.Remove(Job);
+				}
+			}
+			check(!ExecutingJob.IsSet());
 		}
+		Job.OnPostExecute_GameThread();
 	}
 }
 
-void FPhysScene_AsyncPhysicsStateJobQueue::OnAsyncTasksExecuted()
+bool FPhysScene_AsyncPhysicsStateJobQueue::IsCompleted() const
 {
-	check(AreAsyncTasksExecuted());
-	for (const FJobs& Jobs : PipedJobs)
-	{
-		Jobs->OnPostExecute_GameThread();
-	}
-	PipedJobs.Reset();
-	PipeTasks.Reset();
-	IsReadyToRunAsyncTasksEvent.Reset();
+	FReadScopeLock Lock(JobsLock);
+	return AsyncJobTask.IsCompleted() && JobsToExecute.IsEmpty() && !ExecutingJob.IsSet() && CompletedJobs.IsEmpty();
 }
 
-bool FPhysScene_AsyncPhysicsStateJobQueue::AreAsyncTasksExecuted() const
+void FPhysScene_AsyncPhysicsStateJobQueue::OnUpdateLevelStreamingDone()
 {
-	return !PendingJobs.IsValid() && !Pipe.HasWork();
-}
+	++GameThreadEpoch;
 
-bool FPhysScene_AsyncPhysicsStateJobQueue::IsCompleted() const
-{ 
-	return AreAsyncTasksExecuted() && PipedJobs.IsEmpty() && PipeTasks.IsEmpty();
+	LaunchAsyncJobTask();
 }
 
-void FPhysScene_AsyncPhysicsStateJobQueue::Launch()
+void FPhysScene_AsyncPhysicsStateJobQueue::LaunchAsyncJobTask()
 {
-	TRACE_CPUPROFILER_EVENT_SCOPE(FPhysScene_AsyncPhysicsStateJobQueue::Launch);
-	check(IsInGameThread());
-
-	if (!PendingJobs.IsValid() || PendingJobs->IsEmpty())
+	if (!AsyncJobTask.IsCompleted())
 	{
 		return;
 	}
 
-	PendingJobs->OnPreExecute_GameThread();
-
-	// Move Jobs to the async queue
-	FJobs& NewJobs = PipedJobs.Emplace_GetRef();
-	NewJobs = MoveTemp(PendingJobs);
-	if (!IsReadyToRunAsyncTasksEvent)
 	{
-		IsReadyToRunAsyncTasksEvent.Emplace(UE_SOURCE_LOCATION);
+		FReadScopeLock Lock(JobsLock);
+		if (JobsToExecute.IsEmpty())
+		{
+			return;
+		}
 	}
-	// Launch a new task in the pipe
-	PipeTasks.Emplace(Pipe.Launch(UE_SOURCE_LOCATION, [this, NewJobs = NewJobs.Get()]() { NewJobs->Execute(); }, *IsReadyToRunAsyncTasksEvent, UE::Tasks::ETaskPriority::BackgroundHigh));
-}
-
-void FPhysScene_AsyncPhysicsStateJobQueue::OnUpdateLevelStreamingDone()
-{
-	TriggerIsReadyToRunAsyncTasksEvent();
-	
-}
 
-void FPhysScene_AsyncPhysicsStateJobQueue::TriggerIsReadyToRunAsyncTasksEvent()
-{
-	if (IsReadyToRunAsyncTasksEvent)
-	{
-		IsReadyToRunAsyncTasksEvent->Trigger();
-		IsReadyToRunAsyncTasksEvent.Reset();
-	}
+	const double TimeBudgetSeconds = Chaos::GAsyncPhysicsStateTaskTimeBudgetMS > 0 ? Chaos::GAsyncPhysicsStateTaskTimeBudgetMS / 1000 : MAX_dbl;
+	AsyncJobTask = UE::Tasks::Launch(UE_SOURCE_LOCATION, [this, TimeBudgetSeconds]() { ExecuteJobsAsync(TimeBudgetSeconds); }, UE::Tasks::ETaskPriority::BackgroundHigh);
 }
 
-void FPhysScene_AsyncPhysicsStateJobQueue::WaitForAsyncTasksExecution()
+void FPhysScene_AsyncPhysicsStateJobQueue::ExecuteJobsAsync(double TimeBudgetSeconds)
 {
-	TRACE_CPUPROFILER_EVENT_SCOPE(FPhysScene_AsyncTaskQueue::WaitForAsyncTasksExecution);
-	check(IsInGameThread());
+	TRACE_CPUPROFILER_EVENT_SCOPE(FPhysScene_AsyncPhysicsStateJobQueue::ExecuteJobsAsync);
+	UE::FTimeout Timeout(TimeBudgetSeconds - UsedAsyncTaskTimeBudgetSec);
+	bool bWorkRemaining = false;
 
-	TriggerIsReadyToRunAsyncTasksEvent();
-	UE::Tasks::Wait(PipeTasks);
-}
-
-void FPhysScene_AsyncPhysicsStateJobQueue::Tick(bool bWaitForCompletion)
-{
-	TRACE_CPUPROFILER_EVENT_SCOPE(FPhysScene_AsyncPhysicsStateJobQueue::Tick);
-	check(IsInGameThread());
-
-	// Flush all queued jobs
-	Launch();
-
-	// Wait for tasks if the world is inside a blocking load
-	if (bWaitForCompletion || 
-		PhysScene->GetOwningWorld()->GetIsInBlockTillLevelStreamingCompleted() ||
-		PhysScene->GetOwningWorld()->GetShouldForceUnloadStreamingLevels() ||
-		PhysScene->GetOwningWorld()->IsBeingCleanedUp())
+	int32 CompletedJobsCount = 0;
+	int32 Index = 0;
+	do
 	{
-		WaitForAsyncTasksExecution();
-	}
+		if (bIsBlocking)
+		{
+			Timeout = UE::FTimeout::Never();
+		}
+		else if (TaskEpoch != GameThreadEpoch)
+		{
+			// Reset Timeout
+			TaskEpoch = GameThreadEpoch;
+			Timeout = UE::FTimeout(TimeBudgetSeconds);
+			UsedAsyncTaskTimeBudgetSec = 0;
+		}
+		
+		{
+			FWriteScopeLock Lock(JobsLock);
+			if (!JobsToExecute.IsValidIndex(Index))
+			{
+				break;
+			}
+			check(!ExecutingJob.IsSet() || (ExecutingJob == JobsToExecute[Index]));
+			ExecutingJob = JobsToExecute[Index++];
+		}
 
-	// Process completed tasks
-	check(PipeTasks.Num() == PipedJobs.Num());
-	for (int i = 0; i < PipeTasks.Num();)
-	{
-		if (PipeTasks[i].IsCompleted())
+		if (ExecutingJob->Execute(Timeout))
 		{
-			PipedJobs[i]->OnPostExecute_GameThread();
-			PipeTasks.RemoveAt(i, EAllowShrinking::No);
-			PipedJobs.RemoveAt(i, EAllowShrinking::No);
+			FWriteScopeLock Lock(JobsLock);
+			CompletedJobs.Add(*ExecutingJob);
+			ExecutingJob.Reset();
+			++CompletedJobsCount;
 		}
 		else
 		{
-			++i;
+			break;
 		}
-	}
 
-#if DO_CHECK
-	if (AreAsyncTasksExecuted())
+		if (Timeout.IsExpired())
+		{
+			break;
+		}
+	} while (true);
+
+	// Remove completed jobs
 	{
-		check(PipeTasks.IsEmpty());
-		check(PipedJobs.IsEmpty());
-		check(!IsReadyToRunAsyncTasksEvent.IsSet());
+		FWriteScopeLock Lock(JobsLock);
+		JobsToExecute.RemoveAt(0, CompletedJobsCount);
 	}
-#endif
+
+	UsedAsyncTaskTimeBudgetSec += Timeout.GetElapsedSeconds();
 }
 
-void FPhysScene_AsyncPhysicsStateJobQueue::FJobSet::Execute() const
+void FPhysScene_AsyncPhysicsStateJobQueue::Tick(bool bWaitForCompletion)
 {
-	TRACE_CPUPROFILER_EVENT_SCOPE(FPhysScene_AsyncPhysicsStateJobQueue::ExecuteJobs);
-	for (const FJob& Job : Jobs)
+	TRACE_CPUPROFILER_EVENT_SCOPE(FPhysScene_AsyncPhysicsStateJobQueue::Tick);
+	check(IsInGameThread());
+	if (IsCompleted())
 	{
-		Job.Execute();
+		return;
 	}
-}
 
-void FPhysScene_AsyncPhysicsStateJobQueue::FJobSet::OnPreExecute_GameThread() const
-{
-	// Pre-pass to make sure physics meshes are created
+	const bool bShouldWaitForCompletion = bWaitForCompletion ||
+		PhysScene->GetOwningWorld()->GetIsInBlockTillLevelStreamingCompleted() ||
+		PhysScene->GetOwningWorld()->GetShouldForceUnloadStreamingLevels() ||
+		PhysScene->GetOwningWorld()->IsBeingCleanedUp();
+
+	// Wait for tasks if the world is inside a blocking load
+	if (bShouldWaitForCompletion)
 	{
-		TSet<UBodySetup*> UniqueBodySetups;
-		for (const FJob& Job : Jobs)
+		// Tell the async task that there is no time limit anymore
+		bIsBlocking = true;
+		do
 		{
-			if (Job.Type == EJobType::CreatePhysicsState)
+			LaunchAsyncJobTask();
+			AsyncJobTask.Wait();
 			{
-				Job.Processor->CollectBodySetupsWithPhysicsMeshesToCreate(UniqueBodySetups);
+				FReadScopeLock Lock(JobsLock);
+				if (JobsToExecute.IsEmpty())
+				{
+					break;
+				}
 			}
-		}
-		if (UniqueBodySetups.Num())
-		{
-			TArray<UBodySetup*> BodySetups = UniqueBodySetups.Array();
-			ParallelFor(TEXT("CreatePhysicsMeshes.PF"), BodySetups.Num(), 1, [this, &BodySetups](int32 Index) { BodySetups[Index]->CreatePhysicsMeshes(); });
-		}
+		} while (true);
+		bIsBlocking = false;
 	}
 
-	for (const FJob& Job : Jobs)
+	TArray<FJob> CompletedJobsCopy;
 	{
-		Job.OnPreExecute_GameThread();
+		FWriteScopeLock Lock(JobsLock);
+		CompletedJobsCopy = MoveTemp(CompletedJobs);
+	}
+
+	if (!CompletedJobsCopy.IsEmpty())
+	{
+		TRACE_CPUPROFILER_EVENT_SCOPE(FPhysScene_AsyncPhysicsStateJobQueue::Tick_OnPostExecute_GameThread);
+		for (const FJob& Job : CompletedJobsCopy)
+		{
+			Job.OnPostExecute_GameThread();
+		}
 	}
 }
 
@@ -214,35 +250,27 @@ void FPhysScene_AsyncPhysicsStateJobQueue::FJob::OnPreExecute_GameThread() const
 	}
 }
 
-void FPhysScene_AsyncPhysicsStateJobQueue::FJobSet::OnPostExecute_GameThread() const
-{
-	TRACE_CPUPROFILER_EVENT_SCOPE(FPhysScene_AsyncPhysicsStateJobQueue::OnPostExecute_GameThread);
-	for (const FJob& Job : Jobs)
-	{
-		Job.OnPostExecute_GameThread();
-	}
-}
-
-void FPhysScene_AsyncPhysicsStateJobQueue::FJob::Execute() const
+bool FPhysScene_AsyncPhysicsStateJobQueue::FJob::Execute(UE::FTimeout& Timeout) const
 {
 	if (IsValid())
 	{
 		switch (Type)
 		{
-			case EJobType::CreatePhysicsState:
-			{
-				TRACE_CPUPROFILER_EVENT_SCOPE(FPhysScene_AsyncPhysicsStateJobQueue::CreatePhysicsState);
-				Processor->OnAsyncCreatePhysicsState();
-			}
-			break;
-			case EJobType::DestroyPhysicsState:
-			{
-				TRACE_CPUPROFILER_EVENT_SCOPE(FPhysScene_AsyncPhysicsStateJobQueue::DestroyPhysicsState);
-				Processor->OnAsyncDestroyPhysicsState();
-			}
-			break;
+		case EJobType::CreatePhysicsState:
+		{
+			TRACE_CPUPROFILER_EVENT_SCOPE(FPhysScene_AsyncPhysicsStateJobQueue::CreatePhysicsState);
+			return Processor->OnAsyncCreatePhysicsState(Timeout);
+		}
+		break;
+		case EJobType::DestroyPhysicsState:
+		{
+			TRACE_CPUPROFILER_EVENT_SCOPE(FPhysScene_AsyncPhysicsStateJobQueue::DestroyPhysicsState);
+			return Processor->OnAsyncDestroyPhysicsState(Timeout);
+		}
+		break;
 		}
 	}
+	return true;
 }
 
 void FPhysScene_AsyncPhysicsStateJobQueue::FJob::OnPostExecute_GameThread() const
diff --git a/Engine/Source/Runtime/Engine/Private/PhysicsEngine/Experimental/PhysScene_Chaos.cpp b/Engine/Source/Runtime/Engine/Private/PhysicsEngine/Experimental/PhysScene_Chaos.cpp
--- a/Engine/Source/Runtime/Engine/Private/PhysicsEngine/Experimental/PhysScene_Chaos.cpp
+++ b/Engine/Source/Runtime/Engine/Private/PhysicsEngine/Experimental/PhysScene_Chaos.cpp
@@ -1881,10 +1881,6 @@ void FPhysScene_Chaos::ProcessAsyncPhysicsStateJobs(bool bWaitForCompletion)
 	if (AsyncPhysicsStateJobQueue.IsValid())
 	{
 		AsyncPhysicsStateJobQueue->Tick(bWaitForCompletion);
-		if (AsyncPhysicsStateJobQueue->IsCompleted())
-		{
-			AsyncPhysicsStateJobQueue.Reset();
-		}
 	}
 }
 
diff --git a/Engine/Source/Runtime/Engine/Public/Physics/Experimental/AsyncPhysicsStateProcessorInterface.h b/Engine/Source/Runtime/Engine/Public/Physics/Experimental/AsyncPhysicsStateProcessorInterface.h
--- a/Engine/Source/Runtime/Engine/Public/Physics/Experimental/AsyncPhysicsStateProcessorInterface.h
+++ b/Engine/Source/Runtime/Engine/Public/Physics/Experimental/AsyncPhysicsStateProcessorInterface.h
@@ -3,6 +3,7 @@
 #pragma once
 
 #include "Containers/Set.h"
+#include "Misc/Timeout.h"
 
 // Interface used by FPhysScene_AsyncPhysicsStateJobQueue
 
@@ -28,7 +29,7 @@ public:
 	virtual void CollectBodySetupsWithPhysicsMeshesToCreate(TSet<UBodySetup*>& OutBodySetups) const {}
 
 	/** Used to create any physics engine information for this component outside of the GameThread. */
-	virtual void OnAsyncCreatePhysicsState() {}
+	virtual bool OnAsyncCreatePhysicsState(const UE::FTimeout& TimeOut) { return true; }
 
 	/** Called on the GameThread once the component's physic engine information is created. */
 	virtual void OnAsyncCreatePhysicsStateEnd_GameThread() {}
@@ -37,7 +38,7 @@ public:
 	virtual void OnAsyncDestroyPhysicsStateBegin_GameThread() {}
 
 	/** Used to destroy any physics engine information for this component outside of the GameThread. */
-	virtual void OnAsyncDestroyPhysicsState() {}
+	virtual bool OnAsyncDestroyPhysicsState(const UE::FTimeout& TimeOut) { return true; }
 
 	/** Called on the GameThread once the component's physic engine information is destroyed. */
 	virtual void OnAsyncDestroyPhysicsStateEnd_GameThread() {}
diff --git a/Engine/Source/Runtime/Engine/Public/Physics/Experimental/PhysScene_AsyncPhysicsStateJobQueue.h b/Engine/Source/Runtime/Engine/Public/Physics/Experimental/PhysScene_AsyncPhysicsStateJobQueue.h
--- a/Engine/Source/Runtime/Engine/Public/Physics/Experimental/PhysScene_AsyncPhysicsStateJobQueue.h
+++ b/Engine/Source/Runtime/Engine/Public/Physics/Experimental/PhysScene_AsyncPhysicsStateJobQueue.h
@@ -31,7 +31,8 @@ struct FPhysScene_AsyncPhysicsStateJobQueue
 		{
 			return Processor && ::IsValid(Processor->GetAsyncPhysicsStateObject());
 		}
-		void Execute() const;
+
+		bool Execute(UE::FTimeout& Timeout) const;
 		void OnPreExecute_GameThread() const;
 		void OnPostExecute_GameThread() const;
 
@@ -49,31 +50,6 @@ struct FPhysScene_AsyncPhysicsStateJobQueue
 		EJobType Type;
 	};
 
-	struct FJobSet
-	{
-		void Add(const FJob& Job)
-		{
-			Jobs.Add(Job);
-		}
-
-		int32 Remove(const FJob& Job)
-		{
-			return Jobs.Remove(Job);
-		}
-
-		bool IsEmpty() const
-		{
-			return Jobs.IsEmpty();
-		}
-
-		void OnPreExecute_GameThread() const;
-		void Execute() const;
-		void OnPostExecute_GameThread() const;
-
-	private:
-		TSet<FJob> Jobs;
-	};
-
 	void Tick(bool bWaitForCompletion = false);
 	void AddJob(const FJob& Job);
 	void RemoveJob(const FJob& Job);
@@ -81,18 +57,24 @@ struct FPhysScene_AsyncPhysicsStateJobQueue
 
 private:
 
-	void Launch();
-	void WaitForAsyncTasksExecution();
-	bool AreAsyncTasksExecuted() const;
-	void OnAsyncTasksExecuted();
+	void LaunchAsyncJobTask();
+	void ExecuteJobsAsync(double TimeBudgetSeconds);
 	void OnUpdateLevelStreamingDone();
-	void TriggerIsReadyToRunAsyncTasksEvent();
-
-	FPhysScene* PhysScene;
-	using FJobs = TUniquePtr<FJobSet>;
-	FJobs PendingJobs;
-	TArray<FJobs> PipedJobs;
-	TArray<UE::Tasks::FTask> PipeTasks;
-	TOptional<UE::Tasks::FTaskEvent> IsReadyToRunAsyncTasksEvent;
-	UE::Tasks::FPipe Pipe{ TEXT("FPhysScene::AsyncCreatePhysicStateQueue") };
+
+	// GameThread variables
+	FPhysScene* PhysScene = nullptr; // The Physics scene
+	UE::Tasks::FTask AsyncJobTask; // The async task
+
+	// Async Task variables
+	int32 TaskEpoch = 0; // Last epoch used by the async task
+	double UsedAsyncTaskTimeBudgetSec = 0; // Time used in the async task since the last epoch update
+
+	// Variables access/modified by both the GameThread and the async task
+	mutable FRWLock JobsLock; // Lock used to protect jobs to execute, executing and completed
+	TArray<FJob> JobsToExecute;
+	TOptional<FJob> ExecutingJob;
+	TArray<FJob> CompletedJobs;
+
+	std::atomic<bool> bIsBlocking = false; // Notifies to the async task that we are block waiting for the task to complete
+	std::atomic<int32> GameThreadEpoch = 0; // Epoch updated by the GameThread at every frame and used by the async task to reset UsedAsyncTaskTimeBudgetSec
 };
\ No newline at end of file
