Mobile: fixed issues with tangents buffer format when deformers are used on a skeletal mesh

diff --git a/Engine/Source/Runtime/Engine/Private/GPUSkinVertexFactory.cpp b/Engine/Source/Runtime/Engine/Private/GPUSkinVertexFactory.cpp
--- a/Engine/Source/Runtime/Engine/Private/GPUSkinVertexFactory.cpp
+++ b/Engine/Source/Runtime/Engine/Private/GPUSkinVertexFactory.cpp
@@ -1412,14 +1412,16 @@ void FGPUSkinPassthroughVertexFactory::InitRHI(FRHICommandListBase& RHICmdList)
 
 	if (Data.TangentsSRV)
 	{
-		TangentFormat = Data.TangentsSRV->GetDesc().Common.Format;
 		PreSkinnedTangentSRV = Data.TangentsSRV;
-		
-		if (IsOpenGLPlatform(GMaxRHIShaderPlatform))
-		{
-			// See SkinCache TangentFormat
-			TangentFormat = PF_R16G16B16A16_SNORM;
-		}
+	}
+
+	// TODO: SkinCache can work with PF_R8G8B8A8_SNORM, however Deformers expect PF_R16G16B16A16_SNORM
+	// User need to pass an expected format into VF
+	TangentFormat = PF_R16G16B16A16_SNORM;
+	if (IsOpenGLPlatform(GMaxRHIShaderPlatform))
+	{
+		// See SkinCache TangentFormat
+		TangentFormat = PF_R16G16B16A16_SINT;
 	}
 
 	if (EnumHasAnyFlags(VertexAttributesRequested, EVertexAttributeFlags::Tangent) && !IsManualVertexFetch(Data.TangentBasisComponents[0]))
