Chaos - add FBodyInstance::bAllowPartialSleep flag to allow/disallow partial island sleep for the island a specific particle is in.
- flag is also used by GeometryCollections
- wire flag through to Chaos particles in the ControlFlags structure

NOTE: this flag isn't actually being used yet. It will be connected to partial island sleep in an upcoming change.

diff --git a/Engine/Source/Programs/HeadlessChaos/Private/GeometryCollection/GeometryCollectionTestFramework.cpp b/Engine/Source/Programs/HeadlessChaos/Private/GeometryCollection/GeometryCollectionTestFramework.cpp
--- a/Engine/Source/Programs/HeadlessChaos/Private/GeometryCollection/GeometryCollectionTestFramework.cpp
+++ b/Engine/Source/Programs/HeadlessChaos/Private/GeometryCollection/GeometryCollectionTestFramework.cpp
@@ -177,6 +177,7 @@ namespace GeometryCollectionTest
 			SimulationParams.AngularDamping = 0;
 			SimulationParams.UseCCD = false;
 			SimulationParams.UseMACD = false;
+			SimulationParams.AllowPartialIslandSleep = true;
 			SimulationParams.PositionSolverIterations = 8;
 			SimulationParams.VelocitySolverIterations = 1;
 			SimulationParams.ProjectionSolverIterations = 1;
diff --git a/Engine/Source/Runtime/Engine/Classes/Components/PrimitiveComponent.h b/Engine/Source/Runtime/Engine/Classes/Components/PrimitiveComponent.h
--- a/Engine/Source/Runtime/Engine/Classes/Components/PrimitiveComponent.h
+++ b/Engine/Source/Runtime/Engine/Classes/Components/PrimitiveComponent.h
@@ -2820,7 +2820,19 @@ public:
 
 	/** [EXPERIMENTAL] Set whether all bodies in this component should use Motion-Aware Collision Detection */
 	UFUNCTION(BlueprintCallable, Category = "Physics")
-	ENGINE_API virtual void SetAllUseMACD(bool InUseMACD);
+	ENGINE_API virtual void SetAllUseMACD(bool InUseMACD);	
+
+	/** [EXPERIMENTAL] Set whether this component allows Partial Island Sleeping for the island it is in.
+	Note that Partial Island Sleeping is a CPU optimization disabling non-moving bodies in the solver. 
+	This feature is useful for scenes with many rigid bodies but can affect the physics behavior.*/
+	UFUNCTION(BlueprintCallable, Category = "Physics")
+	ENGINE_API virtual void SetAllowPartialIslandSleep(bool InAllowPartialIslandSleep = true, FName BoneName = NAME_None);
+
+	/** [EXPERIMENTAL] Set whether all bodies in this component allow Partial Island Sleeping for the island they are in.
+	Note that Partial Island Sleeping is a CPU optimization disabling non-moving bodies in the solver.
+	This feature is useful for scenes with many rigid bodies but can affect the physics behavior. */
+	UFUNCTION(BlueprintCallable, Category = "Physics")
+	ENGINE_API virtual void SetAllAllowPartialIslandSleep(bool InAllowPartialIslandSleep = true);
 
 	/**
 	 *	Force all bodies in this component to sleep.
diff --git a/Engine/Source/Runtime/Engine/Classes/PhysicsEngine/BodyInstance.h b/Engine/Source/Runtime/Engine/Classes/PhysicsEngine/BodyInstance.h
--- a/Engine/Source/Runtime/Engine/Classes/PhysicsEngine/BodyInstance.h
+++ b/Engine/Source/Runtime/Engine/Classes/PhysicsEngine/BodyInstance.h
@@ -396,6 +396,10 @@ private:
 	UPROPERTY(EditAnywhere, AdvancedDisplay, Category = Collision)
 	uint8 bUseMACD : 1;
 
+	/** [EXPERIMENTAL] If true Partial Island Sleeping is permitted for the entire island this body is in */
+	UPROPERTY(EditAnywhere, AdvancedDisplay, Category = Physics)
+	uint8 bAllowPartialIslandSleep : 1;
+
 public:
 	/** If true ignore analytic collisions and treat objects as a general implicit surface */
 	UPROPERTY(EditAnywhere, AdvancedDisplay, BlueprintReadOnly, Category = Collision)
@@ -579,6 +583,8 @@ public:
 
 	bool IsUsingMACD() const { return bUseMACD; }
 
+	bool IsPartialIslandSleepAllowed() const { return bAllowPartialIslandSleep; }
+
 	/** Sets the mass override */
 	ENGINE_API void SetMassOverride(float MassInKG, bool bNewOverrideMass = true);
 
@@ -1035,6 +1041,14 @@ public:
 	 */
 	ENGINE_API void SetUseMACD(bool bInUseMACD);
 
+	/**
+	 * [EXPERIMENTAL] Enable/disable Partial Island Sleeping for the entire island this body is in.
+	 * Partial Island Sleeping permits the solver to freeze the body's dynamics ("sleep") if it has not moved significantly
+	 * in the previous simulation frames even if connected bodies are still moving. The body will be unfrozen ("wake up")
+	 * due to several possible triggers such as the detection of new collisions, creation of new constraints, and user interactions.
+	 */
+	ENGINE_API void SetAllowPartialIslandSleep(bool bInAllowPartialIslandSleep);
+
 	ENGINE_API void SetPositionSolverIterationCount(uint8 PositionSolverIterationCountIn);
 
 	ENGINE_API void SetVelocitySolverIterationCount(uint8 VelocitySolverIterationCountIn);
@@ -1044,6 +1058,9 @@ public:
 	/** [EXPERIMENTAL] Whether Motion-Aware Collision Detection is enabled */
 	bool GetUseMACD() const { return bUseMACD != 0; }
 
+	/** [EXPERIMENTAL] Whether Partial Island Sleep is permitted for the island this body is in */
+	bool GetAllowPartialIslandSleep() const { return bAllowPartialIslandSleep != 0; }
+
 	uint8 GetPositionSolverIterationCount() const { return PositionSolverIterationCount; }
 	uint8 GetVelocitySolverIterationCount() const { return VelocitySolverIterationCount; }
 	uint8 GetProjectionSolverIterationCount() const { return ProjectionSolverIterationCount; }
diff --git a/Engine/Source/Runtime/Engine/Private/PhysicsEngine/BodyInstance.cpp b/Engine/Source/Runtime/Engine/Private/PhysicsEngine/BodyInstance.cpp
--- a/Engine/Source/Runtime/Engine/Private/PhysicsEngine/BodyInstance.cpp
+++ b/Engine/Source/Runtime/Engine/Private/PhysicsEngine/BodyInstance.cpp
@@ -370,6 +370,7 @@ FBodyInstance::FBodyInstance()
 	, bShapeCollisionEnabledIsSet(false)
 	, bShapeCollisionResponsesIsSet(false)
 	, bUseMACD(false)
+	, bAllowPartialIslandSleep(true)
 	, bIgnoreAnalyticCollisions(false)
 	, bNotifyRigidBodyCollision(false)
 	, bContactModification(false)
@@ -1305,6 +1306,7 @@ void TInitBodiesHelperBase<InAllocatorType>::CreateActor_AssumesLocked(FBodyInst
 	{
 		FPhysicsInterface::SetCcdEnabled_AssumesLocked(ActorHandle, Instance->bUseCCD);
 		FPhysicsInterface::SetMACDEnabled_AssumesLocked(ActorHandle, Instance->IsUsingMACD());
+		FPhysicsInterface::SetPartialIslandSleepAllowed_AssumesLocked(ActorHandle, Instance->IsPartialIslandSleepAllowed());
 		FPhysicsInterface::SetIsKinematic_AssumesLocked(ActorHandle, !Instance->ShouldInstanceSimulatingPhysics());
 
 		FPhysicsInterface::SetMaxLinearVelocity_AssumesLocked(ActorHandle, TNumericLimits<float>::Max());
@@ -2555,6 +2557,7 @@ void FBodyInstance::UpdateInstanceSimulatePhysics()
 		FPhysicsInterface::SetIsKinematic_AssumesLocked(Actor, !bUseSimulate);
 		FPhysicsInterface::SetCcdEnabled_AssumesLocked(Actor, bUseCCD);
 		FPhysicsInterface::SetMACDEnabled_AssumesLocked(Actor, bUseMACD);
+		FPhysicsInterface::SetPartialIslandSleepAllowed_AssumesLocked(Actor, bAllowPartialIslandSleep);
 		FPhysicsInterface::SetOneWayInteraction_AssumesLocked(Actor, bOneWayInteraction);
 		FPhysicsInterface::SetPositionSolverIterationCount_AssumesLocked(Actor, PositionSolverIterationCount);
 		FPhysicsInterface::SetVelocitySolverIterationCount_AssumesLocked(Actor, VelocitySolverIterationCount);
@@ -4001,6 +4004,22 @@ void FBodyInstance::SetUseMACD(bool bInUseMACD)
 	}
 }
 
+void FBodyInstance::SetAllowPartialIslandSleep(bool bInAllowPartialIslandSleep)
+{
+	if (bAllowPartialIslandSleep != bInAllowPartialIslandSleep)
+	{
+		bAllowPartialIslandSleep = bInAllowPartialIslandSleep;
+
+		FPhysicsCommand::ExecuteWrite(GetPhysicsActor(), [this, bInAllowPartialIslandSleep](const FPhysicsActorHandle& Actor)
+			{
+				if (FPhysicsInterface::IsValid(Actor) && FPhysicsInterface::IsRigidBody(Actor))
+				{
+					FPhysicsInterface::SetPartialIslandSleepAllowed_AssumesLocked(Actor, bInAllowPartialIslandSleep);
+				}
+			});
+	}
+}
+
 void FBodyInstance::SetPositionSolverIterationCount(uint8 PositionSolverIterationCountIn)
 {
 	if (PositionSolverIterationCount != PositionSolverIterationCountIn)
diff --git a/Engine/Source/Runtime/Engine/Private/PrimitiveComponentPhysics.cpp b/Engine/Source/Runtime/Engine/Private/PrimitiveComponentPhysics.cpp
--- a/Engine/Source/Runtime/Engine/Private/PrimitiveComponentPhysics.cpp
+++ b/Engine/Source/Runtime/Engine/Private/PrimitiveComponentPhysics.cpp
@@ -797,6 +797,20 @@ void UPrimitiveComponent::SetAllUseMACD(bool bInUseMACD)
 	SetUseMACD(bInUseMACD, NAME_None);
 }
 
+void UPrimitiveComponent::SetAllowPartialIslandSleep(bool bInAllowPartialIslandSleep /* = true */, FName BoneName /* = NAME_None */)
+{
+	FBodyInstance* BI = GetBodyInstance(BoneName);
+	if (BI)
+	{
+		BI->SetAllowPartialIslandSleep(bInAllowPartialIslandSleep);
+	}
+}
+
+void UPrimitiveComponent::SetAllAllowPartialIslandSleep(bool bInAllowPartialIslandSleep /* = true */)
+{
+	SetAllowPartialIslandSleep(bInAllowPartialIslandSleep, NAME_None);
+}
+
 void UPrimitiveComponent::PutRigidBodyToSleep(FName BoneName)
 {
 	FBodyInstance* BI = GetBodyInstance(BoneName);
diff --git a/Engine/Source/Runtime/Experimental/Chaos/Private/PhysicsProxy/GeometryCollectionPhysicsProxy.cpp b/Engine/Source/Runtime/Experimental/Chaos/Private/PhysicsProxy/GeometryCollectionPhysicsProxy.cpp
--- a/Engine/Source/Runtime/Experimental/Chaos/Private/PhysicsProxy/GeometryCollectionPhysicsProxy.cpp
+++ b/Engine/Source/Runtime/Experimental/Chaos/Private/PhysicsProxy/GeometryCollectionPhysicsProxy.cpp
@@ -1693,6 +1693,7 @@ void FGeometryCollectionPhysicsProxy::InitializeBodiesPT(Chaos::FPBDRigidsSolver
 				Handle->SetGravityGroupIndex(Parameters.GravityGroupIndex);
 				Handle->SetCCDEnabled(Parameters.UseCCD);
 				Handle->SetMACDEnabled(Parameters.UseMACD);
+				Handle->SetPartialIslandSleepAllowed(Parameters.AllowPartialIslandSleep);
 				Handle->SetIterationSettings(Chaos::Private::FIterationSettings(Parameters.PositionSolverIterations, Parameters.VelocitySolverIterations, Parameters.ProjectionSolverIterations));
 				Handle->SetOneWayInteraction(bIsOneWayInteraction);
 				Handle->SetInertiaConditioningEnabled(Parameters.UseInertiaConditioning);
diff --git a/Engine/Source/Runtime/Experimental/Chaos/Public/Chaos/ParticleDirtyFlags.h b/Engine/Source/Runtime/Experimental/Chaos/Public/Chaos/ParticleDirtyFlags.h
--- a/Engine/Source/Runtime/Experimental/Chaos/Public/Chaos/ParticleDirtyFlags.h
+++ b/Engine/Source/Runtime/Experimental/Chaos/Public/Chaos/ParticleDirtyFlags.h
@@ -549,6 +549,9 @@ public:
 	bool MACDEnabled() const { return MControlFlags.GetMACDEnabled(); }
 	void SetMACDEnabled(bool bInMACDEnabled) { MControlFlags.SetMACDEnabled(bInMACDEnabled); }
 
+	bool PartialIslandSleepAllowed() const { return MControlFlags.GetPartialIslandSleepAllowed(); }
+	void SetPartialIslandSleepAllowed(bool bInPartialIslandSleepAllowed) { MControlFlags.SetPartialIslandSleepAllowed(bInPartialIslandSleepAllowed); }
+
 	void SetPositionSolverIterationCount(const int32 PositionSolverIterationCountIn) { MIterationSettings.SetNumPositionIterations(PositionSolverIterationCountIn); }
 	void SetVelocitySolverIterationCount(const int32 VelocitySolverIterationCountIn) { MIterationSettings.SetNumVelocityIterations(VelocitySolverIterationCountIn); }
 	void SetProjectionSolverIterationCount(const int32 ProjectionSolverIterationCountIn) { MIterationSettings.SetNumProjectionIterations(ProjectionSolverIterationCountIn); }
diff --git a/Engine/Source/Runtime/Experimental/Chaos/Public/Chaos/ParticleHandle.h b/Engine/Source/Runtime/Experimental/Chaos/Public/Chaos/ParticleHandle.h
--- a/Engine/Source/Runtime/Experimental/Chaos/Public/Chaos/ParticleHandle.h
+++ b/Engine/Source/Runtime/Experimental/Chaos/Public/Chaos/ParticleHandle.h
@@ -129,6 +129,7 @@ void PBDRigidParticleDefaultConstruct(FConcrete& Concrete, const FPBDRigidPartic
 	Concrete.SetGravityGroupIndex(0);
 	Concrete.SetCCDEnabled(Params.bCCDEnabled);
 	Concrete.SetMACDEnabled(false);
+	Concrete.SetPartialIslandSleepAllowed(true);
 	Concrete.SetIterationSettings(Private::FIterationSettings(INDEX_NONE, INDEX_NONE, INDEX_NONE));
 	Concrete.SetDisabled(Params.bDisabled);
 	Concrete.SetSleepType(ESleepType::MaterialSleep);
@@ -1321,6 +1322,16 @@ public:
 		PBDRigidParticles->ControlFlags(ParticleIdx).SetMACDEnabled(bEnabled);
 	}
 
+	inline bool PartialIslandSleepAllowed() const
+	{
+		return ControlFlags().GetPartialIslandSleepAllowed();
+	}
+
+	inline void SetPartialIslandSleepAllowed(bool bEnabled)
+	{
+		PBDRigidParticles->ControlFlags(ParticleIdx).SetPartialIslandSleepAllowed(bEnabled);
+	}
+
 	inline Private::FIterationSettings IterationSettings() const
 	{
 		return PBDRigidParticles->ParticleIterationCounts(ParticleIdx);
@@ -1884,6 +1895,16 @@ public:
 		return false;
 	}
 
+	bool PartialIslandSleepAllowed() const
+	{
+		if (MHandle->CastToRigidParticle())
+		{
+			return MHandle->CastToRigidParticle()->PartialIslandSleepAllowed();
+		}
+
+		return false;
+	}
+
 	Private::FIterationSettings IterationSettings() const
 	{
 		if (MHandle->CastToRigidParticle())
@@ -3458,6 +3479,12 @@ public:
 		MMiscData.Modify(true, MDirtyFlags, Proxy, [bInEnabled](auto& Data) { Data.SetMACDEnabled(bInEnabled); });
 	}
 
+	bool PartialIslandSleepAllowed() const { return MMiscData.Read().PartialIslandSleepAllowed(); }
+	void SetPartialIslandSleepAllowed(bool bInEnabled)
+	{
+		MMiscData.Modify(true, MDirtyFlags, Proxy, [bInEnabled](auto& Data) { Data.SetPartialIslandSleepAllowed(bInEnabled); });
+	}
+
 	Private::FIterationSettings IterationSettings() const { return MMiscData.Read().IterationSettings(); }
 
 	void SetIterationSettings(const Private::FIterationSettings& SolverIterationSettingsIn)
diff --git a/Engine/Source/Runtime/Experimental/Chaos/Public/Chaos/RigidParticleControlFlags.h b/Engine/Source/Runtime/Experimental/Chaos/Public/Chaos/RigidParticleControlFlags.h
--- a/Engine/Source/Runtime/Experimental/Chaos/Public/Chaos/RigidParticleControlFlags.h
+++ b/Engine/Source/Runtime/Experimental/Chaos/Public/Chaos/RigidParticleControlFlags.h
@@ -50,6 +50,9 @@ namespace Chaos
 		bool GetMACDEnabled() const { return Flags.bMACDEnabled; }
 		FRigidParticleControlFlags& SetMACDEnabled(const bool bEnabled) { Flags.bMACDEnabled = bEnabled; return *this; }
 
+		bool GetPartialIslandSleepAllowed() const { return Flags.bPartialIslandSleepAllowed; }
+		FRigidParticleControlFlags& SetPartialIslandSleepAllowed(const bool bEnabled) { Flags.bPartialIslandSleepAllowed = bEnabled; return *this; }
+
 		bool GetOneWayInteractionEnabled() const { return Flags.bOneWayInteractionEnabled; }
 		FRigidParticleControlFlags& SetOneWayInteractionEnabled(const bool bEnabled) { Flags.bOneWayInteractionEnabled = bEnabled; return *this; }
 
@@ -84,6 +87,7 @@ namespace Chaos
 			FStorage GravityGroupIndex : 3;
 			FStorage bUpdateKinematicFromSimulation : 1;
 			FStorage bMACDEnabled : 1;
+			FStorage bPartialIslandSleepAllowed : 1;
 			FStorage bGyroscopicTorqueEnabled : 1;
 			// Add new properties above this line
 			// Change FStorage typedef if we exceed the max bits
diff --git a/Engine/Source/Runtime/Experimental/Chaos/Public/GeometryCollection/GeometryCollectionSimulationCoreTypes.h b/Engine/Source/Runtime/Experimental/Chaos/Public/GeometryCollection/GeometryCollectionSimulationCoreTypes.h
--- a/Engine/Source/Runtime/Experimental/Chaos/Public/GeometryCollection/GeometryCollectionSimulationCoreTypes.h
+++ b/Engine/Source/Runtime/Experimental/Chaos/Public/GeometryCollection/GeometryCollectionSimulationCoreTypes.h
@@ -212,6 +212,7 @@ PRAGMA_ENABLE_DEPRECATION_WARNINGS
 		, UseInertiaConditioning(true)
 		, UseCCD(false)
 		, UseMACD(false)
+		, AllowPartialIslandSleep(true)
 		, bEnableStrainOnCollision(true)
 		, bUseStaticMeshCollisionForTraces(false)
 		, bOptimizeConvexes(true)
@@ -290,6 +291,7 @@ PRAGMA_ENABLE_DEPRECATION_WARNINGS
 		, UseInertiaConditioning(Other.UseInertiaConditioning)
 		, UseCCD(Other.UseCCD)
 		, UseMACD(Other.UseMACD)
+		, AllowPartialIslandSleep(Other.AllowPartialIslandSleep)
 		, bEnableStrainOnCollision(Other.bEnableStrainOnCollision)
 		, bUseStaticMeshCollisionForTraces(Other.bUseStaticMeshCollisionForTraces)
 		, bOptimizeConvexes(Other.bOptimizeConvexes)
@@ -402,6 +404,7 @@ PRAGMA_ENABLE_DEPRECATION_WARNINGS
 	bool UseInertiaConditioning : 1;
 	bool UseCCD : 1;
 	bool UseMACD : 1;
+	bool AllowPartialIslandSleep : 1;
 	bool bEnableStrainOnCollision : 1;
 	bool bUseStaticMeshCollisionForTraces : 1;
 	bool bOptimizeConvexes : 1;
diff --git a/Engine/Source/Runtime/Experimental/Chaos/Public/PhysicsProxy/SingleParticlePhysicsProxy.h b/Engine/Source/Runtime/Experimental/Chaos/Public/PhysicsProxy/SingleParticlePhysicsProxy.h
--- a/Engine/Source/Runtime/Experimental/Chaos/Public/PhysicsProxy/SingleParticlePhysicsProxy.h
+++ b/Engine/Source/Runtime/Experimental/Chaos/Public/PhysicsProxy/SingleParticlePhysicsProxy.h
@@ -520,13 +520,37 @@ public:
 			});
 	}
 
-	void SetMACDEnabled(const bool InCCDEnabled)
+	void SetMACDEnabled(const bool InMACDEnabled)
 	{
-		Write([InCCDEnabled](auto* Particle)
+		Write([InMACDEnabled](auto* Particle)
+			{
+				if (auto Rigid = Particle->CastToRigidParticle())
+				{
+					return Rigid->SetMACDEnabled(InMACDEnabled);
+				}
+			});
+	}
+
+	bool PartialIslandSleepAllowed() const
+	{
+		return Read([](auto* Particle)
+			{
+				if (auto Rigid = Particle->CastToRigidParticle())
+				{
+					return Rigid->PartialIslandSleepAllowed();
+				}
+
+				return false;
+			});
+	}
+
+	void SetPartialIslandSleepAllowed(const bool InPartialIslandSleepAllowed)
+	{
+		Write([InPartialIslandSleepAllowed](auto* Particle)
 			{
 				if (auto Rigid = Particle->CastToRigidParticle())
 				{
-					return Rigid->SetMACDEnabled(InCCDEnabled);
+					return Rigid->SetPartialIslandSleepAllowed(InPartialIslandSleepAllowed);
 				}
 			});
 	}
@@ -1386,6 +1410,26 @@ public:
 		return false;
 	}
 
+	void SetPartialIslandSleepAllowed(bool bEnabled)
+	{
+		VerifyContext();
+		if (TPBDRigidParticle<FReal, 3>*Rigid = GetParticle_LowLevel()->CastToRigidParticle())
+		{
+			Rigid->SetPartialIslandSleepAllowed(bEnabled);
+		}
+	}
+
+	bool PartialIslandSleepAllowed() const
+	{
+		VerifyContext();
+		if (const TPBDRigidParticle<FReal, 3>*Rigid = GetParticle_LowLevel()->CastToRigidParticle())
+		{
+			return Rigid->PartialIslandSleepAllowed();
+		}
+
+		return false;
+	}
+
 	Private::FIterationSettings IterationSettings() const
 	{
 		VerifyContext();
diff --git a/Engine/Source/Runtime/Experimental/ChaosVisualDebugger/Public/DataWrappers/ChaosVDParticleDataWrapper.h b/Engine/Source/Runtime/Experimental/ChaosVisualDebugger/Public/DataWrappers/ChaosVDParticleDataWrapper.h
--- a/Engine/Source/Runtime/Experimental/ChaosVisualDebugger/Public/DataWrappers/ChaosVDParticleDataWrapper.h
+++ b/Engine/Source/Runtime/Experimental/ChaosVisualDebugger/Public/DataWrappers/ChaosVDParticleDataWrapper.h
@@ -86,6 +86,7 @@ struct FChaosVDFRigidParticleControlFlags : public FChaosVDWrapperDataBase
 		  bInertiaConditioningEnabled(false), 
 		  GravityGroupIndex(0),
 		  bMACDEnabled(false), 
+		  bPartialIslandSleepAllowed(true),
 		  bGyroscopicTorqueEnabled(false),
 		  PositionSolverIterationCount(8),
 		  VelocitySolverIterationCount(1),
@@ -104,6 +105,7 @@ struct FChaosVDFRigidParticleControlFlags : public FChaosVDWrapperDataBase
 		bInertiaConditioningEnabled = Other.GetInertiaConditioningEnabled();
 		GravityGroupIndex = Other.GetGravityGroupIndex();
 		bMACDEnabled = Other.GetMACDEnabled();
+		bPartialIslandSleepAllowed = Other.GetPartialIslandSleepAllowed();
 		bGyroscopicTorqueEnabled = Other.GetGyroscopicTorqueEnabled();
 
 		bHasValidData = true;
@@ -118,6 +120,7 @@ struct FChaosVDFRigidParticleControlFlags : public FChaosVDWrapperDataBase
 		Other.SetInertiaConditioningEnabled(bInertiaConditioningEnabled);
 		Other.SetGravityGroupIndex(GravityGroupIndex);
 		Other.SetMACDEnabled(bMACDEnabled);
+		Other.SetPartialIslandSleepAllowed(bPartialIslandSleepAllowed);
 		Other.SetGyroscopicTorqueEnabled(bGyroscopicTorqueEnabled);
 	}
 
@@ -134,6 +137,8 @@ struct FChaosVDFRigidParticleControlFlags : public FChaosVDWrapperDataBase
 	UPROPERTY(VisibleAnywhere, Category = "Particle Control Flags")
 	bool bMACDEnabled;
 	UPROPERTY(VisibleAnywhere, Category = "Particle Control Flags")
+	bool bPartialIslandSleepAllowed;
+	UPROPERTY(VisibleAnywhere, Category = "Particle Control Flags")
 	bool bGyroscopicTorqueEnabled;
 	UPROPERTY(VisibleAnywhere, Category = "Particle Control Flags")
 	uint8 PositionSolverIterationCount;
diff --git a/Engine/Source/Runtime/Experimental/GeometryCollectionEngine/Private/GeometryCollection/GeometryCollectionComponent.cpp b/Engine/Source/Runtime/Experimental/GeometryCollectionEngine/Private/GeometryCollection/GeometryCollectionComponent.cpp
--- a/Engine/Source/Runtime/Experimental/GeometryCollectionEngine/Private/GeometryCollection/GeometryCollectionComponent.cpp
+++ b/Engine/Source/Runtime/Experimental/GeometryCollectionEngine/Private/GeometryCollection/GeometryCollectionComponent.cpp
@@ -3891,6 +3891,7 @@ PRAGMA_ENABLE_DEPRECATION_WARNINGS
 		SimulationParameters.UseInertiaConditioning = BodyInstance.IsInertiaConditioningEnabled();
 		SimulationParameters.UseCCD = BodyInstance.bUseCCD;
 		SimulationParameters.UseMACD = BodyInstance.GetUseMACD();
+		SimulationParameters.AllowPartialIslandSleep = BodyInstance.GetAllowPartialIslandSleep();
 		SimulationParameters.PositionSolverIterations = BodyInstance.GetPositionSolverIterationCount();
 		SimulationParameters.VelocitySolverIterations = BodyInstance.GetVelocitySolverIterationCount();
 		SimulationParameters.ProjectionSolverIterations = BodyInstance.GetProjectionSolverIterationCount();
diff --git a/Engine/Source/Runtime/PhysicsCore/Private/ChaosEngineInterface.cpp b/Engine/Source/Runtime/PhysicsCore/Private/ChaosEngineInterface.cpp
--- a/Engine/Source/Runtime/PhysicsCore/Private/ChaosEngineInterface.cpp
+++ b/Engine/Source/Runtime/PhysicsCore/Private/ChaosEngineInterface.cpp
@@ -553,6 +553,11 @@ void FChaosEngineInterface::SetMACDEnabled_AssumesLocked(const FPhysicsActorHand
 	InActorReference->GetGameThreadAPI().SetMACDEnabled(bIsMACDEnabled);
 }
 
+void FChaosEngineInterface::SetPartialIslandSleepAllowed_AssumesLocked(const FPhysicsActorHandle& InActorReference, bool bIsPartialIslandSleepAllowed)
+{
+	InActorReference->GetGameThreadAPI().SetPartialIslandSleepAllowed(bIsPartialIslandSleepAllowed);
+}
+
 void FChaosEngineInterface::SetPositionSolverIterationCount_AssumesLocked(const FPhysicsActorHandle& InActorReference, uint8 PositionSolverIterationCount)
 {
 	InActorReference->GetGameThreadAPI().SetPositionSolverIterationCount(PositionSolverIterationCount);
diff --git a/Engine/Source/Runtime/PhysicsCore/Public/Chaos/ChaosEngineInterface.h b/Engine/Source/Runtime/PhysicsCore/Public/Chaos/ChaosEngineInterface.h
--- a/Engine/Source/Runtime/PhysicsCore/Public/Chaos/ChaosEngineInterface.h
+++ b/Engine/Source/Runtime/PhysicsCore/Public/Chaos/ChaosEngineInterface.h
@@ -414,6 +414,7 @@ public:
 	static PHYSICSCORE_API void SetIsKinematic_AssumesLocked(const FPhysicsActorHandle& InActorReference,bool bIsKinematic);
 	static PHYSICSCORE_API void SetCcdEnabled_AssumesLocked(const FPhysicsActorHandle& InActorReference,bool bIsCcdEnabled);
 	static PHYSICSCORE_API void SetMACDEnabled_AssumesLocked(const FPhysicsActorHandle& InActorReference, bool bIsMACDEnabled);
+	static PHYSICSCORE_API void SetPartialIslandSleepAllowed_AssumesLocked(const FPhysicsActorHandle& InActorReference, bool bIsPartialIslandSleepAllowed);
 	static PHYSICSCORE_API void SetPositionSolverIterationCount_AssumesLocked(const FPhysicsActorHandle& InActorReference, uint8 PositionSolverIterationCount);
 	static PHYSICSCORE_API void SetVelocitySolverIterationCount_AssumesLocked(const FPhysicsActorHandle& InActorReference, uint8 VelocitySolverIterationCount);
 	static PHYSICSCORE_API void SetProjectionSolverIterationCount_AssumesLocked(const FPhysicsActorHandle& InActorReference, uint8 ProjectionSolverIterationCount);
