[NetPhysics] Move ensure for failed delta serialization behind debugging definition.
The ensure has a valid callflow that can trigger it, if a possession is made, but it should still be part of debugging since in any other case this is an issue.

#rnx

diff --git a/Engine/Source/Runtime/Engine/Private/PhysicsEngine/NetworkPhysicsComponent.cpp b/Engine/Source/Runtime/Engine/Private/PhysicsEngine/NetworkPhysicsComponent.cpp
--- a/Engine/Source/Runtime/Engine/Private/PhysicsEngine/NetworkPhysicsComponent.cpp
+++ b/Engine/Source/Runtime/Engine/Private/PhysicsEngine/NetworkPhysicsComponent.cpp
@@ -1048,7 +1048,9 @@ void UNetworkPhysicsComponent::OnRep_SetReplicatedStates()
 {
 	if (ReplicatedStates.bDeltaSerializationIssue)
 	{
-		ensureMsgf(!IsLocallyControlled(), TEXT("OnRep_SetReplicatedStates failed delta serialization, should not happen on the owning client."));
+#if DEBUG_NETWORK_PHYSICS_DELTASERIALIZATION
+		ensureMsgf(!IsLocallyControlled(), TEXT("OnRep_SetReplicatedStates failed delta serialization, should not happen on the owning client unless the pawn just got possessed."));
+#endif
 		return;
 	}
 
@@ -1074,7 +1076,9 @@ void UNetworkPhysicsComponent::OnRep_SetReplicatedInputs()
 {
 	if (ReplicatedInputs.bDeltaSerializationIssue)
 	{
-		ensureMsgf(!IsLocallyControlled(), TEXT("OnRep_SetReplicatedInputs failed delta serialization, should not happen on the owning client."));
+#if DEBUG_NETWORK_PHYSICS_DELTASERIALIZATION
+		ensureMsgf(!IsLocallyControlled(), TEXT("OnRep_SetReplicatedInputs failed delta serialization, should not happen on the owning client unless the pawn just got possessed."));
+#endif
 		return;
 	}
 
@@ -1100,7 +1104,9 @@ void UNetworkPhysicsComponent::OnRep_SetReplicatedRemoteInputs()
 {
 	if (ReplicatedRemoteInputs.bDeltaSerializationIssue)
 	{
-		ensureMsgf(!IsLocallyControlled(), TEXT("OnRep_SetReplicatedRemoteInputs failed delta serialization, should not happen on the owning client."));
+#if DEBUG_NETWORK_PHYSICS_DELTASERIALIZATION
+		ensureMsgf(!IsLocallyControlled(), TEXT("OnRep_SetReplicatedRemoteInputs failed delta serialization, should not happen on the owning client unless the pawn just got possessed."));
+#endif
 		return;
 	}
 
@@ -1199,7 +1205,9 @@ void UNetworkPhysicsComponent::MulticastReceiveImportantInputData_Implementation
 
 	if (ServerInputs.bDeltaSerializationIssue)
 	{
-		ensureMsgf(!IsLocallyControlled(), TEXT("MulticastReceiveImportantInputData_Implementation failed delta serialization, should not happen on the owning client."));
+#if DEBUG_NETWORK_PHYSICS_DELTASERIALIZATION
+		ensureMsgf(!IsLocallyControlled(), TEXT("MulticastReceiveImportantInputData_Implementation failed delta serialization, should not happen on the owning client unless the pawn just got possessed."));
+#endif
 		return;
 	}
 
@@ -1228,7 +1236,9 @@ void UNetworkPhysicsComponent::MulticastReceiveImportantStateData_Implementation
 
 	if (ServerStates.bDeltaSerializationIssue)
 	{
-		ensureMsgf(!IsLocallyControlled(), TEXT("MulticastReceiveImportantStateData_Implementation failed delta serialization, should not happen on the owning client."));
+#if DEBUG_NETWORK_PHYSICS_DELTASERIALIZATION
+		ensureMsgf(!IsLocallyControlled(), TEXT("MulticastReceiveImportantStateData_Implementation failed delta serialization, should not happen on the owning client unless the pawn just got possessed."));
+#endif
 		return;
 	}
 
