[Audio Insights] Fix for Editor crash when playing Sound Wave using Audio Insights

diff --git a/Engine/Source/Runtime/Engine/Private/AudioDevice.cpp b/Engine/Source/Runtime/Engine/Private/AudioDevice.cpp
--- a/Engine/Source/Runtime/Engine/Private/AudioDevice.cpp
+++ b/Engine/Source/Runtime/Engine/Private/AudioDevice.cpp
@@ -5087,7 +5087,10 @@ void FAudioDevice::SendSoundStartedTraceMessage(const FActiveSound& InActiveSoun
 #if WITH_EDITOR
 		if (const TObjectPtr<UAudioComponent> AudioComponent = UAudioComponent::GetAudioComponentFromID(InActiveSound.GetAudioComponentID()))
 		{
-			OwnerActorName = AudioComponent->GetOwner()->GetActorLabel(false);
+			if (const AActor* OwnerActor = AudioComponent->GetOwner())
+			{
+				OwnerActorName = OwnerActor->GetActorLabel(false);
+			}
 		}
 #endif // WITH_EDITOR
 
diff --git a/Engine/Source/Runtime/Engine/Private/SoundWave.cpp b/Engine/Source/Runtime/Engine/Private/SoundWave.cpp
--- a/Engine/Source/Runtime/Engine/Private/SoundWave.cpp
+++ b/Engine/Source/Runtime/Engine/Private/SoundWave.cpp
@@ -3884,7 +3884,10 @@ void USoundWave::Parse(FAudioDevice* AudioDevice, const UPTRINT NodeWaveInstance
 #if WITH_EDITOR
 				if (const TObjectPtr<UAudioComponent> AudioComponent = UAudioComponent::GetAudioComponentFromID(ActiveSound.GetAudioComponentID()))
 				{
-					OwnerActorName = AudioComponent->GetOwner()->GetActorLabel(false);
+					if (const AActor* OwnerActor = AudioComponent->GetOwner())
+					{
+						OwnerActorName = OwnerActor->GetActorLabel(false);
+					}
 				}
 #endif // WITH_EDITOR
 
