Temporarily remove DownloadParams from FPrimaryAssetLoadState to save memory.  Will be re-implementing this in a more memory friendly way in a follow up.

diff --git a/Engine/Source/Runtime/Engine/Private/AssetManager.cpp b/Engine/Source/Runtime/Engine/Private/AssetManager.cpp
--- a/Engine/Source/Runtime/Engine/Private/AssetManager.cpp
+++ b/Engine/Source/Runtime/Engine/Private/AssetManager.cpp
@@ -57,6 +57,9 @@ LLM_DEFINE_TAG(AssetManager);
 
 DEFINE_LOG_CATEGORY(LogAssetManager);
 
+// Temp hack to claw back memory. Proper fix to come later
+// TODO: FIXME: JMarcus
+#define UE_ASSETMANAGER_DOWNLOADPARMS_IN_LOADSTATE 0
 
 /** Structure defining the current loading state of an asset */
 struct FPrimaryAssetLoadState
@@ -66,12 +69,19 @@ struct FPrimaryAssetLoadState
 
 	/** The set of bundles to be loaded by the handle */
 	TArray<FName> BundleNames;
+#if UE_ASSETMANAGER_DOWNLOADPARMS_IN_LOADSTATE
 	TOptional<FAssetManagerDownloadParams> DownloadParams;
 
 	bool IsSame(TConstArrayView<FName> InBundleNames, const TOptional<FAssetManagerDownloadParams>& InDownloadParams) const
 	{
 		return Algo::Compare(InBundleNames, BundleNames) && (InDownloadParams == DownloadParams);
 	}
+#else
+	bool IsSame(TConstArrayView<FName> InBundleNames, const TOptional<FAssetManagerDownloadParams>& InDownloadParams) const
+	{
+		return Algo::Compare(InBundleNames, BundleNames);
+	}
+#endif // UE_ASSETMANAGER_DOWNLOADPARMS_IN_LOADSTATE
 
 	/** If this state is keeping things in memory */
 	bool IsValid() const { return Handle.IsValid() && Handle->IsActive(); }
@@ -90,7 +100,9 @@ struct FPrimaryAssetLoadState
 			Handle = nullptr;
 		}
 		BundleNames.Reset();
+#if UE_ASSETMANAGER_DOWNLOADPARMS_IN_LOADSTATE
 		DownloadParams.Reset();
+#endif // UE_ASSETMANAGER_DOWNLOADPARMS_IN_LOADSTATE
 	}
 };
 
@@ -1910,7 +1922,9 @@ TSharedPtr<FStreamableHandle> UAssetManager::ChangeBundleStateForPrimaryAssets(
 			{
 				// New state has no assets to load. Set the CurrentState's bundles and clear the handle
 				NameData->CurrentState.BundleNames = NewBundleState;
+#if UE_ASSETMANAGER_DOWNLOADPARMS_IN_LOADSTATE
 				NameData->CurrentState.DownloadParams = CopyTemp(*LoadParams.DownloadParams);
+#endif // UE_ASSETMANAGER_DOWNLOADPARMS_IN_LOADSTATE
 				NameData->CurrentState.Handle.Reset();
 				continue;
 			}
@@ -1939,7 +1953,9 @@ TSharedPtr<FStreamableHandle> UAssetManager::ChangeBundleStateForPrimaryAssets(
 			// Set pending bundle state before issuing the request, because it can be resolved immediately.
 			// OnAssetStateChangeCompleted requires this data to handle state transition correctly.
 			NameData->PendingState.BundleNames = NewBundleState;
+#if UE_ASSETMANAGER_DOWNLOADPARMS_IN_LOADSTATE
 			NameData->PendingState.DownloadParams = CopyTemp(LoadParams.DownloadParams);
+#endif // UE_ASSETMANAGER_DOWNLOADPARMS_IN_LOADSTATE
 
 			FAssetManagerLoadParams AssetListLoadParams;
 			AssetListLoadParams.Priority = LoadParams.Priority;
@@ -2133,7 +2149,9 @@ void UAssetManager::OnAssetStateChangeCompleted(TSharedPtr<FStreamableHandle> Bo
 			else
 			{
 				NameData->CurrentState.Handle = BoundHandle;
+#if UE_ASSETMANAGER_DOWNLOADPARMS_IN_LOADSTATE
 				NameData->CurrentState.DownloadParams = NameData->PendingState.DownloadParams;
+#endif // UE_ASSETMANAGER_DOWNLOADPARMS_IN_LOADSTATE
 				NameData->CurrentState.BundleNames = NameData->PendingState.BundleNames;
 				
 				NotifyOnAssetStateChangeCompleted(PrimaryAssetId);
