Enable UE_ASSETMANAGER_DOWNLOADPARMS_IN_LOADSTATE for 36.20+.
I'll be replacing this hack with a proper solution in 36.20

diff --git a/Engine/Source/Runtime/Engine/Private/AssetManager.cpp b/Engine/Source/Runtime/Engine/Private/AssetManager.cpp
--- a/Engine/Source/Runtime/Engine/Private/AssetManager.cpp
+++ b/Engine/Source/Runtime/Engine/Private/AssetManager.cpp
@@ -59,7 +59,7 @@ DEFINE_LOG_CATEGORY(LogAssetManager);
 
 // Temp hack to claw back memory. Proper fix to come later
 // TODO: FIXME: JMarcus
-#define UE_ASSETMANAGER_DOWNLOADPARMS_IN_LOADSTATE 0
+#define UE_ASSETMANAGER_DOWNLOADPARMS_IN_LOADSTATE 1
 
 /** Structure defining the current loading state of an asset */
 struct FPrimaryAssetLoadState
