Fixed TSAN data race error on LevelStreaming CurrentState when using wp.runtime.UpdateStreaming.EnableAsyncUpdate.

diff --git a/Engine/Source/Runtime/Engine/Classes/Engine/LevelStreaming.h b/Engine/Source/Runtime/Engine/Classes/Engine/LevelStreaming.h
--- a/Engine/Source/Runtime/Engine/Classes/Engine/LevelStreaming.h
+++ b/Engine/Source/Runtime/Engine/Classes/Engine/LevelStreaming.h
@@ -279,7 +279,7 @@ protected:
 
 private:
 	/** What the current streamed state of the streaming level is */
-	ELevelStreamingState CurrentState;
+	TAtomic<ELevelStreamingState> CurrentState;
 
 	/** What streamed state the streaming level is transitioning towards */
 	ELevelStreamingTargetState TargetState;
@@ -361,7 +361,7 @@ public:
 	/** Returns the current loaded/visible state of the streaming level. */
 PRAGMA_DISABLE_DEPRECATION_WARNINGS
 	UE_DEPRECATED(5.2, "ULevelStreaming::ECurrentState has been replaced by ELevelStreamingState. Use GetLevelStreamingState instead.")
-	ECurrentState GetCurrentState() const { return (ECurrentState)CurrentState; }
+	ECurrentState GetCurrentState() const { return (ECurrentState)GetLevelStreamingState(); }
 PRAGMA_ENABLE_DEPRECATION_WARNINGS
 
 	ELevelStreamingState GetLevelStreamingState() const { return CurrentState; }
diff --git a/Engine/Source/Runtime/Engine/Private/LevelStreaming.cpp b/Engine/Source/Runtime/Engine/Private/LevelStreaming.cpp
--- a/Engine/Source/Runtime/Engine/Private/LevelStreaming.cpp
+++ b/Engine/Source/Runtime/Engine/Private/LevelStreaming.cpp
@@ -591,9 +591,7 @@ void ULevelStreaming::OnLevelRemoved()
 
 void ULevelStreaming::SetCurrentState(ELevelStreamingState NewState)
 {
-	ELevelStreamingState OldState = CurrentState;
-	CurrentState = NewState;
-
+	ELevelStreamingState OldState = CurrentState.Exchange(NewState);
 	if (OldState != NewState)
 	{
 		OnCurrentStateChanged(OldState, NewState);
