[HWRT] Increase default r.RayTracing.ResidentGeometryMemoryPoolSizeInMB to 400MB and approximate compacted size in FRayTracingGeometryManager.

diff --git a/Engine/Source/Runtime/Apple/MetalRHI/Private/MetalRHI.cpp b/Engine/Source/Runtime/Apple/MetalRHI/Private/MetalRHI.cpp
--- a/Engine/Source/Runtime/Apple/MetalRHI/Private/MetalRHI.cpp
+++ b/Engine/Source/Runtime/Apple/MetalRHI/Private/MetalRHI.cpp
@@ -1178,6 +1178,9 @@ FMetalDynamicRHI::FMetalDynamicRHI(ERHIFeatureLevel::Type RequestedFeatureLevel)
 		GRHIRayTracingAccelerationStructureAlignment = 16;
 		GRHIRayTracingScratchBufferAlignment = 4;
 		GRHIRayTracingInstanceDescriptorSize = uint32(sizeof(MTLAccelerationStructureUserIDInstanceDescriptor));
+
+		static auto CVarRayTracingAllowCompaction = IConsoleManager::Get().FindConsoleVariable(TEXT("r.Metal.RayTracing.AllowCompaction"));
+		GRHIGlobals.RayTracing.SupportsAccelerationStructureCompaction = CVarRayTracingAllowCompaction->GetInt() != 0;
 	}
 #endif
 		
diff --git a/Engine/Source/Runtime/D3D12RHI/Private/D3D12Adapter.cpp b/Engine/Source/Runtime/D3D12RHI/Private/D3D12Adapter.cpp
--- a/Engine/Source/Runtime/D3D12RHI/Private/D3D12Adapter.cpp
+++ b/Engine/Source/Runtime/D3D12RHI/Private/D3D12Adapter.cpp
@@ -1181,6 +1181,9 @@ void FD3D12Adapter::InitializeDevices()
 			}
 #endif // PLATFORM_WINDOWS
 
+			static auto CVarRayTracingAllowCompaction = IConsoleManager::Get().FindConsoleVariable(TEXT("r.D3D12.RayTracing.AllowCompaction"));
+			GRHIGlobals.RayTracing.SupportsAccelerationStructureCompaction = CVarRayTracingAllowCompaction->GetInt() != 0;
+
 			GRHIRayTracingAccelerationStructureAlignment = uint32(D3D12_RAYTRACING_ACCELERATION_STRUCTURE_BYTE_ALIGNMENT);
 			GRHIRayTracingScratchBufferAlignment = uint32(D3D12_RAYTRACING_ACCELERATION_STRUCTURE_BYTE_ALIGNMENT);
 			GRHIRayTracingShaderTableAlignment = uint32(D3D12_RAYTRACING_SHADER_TABLE_BYTE_ALIGNMENT);
diff --git a/Engine/Source/Runtime/Engine/Private/Rendering/RayTracingGeometryManager.cpp b/Engine/Source/Runtime/Engine/Private/Rendering/RayTracingGeometryManager.cpp
--- a/Engine/Source/Runtime/Engine/Private/Rendering/RayTracingGeometryManager.cpp
+++ b/Engine/Source/Runtime/Engine/Private/Rendering/RayTracingGeometryManager.cpp
@@ -69,7 +69,7 @@ static FAutoConsoleVariableRef CVarNaniteStreamingMaxPendingRequests(
 	ECVF_ReadOnly
 );
 
-static int32 GRayTracingResidentGeometryMemoryPoolSizeInMB = 256;
+static int32 GRayTracingResidentGeometryMemoryPoolSizeInMB = 400;
 static FAutoConsoleVariableRef CVarRayTracingResidentGeometryMemoryPoolSizeInMB(
 	TEXT("r.RayTracing.ResidentGeometryMemoryPoolSizeInMB"),
 	GRayTracingResidentGeometryMemoryPoolSizeInMB,
@@ -78,6 +78,15 @@ static FAutoConsoleVariableRef CVarRayTracingResidentGeometryMemoryPoolSizeInMB(
 	ECVF_RenderThreadSafe
 );
 
+static float GRayTracingApproximateCompactionRatio = 0.5f;
+static FAutoConsoleVariableRef CVarRayTracingApproximateCompactionRatio(
+	TEXT("r.RayTracing.ApproximateCompactionRatio"),
+	GRayTracingApproximateCompactionRatio,
+	TEXT("Ratio used by Ray Tracing Geometry Manager to approximate the ray tracing geometry size after compaction.\n")
+	TEXT("This will be removed in a future version once Ray Tracing Geometry Manager tracks the actual compacted sizes."),
+	ECVF_RenderThreadSafe
+);
+
 static bool bRefreshAlwaysResidentRayTracingGeometries = false;
 
 static int32 GRayTracingNumAlwaysResidentLODs = 1;
@@ -452,6 +461,15 @@ void FRayTracingGeometryManager::SetRayTracingGeometryGroupCurrentFirstLODIndex(
 	Group.CurrentFirstLODIdx = NewCurrentFirstLODIdx;
 }
 
+static bool ShouldCompactAfterBuild(const FRayTracingGeometryInitializer& Initializer)
+{
+	return Initializer.bAllowCompaction
+		&& !Initializer.bFastBuild
+		&& !Initializer.bAllowUpdate
+		&& !Initializer.OfflineDataHeader.IsValid()
+		&& GRHIGlobals.RayTracing.SupportsAccelerationStructureCompaction;
+}
+
 void FRayTracingGeometryManager::RefreshRegisteredGeometry(FGeometryHandle Handle)
 {
 	FScopeLock ScopeLock(&MainCS);
@@ -474,7 +492,19 @@ void FRayTracingGeometryManager::RefreshRegisteredGeometry(FGeometryHandle Handl
 				}
 			}
 
-			RegisteredGeometry.Size = bAllSegmentsAreValid ? RHICalcRayTracingGeometrySize(RegisteredGeometry.Geometry->Initializer).ResultSize : 0;
+			if (bAllSegmentsAreValid)
+			{
+				RegisteredGeometry.Size = RHICalcRayTracingGeometrySize(RegisteredGeometry.Geometry->Initializer).ResultSize;
+
+				if (ShouldCompactAfterBuild(RegisteredGeometry.Geometry->Initializer))
+				{
+					RegisteredGeometry.Size = uint32(RegisteredGeometry.Size * GRayTracingApproximateCompactionRatio);
+				}
+			}
+			else
+			{
+				RegisteredGeometry.Size = 0;
+			}
 		}
 
 		RegisteredGeometry.bEvicted = RegisteredGeometry.Geometry->IsEvicted();
diff --git a/Engine/Source/Runtime/RHI/Public/RHIGlobals.h b/Engine/Source/Runtime/RHI/Public/RHIGlobals.h
--- a/Engine/Source/Runtime/RHI/Public/RHIGlobals.h
+++ b/Engine/Source/Runtime/RHI/Public/RHIGlobals.h
@@ -458,6 +458,9 @@ struct FRHIGlobals
 		/** Wether an extra uniform buffer parameter is required when loose parameters are present, or if they are stored directly in the shader record. */
 		bool SupportsLooseParamsInShaderRecord = false;
 
+		/** Whether or not the RHI supports compacting acceleration structures */
+		bool SupportsAccelerationStructureCompaction = false;
+
 		/** Whether or not the RHI supports serializing acceleration structures to a file */
 		bool SupportsSerializeAccelerationStructure = false;
 
diff --git a/Engine/Source/Runtime/VulkanRHI/Private/VulkanRHI.cpp b/Engine/Source/Runtime/VulkanRHI/Private/VulkanRHI.cpp
--- a/Engine/Source/Runtime/VulkanRHI/Private/VulkanRHI.cpp
+++ b/Engine/Source/Runtime/VulkanRHI/Private/VulkanRHI.cpp
@@ -911,6 +911,9 @@ void FVulkanDynamicRHI::InitInstance()
 			// Inline RayTracing SBT is needed if raytracing position fetch isn't available
 			GRHIGlobals.RayTracing.RequiresInlineRayTracingSBT = !VULKAN_SUPPORTS_RAY_TRACING_POSITION_FETCH;
 
+			static auto CVarRayTracingAllowCompaction = IConsoleManager::Get().FindConsoleVariable(TEXT("r.Vulkan.RayTracing.AllowCompaction"));
+			GRHIGlobals.RayTracing.SupportsAccelerationStructureCompaction = CVarRayTracingAllowCompaction->GetInt() != 0;
+
 			GRHIRayTracingAccelerationStructureAlignment = 256; // TODO (currently handled by FVulkanAccelerationStructureBuffer)
 			//Some devices have 64 for min AS offset alignment meanwhile engine AS alignment is 256. hence using round up value
 			GRHIRayTracingScratchBufferAlignment = FPlatformMath::Max<uint32>(GRHIRayTracingAccelerationStructureAlignment, 
