Changed draw behavior for selected meshes which only write to runtime virtual texture.
Previously we would fully show those meshes on selection (they are hidden in the main pass otherwise).
Production feedback is that this isn't good for level editing, since selection can then hide the surface that shows the runtime virtual texture results.
Now we show the selection outline for the mesh only. This gives sufficient information about mesh shape without hiding what is below.

diff --git a/Engine/Source/Runtime/Engine/Private/StaticMeshSceneProxy.cpp b/Engine/Source/Runtime/Engine/Private/StaticMeshSceneProxy.cpp
--- a/Engine/Source/Runtime/Engine/Private/StaticMeshSceneProxy.cpp
+++ b/Engine/Source/Runtime/Engine/Private/StaticMeshSceneProxy.cpp
@@ -1772,6 +1772,16 @@ void FStaticMeshSceneProxy::GetDynamicMeshElements(const TArray<const FSceneView
 											}
 										}
 
+										// Only draw the selection outline for selected items which only write to runtime virtual texture.
+										if (bProxyIsSelected && DrawInRuntimeVirtualTextureOnly(ViewFamily.EngineShowFlags.Editor))
+										{
+											if (FColoredMaterialRenderProxy* DebugMaterialProxy = new FColoredMaterialRenderProxy(GEngine->EmissiveMeshMaterial->GetRenderProxy(), FColor::Transparent))
+											{
+												Collector.RegisterOneFrameMaterialProxy(DebugMaterialProxy);
+												MeshElement.MaterialRenderProxy = DebugMaterialProxy;
+												bDebugMaterialRenderProxySet = true;
+											}
+										}
 	#endif // WITH_EDITOR
 										// Override the mesh's material with our material that draws vertex color
 										if (!bDebugMaterialRenderProxySet && bProxyIsSelected && EngineShowFlags.VertexColors && AllowDebugViewmodes())
@@ -2253,6 +2263,7 @@ FPrimitiveViewRelevance FStaticMeshSceneProxy::GetViewRelevance(const FSceneView
 #if WITH_EDITOR
 		(IsSelected() && View->Family->EngineShowFlags.VertexColors) ||
 		(IsSelected() && View->Family->EngineShowFlags.PhysicalMaterialMasks) ||
+		(IsSelected() && DrawInRuntimeVirtualTextureOnly(View->Family->EngineShowFlags.Editor)) ||
 #endif
 #if STATICMESH_ENABLE_DEBUG_RENDERING
 		bDrawMeshCollisionIfComplex ||
