FSimpleStreamableAssetMaanger adding index validation code for elements indices
In editor after bulding the asset it will trigger recreation that will empty the asset data before unregistration for every object that was using the asset

diff --git a/Engine/Source/Runtime/Engine/Public/Streaming/SimpleStreamableAssetManager.h b/Engine/Source/Runtime/Engine/Public/Streaming/SimpleStreamableAssetManager.h
--- a/Engine/Source/Runtime/Engine/Public/Streaming/SimpleStreamableAssetManager.h
+++ b/Engine/Source/Runtime/Engine/Public/Streaming/SimpleStreamableAssetManager.h
@@ -274,7 +274,7 @@ private:
 		
 		void Reset(int32 Index)
 		{
-			if (UsedElements[Index])
+			if (UsedElements.Num() > Index && UsedElements[Index])
 			{
 				UsedElements[Index] = false;
 				FreeElementIndexHint = FMath::Min(FreeElementIndexHint, Index);
