Fix for NoPawnPlayerController not having the correct view target position when being used for replication relevancy.

diff --git a/Engine/Source/Runtime/Engine/Private/PlayerController.cpp b/Engine/Source/Runtime/Engine/Private/PlayerController.cpp
--- a/Engine/Source/Runtime/Engine/Private/PlayerController.cpp
+++ b/Engine/Source/Runtime/Engine/Private/PlayerController.cpp
@@ -6663,8 +6663,8 @@ class UNetConnection* ANoPawnPlayerController::GetNetConnection() const
 
 void ANoPawnPlayerController::ServerSetViewTargetPosition_Implementation(FVector ViewTargetPos)
 {
-	// Set the value of FRepMovement because SetActorLocation() is marked as private.
-	GetReplicatedMovement_Mutable().Location = ViewTargetPos;
+	AActor* Actor = Cast<AActor>(this);
+	Actor->SetActorLocation(ViewTargetPos);
 }
 
 #undef LOCTEXT_NAMESPACE
