Fixed incorrect increment of pool bucket index.

diff --git a/Engine/Source/Runtime/Engine/Internal/SkeletalMeshUpdater.h b/Engine/Source/Runtime/Engine/Internal/SkeletalMeshUpdater.h
--- a/Engine/Source/Runtime/Engine/Internal/SkeletalMeshUpdater.h
+++ b/Engine/Source/Runtime/Engine/Internal/SkeletalMeshUpdater.h
@@ -187,7 +187,7 @@ public:
 			if (NumFreeItemBytes > 0)
 			{
 				NumFreeBytes.fetch_sub(NumFreeItemBytes, std::memory_order_relaxed);
-				PoolBucketIndex = PoolBucketIndex++ % NumPoolBuckets;
+				PoolBucketIndex = (PoolBucketIndex + 1) % NumPoolBuckets;
 			}
 		}
 	}
