Moved several light component property updates to propagate directly to the proxy without marking the full render state dirty.

diff --git a/Engine/Source/Runtime/Engine/Private/Components/DirectionalLightComponent.cpp b/Engine/Source/Runtime/Engine/Private/Components/DirectionalLightComponent.cpp
--- a/Engine/Source/Runtime/Engine/Private/Components/DirectionalLightComponent.cpp
+++ b/Engine/Source/Runtime/Engine/Private/Components/DirectionalLightComponent.cpp
@@ -384,6 +384,42 @@ public:
 			});
 	}
 
+	void UpdateLightSourceAngle_GameThread(const UDirectionalLightComponent* Component)
+	{
+		ENQUEUE_RENDER_COMMAND(FUpdateLightSourceAngleCommand)(
+			[Proxy = this, LightSourceAngle = Component->LightSourceAngle](FRHICommandList& RHICmdList)
+		{
+			Proxy->LightSourceAngle = LightSourceAngle;
+		});
+	}
+
+	void UpdateLightSourceSoftAngle_GameThread(const UDirectionalLightComponent* Component)
+	{
+		ENQUEUE_RENDER_COMMAND(FUpdateLightSourceSoftAngleCommand)(
+			[Proxy = this, NewLightSourceSoftAngle = Component->LightSourceSoftAngle](FRHICommandList& RHICmdList)
+		{
+			Proxy->LightSourceSoftAngle = NewLightSourceSoftAngle;
+		});
+	}
+
+	void UpdateShadowSourceAngleFactor_GameThread(const UDirectionalLightComponent* Component)
+	{
+		ENQUEUE_RENDER_COMMAND(FUpdateShadowSourceAngleFactorCommand)(
+			[Proxy = this, NewShadowSourceAngleFactor = Component->ShadowSourceAngleFactor](FRHICommandList& RHICmdList)
+		{
+			Proxy->ShadowSourceAngleFactor = NewShadowSourceAngleFactor;
+		});
+	}
+
+	void UpdateEnableLightShaftOcclusion_GameThread(const UDirectionalLightComponent* Component)
+	{
+		ENQUEUE_RENDER_COMMAND(FUpdateEnableLightShaftOcclusionCommand)(
+			[Proxy = this, bNewEnableLightShaftOcclusion = (bool)Component->bEnableLightShaftOcclusion](FRHICommandList& RHICmdList)
+		{
+			Proxy->bEnableLightShaftOcclusion = bNewEnableLightShaftOcclusion;
+		});
+	}
+
 	void UpdateAtmosphereSunDiskColorScale_GameThread(const UDirectionalLightComponent* Component)
 	{
 		FLinearColor NewAtmosphereSunDiskColorScale = Component->AtmosphereSunDiskColorScale;
@@ -1288,7 +1324,11 @@ void UDirectionalLightComponent::SetEnableLightShaftOcclusion(bool bNewValue)
 		&& bEnableLightShaftOcclusion != bNewValue)
 	{
 		bEnableLightShaftOcclusion = bNewValue;
-		MarkRenderStateDirty();
+		if (SceneProxy)
+		{
+			FDirectionalLightSceneProxy* DirectionalLightSceneProxy = (FDirectionalLightSceneProxy*)SceneProxy;
+			DirectionalLightSceneProxy->UpdateEnableLightShaftOcclusion_GameThread(this);
+		}
 	}
 }
 
@@ -1331,7 +1371,11 @@ void UDirectionalLightComponent::SetLightSourceAngle(float NewValue)
 		&& LightSourceAngle != NewValue)
 	{
 		LightSourceAngle = NewValue;
-		MarkRenderStateDirty();
+		if (SceneProxy)
+		{
+			FDirectionalLightSceneProxy* DirectionalLightSceneProxy = (FDirectionalLightSceneProxy*)SceneProxy;
+			DirectionalLightSceneProxy->UpdateLightSourceAngle_GameThread(this);
+		}
 	}
 }
 
@@ -1341,7 +1385,11 @@ void UDirectionalLightComponent::SetLightSourceSoftAngle(float NewValue)
 		&& LightSourceSoftAngle != NewValue)
 	{
 		LightSourceSoftAngle = NewValue;
-		MarkRenderStateDirty();
+		if (SceneProxy)
+		{
+			FDirectionalLightSceneProxy* DirectionalLightSceneProxy = (FDirectionalLightSceneProxy*)SceneProxy;
+			DirectionalLightSceneProxy->UpdateLightSourceSoftAngle_GameThread(this);
+		}
 	}
 }
 
@@ -1351,7 +1399,11 @@ void UDirectionalLightComponent::SetShadowSourceAngleFactor(float NewValue)
 		&& ShadowSourceAngleFactor != NewValue)
 	{
 		ShadowSourceAngleFactor = NewValue;
-		MarkRenderStateDirty();
+		if (SceneProxy)
+		{
+			FDirectionalLightSceneProxy* DirectionalLightSceneProxy = (FDirectionalLightSceneProxy*)SceneProxy;
+			DirectionalLightSceneProxy->UpdateShadowSourceAngleFactor_GameThread(this);
+		}
 	}
 }
 
diff --git a/Engine/Source/Runtime/Engine/Private/Components/LightComponent.cpp b/Engine/Source/Runtime/Engine/Private/Components/LightComponent.cpp
--- a/Engine/Source/Runtime/Engine/Private/Components/LightComponent.cpp
+++ b/Engine/Source/Runtime/Engine/Private/Components/LightComponent.cpp
@@ -1112,7 +1112,10 @@ void ULightComponent::SetLightFunctionMaterial(UMaterialInterface* NewLightFunct
 		StashedLightFunctionMaterial = nullptr;
 #endif
 		LightFunctionMaterial = NewLightFunctionMaterial;
-		MarkRenderStateDirty();
+		if (SceneProxy)
+		{
+			SceneProxy->SetLightFunctionMaterial_GameThread(this);
+		}
 	}
 }
 
@@ -1163,18 +1166,20 @@ void ULightComponent::SetAffectTranslucentLighting(bool bNewValue)
 		MarkRenderStateDirty();
 	}
 }
+
 void ULightComponent::SetTransmission(bool bNewValue)
 {
 	if (AreDynamicDataChangesAllowed()
 		&& bTransmission != bNewValue)
 	{
 		bTransmission = bNewValue;
-		MarkRenderStateDirty();
+		if (SceneProxy)
+		{
+			SceneProxy->SetTransmission_GameThread(this);
+		}
 	}
 }
 
-
-
 void ULightComponent::SetEnableLightShaftBloom(bool bNewValue)
 {
 	if (AreDynamicDataChangesAllowed()
diff --git a/Engine/Source/Runtime/Engine/Private/LightSceneProxy.cpp b/Engine/Source/Runtime/Engine/Private/LightSceneProxy.cpp
--- a/Engine/Source/Runtime/Engine/Private/LightSceneProxy.cpp
+++ b/Engine/Source/Runtime/Engine/Private/LightSceneProxy.cpp
@@ -198,3 +198,36 @@ FTexture* FLightSceneProxy::GetIESTextureResource() const
 {
 	return IESTexture ? IESTexture->GetResource() : nullptr;
 }
+
+void FLightSceneProxy::SetTransmission_GameThread(const ULightComponent* Component)
+{
+	ENQUEUE_RENDER_COMMAND(FSetTransmissionCommand)(
+		[Proxy = this, ComponentTransmission = Component->bTransmission, ComponentCastsShadow = Component->CastShadows && Component->CastDynamicShadows](FRHICommandList& RHICmdList)
+	{
+		if (ComponentCastsShadow && IsMobilePlatform(Proxy->SceneInterface->GetShaderPlatform()))
+		{
+			Proxy->bTransmission = false;
+		}
+		else
+		{
+			Proxy->bTransmission = ComponentTransmission && Proxy->bCastDynamicShadow && !Proxy->bStaticShadowing;
+		}
+	});
+}
+
+void FLightSceneProxy::SetLightFunctionMaterial_GameThread(const ULightComponent* Component)
+{
+	const FMaterialRenderProxy* NewLightFunctionMaterial = nullptr;
+
+	if (Component->LightFunctionMaterial &&
+		Component->LightFunctionMaterial->GetMaterial()->MaterialDomain == MD_LightFunction)
+	{
+		NewLightFunctionMaterial = Component->LightFunctionMaterial->GetRenderProxy();
+	}
+
+	ENQUEUE_RENDER_COMMAND(FUpdateTransmissionCommand)(
+		[Proxy = this, NewLightFunctionMaterial](FRHICommandList& RHICmdList)
+	{
+		Proxy->LightFunctionMaterial = NewLightFunctionMaterial;
+	});
+}
\ No newline at end of file
diff --git a/Engine/Source/Runtime/Engine/Public/LightSceneProxy.h b/Engine/Source/Runtime/Engine/Public/LightSceneProxy.h
--- a/Engine/Source/Runtime/Engine/Public/LightSceneProxy.h
+++ b/Engine/Source/Runtime/Engine/Public/LightSceneProxy.h
@@ -57,6 +57,9 @@ public:
 		return true;
 	}
 
+	ENGINE_API void SetLightFunctionMaterial_GameThread(const ULightComponent* Component);
+	ENGINE_API void SetTransmission_GameThread(const ULightComponent* Component);
+
 	ENGINE_API virtual FSphere GetBoundingSphere() const;
 
 	/** @return radius of the light */
