Fix inconsistencies occuring with multiple worlds trying to share the same FPrecomputedVolumetricLightmapData (it's not shareable right now)
Added optionnal World* in FStreamingViewInfo so that we can differentiate between views in calls to UpdateResourceStreaming

diff --git a/Engine/Source/Editor/UnrealEd/Private/EditorEngine.cpp b/Engine/Source/Editor/UnrealEd/Private/EditorEngine.cpp
--- a/Engine/Source/Editor/UnrealEd/Private/EditorEngine.cpp
+++ b/Engine/Source/Editor/UnrealEd/Private/EditorEngine.cpp
@@ -2592,7 +2592,7 @@ bool UEditorEngine::UpdateSingleViewportClient(FEditorViewportClient* InViewport
 	{
 		float XSize = static_cast<float>(InViewportClient->Viewport->GetSizeXY().X);
 
-		IStreamingManager::Get().AddViewInformation( InViewportClient->GetViewLocation(), XSize, XSize / FMath::Tan(FMath::DegreesToRadians(InViewportClient->ViewFOV * 0.5f)) );
+		IStreamingManager::Get().AddViewInformation( InViewportClient->GetViewLocation(), XSize, XSize / FMath::Tan(FMath::DegreesToRadians(InViewportClient->ViewFOV * 0.5f)), 1.0f, false, 0.0f, nullptr, InViewportClient->GetWorld());
 	}
 	
 	// Only allow viewports to be drawn if we are not throttling for slate UI responsiveness or if the viewport client requested a redraw
diff --git a/Engine/Source/Runtime/Engine/Private/ContentStreaming.cpp b/Engine/Source/Runtime/Engine/Private/ContentStreaming.cpp
--- a/Engine/Source/Runtime/Engine/Private/ContentStreaming.cpp
+++ b/Engine/Source/Runtime/Engine/Private/ContentStreaming.cpp
@@ -425,7 +425,7 @@ bool IStreamingManager::HasShutdown()
  * @param Duration				How long the streaming system should keep checking this location (in seconds). 0 means just for the next Tick.
  * @param InActorToBoost		Optional pointer to an actor who's textures should have their streaming priority boosted
  */
-void IStreamingManager::AddViewInfoToArray( TArray<FStreamingViewInfo> &ViewInfos, const FVector& ViewOrigin, float ScreenSize, float FOVScreenSize, float BoostFactor, bool bOverrideLocation, float Duration, TWeakObjectPtr<AActor> InActorToBoost )
+void IStreamingManager::AddViewInfoToArray( TArray<FStreamingViewInfo> &ViewInfos, const FVector& ViewOrigin, float ScreenSize, float FOVScreenSize, float BoostFactor, bool bOverrideLocation, float Duration, TWeakObjectPtr<AActor> InActorToBoost, TWeakObjectPtr<UWorld> InWorld)
 {
 	// Check for duplicates and existing overrides.
 	bool bShouldAddView = true;
@@ -447,7 +447,7 @@ void IStreamingManager::AddViewInfoToArray( TArray<FStreamingViewInfo> &ViewInfo
 
 	if ( bShouldAddView )
 	{
-		new (ViewInfos) FStreamingViewInfo( ViewOrigin, ScreenSize, FOVScreenSize, BoostFactor, bOverrideLocation, Duration, InActorToBoost );
+		new (ViewInfos) FStreamingViewInfo( ViewOrigin, ScreenSize, FOVScreenSize, BoostFactor, bOverrideLocation, Duration, InActorToBoost, InWorld );
 	}
 }
 
@@ -553,7 +553,7 @@ void IStreamingManager::SetupViewInfos( float DeltaTime )
 		const FStreamingViewInfo& ViewInfo = LastingViewInfos[ ViewIndex ];
 		if ( bUseOverrideViews == ViewInfo.bOverrideLocation )
 		{
-			AddViewInfoToArray( CurrentViewInfos, ViewInfo.ViewOrigin, ViewInfo.ScreenSize * SplitScreenFactor, ViewInfo.FOVScreenSize, ViewInfo.BoostFactor, ViewInfo.bOverrideLocation, ViewInfo.Duration, ViewInfo.ActorToBoost );
+			AddViewInfoToArray( CurrentViewInfos, ViewInfo.ViewOrigin, ViewInfo.ScreenSize * SplitScreenFactor, ViewInfo.FOVScreenSize, ViewInfo.BoostFactor, ViewInfo.bOverrideLocation, ViewInfo.Duration, ViewInfo.ActorToBoost, ViewInfo.World );
 		}
 	}
 
@@ -563,7 +563,7 @@ void IStreamingManager::SetupViewInfos( float DeltaTime )
 		const FStreamingViewInfo& ViewInfo = PendingViewInfos[ ViewIndex ];
 		if ( bUseOverrideViews == ViewInfo.bOverrideLocation )
 		{
-			AddViewInfoToArray( CurrentViewInfos, ViewInfo.ViewOrigin, ViewInfo.ScreenSize * SplitScreenFactor, ViewInfo.FOVScreenSize, ViewInfo.BoostFactor, ViewInfo.bOverrideLocation, ViewInfo.Duration, ViewInfo.ActorToBoost );
+			AddViewInfoToArray( CurrentViewInfos, ViewInfo.ViewOrigin, ViewInfo.ScreenSize * SplitScreenFactor, ViewInfo.FOVScreenSize, ViewInfo.BoostFactor, ViewInfo.bOverrideLocation, ViewInfo.Duration, ViewInfo.ActorToBoost, ViewInfo.World );
 		}
 	}
 
@@ -643,7 +643,7 @@ void IStreamingManager::SetupViewInfos( float DeltaTime )
  * @param Duration				How long the streaming system should keep checking this location (in seconds). 0 means just for the next Tick.
  * @param InActorToBoost		Optional pointer to an actor who's textures should have their streaming priority boosted
  */
-void IStreamingManager::AddViewInformation( const FVector& ViewOrigin, float ScreenSize, float FOVScreenSize, float BoostFactor/*=1.0f*/, bool bOverrideLocation/*=false*/, float Duration/*=0.0f*/, TWeakObjectPtr<AActor> InActorToBoost /*=NULL*/ )
+void IStreamingManager::AddViewInformation( const FVector& ViewOrigin, float ScreenSize, float FOVScreenSize, float BoostFactor/*=1.0f*/, bool bOverrideLocation/*=false*/, float Duration/*=0.0f*/, TWeakObjectPtr<AActor> InActorToBoost /*=NULL*/, TWeakObjectPtr<UWorld> InWorld /*= NULL*/)
 {
 	// Is this a reasonable location?
 	if ( FMath::Abs(ViewOrigin.X) < (1.0e+20f) && FMath::Abs(ViewOrigin.Y) < (1.0e+20f) && FMath::Abs(ViewOrigin.Z) < (1.0e+20f) )
@@ -669,11 +669,11 @@ void IStreamingManager::AddViewInformation( const FVector& ViewOrigin, float Scr
 		// Should we remember this location for a while?
 		if ( Duration > 0.0f )
 		{
-			AddViewInfoToArray( LastingViewInfos, ViewOrigin, ScreenSize, FOVScreenSize, BoostFactor, bOverrideLocation, Duration, InActorToBoost );
+			AddViewInfoToArray( LastingViewInfos, ViewOrigin, ScreenSize, FOVScreenSize, BoostFactor, bOverrideLocation, Duration, InActorToBoost, InWorld );
 		}
 		else
 		{
-			AddViewInfoToArray( PendingViewInfos, ViewOrigin, ScreenSize, FOVScreenSize, BoostFactor, bOverrideLocation, 0.0f, InActorToBoost );
+			AddViewInfoToArray( PendingViewInfos, ViewOrigin, ScreenSize, FOVScreenSize, BoostFactor, bOverrideLocation, 0.0f, InActorToBoost, InWorld );
 		}
 	}
 	else
diff --git a/Engine/Source/Runtime/Engine/Private/GameViewportClient.cpp b/Engine/Source/Runtime/Engine/Private/GameViewportClient.cpp
--- a/Engine/Source/Runtime/Engine/Private/GameViewportClient.cpp
+++ b/Engine/Source/Runtime/Engine/Private/GameViewportClient.cpp
@@ -1766,7 +1766,7 @@ void UGameViewportClient::Draw(FViewport* InViewport, FCanvas* SceneCanvas)
 
 					// Add view information for resource streaming. Allow up to 5X boost for small FOV.
 					const float StreamingScale = 1.f / FMath::Clamp<float>(View->LODDistanceFactor, .2f, 1.f);
-					IStreamingManager::Get().AddViewInformation(View->ViewMatrices.GetViewOrigin(), View->UnscaledViewRect.Width(), View->UnscaledViewRect.Width() * View->ViewMatrices.GetProjectionMatrix().M[0][0], StreamingScale);
+					IStreamingManager::Get().AddViewInformation(View->ViewMatrices.GetViewOrigin(), View->UnscaledViewRect.Width(), View->UnscaledViewRect.Width() * View->ViewMatrices.GetProjectionMatrix().M[0][0], StreamingScale, false, 0, nullptr, World);
 					MyWorld->ViewLocationsRenderedLastFrame.Add(View->ViewMatrices.GetViewOrigin());
 					FWorldCachedViewInfo& WorldViewInfo = World->CachedViewInfoRenderedLastFrame.AddDefaulted_GetRef();
 					WorldViewInfo.ViewMatrix = View->ViewMatrices.GetViewMatrix();
diff --git a/Engine/Source/Runtime/Engine/Private/PrecomputedVolumetricLightmapStreaming.cpp b/Engine/Source/Runtime/Engine/Private/PrecomputedVolumetricLightmapStreaming.cpp
--- a/Engine/Source/Runtime/Engine/Private/PrecomputedVolumetricLightmapStreaming.cpp
+++ b/Engine/Source/Runtime/Engine/Private/PrecomputedVolumetricLightmapStreaming.cpp
@@ -17,6 +17,8 @@
 
 using FVersionedMemoryReaderView = TVersionedReader<FMemoryReaderView>;
 
+DEFINE_LOG_CATEGORY_STATIC(LogVolumetricLightmapStreaming, Warning, All);
+
 enum ETimedExecutionControl
 {
 	Continue, 
@@ -91,16 +93,34 @@ public:
 
 	virtual void UpdateResourceStreaming(float DeltaTime, bool bProcessEverything = false) override
 	{
-		 int32 nbViews = IStreamingManager::Get().GetNumViews();
-
-		 //@todo_ow: Manage multiple views properly, add better render/world setting to determine extent / budget
-		 if (nbViews > 0)
-		 {			
-			FBox::FReal StreamDistance = Owner->World->GetWorldSettings()->VolumetricLightmapLoadingRange;
+		 int32 NbViews = IStreamingManager::Get().GetNumViews();
+		 UWorld* ViewWorld = Owner->World;
+		 const FStreamingViewInfo* ViewInfoForWorld = nullptr;
+		 const FStreamingViewInfo* ViewInfoNoWorld = nullptr;
+
+		 // @todo_ow: Improve multiple views handling, add better render/world setting to determine extent / budget
+		 // Find the best ViewInfo for the OwnerWorld
+		 for (int i = 0; i < NbViews; i++)
+		 {
+			 const FStreamingViewInfo& ViewInfo = IStreamingManager::Get().GetViewInformation(i);
+
+			 if (ViewInfo.World == ViewWorld && !ViewInfoForWorld)
+			 {
+				 ViewInfoForWorld = &ViewInfo;
+				 break;
+			 }
+			 else if (!ViewInfo.World.IsValid() && !ViewInfoNoWorld)
+			 {
+				 ViewInfoNoWorld = &ViewInfo;
+			 }
+		 }
 
-			FStreamingViewInfo ViewInfo = IStreamingManager::Get().GetViewInformation(0);
-			FBox Bounds = FBox( ViewInfo.ViewOrigin - FVector(StreamDistance, StreamDistance, StreamDistance), ViewInfo.ViewOrigin  + FVector(StreamDistance, StreamDistance, StreamDistance));
-			Owner->UpdateBounds(Bounds);
+		 if (ViewInfoForWorld || ViewInfoNoWorld)
+		 {
+			 const FStreamingViewInfo& ViewInfo = ViewInfoForWorld ? *ViewInfoForWorld : *ViewInfoNoWorld;
+			 FBox::FReal StreamDistance = ViewWorld->GetWorldSettings()->VolumetricLightmapLoadingRange;
+			 FBox Bounds = FBox(ViewInfo.ViewOrigin - FVector(StreamDistance, StreamDistance, StreamDistance), ViewInfo.ViewOrigin + FVector(StreamDistance, StreamDistance, StreamDistance));
+			 Owner->UpdateBounds(Bounds);
 		 }
 	}
 
@@ -138,12 +158,13 @@ FVolumetricLightmapGridManager::~FVolumetricLightmapGridManager()
 
 int32 FVolumetricLightmapGridManager::GetNumPendingRequests()
 {
+	FScopeLock PendingRequestsLock(&PendingRequestsCriticalSection);
 	return PendingCellRequests.Num();
 }
 
 int32 FVolumetricLightmapGridManager::WaitForPendingRequest(float TimeLimit)
 {
-	bool bTimedOut = TimedExecution(PendingCellRequests, TimeLimit, [this](float TimeLimit, CellRequest& Request) -> ETimedExecutionControl
+	bool bTimedOut = TimedExecution(PendingCellRequests, TimeLimit, [this](float TimeLimit, FCellRequest& Request) -> ETimedExecutionControl
 	{		
 		if (Request.IORequest && Request.IORequest->WaitCompletion(TimeLimit))
 		{	
@@ -158,12 +179,15 @@ int32 FVolumetricLightmapGridManager::WaitForPendingRequest(float TimeLimit)
 }
 
 
-void FVolumetricLightmapGridManager::ReleaseCellData(FVolumetricLightMapGridCell* GridCell, FSceneInterface* InScene)
+void FVolumetricLightmapGridManager::ReleaseCellData(FLoadedCellData* LoadedCell, FSceneInterface* InScene)
 {
-	if (GridCell->Data)
-	{			
-		FPrecomputedVolumetricLightmapData* Data = GridCell->Data;
-		GridCell->Data = nullptr;
+	if (LoadedCell->Data)
+	{
+		FVolumetricLightMapGridCell* Cell = LoadedCell->Cell;
+		UE_LOG(LogVolumetricLightmapStreaming, Log, TEXT("Releasing cell data for streaming cell %i (%p, %s)"), Cell->CellID, Cell, *World->GetFullName());
+
+		FPrecomputedVolumetricLightmapData* Data = LoadedCell->Data;
+		LoadedCell->Data = nullptr;
 
 		ENQUEUE_RENDER_COMMAND(DeleteVolumetricLightDataCommand)
 			([Data] (FRHICommandListBase&)
@@ -176,17 +200,20 @@ void FVolumetricLightmapGridManager::ReleaseCellData(FVolumetricLightMapGridCell
 
 void FVolumetricLightmapGridManager::RemoveFromScene(FSceneInterface* InScene)
 {
+	// In the unlikely event we have pending requests, wait for pending requests to finish so that we can release the data properly
+	WaitForPendingRequest(0.0f);
+
 	for (auto& It : LoadedCells)
 	{
-		FPrecomputedVolumetricLightmap* Lightmap = It.Value;
+		FLoadedCellData& LoadedCell = It.Value;
 		FVolumetricLightMapGridCell* GridCell = It.Key;
 
-		if (Lightmap)
+		if (LoadedCell.Lightmap)
 		{
-			Lightmap->RemoveFromScene(InScene);
+			LoadedCell.Lightmap->RemoveFromScene(InScene);
 		}
 
-		ReleaseCellData(GridCell, InScene);
+		ReleaseCellData(&LoadedCell, InScene);
 	}
 
 	LoadedCells.Reset();
@@ -194,14 +221,21 @@ void FVolumetricLightmapGridManager::RemoveFromScene(FSceneInterface* InScene)
 	StreamingManager.Reset();	
 }
 
-IBulkDataIORequest* FVolumetricLightmapGridManager::RequestVolumetricLightMapCell(FVolumetricLightMapGridCell& Cell)
-{	
-	if (!Cell.BulkData.GetElementCount())
+void FVolumetricLightmapGridManager::RequestVolumetricLightMapCell(FCellRequest& InCellRequest)
+{		
+	FVolumetricLightMapGridCell* Cell = InCellRequest.Cell;
+	check(Cell);
+
+	// should not be already loaded
+	check(!InCellRequest.Data);
+
+	if (!Cell->BulkData.GetElementCount())
 	{
-		return nullptr;
+		InCellRequest.Status = FCellRequest::Ready;
+		return;
 	}
-
-	FBulkDataIORequestCallBack RequestCallback = [&Cell](bool bWasCancelled, IBulkDataIORequest* IORequest)
+		
+	FBulkDataIORequestCallBack RequestCallback = [Cell, this](bool bWasCancelled, IBulkDataIORequest* IORequest)
 	{
 		if (!bWasCancelled)
 		{
@@ -215,11 +249,18 @@ IBulkDataIORequest* FVolumetricLightmapGridManager::RequestVolumetricLightMapCel
 
 				FileDataAr << Data;
 
-				//@todo_ow: this is fine for now since the other thread checks for both IO completion and 
-				// the presence of the ptr, when we implement cancellation we'll need to have proper synchronization
-				check(Data);
-				check(!Cell.Data);
-				Cell.Data = Data;
+				{
+					FScopeLock PendingRequestsLock(&PendingRequestsCriticalSection);
+
+					FCellRequest* Request = PendingCellRequests.FindByPredicate([Cell](FCellRequest& Request) { return Request.Cell == Cell; });
+					check(Request);
+					
+					UE_LOG(LogVolumetricLightmapStreaming, Log, TEXT("IO Request Callback for streaming cell %i (%p, %s)"), Cell->CellID, Cell, *this->World->GetFullName());
+
+					check(Data);
+					check(!Request->Data); // redundant check added because this runs in the async thread
+					Request->Data = Data;
+				}
 				
 				FMemory::Free(Memory);
 			}
@@ -227,25 +268,41 @@ IBulkDataIORequest* FVolumetricLightmapGridManager::RequestVolumetricLightMapCel
 		else
 		{
 			check(!IORequest->GetReadResults());
-			check(!Cell.Data);
 		}
 	};
 
 	TRACE_IOSTORE_METADATA_SCOPE_TAG("PrecomputedVolumetricLightmap");
-	IBulkDataIORequest* IORequest = Cell.BulkData.CreateStreamingRequest(AIOP_Normal, &RequestCallback, nullptr);
 
-	if (!IORequest)
+	if (!Cell->BulkData.IsBulkDataLoaded())
 	{
-		check(Cell.BulkData.IsBulkDataLoaded());
-		// For unsaved data we can't stream it, we need to do an immediate load
-		Grid->LoadVolumetricLightMapCell(Cell);
+		UE_LOG(LogVolumetricLightmapStreaming, Log, TEXT("Request streaming for cell %i (%p, %s)"), Cell->CellID, Cell, *World->GetFullName());
+		InCellRequest.IORequest = Cell->BulkData.CreateStreamingRequest(AIOP_Normal, &RequestCallback, nullptr);
+		InCellRequest.Status = FCellRequest::Requested;
 	}
-
-	return IORequest;
+#if WITH_EDITOR
+	else
+	{
+		UE_LOG(LogVolumetricLightmapStreaming, Log, TEXT("Loading streaming cell %i (%p, %s) without streaming it"), Cell->CellID, Cell, *World->GetFullName());
+		// For unsaved data we can't stream it, we need to do an immediate load (Cell.BulkData.IsBulkDataLoaded() will always return true for unsaved data)
+		Grid->LoadVolumetricLightMapCell(*Cell, InCellRequest.Data);
+		InCellRequest.Status = FCellRequest::Ready;
+	}
+#else
+	else
+	{	
+		// We should never end up in this state in non-editor builds
+		check(false);
+		InCellRequest.Status = FCellRequest::Ready;
+	}
+#endif
 }
 
 void FVolumetricLightmapGridManager::UpdateBounds(const FBox& InBounds)
 {
+	UE_LOG(LogVolumetricLightmapStreaming, Verbose, TEXT("Updating bounds: %s, for : %s"), *InBounds.ToString(), *World->GetFullName());
+
+	FScopeLock PendingRequestsLock(&PendingRequestsCriticalSection);
+
 	//@todo_ow: Make a pass on this to minimize heap allocations
 	check(Grid);	
 
@@ -260,7 +317,7 @@ void FVolumetricLightmapGridManager::UpdateBounds(const FBox& InBounds)
 
 	for (FVolumetricLightMapGridCell* Cell : IntersectingCells)
 	{
-		if (FPrecomputedVolumetricLightmap** LightMap = LoadedCells.Find(Cell))
+		if (FLoadedCellData* LoadedCell = LoadedCells.Find(Cell))
 		{		
 			CellsToRemove.Remove(Cell);
 		}
@@ -273,33 +330,34 @@ void FVolumetricLightmapGridManager::UpdateBounds(const FBox& InBounds)
 	// Make the necessary IO requests for new cells
 	for (FVolumetricLightMapGridCell* Cell : CellsToRequest)
 	{
-		if (!PendingCellRequests.ContainsByPredicate([Cell](CellRequest& Request) { return Request.Cell == Cell; }))
+		if (!PendingCellRequests.ContainsByPredicate([Cell](FCellRequest& Request) { return Request.Cell == Cell; }))
 		{
-			CellRequest Request;
-			Request.Cell = Cell;
-			Request.IORequest = RequestVolumetricLightMapCell(*Cell);	
-			if (Request.IORequest)
-			{
-				Request.Status = CellRequest::Requested;
-			}
-			else
-			{
-				Request.Status = CellRequest::Ready;
-			}
+			FCellRequest Request { .Cell = Cell };
+
+			RequestVolumetricLightMapCell(Request);
 
 			PendingCellRequests.Add(Request);
 		}
 	}
 
+#if DO_CHECK
+	// Since cells to remove are obtained from subtracting all IntersectingCells from the LoadedCells set, pending requested cells 
+	// should never be in CellsToRemove
+	for (FVolumetricLightMapGridCell* Cell : CellsToRemove)
+	{	
+		check(!PendingCellRequests.FindByPredicate([Cell](FCellRequest& InRequest) { return InRequest.Cell == Cell; }));
+	}
+#endif
+
 	// Since cells to remove are obtained from subtracting all IntersectingCells from the LoadedCells set, pending requested cells are never
 	// in CellsToRemove, so this bit is not useful for now. We could extract cells to remove from the pending rq
 	// @todo_ow: Handle removal of requested cells, an optimization for rarer/edgier cases so we'll do this later, plus this ones needs actual synchronization with the async request
 	/*for (FVolumetricLightMapGridCell* Cell : CellsToRemove)
 	{	
-		if (CellRequest* Request = PendingCellRequests.FindByPredicate([Cell](CellRequest& InRequest) { return InRequest.Cell == Cell; }))
+		if (FCellRequest* Request = PendingCellRequests.FindByPredicate([Cell](FCellRequest& InRequest) { return InRequest.Cell == Cell; }))
 		{
-			check(Request->Status == CellRequest::Requested);
-			Request->Status = CellRequest::Cancelled;			
+			check(Request->Status == FCellRequest::Requested);
+			Request->Status = FCellRequest::Cancelled;			
 			Request->IORequest->Cancel();
 			delete Request->IORequest;
 			Request->IORequest = nullptr;
@@ -309,15 +367,15 @@ void FVolumetricLightmapGridManager::UpdateBounds(const FBox& InBounds)
 	// Remove all unnecessary cells
 	for (FVolumetricLightMapGridCell* Cell : CellsToRemove)
 	{
-		FPrecomputedVolumetricLightmap* Lightmap = LoadedCells[Cell];
-		if (Lightmap)
-		{			
-			Lightmap->RemoveFromScene(World->Scene);
+		FLoadedCellData& LoadedCell = LoadedCells[Cell];
+		if (LoadedCell.Lightmap)
+		{
+			LoadedCell.Lightmap->RemoveFromScene(World->Scene);
 		}
 
-		LoadedCells.Remove(Cell);
+		verify(LoadedCells.Remove(Cell) > 0);
 
-		ReleaseCellData(Cell, World->Scene);
+		ReleaseCellData(&LoadedCell, World->Scene);
 	}
 
 	// Update currently tracked bounds
@@ -328,29 +386,32 @@ void FVolumetricLightmapGridManager::UpdateBounds(const FBox& InBounds)
 
 int32 FVolumetricLightmapGridManager::ProcessRequests()
 {	
-	TArray<FVolumetricLightMapGridCell*> CellsToAdd;
+	FScopeLock PendingRequestsLock(&PendingRequestsCriticalSection);
+	
+	TArray<FCellRequest*> CellsToAdd;
 
 	// Process pending IO requests and move to add list if ready
 	//@todo_ow: Use an inline array and just mark the requests that need to be removed
-	TArray<CellRequest> UpdatedRequests;
+	TArray<FCellRequest> UpdatedRequests;
 
-	for (CellRequest& Request : PendingCellRequests)
+	for (FCellRequest& Request : PendingCellRequests)
 	{
-		if (Request.Status == CellRequest::Ready)
+		if (Request.Status == FCellRequest::Ready)
 		{
-			check(Request.Cell->Data);
-			check(Request.IORequest == nullptr);
-			CellsToAdd.Add(Request.Cell);
+			check(Request.Data);
+			check(!Request.IORequest);
+			CellsToAdd.Add(&Request);
 		}
-		else if (Request.Status == CellRequest::Requested)
+		else if (Request.Status == FCellRequest::Requested)
 		{
-			if (Request.IORequest->PollCompletion() && Request.Cell->Data)
+			if (Request.IORequest->PollCompletion())
 			{
-				CellsToAdd.Add(Request.Cell);
+				check(Request.Data);
+				CellsToAdd.Add(&Request);
 				
 				delete Request.IORequest;
 				Request.IORequest = nullptr;
-				Request.Status = CellRequest::Ready;
+				Request.Status = FCellRequest::Ready;
 			}
 			else
 			{
@@ -363,23 +424,26 @@ int32 FVolumetricLightmapGridManager::ProcessRequests()
 		}
 	}
 
-	PendingCellRequests = MoveTemp(UpdatedRequests);
-
 	// Add all necessary cells
-	for (FVolumetricLightMapGridCell* Cell : CellsToAdd)
+	for (FCellRequest* FCellRequest : CellsToAdd)
 	{
-		FPrecomputedVolumetricLightmap* Lightmap = nullptr;
-		if (Cell->Data)
-		{	
-			Lightmap = new FPrecomputedVolumetricLightmap();
-			Lightmap->AddToScene(World->Scene, Registry, Cell->Data, false);
+		FLoadedCellData LoadedCell { .Data = FCellRequest->Data, .Cell = FCellRequest->Cell };
+
+		if (LoadedCell.Data)
+		{
+			LoadedCell.Lightmap = new FPrecomputedVolumetricLightmap();
+			LoadedCell.Lightmap->AddToScene(World->Scene, Registry, LoadedCell.Data, false);
 		}
 		else
 		{
-			check(Cell->BulkData.GetElementCount() == 0);
+			check(LoadedCell.Cell->BulkData.GetElementCount() == 0);
 		}
-		LoadedCells.Add(Cell, Lightmap);
+
+		LoadedCells.Add(LoadedCell.Cell, LoadedCell);
 	}
 
+	// We can now overwrite the pending requests array since we're not using ptrs into it anymore
+	PendingCellRequests = MoveTemp(UpdatedRequests);
+
 	return PendingCellRequests.Num();
 }
\ No newline at end of file
diff --git a/Engine/Source/Runtime/Engine/Private/StaticLightingBuildContext.cpp b/Engine/Source/Runtime/Engine/Private/StaticLightingBuildContext.cpp
--- a/Engine/Source/Runtime/Engine/Private/StaticLightingBuildContext.cpp
+++ b/Engine/Source/Runtime/Engine/Private/StaticLightingBuildContext.cpp
@@ -162,11 +162,11 @@ FPrecomputedVolumetricLightmapData& FStaticLightingBuildContext::GetOrCreateLeve
 	{
 		if (FVolumetricLightMapGridCell* Cell = VolumetricLightMapGridDesc->GetCell(LevelId))
 		{	
-			if (!Cell->Data)
+			if (!Cell->EditorData)
 			{
-				Cell->Data = new FPrecomputedVolumetricLightmapData();
+				Cell->EditorData = new FPrecomputedVolumetricLightmapData();
 			}
-			return *Cell->Data;
+			return *Cell->EditorData;
 		}
 	}
 
diff --git a/Engine/Source/Runtime/Engine/Private/WorldPartition/StaticLightingData/VolumetricLightmapGrid.cpp b/Engine/Source/Runtime/Engine/Private/WorldPartition/StaticLightingData/VolumetricLightmapGrid.cpp
--- a/Engine/Source/Runtime/Engine/Private/WorldPartition/StaticLightingData/VolumetricLightmapGrid.cpp
+++ b/Engine/Source/Runtime/Engine/Private/WorldPartition/StaticLightingData/VolumetricLightmapGrid.cpp
@@ -113,34 +113,48 @@ FGuid FVolumetricLightMapGridDesc::GetCellGuid(uint32 CellID) const
 
 FPrecomputedVolumetricLightmapData* FVolumetricLightMapGridDesc::GetPrecomputedVolumetricLightmapBuildData(FGuid LevelId) const
 {
+#if WITH_EDITOR	
 	if (FVolumetricLightMapGridCell* Cell = GetCell(LevelId))
 	{
-		return Cell->Data;
+		return Cell->EditorData;
 	}
+#endif
 
 	return nullptr;
 }
-	
+
+
 FPrecomputedVolumetricLightmapData* FVolumetricLightMapGridDesc::GetOrCreatePrecomputedVolumetricLightmapBuildData(FGuid LevelId)
 {
+#if WITH_EDITOR	
 	FVolumetricLightMapGridCell* Cell = GetCell(LevelId);
 	check(Cell); 
 
-	if (!Cell->Data)
+	if (!Cell->EditorData)
 	{
-		Cell->Data = new FPrecomputedVolumetricLightmapData();
+		Cell->EditorData  = new FPrecomputedVolumetricLightmapData();
 	}
 
-	return Cell->Data;
+	return Cell->EditorData;
+#else
+	return nullptr;
+#endif
 }
 
+
 TArray<FVolumetricLightMapGridCell*> FVolumetricLightMapGridDesc::GetIntersectingCells(const FBox& InBounds, bool bInWithData)
 {
 	TArray<FVolumetricLightMapGridCell*> IntersectingCells;
 
 	for (FVolumetricLightMapGridCell& Cell : Cells)
-	{
-		if (Cell.Bounds.Intersect(InBounds) && (Cell.BulkData.GetElementCount() || Cell.Data || !bInWithData) )
+	{		
+#if WITH_EDITOR
+		bool bCellHasData = Cell.BulkData.GetElementCount() || Cell.EditorData;
+#else
+		bool bCellHasData = Cell.BulkData.GetElementCount() > 0;
+#endif
+
+		if (Cell.Bounds.Intersect(InBounds) && (bCellHasData || !bInWithData) )
 		{
 			IntersectingCells.Add(&Cell);
 		}
@@ -152,9 +166,9 @@ TArray<FVolumetricLightMapGridCell*> FVolumetricLightMapGridDesc::GetIntersectin
 void FVolumetricLightMapGridDesc::SerializeBulkData(FArchive& Ar, UObject* Owner)
 {	
 	for (FVolumetricLightMapGridCell& Cell : Cells)
-	{		
-		Cell.BulkData.Serialize(Ar, Owner);		
-	}	
+	{
+		Cell.BulkData.Serialize(Ar, Owner);
+	}
 }
 
 #if WITH_EDITOR
@@ -165,24 +179,25 @@ void FVolumetricLightMapGridDesc::InitializeBulkData()
 	{
 		Cell.BulkData.SetBulkDataFlags(BULKDATA_Force_NOT_InlinePayload);
 
-		if (!Cell.Data)
+		if (!Cell.EditorData)
 		{
 			Cell.BulkData.RemoveBulkData();
 			continue;
 		}
 			
 		FVersionedBulkDataWriter Ar(Cell.BulkData, true);
-		Ar << Cell.Data;
+		FPrecomputedVolumetricLightmapData* LightmapData = Cell.EditorData;
+		Ar << LightmapData;
 	}
 }
 #endif
 
-void FVolumetricLightMapGridDesc::LoadVolumetricLightMapCell(FVolumetricLightMapGridCell& Cell)
+void FVolumetricLightMapGridDesc::LoadVolumetricLightMapCell(FVolumetricLightMapGridCell& Cell, FPrecomputedVolumetricLightmapData*& OutData)
 {
 	if (Cell.BulkData.GetElementCount())
 	{
 		FVersionedBulkDataReader Ar(Cell.BulkData, true);
-		Ar << Cell.Data;
+		Ar << OutData;
 	}
 }
 
diff --git a/Engine/Source/Runtime/Engine/Public/ContentStreaming.h b/Engine/Source/Runtime/Engine/Public/ContentStreaming.h
--- a/Engine/Source/Runtime/Engine/Public/ContentStreaming.h
+++ b/Engine/Source/Runtime/Engine/Public/ContentStreaming.h
@@ -72,7 +72,7 @@ enum ERemoveStreamingViews
  */
 struct FStreamingViewInfo
 {
-	FStreamingViewInfo( const FVector& InViewOrigin, float InScreenSize, float InFOVScreenSize, float InBoostFactor, bool bInOverrideLocation, float InDuration, TWeakObjectPtr<AActor> InActorToBoost )
+	FStreamingViewInfo( const FVector& InViewOrigin, float InScreenSize, float InFOVScreenSize, float InBoostFactor, bool bInOverrideLocation, float InDuration, TWeakObjectPtr<AActor> InActorToBoost, TWeakObjectPtr<UWorld> InWorld)
 	:	ViewOrigin( InViewOrigin )
 	,	ScreenSize( InScreenSize )
 	,	FOVScreenSize( InFOVScreenSize )
@@ -80,6 +80,7 @@ struct FStreamingViewInfo
 	,	Duration( InDuration )
 	,	bOverrideLocation( bInOverrideLocation )
 	,	ActorToBoost( InActorToBoost )
+	,	World( InWorld )
 	{
 	}
 	/** View origin */
@@ -96,6 +97,8 @@ struct FStreamingViewInfo
 	bool	bOverrideLocation;
 	/** Optional pointer to an actor who's textures should have their streaming priority boosted */
 	TWeakObjectPtr<AActor> ActorToBoost;
+	/** Optional pointer to the world the view is located into */
+	TWeakObjectPtr<UWorld> World;
 };
 
 /**
@@ -234,8 +237,9 @@ struct IStreamingManager
 	 * @param bOverrideLocation		Whether this is an override location, which forces the streaming system to ignore all other regular locations
 	 * @param Duration				How long the streaming system should keep checking this location, in seconds. 0 means just for the next Tick.
 	 * @param InActorToBoost		Optional pointer to an actor who's textures should have their streaming priority boosted
+	 * @param InWorld				The world this view is located in
 	 */
-	ENGINE_API void AddViewInformation(const FVector& ViewOrigin, float ScreenSize, float FOVScreenSize, float BoostFactor = 1.0f, bool bOverrideLocation = false, float Duration = 0.0f, TWeakObjectPtr<AActor> InActorToBoost = NULL);
+	ENGINE_API void AddViewInformation(const FVector& ViewOrigin, float ScreenSize, float FOVScreenSize, float BoostFactor = 1.0f, bool bOverrideLocation = false, float Duration = 0.0f, TWeakObjectPtr<AActor> InActorToBoost = NULL, TWeakObjectPtr<UWorld> InWorld = NULL);
 
 	/**
 	 * Queue up view locations to the streaming system. These locations will be added properly at the next call to AddViewInformation,
@@ -360,8 +364,9 @@ protected:
 	 * @param bOverrideLocation		Whether this is an override location, which forces the streaming system to ignore all other regular locations
 	 * @param Duration				How long the streaming system should keep checking this location (in seconds). 0 means just for the next Tick.
 	 * @param InActorToBoost		Optional pointer to an actor who's textures should have their streaming priority boosted
+	 * @param InWorld				The world this view is located in
 	 */
-	static void AddViewInfoToArray( TArray<FStreamingViewInfo> &ViewInfos, const FVector& ViewOrigin, float ScreenSize, float FOVScreenSize, float BoostFactor, bool bOverrideLocation, float Duration, TWeakObjectPtr<AActor> InActorToBoost );
+	static void AddViewInfoToArray( TArray<FStreamingViewInfo> &ViewInfos, const FVector& ViewOrigin, float ScreenSize, float FOVScreenSize, float BoostFactor, bool bOverrideLocation, float Duration, TWeakObjectPtr<AActor> InActorToBoost, TWeakObjectPtr<UWorld> InWorld);
 
 	/**
 	 * Remove view infos with the same location from the given array.
diff --git a/Engine/Source/Runtime/Engine/Public/PrecomputedVolumetricLightmapStreaming.h b/Engine/Source/Runtime/Engine/Public/PrecomputedVolumetricLightmapStreaming.h
--- a/Engine/Source/Runtime/Engine/Public/PrecomputedVolumetricLightmapStreaming.h
+++ b/Engine/Source/Runtime/Engine/Public/PrecomputedVolumetricLightmapStreaming.h
@@ -14,6 +14,7 @@ struct FVolumetricLightMapGridDesc;
 struct FVolumetricLightMapGridCell;
 class FVolumetricLightmapGridStreamingManager;
 class UMapBuildDataRegistry;
+class FPrecomputedVolumetricLightmapData;
 
 class FVolumetricLightmapGridManager
 {	
@@ -26,19 +27,17 @@ public:
 	int32 ProcessRequests();
 	void RemoveFromScene(class FSceneInterface* InScene);
 	int32 WaitForPendingRequest(float Timelimit);
-	int32 GetNumPendingRequests();	
+	int32 GetNumPendingRequests();
 
 private:
 	
 	friend class FVolumetricLightmapGridStreamingManager;
 
-	void ReleaseCellData(FVolumetricLightMapGridCell* GridCell, FSceneInterface* InScene);	
-
 	UWorld* World = nullptr;
 	UMapBuildDataRegistry* Registry = nullptr;
 	FVolumetricLightMapGridDesc* Grid = nullptr;
 	
-	struct CellRequest
+	struct FCellRequest
 	{
 		enum EStatus
 		{
@@ -51,16 +50,27 @@ private:
 
 		EStatus Status = Created;
 		IBulkDataIORequest* IORequest = nullptr;
-		FVolumetricLightMapGridCell* Cell = nullptr;			
+		TNotNull<FVolumetricLightMapGridCell*> Cell;
+		FPrecomputedVolumetricLightmapData* Data = nullptr;
 	};	
 
-	IBulkDataIORequest* RequestVolumetricLightMapCell(FVolumetricLightMapGridCell& Cell);
-	
-	TArray<CellRequest>	PendingCellRequests;
-	TMap<FVolumetricLightMapGridCell*, FPrecomputedVolumetricLightmap*> LoadedCells;	
+	struct FLoadedCellData
+	{
+		FPrecomputedVolumetricLightmapData* Data = nullptr;
+		TNotNull<FVolumetricLightMapGridCell*> Cell;
+		FPrecomputedVolumetricLightmap* Lightmap = nullptr;
+	};
 
+	void RequestVolumetricLightMapCell(FCellRequest& CellRequest);
+	void ReleaseCellData(FLoadedCellData* LoadedCell, FSceneInterface* InScene);
+	
+	TArray<FCellRequest>	PendingCellRequests;
+	TMap<FVolumetricLightMapGridCell*, FLoadedCellData> LoadedCells;
+	
 	FBox Bounds;
 
+	FCriticalSection PendingRequestsCriticalSection;
+
 	TUniquePtr<FVolumetricLightmapGridStreamingManager> StreamingManager;
 };
 
diff --git a/Engine/Source/Runtime/Engine/Public/WorldPartition/StaticLightingData/VolumetricLightmapGrid.h b/Engine/Source/Runtime/Engine/Public/WorldPartition/StaticLightingData/VolumetricLightmapGrid.h
--- a/Engine/Source/Runtime/Engine/Public/WorldPartition/StaticLightingData/VolumetricLightmapGrid.h
+++ b/Engine/Source/Runtime/Engine/Public/WorldPartition/StaticLightingData/VolumetricLightmapGrid.h
@@ -19,9 +19,13 @@ struct FVolumetricLightMapGridCell
 	UPROPERTY()
 	FBox Bounds;
 
-	FPrecomputedVolumetricLightmapData* Data = nullptr;	
+	FByteBulkData BulkData;
 
-	FByteBulkData		BulkData;
+#if WITH_EDITOR
+	// ptr to the just recently built version of the FPrecomputedVolumetricLightmapData
+	// not used at runtime
+	FPrecomputedVolumetricLightmapData* EditorData = nullptr;
+#endif
 	
 	UPROPERTY()
 	uint32 CellID = -1;
@@ -60,7 +64,7 @@ struct FVolumetricLightMapGridDesc
 	
 	[[nodiscard]] ENGINE_API TArray<FVolumetricLightMapGridCell*> GetIntersectingCells(const FBox& InBounds, bool bInWithData);	
 	
-	ENGINE_API void LoadVolumetricLightMapCell(FVolumetricLightMapGridCell& Cell);
+	ENGINE_API void LoadVolumetricLightMapCell(FVolumetricLightMapGridCell& Cell, FPrecomputedVolumetricLightmapData*& OutData);
 
 	ENGINE_API FPrecomputedVolumetricLightmapData* GetPrecomputedVolumetricLightmapBuildData(FGuid LevelId) const;
 	ENGINE_API FPrecomputedVolumetricLightmapData* GetOrCreatePrecomputedVolumetricLightmapBuildData(FGuid LevelId);
