RBAN - cvar to force use of linear or non-linear solver.
Intended for testing and debugging only.
Will affect any system that uses the ImmediatePhysics solver.
Will only take effect when joints (or solver/RBAN nodes) are recreated.

diff --git a/Engine/Source/Runtime/Engine/Private/PhysicsEngine/ImmediatePhysics/ImmediatePhysicsAdapters.cpp b/Engine/Source/Runtime/Engine/Private/PhysicsEngine/ImmediatePhysics/ImmediatePhysicsAdapters.cpp
--- a/Engine/Source/Runtime/Engine/Private/PhysicsEngine/ImmediatePhysics/ImmediatePhysicsAdapters.cpp
+++ b/Engine/Source/Runtime/Engine/Private/PhysicsEngine/ImmediatePhysics/ImmediatePhysicsAdapters.cpp
@@ -18,6 +18,13 @@
 #include "PhysicsEngine/ConstraintInstance.h"
 #include "PhysicsProxy/SingleParticlePhysicsProxy.h"
 
+int32 RBANForceJointSolverType = -1;
+FAutoConsoleVariableRef CVarRigidBodyNodeForceJointSolverType(
+	TEXT("p.RigidBodyNode.ForceJointSolverType"), 
+	RBANForceJointSolverType, 
+	TEXT("-1: RBAN joints use solver specific in settings; 0: All RBAN joints use the linear solver. 1: All RBAN joints use non-linear solver"), 
+	ECVF_Default);
+
 namespace ImmediatePhysics
 {
 	bool CreateDefaultGeometry(const FVector& Scale, Chaos::FReal& OutMass, Chaos::FVec3& OutInertia, Chaos::FRigidTransform3& OutCoMTransform, Chaos::FImplicitObjectPtr& OutGeom, TArray<TUniquePtr<Chaos::FPerShapeData>>& OutShapes)
@@ -380,6 +387,10 @@ namespace ImmediatePhysics
 		JointSettings.bMassConditioningEnabled = Profile.bEnableMassConditioning;
 
 		JointSettings.bUseLinearSolver = Profile.bUseLinearJointSolver;
+		if (RBANForceJointSolverType >= 0)
+		{
+			JointSettings.bUseLinearSolver = (RBANForceJointSolverType == 0);
+		}
 
 		UpdateJointSettingsFromLinearConstraint(Profile.LinearLimit, JointSettings);
 		UpdateJointSettingsFromConeConstraint(Profile.ConeLimit, JointSettings);
