Fix Inheritance path not working properly for non-volume modulation settings

diff --git a/Engine/Source/Runtime/Engine/Private/Audio.cpp b/Engine/Source/Runtime/Engine/Private/Audio.cpp
--- a/Engine/Source/Runtime/Engine/Private/Audio.cpp
+++ b/Engine/Source/Runtime/Engine/Private/Audio.cpp
@@ -276,7 +276,7 @@ void FWaveInstance::CalculateInheritedModulationSettingsAndValue(const EModulati
 			}
 		};
 
-	switch (InActiveSound.ModulationRouting.VolumeRouting)
+	switch (InActiveSound.ModulationRouting.GetRoutingForDestination(InDestination))
 	{
 	case EModulationRouting::Union:
 	{
