fix incorrect ifdef check for PLATFORM_WINDOWS_ARM64EC

#rnx

diff --git a/Engine/Source/Runtime/Engine/Private/ShaderCompiler/ShaderCompiler.cpp b/Engine/Source/Runtime/Engine/Private/ShaderCompiler/ShaderCompiler.cpp
--- a/Engine/Source/Runtime/Engine/Private/ShaderCompiler/ShaderCompiler.cpp
+++ b/Engine/Source/Runtime/Engine/Private/ShaderCompiler/ShaderCompiler.cpp
@@ -728,7 +728,7 @@ FShaderCompilingManager::FShaderCompilingManager() :
 	WorkersBusyTime = 0;
 
 #if PLATFORM_WINDOWS
-#ifdef PLATFORM_WINDOWS_ARM64EC
+#if PLATFORM_WINDOWS_ARM64EC
 	FString ExecutableName("ShaderCompileWorkerarm64ec.exe");
 #elif PLATFORM_CPU_ARM_FAMILY
 	FString ExecutableName("ShaderCompileWorkerarm64.exe");
