Texture update shaders now compile on all platforms, including mobile

#rnx

diff --git a/Engine/Source/Runtime/Engine/Public/UpdateTextureShaders.h b/Engine/Source/Runtime/Engine/Public/UpdateTextureShaders.h
--- a/Engine/Source/Runtime/Engine/Public/UpdateTextureShaders.h
+++ b/Engine/Source/Runtime/Engine/Public/UpdateTextureShaders.h
@@ -27,11 +27,6 @@ public:
 	{
 	}
 
-	static bool ShouldCompilePermutation(const FGlobalShaderPermutationParameters& Parameters)
-	{
-		return IsFeatureLevelSupported(Parameters.Platform, ERHIFeatureLevel::SM5);
-	}
-
 	static inline void ModifyCompilationEnvironment(const FGlobalShaderPermutationParameters& Parameters, FShaderCompilerEnvironment& OutEnvironment)
 	{
 		FGlobalShader::ModifyCompilationEnvironment(Parameters, OutEnvironment);
