[NewMaterialTranslator] - Fix crash bug when plugging in a Texture Object as a Material Function Input.

- When iterating over all the function inputs and determining their types we call `MIR::FType::FromMaterialValueType`, which wasn't handling Texture2D at all.

#rnx

diff --git a/Engine/Source/Runtime/Engine/Private/Materials/MaterialIRTypes.cpp b/Engine/Source/Runtime/Engine/Private/Materials/MaterialIRTypes.cpp
--- a/Engine/Source/Runtime/Engine/Private/Materials/MaterialIRTypes.cpp
+++ b/Engine/Source/Runtime/Engine/Private/Materials/MaterialIRTypes.cpp
@@ -70,6 +70,9 @@ const FType* FType::FromMaterialValueType(EMaterialValueType Type)
 		case MCT_Float:
 			return FPrimitiveType::GetVector(ScalarKind_Float, 4);
 
+		case MCT_Texture2D:
+			return FObjectType::GetTexture2D();
+
 		case MCT_UInt1:
 			return FPrimitiveType::GetVector(ScalarKind_Int, 1);
 		case MCT_UInt2:
