Invalidate GPU Scene Lights entry when updating directional light scene proxy.
- Extended FSceneInterface with a helper function UpdateLightProxy(...) that enqueues a command to update to invalidate the corresponding state.

diff --git a/Engine/Source/Runtime/Engine/Classes/Components/LightComponent.h b/Engine/Source/Runtime/Engine/Classes/Components/LightComponent.h
--- a/Engine/Source/Runtime/Engine/Classes/Components/LightComponent.h
+++ b/Engine/Source/Runtime/Engine/Classes/Components/LightComponent.h
@@ -12,6 +12,7 @@
 #include "Components/LightComponentBase.h"
 #include "LightComponent.generated.h"
 
+class FLightSceneProxy;
 class FLightComponentMapBuildData;
 class FStaticLightingBuildContext;
 class FStaticShadowDepthMapData;
@@ -21,6 +22,8 @@ class UPrimitiveComponent;
 class UTextureLightProfile;
 enum class ELightUnits : uint8;
 
+using FUpdateLightProxyCallback = TFunction<void(FLightSceneProxy*)>;
+
 /** 
  * A texture containing depth values of static objects that was computed during the lighting build.
  * Used by Stationary lights to shadow translucency.
@@ -356,7 +359,7 @@ public:
 
 public:
 	/** The light's scene info. */
-	class FLightSceneProxy* SceneProxy;
+	FLightSceneProxy* SceneProxy;
 
 	/**
 	 * Pushes new selection state to the render thread light proxy
@@ -498,6 +501,8 @@ protected:
 #endif // WITH_EDITOR
 	//~ End USceneComponent Interface
 
+	void UpdateProxy(FUpdateLightProxyCallback&& Func);
+
 private:
 
 	bool CanTraceDistanceFieldShadows() const;
diff --git a/Engine/Source/Runtime/Engine/Private/Components/DirectionalLightComponent.cpp b/Engine/Source/Runtime/Engine/Private/Components/DirectionalLightComponent.cpp
--- a/Engine/Source/Runtime/Engine/Private/Components/DirectionalLightComponent.cpp
+++ b/Engine/Source/Runtime/Engine/Private/Components/DirectionalLightComponent.cpp
@@ -372,63 +372,9 @@ public:
 		ShadowAmount = Component->ShadowAmount;
 	}
 
-	void UpdateLightShaftOverrideDirection_GameThread(const UDirectionalLightComponent* Component)
+	void SetAtmosphereSunDiskColorScale(FLinearColor NewValue)
 	{
-		FVector NewLightShaftOverrideDirection = Component->LightShaftOverrideDirection;
-		NewLightShaftOverrideDirection.Normalize();
-		FDirectionalLightSceneProxy* Proxy = this;
-		ENQUEUE_RENDER_COMMAND(FUpdateLightShaftOverrideDirectionCommand)(
-			[Proxy, NewLightShaftOverrideDirection](FRHICommandList& RHICmdList)
-			{
-				Proxy->UpdateLightShaftOverrideDirection_RenderThread(NewLightShaftOverrideDirection);
-			});
-	}
-
-	void UpdateLightSourceAngle_GameThread(const UDirectionalLightComponent* Component)
-	{
-		ENQUEUE_RENDER_COMMAND(FUpdateLightSourceAngleCommand)(
-			[Proxy = this, LightSourceAngle = Component->LightSourceAngle](FRHICommandList& RHICmdList)
-		{
-			Proxy->LightSourceAngle = LightSourceAngle;
-		});
-	}
-
-	void UpdateLightSourceSoftAngle_GameThread(const UDirectionalLightComponent* Component)
-	{
-		ENQUEUE_RENDER_COMMAND(FUpdateLightSourceSoftAngleCommand)(
-			[Proxy = this, NewLightSourceSoftAngle = Component->LightSourceSoftAngle](FRHICommandList& RHICmdList)
-		{
-			Proxy->LightSourceSoftAngle = NewLightSourceSoftAngle;
-		});
-	}
-
-	void UpdateShadowSourceAngleFactor_GameThread(const UDirectionalLightComponent* Component)
-	{
-		ENQUEUE_RENDER_COMMAND(FUpdateShadowSourceAngleFactorCommand)(
-			[Proxy = this, NewShadowSourceAngleFactor = Component->ShadowSourceAngleFactor](FRHICommandList& RHICmdList)
-		{
-			Proxy->ShadowSourceAngleFactor = NewShadowSourceAngleFactor;
-		});
-	}
-
-	void UpdateEnableLightShaftOcclusion_GameThread(const UDirectionalLightComponent* Component)
-	{
-		ENQUEUE_RENDER_COMMAND(FUpdateEnableLightShaftOcclusionCommand)(
-			[Proxy = this, bNewEnableLightShaftOcclusion = (bool)Component->bEnableLightShaftOcclusion](FRHICommandList& RHICmdList)
-		{
-			Proxy->bEnableLightShaftOcclusion = bNewEnableLightShaftOcclusion;
-		});
-	}
-
-	void UpdateAtmosphereSunDiskColorScale_GameThread(const UDirectionalLightComponent* Component)
-	{
-		FLinearColor NewAtmosphereSunDiskColorScale = Component->AtmosphereSunDiskColorScale;
-		FDirectionalLightSceneProxy* Proxy = this;
-		ENQUEUE_RENDER_COMMAND(FUpdateLightShaftOverrideDirectionCommand)(
-			[Proxy, NewAtmosphereSunDiskColorScale](FRHICommandList& RHICmdList)
-			{
-				Proxy->AtmosphereSunDiskColorScale = NewAtmosphereSunDiskColorScale;
-			});
+		AtmosphereSunDiskColorScale = NewValue;
 	}
 
 	/** Accesses parameters needed for rendering the light. */
@@ -1324,11 +1270,11 @@ void UDirectionalLightComponent::SetEnableLightShaftOcclusion(bool bNewValue)
 		&& bEnableLightShaftOcclusion != bNewValue)
 	{
 		bEnableLightShaftOcclusion = bNewValue;
-		if (SceneProxy)
-		{
-			FDirectionalLightSceneProxy* DirectionalLightSceneProxy = (FDirectionalLightSceneProxy*)SceneProxy;
-			DirectionalLightSceneProxy->UpdateEnableLightShaftOcclusion_GameThread(this);
-		}
+		UpdateProxy([bEnableLightShaftOcclusion = (bool)bEnableLightShaftOcclusion](FLightSceneProxy* Proxy)
+			{
+				FDirectionalLightSceneProxy* DirectionalLightSceneProxy = (FDirectionalLightSceneProxy*)Proxy;
+				DirectionalLightSceneProxy->bEnableLightShaftOcclusion = bEnableLightShaftOcclusion;
+			});
 	}
 }
 
@@ -1357,11 +1303,11 @@ void UDirectionalLightComponent::SetLightShaftOverrideDirection(FVector NewValue
 		&& LightShaftOverrideDirection != NewValue)
 	{
 		LightShaftOverrideDirection = NewValue;
-		if (SceneProxy)
-		{
-			FDirectionalLightSceneProxy* DirectionalLightSceneProxy = (FDirectionalLightSceneProxy*)SceneProxy;
-			DirectionalLightSceneProxy->UpdateLightShaftOverrideDirection_GameThread(this);
-		}
+		UpdateProxy([LightShaftOverrideDirection = LightShaftOverrideDirection](FLightSceneProxy* Proxy)
+			{
+				FDirectionalLightSceneProxy* DirectionalLightSceneProxy = (FDirectionalLightSceneProxy*)Proxy;
+				DirectionalLightSceneProxy->LightShaftOverrideDirection = LightShaftOverrideDirection;
+			});
 	}
 }
 
@@ -1371,11 +1317,11 @@ void UDirectionalLightComponent::SetLightSourceAngle(float NewValue)
 		&& LightSourceAngle != NewValue)
 	{
 		LightSourceAngle = NewValue;
-		if (SceneProxy)
-		{
-			FDirectionalLightSceneProxy* DirectionalLightSceneProxy = (FDirectionalLightSceneProxy*)SceneProxy;
-			DirectionalLightSceneProxy->UpdateLightSourceAngle_GameThread(this);
-		}
+		UpdateProxy([LightSourceAngle = LightSourceAngle](FLightSceneProxy* Proxy)
+			{
+				FDirectionalLightSceneProxy* DirectionalLightSceneProxy = (FDirectionalLightSceneProxy*)Proxy;
+				DirectionalLightSceneProxy->LightSourceAngle = LightSourceAngle;
+			});
 	}
 }
 
@@ -1385,11 +1331,11 @@ void UDirectionalLightComponent::SetLightSourceSoftAngle(float NewValue)
 		&& LightSourceSoftAngle != NewValue)
 	{
 		LightSourceSoftAngle = NewValue;
-		if (SceneProxy)
-		{
-			FDirectionalLightSceneProxy* DirectionalLightSceneProxy = (FDirectionalLightSceneProxy*)SceneProxy;
-			DirectionalLightSceneProxy->UpdateLightSourceSoftAngle_GameThread(this);
-		}
+		UpdateProxy([LightSourceSoftAngle = LightSourceSoftAngle](FLightSceneProxy* Proxy)
+			{
+				FDirectionalLightSceneProxy* DirectionalLightSceneProxy = (FDirectionalLightSceneProxy*)Proxy;
+				DirectionalLightSceneProxy->LightSourceSoftAngle = LightSourceSoftAngle;
+			});
 	}
 }
 
@@ -1399,11 +1345,11 @@ void UDirectionalLightComponent::SetShadowSourceAngleFactor(float NewValue)
 		&& ShadowSourceAngleFactor != NewValue)
 	{
 		ShadowSourceAngleFactor = NewValue;
-		if (SceneProxy)
-		{
-			FDirectionalLightSceneProxy* DirectionalLightSceneProxy = (FDirectionalLightSceneProxy*)SceneProxy;
-			DirectionalLightSceneProxy->UpdateShadowSourceAngleFactor_GameThread(this);
-		}
+		UpdateProxy([ShadowSourceAngleFactor = ShadowSourceAngleFactor](FLightSceneProxy* Proxy)
+			{
+				FDirectionalLightSceneProxy* DirectionalLightSceneProxy = (FDirectionalLightSceneProxy*)Proxy;
+				DirectionalLightSceneProxy->ShadowSourceAngleFactor = ShadowSourceAngleFactor;
+			});
 	}
 }
 
@@ -1443,11 +1389,11 @@ void UDirectionalLightComponent::SetAtmosphereSunDiskColorScale(FLinearColor New
 		&& AtmosphereSunDiskColorScale != NewValue)
 	{
 		AtmosphereSunDiskColorScale = NewValue;
-		if (SceneProxy)
-		{
-			FDirectionalLightSceneProxy* DirectionalLightSceneProxy = (FDirectionalLightSceneProxy*)SceneProxy;
-			DirectionalLightSceneProxy->UpdateAtmosphereSunDiskColorScale_GameThread(this);
-		}
+		UpdateProxy([AtmosphereSunDiskColorScale = AtmosphereSunDiskColorScale](FLightSceneProxy* Proxy)
+			{
+				FDirectionalLightSceneProxy* DirectionalLightSceneProxy = (FDirectionalLightSceneProxy*)Proxy;
+				DirectionalLightSceneProxy->SetAtmosphereSunDiskColorScale(AtmosphereSunDiskColorScale);
+			});
 	}
 }
 
diff --git a/Engine/Source/Runtime/Engine/Private/Components/LightComponent.cpp b/Engine/Source/Runtime/Engine/Private/Components/LightComponent.cpp
--- a/Engine/Source/Runtime/Engine/Private/Components/LightComponent.cpp
+++ b/Engine/Source/Runtime/Engine/Private/Components/LightComponent.cpp
@@ -973,6 +973,18 @@ void ULightComponent::DestroyRenderState_Concurrent()
 	bAddedToSceneVisible = false;
 }
 
+void ULightComponent::UpdateProxy(FUpdateLightProxyCallback&& Func)
+{
+	if (SceneProxy)
+	{
+		UWorld* World = GetWorld();
+		if (World && World->Scene)
+		{
+			World->Scene->UpdateLightProxy(this, MoveTemp(Func));
+		}
+	}
+}
+
 #if WITH_EDITOR
 bool ULightComponent::GetMaterialPropertyPath(int32 ElementIndex, UObject*& OutOwner, FString& OutPropertyPath, FProperty*& OutProperty)
 {
diff --git a/Engine/Source/Runtime/Engine/Public/SceneInterface.h b/Engine/Source/Runtime/Engine/Public/SceneInterface.h
--- a/Engine/Source/Runtime/Engine/Public/SceneInterface.h
+++ b/Engine/Source/Runtime/Engine/Public/SceneInterface.h
@@ -333,6 +333,8 @@ public:
 	 */
 	virtual void UpdateLightColorAndBrightness(ULightComponent* Light) = 0;
 
+	virtual void UpdateLightProxy(ULightComponent* Light, TFunction<void(FLightSceneProxy*)>&& Func) = 0;
+
 	/** Sets the precomputed visibility handler for the scene, or NULL to clear the current one. */
 	virtual void SetPrecomputedVisibility(const class FPrecomputedVisibilityHandler* PrecomputedVisibilityHandler) {}
 
diff --git a/Engine/Source/Runtime/Renderer/Private/RendererScene.cpp b/Engine/Source/Runtime/Renderer/Private/RendererScene.cpp
--- a/Engine/Source/Runtime/Renderer/Private/RendererScene.cpp
+++ b/Engine/Source/Runtime/Renderer/Private/RendererScene.cpp
@@ -3437,6 +3437,17 @@ void FScene::UpdateLightTransform_RenderThread(int32 LightId, FLightSceneInfo* L
 	}
 }
 
+template <typename UpdatePayloadType>
+void FScene::UpdateLightInternal_RenderThread(FLightSceneProxy* LightSceneProxy, UpdatePayloadType&& InUpdatePayload)
+{
+	FLightSceneInfo* LightSceneInfo = LightSceneProxy->GetLightSceneInfo();
+	if (LightSceneInfo->bVisible)
+	{
+		FScopeCycleCounter Context(LightSceneProxy->GetStatId());
+		SceneLightInfoUpdates->Enqueue(LightSceneInfo, MoveTemp(InUpdatePayload));
+	}
+}
+
 template <typename UpdatePayloadType>
 void FScene::UpdateLightInternal(FLightSceneProxy* LightSceneProxy, UpdatePayloadType&& InUpdatePayload)
 {
@@ -3446,9 +3457,9 @@ void FScene::UpdateLightInternal(FLightSceneProxy* LightSceneProxy, UpdatePayloa
 		if (LightSceneInfo->bVisible)
 		{
 			ENQUEUE_RENDER_COMMAND(UpdateLightTransform)(
-				[this, LightSceneInfo, UpdatePayload = MoveTemp(InUpdatePayload)] (FRHICommandListBase&) mutable
+				[this, LightSceneInfo, LightSceneProxy, UpdatePayload = MoveTemp(InUpdatePayload)] (FRHICommandListBase&) mutable
 				{
-					FScopeCycleCounter Context(LightSceneInfo->Proxy->GetStatId());
+					FScopeCycleCounter Context(LightSceneProxy->GetStatId());
 					SceneLightInfoUpdates->Enqueue(LightSceneInfo, MoveTemp(UpdatePayload));
 				});
 		}
@@ -3468,6 +3479,20 @@ void FScene::UpdateLightColorAndBrightness(ULightComponent* Light)
 		.NewVolumetricScatteringIntensity = Light->VolumetricScatteringIntensity});
 }
 
+void FScene::UpdateLightProxy(ULightComponent* Light, TFunction<void(FLightSceneProxy*)>&& Func)
+{
+	FLightSceneProxy* LightSceneProxy = Light->SceneProxy;
+	if (LightSceneProxy)
+	{
+		ENQUEUE_RENDER_COMMAND(UpdateLightProxy)(
+			[this, LightSceneProxy, Func](FRHICommandListBase&)
+			{
+				Func(LightSceneProxy);
+				UpdateLightInternal_RenderThread(LightSceneProxy, FUpdateLightInvalidateGPUState{});
+			});
+	}
+}
+
 void FScene::RemoveLightSceneInfo_RenderThread(FLightSceneInfo* LightSceneInfo)
 {
 	SCOPE_CYCLE_COUNTER(STAT_RemoveSceneLightTime);
@@ -6577,6 +6602,8 @@ public:
 	virtual void UpdateLightTransform(ULightComponent* Light) override {}
 	virtual void UpdateLightColorAndBrightness(ULightComponent* Light) override {}
 
+	virtual void UpdateLightProxy(ULightComponent* Light, TFunction<void(FLightSceneProxy*)>&& Func) override {}
+
 	virtual void AddExponentialHeightFog(uint64 Id, const FExponentialHeightFogDynamicState& State) override {}
 	virtual void RemoveExponentialHeightFog(uint64 Id) override {}
 	virtual bool HasAnyExponentialHeightFog() const override { return false; }
diff --git a/Engine/Source/Runtime/Renderer/Private/ScenePrivate.h b/Engine/Source/Runtime/Renderer/Private/ScenePrivate.h
--- a/Engine/Source/Runtime/Renderer/Private/ScenePrivate.h
+++ b/Engine/Source/Runtime/Renderer/Private/ScenePrivate.h
@@ -2816,6 +2816,7 @@ ENUM_CLASS_FLAGS(ELightDirtyFlags);
 
 enum class ELightUpdateId : uint32 
 {
+	Invalidate,
 	Transform,
 	Color,
 	MAX
@@ -2824,6 +2825,8 @@ enum class ELightUpdateId : uint32
 using FSceneLightInfoUpdates = TSceneUpdateCommandQueue<FLightSceneInfo, ELightDirtyFlags, ELightUpdateId>;
 using FUpdateLightCommand = FSceneLightInfoUpdates::FUpdateCommand;
 
+using FUpdateLightInvalidateGPUState = FSceneLightInfoUpdates::TPayloadBase<ELightUpdateId::Invalidate, ELightDirtyFlags::None>;
+
 struct FUpdateLightTransformParameters : public FSceneLightInfoUpdates::TPayloadBase<ELightUpdateId::Transform, ELightDirtyFlags::None>
 {
 	FMatrix LightToWorld;
@@ -3428,6 +3431,7 @@ public:
 	virtual void GetPrimitiveUniformShaderParameters_RenderThread(const FPrimitiveSceneInfo* PrimitiveSceneInfo, bool& bHasPrecomputedVolumetricLightmap, FMatrix& PreviousLocalToWorld, int32& SingleCaptureIndex, bool& bOutputVelocity) const override;
 	virtual void UpdateLightTransform(ULightComponent* Light) override;
 	virtual void UpdateLightColorAndBrightness(ULightComponent* Light) override;
+	virtual void UpdateLightProxy(ULightComponent* Light, TFunction<void(FLightSceneProxy*)>&& Func) override;
 	virtual void AddExponentialHeightFog(uint64 Id, const FExponentialHeightFogDynamicState& State) override;
 	virtual void RemoveExponentialHeightFog(uint64 Id) override;
 	virtual bool HasAnyExponentialHeightFog() const override;
@@ -3998,6 +4002,9 @@ private:
 	template <typename UpdatePayloadType>
 	void UpdateLightInternal(FLightSceneProxy* LightSceneProxy, UpdatePayloadType &&InUpdatePayload);
 
+	template <typename UpdatePayloadType>
+	void UpdateLightInternal_RenderThread(FLightSceneProxy* LightSceneProxy, UpdatePayloadType&& InUpdatePayload);
+
 	/**
 	 * Update tracked scene state for cached CSM shadows
 	 */
