Fix texture streaming CSV stats not updating correctly.
There was a typo causing us to early out.
#rnx

diff --git a/Engine/Source/Runtime/Engine/Private/Streaming/AsyncTextureStreaming.cpp b/Engine/Source/Runtime/Engine/Private/Streaming/AsyncTextureStreaming.cpp
--- a/Engine/Source/Runtime/Engine/Private/Streaming/AsyncTextureStreaming.cpp
+++ b/Engine/Source/Runtime/Engine/Private/Streaming/AsyncTextureStreaming.cpp
@@ -1294,7 +1294,7 @@ void FRenderAssetStreamingMipCalcTask::UpdateCSVOnlyStats_Async()
 	{
 		if (IsAborted()) break;
 		if (!StreamingRenderAsset.RenderAsset) continue;
-		if (!StreamingRenderAsset.IsVirtualTexture()) continue;
+		if (StreamingRenderAsset.IsVirtualTexture()) continue;
 
 		const int64 ResidentSize = StreamingRenderAsset.GetSize(StreamingRenderAsset.ResidentMips);
 		const int64 RequiredSize = StreamingRenderAsset.GetSize(StreamingRenderAsset.GetPerfectWantedMips());
