Change check to ensure to avoid crash using rewind debugger

diff --git a/Engine/Source/Runtime/Engine/Private/Components/SkeletalMeshComponent.cpp b/Engine/Source/Runtime/Engine/Private/Components/SkeletalMeshComponent.cpp
--- a/Engine/Source/Runtime/Engine/Private/Components/SkeletalMeshComponent.cpp
+++ b/Engine/Source/Runtime/Engine/Private/Components/SkeletalMeshComponent.cpp
@@ -2208,10 +2208,15 @@ void USkeletalMeshComponent::ComputeRequiredBones(TArray<FBoneIndexType>& OutReq
 	const TArray<uint8>& EditableBoneVisibilityStates = SkeletalMeshComponent->GetEditableBoneVisibilityStates();
 	if (SkeletalMeshComponent->ShouldUpdateBoneVisibility() && EditableBoneVisibilityStates.Num() > 0)
 	{
-		check(EditableBoneVisibilityStates.Num() == SkeletalMeshComponent->GetNumComponentSpaceTransforms());
+		if (!ensureMsgf(EditableBoneVisibilityStates.Num() == SkeletalMeshComponent->GetNumComponentSpaceTransforms(),
+		TEXT("Skeletal Mesh asset '%s' has incorrect BoneVisibilityStates. # of BoneVisibilityStates (%d), # of SkeletalMeshComponent->GetNumComponentSpaceTransforms() (%d)"),
+		*SkeletalMesh->GetName(), EditableBoneVisibilityStates.Num(), SkeletalMeshComponent->GetNumComponentSpaceTransforms()))
+		{
+			return;
+		}
 
 		if (ensureMsgf(EditableBoneVisibilityStates.Num() >= OutRequiredBones.Num(),
-			TEXT("Skeletal Mesh asset '%s' has incorrect BoneVisibilityStates. # of BoneVisibilityStatese (%d), # of OutRequiredBones (%d)"),
+			TEXT("Skeletal Mesh asset '%s' has incorrect BoneVisibilityStates. # of BoneVisibilityStates (%d), # of OutRequiredBones (%d)"),
 			*SkeletalMesh->GetName(), EditableBoneVisibilityStates.Num(), OutRequiredBones.Num()))
 		{
 			int32 VisibleBoneWriteIndex = 0;
