Fix crash when attempting to export a Waveform with more than 1 Channel

diff --git a/Engine/Source/Runtime/Engine/Private/SampleBufferIO.cpp b/Engine/Source/Runtime/Engine/Private/SampleBufferIO.cpp
--- a/Engine/Source/Runtime/Engine/Private/SampleBufferIO.cpp
+++ b/Engine/Source/Runtime/Engine/Private/SampleBufferIO.cpp
@@ -220,12 +220,15 @@ namespace Audio
 			CurrentBuffer.MixBufferToChannels(2);
 		}
 	
-		CurrentOperation.Reset(new FAsyncSoundWavePCMWriterTask(this, ESoundWavePCMWriteTaskType::GenerateAndWriteSoundWave, OnSuccess));
-		CurrentOperation->StartBackgroundTask();
+		CurrentOperation.Reset(new FAsyncSoundWavePCMWriterTask(this, ESoundWavePCMWriteTaskType::GenerateAndWriteSoundWave, [this, OnSuccess](const USoundWave* InSoundWave)
+		{
+			// Invalidate compressed data on new soundwave so that it gets refreshed or the soundwave is not seekable
+			// and causes this error: FMixerBuffer::CreateStreamingBuffer failed to StreamCompressedInfo on SoundWave
+			CurrentSoundWave->InvalidateCompressedData(true); // calls CurrentSoundWave->FreeResources(), note the game thread is the audio thread in the editor
 
-		// Invalidate compressed data on new soundwave so that it gets refreshed or the soundwave is not seekable
-		// and causes this error: FMixerBuffer::CreateStreamingBuffer failed to StreamCompressedInfo on SoundWave
-		CurrentSoundWave->InvalidateCompressedData(true);
+			OnSuccess(InSoundWave);
+		}));
+		CurrentOperation->StartBackgroundTask();
 
 		return true;
 	}
