SVE Context: Add  IsA(..) that accepts l-value.

Also made FSceneViewExtensionIsActiveFunctor::GetGuid a const function.

#rnx

diff --git a/Engine/Source/Runtime/Engine/Public/SceneViewExtensionContext.h b/Engine/Source/Runtime/Engine/Public/SceneViewExtensionContext.h
--- a/Engine/Source/Runtime/Engine/Public/SceneViewExtensionContext.h
+++ b/Engine/Source/Runtime/Engine/Public/SceneViewExtensionContext.h
@@ -36,11 +36,17 @@ public:
 	virtual ~FSceneViewExtensionContext() {}
 	
 	// Returns true if the given object is of the same type.
-	bool IsA(const FSceneViewExtensionContext&& Other) const
+	bool IsA(const FSceneViewExtensionContext& Other) const
 	{ 
 		return GetRTTI() == Other.GetRTTI(); 
 	}
 
+	// Returns true if the given object is of the same type.
+	bool IsA(const FSceneViewExtensionContext&& Other) const
+	{
+		return IsA(static_cast<const FSceneViewExtensionContext&>(Other));
+	}
+
 	/** Retrieve the world pointer for this context */
 	UWorld* GetWorld() const
 	{
@@ -100,7 +106,7 @@ public:
 	{}
 
 	// Returns the Guid of this Functor.
-	FGuid GetGuid()
+	FGuid GetGuid() const
 	{
 		return Guid;
 	}
