CurveLinearColor : fix RGB very slightly negative confusing HSV conversion, making black turn to white
caused a visible banding artifact on the 3rd person Mannequin plastic diffraction material

diff --git a/Engine/Source/Runtime/Core/Private/Math/Color.cpp b/Engine/Source/Runtime/Core/Private/Math/Color.cpp
--- a/Engine/Source/Runtime/Core/Private/Math/Color.cpp
+++ b/Engine/Source/Runtime/Core/Private/Math/Color.cpp
@@ -399,6 +399,11 @@ FLinearColor FLinearColor::LinearRGBToHSV() const
 	const float RGBMax = FMath::Max3(R, G, B);
 	const float RGBRange = RGBMax - RGBMin;
 
+	if ( RGBMax <= 0.f )
+	{
+		return FLinearColor(0.f,0.f,0.f,A);
+	}
+
 	const float Hue = (RGBMax == RGBMin ? 0.0f :
 					   RGBMax == R    ? FMath::Fmod((((G - B) / RGBRange) * 60.0f) + 360.0f, 360.0f) :
 					   RGBMax == G    ?             (((B - R) / RGBRange) * 60.0f) + 120.0f :
@@ -417,7 +422,7 @@ FLinearColor FLinearColor::LinearRGBToHSV() const
 /** Converts an HSV color to a linear space RGB color */
 FLinearColor FLinearColor::HSVToLinearRGB() const
 {
-	// In this color, R = H, G = S, B = V
+	// In this color, R = H, G = S, B = V  ; H in [0,360] , S&V in [0,1]
 	const float Hue = R;
 	const float Saturation = G;
 	const float Value = B;
@@ -440,7 +445,7 @@ FLinearColor FLinearColor::HSVToLinearRGB() const
 		{3, 1, 0},
 		{0, 1, 2},
 	};
-	const uint32 SwizzleIndex = ((uint32)HDiv60_Floor) % 6;
+	const uint32 SwizzleIndex = ((uint32)HDiv60_Floor) % 6;  // use of uint32 ensures SwizzleIndex can't be negative even if Hue was
 
 	return FLinearColor(RGBValues[RGBSwizzle[SwizzleIndex][0]],
 						RGBValues[RGBSwizzle[SwizzleIndex][1]],
diff --git a/Engine/Source/Runtime/Engine/Private/Curves/CurveLinearColor.cpp b/Engine/Source/Runtime/Engine/Private/Curves/CurveLinearColor.cpp
--- a/Engine/Source/Runtime/Engine/Private/Curves/CurveLinearColor.cpp
+++ b/Engine/Source/Runtime/Engine/Private/Curves/CurveLinearColor.cpp
@@ -68,10 +68,24 @@ UCurveLinearColor::UCurveLinearColor(const FObjectInitializer& ObjectInitializer
 #endif
 }
 
+static FLinearColor ClampNonNegative(const FLinearColor & InColor)
+{
+	FLinearColor OutColor;
+	OutColor.R = FMath::Max(InColor.R,0.f);
+	OutColor.G = FMath::Max(InColor.G,0.f);
+	OutColor.B = FMath::Max(InColor.B,0.f);
+	OutColor.A = FMath::Max(InColor.A,0.f);
+	return OutColor;
+}
+
 FLinearColor UCurveLinearColor::GetLinearColorValue( float InTime ) const
 {
 	// Logic copied from .\Engine\Source\Developer\TextureCompressor\Private\TextureCompressorModule.cpp
-	const FLinearColor OriginalColor = GetUnadjustedLinearColorValue(InTime);
+	FLinearColor OriginalColor = GetUnadjustedLinearColorValue(InTime);
+	
+	// allow > 1.0 at this point but not < 0.0
+	//	we're just about to convert this RGB to HSV, that doesn't support negatives
+	OriginalColor = ClampNonNegative(OriginalColor);
 
 	// Only clamp value if the color is RGB < 1
 	const bool bShouldClampValue = (OriginalColor.R <= 1.0f && OriginalColor.G <= 1.0f && OriginalColor.B <= 1.0f);
@@ -140,7 +154,11 @@ FLinearColor UCurveLinearColor::GetLinearColorValue( float InTime ) const
 FLinearColor UCurveLinearColor::GetClampedLinearColorValue(float InTime) const
 {
 	// Logic copied from .\Engine\Source\Developer\TextureCompressor\Private\TextureCompressorModule.cpp
-	const FLinearColor OriginalColor = GetUnadjustedLinearColorValue(InTime);
+	FLinearColor OriginalColor = GetUnadjustedLinearColorValue(InTime);
+	
+	// allow > 1.0 at this point but not < 0.0
+	//	we're just about to convert this RGB to HSV, that doesn't support negatives
+	OriginalColor = ClampNonNegative(OriginalColor);
 
 	// Always clamp value
 	const bool bShouldClampValue = true;
@@ -222,6 +240,8 @@ FLinearColor UCurveLinearColor::GetUnadjustedLinearColorValue(float InTime) cons
 	{
 		Result.A = FloatCurves[3].Eval(InTime);
 	}
+	
+	// Result outside [0,1] intentionally allowed here
 
 	return Result;
 }
