DataTable: Add error handling for a missing RowStruct type when using UnversionedPropertySerialization (UPS) and ZenLoader.
Prepare for enabling UPS on dedicated servers without introducing new crashes.
Note that the right solution is to strip the whole DataTable when cooking, but there is no guarantee that it will actually be the case.

#rnx
#tests client,server,editor

diff --git a/Engine/Source/Runtime/CoreUObject/Private/Serialization/AsyncLoading2.cpp b/Engine/Source/Runtime/CoreUObject/Private/Serialization/AsyncLoading2.cpp
--- a/Engine/Source/Runtime/CoreUObject/Private/Serialization/AsyncLoading2.cpp
+++ b/Engine/Source/Runtime/CoreUObject/Private/Serialization/AsyncLoading2.cpp
@@ -8511,10 +8511,23 @@ bool FAsyncPackage2::EventDrivenSerializeExport(const FAsyncPackageHeaderData& H
 	}
 	Ar->TemplateForGetArchetypeFromLoader = nullptr;
 
-	UE_ASYNC_PACKAGE_CLOG(
-		Export.CookedSerialSize != uint64(Ar->Tell() - Pos), Fatal, Desc, TEXT("ObjectSerializationError"),
-		TEXT("%s: Serial size mismatch: Expected read size %lld, Actual read size %lld"),
-		Object ? *Object->GetFullName() : TEXT("null"), Export.CookedSerialSize, uint64(Ar->Tell() - Pos));
+	const uint64 ReadBytes = uint64(Ar->Tell() - Pos);
+	if (Ar->IsError())
+	{
+		const uint64 SeekBytes = Export.CookedSerialSize - ReadBytes;
+		Ar->Seek(SeekBytes);
+		Ar->ClearError();
+		UE_ASYNC_PACKAGE_LOG(Error, Desc, TEXT("ObjectSerializationError"),
+			TEXT("%s: Serialization error handled, %lld of %lld bytes was read, %lld bytes was seeked. The serialized object may be in a bad state."),
+			Object ? *Object->GetFullName() : TEXT("null"), ReadBytes, Export.CookedSerialSize, SeekBytes);
+	}
+	else
+	{
+		UE_ASYNC_PACKAGE_CLOG(
+			Export.CookedSerialSize != ReadBytes, Fatal, Desc, TEXT("ObjectSerializationError"),
+			TEXT("%s: Serial size mismatch: Expected read size %lld, Actual read size %lld"),
+			Object ? *Object->GetFullName() : TEXT("null"), Export.CookedSerialSize, ReadBytes);
+	}
 
 	Ar->ExportBufferEnd();
 
diff --git a/Engine/Source/Runtime/Engine/Private/DataTable.cpp b/Engine/Source/Runtime/Engine/Private/DataTable.cpp
--- a/Engine/Source/Runtime/Engine/Private/DataTable.cpp
+++ b/Engine/Source/Runtime/Engine/Private/DataTable.cpp
@@ -112,9 +112,18 @@ void UDataTable::PostEditChangeProperty(FPropertyChangedEvent& PropertyChangedEv
 
 void UDataTable::LoadStructData(FStructuredArchiveSlot Slot)
 {
+	int32 NumRows;
+	FStructuredArchiveArray Array = Slot.EnterArray(NumRows);
+
 	UScriptStruct* LoadUsingStruct = RowStruct;
 	if (!LoadUsingStruct)
 	{
+		if (NumRows > 0 && Slot.GetUnderlyingArchive().UseUnversionedPropertySerialization())
+		{
+			UE_LOG(LogDataTable, Error, TEXT("Failed to load DataTable '%s' with %d rows. The rows can't be loaded with UnversionedPropertySerialization when the RowStruct is missing."), *GetPathName(), NumRows);
+			Slot.GetUnderlyingArchive().SetError();
+			return;
+		}
 		if (!HasAnyFlags(RF_ClassDefaultObject) && GetOutermost() != GetTransientPackage())
 		{
 			UE_LOG(LogDataTable, Warning, TEXT("Missing RowStruct while loading DataTable '%s', NeedLoad: '%s'!"), *GetPathName(), HasAnyFlags(RF_NeedLoad) ? TEXT("true") : TEXT("false"));
@@ -122,9 +131,6 @@ void UDataTable::LoadStructData(FStructuredArchiveSlot Slot)
 		LoadUsingStruct = FTableRowBase::StaticStruct();
 	}
 
-	int32 NumRows;
-	FStructuredArchiveArray Array = Slot.EnterArray(NumRows);
-
 	DATATABLE_CHANGE_SCOPE();
 
 	RowMap.Reserve(NumRows);
@@ -1083,4 +1089,4 @@ bool FDataTableCategoryHandle::operator != (FDataTableCategoryHandle const& Othe
 	return DataTable != Other.DataTable || ColumnName != Other.ColumnName || RowContents != Other.RowContents;
 }
 
-#undef LOCTEXT_NAMESPACE
\ No newline at end of file
+#undef LOCTEXT_NAMESPACE
