FWorldCookPackageSplitter:
- Renamed FSubSplitter to ISubSplitter as it is an interface.
- Made ISubSplitter unaware of its owning world (implementations can store it when invoked if needed).
- Transfer splitters ownership to the factories, which allows a class to derive from ISubSplitter instead of doing memory allocations.

diff --git a/Engine/Source/Runtime/Engine/Private/WorldCookPackageSplitter.cpp b/Engine/Source/Runtime/Engine/Private/WorldCookPackageSplitter.cpp
--- a/Engine/Source/Runtime/Engine/Private/WorldCookPackageSplitter.cpp
+++ b/Engine/Source/Runtime/Engine/Private/WorldCookPackageSplitter.cpp
@@ -87,14 +87,14 @@ TArray<ICookPackageSplitter::FGeneratedPackage> FWorldCookPackageSplitter::GetGe
 		{
 			if (WorldCookPackageSplitterFactory.ShouldSplit(OwnerWorld))
 			{
-				CookPackageSubSplitters.Emplace(WorldCookPackageSplitterFactory.MakeInstance(OwnerWorld));
+				CookPackageSubSplitters.Add({ &WorldCookPackageSplitterFactory, WorldCookPackageSplitterFactory.MakeInstance(OwnerWorld) });
 			}
 		});
 	check(CookPackageSubSplitters.Num());
 
 	TArray<ICookPackageSplitter::FGeneratedPackage> PackagesToGenerate;
 	ForEachCookPackageSubSplitters(
-		[&PackagesToGenerate, OwnerPackage, OwnerWorld](FWorldCookPackageSplitter::FSubSplitter* WorldCookPackageSplitter)
+		[&PackagesToGenerate, OwnerPackage, OwnerWorld](FWorldCookPackageSplitter::ISubSplitter* WorldCookPackageSplitter)
 		{
 			PackagesToGenerate.Append(WorldCookPackageSplitter->GetGenerateList(OwnerPackage));
 		});
@@ -108,7 +108,7 @@ bool FWorldCookPackageSplitter::PopulateGeneratedPackage(FPopulateContext& Popul
 {
 	bool bResult = false;
 	ForEachCookPackageSubSplitters(
-		[&bResult, &PopulateContext](FWorldCookPackageSplitter::FSubSplitter* WorldCookPackageSplitter)
+		[&bResult, &PopulateContext](FWorldCookPackageSplitter::ISubSplitter* WorldCookPackageSplitter)
 		{
 			bResult = WorldCookPackageSplitter->PopulateGeneratedPackage(PopulateContext);
 			return bResult;
@@ -120,7 +120,7 @@ bool FWorldCookPackageSplitter::PopulateGeneratorPackage(FPopulateContext& Popul
 {
 	bool bResult = false;
 	ForEachCookPackageSubSplitters(
-		[&bResult, &PopulateContext](FWorldCookPackageSplitter::FSubSplitter* WorldCookPackageSplitter)
+		[&bResult, &PopulateContext](FWorldCookPackageSplitter::ISubSplitter* WorldCookPackageSplitter)
 		{
 			bResult = WorldCookPackageSplitter->PopulateGeneratorPackage(PopulateContext);
 			return bResult;
@@ -138,11 +138,17 @@ void FWorldCookPackageSplitter::Teardown(ICookPackageSplitter::ETeardown Status)
 	FWorldDelegates::OnWorldCleanup.RemoveAll(this);
 
 	ForEachCookPackageSubSplitters(
-		[Status](FWorldCookPackageSplitter::FSubSplitter* WorldCookPackageSplitter)
+		[Status](FWorldCookPackageSplitter::ISubSplitter* WorldCookPackageSplitter)
 		{
 			WorldCookPackageSplitter->Teardown(Status);
 		});
 
+	for (const TPair<const FWorldCookPackageSplitter::FSubSplitterFactory*, ISubSplitter*>& CookPackageSubSplitterPair : CookPackageSubSplitters)
+	{
+		CookPackageSubSplitterPair.Key->ReleaseInstance(ReferencedWorld, CookPackageSubSplitterPair.Value);
+	}
+	CookPackageSubSplitters.Empty();
+
 	if (bInitializedPhysicsSceneForSave)
 	{
 		GEditor->CleanupPhysicsSceneThatWasInitializedForSave(ReferencedWorld, bForceInitializedWorld);
diff --git a/Engine/Source/Runtime/Engine/Private/WorldPartition/WorldPartition.cpp b/Engine/Source/Runtime/Engine/Private/WorldPartition/WorldPartition.cpp
--- a/Engine/Source/Runtime/Engine/Private/WorldPartition/WorldPartition.cpp
+++ b/Engine/Source/Runtime/Engine/Private/WorldPartition/WorldPartition.cpp
@@ -399,9 +399,14 @@ UWorldPartition::UWorldPartition(const FObjectInitializer& ObjectInitializer)
 						return World->IsPartitionedWorld();
 					},
 				.MakeInstance =
-					[](TNotNull<const UWorld*> World) -> FWorldCookPackageSplitter::FSubSplitter*
+					[](TNotNull<const UWorld*> World) -> FWorldCookPackageSplitter::ISubSplitter*
         			{
-						return new FWorldPartitionCookPackageSplitter(World);
+						return new FWorldPartitionCookPackageSplitter();
+					},
+				.ReleaseInstance = 
+					[](TNotNull<const UWorld*> World, TNotNull<FWorldCookPackageSplitter::ISubSplitter*> CookPackageSplitter)
+					{
+						delete NotNullGet(CookPackageSplitter);
 					}
 			});
 	}
diff --git a/Engine/Source/Runtime/Engine/Private/WorldPartition/WorldPartitionCookPackageSplitter.cpp b/Engine/Source/Runtime/Engine/Private/WorldPartition/WorldPartitionCookPackageSplitter.cpp
--- a/Engine/Source/Runtime/Engine/Private/WorldPartition/WorldPartitionCookPackageSplitter.cpp
+++ b/Engine/Source/Runtime/Engine/Private/WorldPartition/WorldPartitionCookPackageSplitter.cpp
@@ -13,9 +13,10 @@
 
 TArray<ICookPackageSplitter::FGeneratedPackage> FWorldPartitionCookPackageSplitter::GetGenerateList(const UPackage* OwnerPackage)
 {
-	UWorld* ReferencedWorld = const_cast<UWorld*>(NotNullGet(World));
+	World = UWorld::FindWorldInPackage(const_cast<UPackage*>(OwnerPackage));
+	check(World);
 
-	UE_LOG(LogWorldPartition, Display, TEXT("[Cook] Gathering packages to cook from generators for owner object %s."), *GetFullNameSafe(ReferencedWorld));
+	UE_LOG(LogWorldPartition, Display, TEXT("[Cook] Gathering packages to cook from generators for owner object %s."), *GetFullNameSafe(World));
 
 	// Assume that the world is partitioned as per FWorldPartitionCookPackageSplitter::ShouldSplit
 	UWorldPartition* WorldPartition = World->PersistentLevel->GetWorldPartition();
@@ -23,11 +24,11 @@ TArray<ICookPackageSplitter::FGeneratedPackage> FWorldPartitionCookPackageSplitt
 
 	// We expect the WorldPartition has not yet been initialized
 	ensure(!WorldPartition->IsInitialized());
-	WorldPartition->Initialize(ReferencedWorld, FTransform::Identity);
+	WorldPartition->Initialize(World, FTransform::Identity);
 	WorldPartition->BeginCook(CookContext);
 
 	bool bIsSuccess = CookContext.GatherPackagesToCook();
-	UE_CLOG(!bIsSuccess, LogWorldPartition, Warning, TEXT("[Cook] Errors while gathering packages to took from generators for owner object %s."), *GetFullNameSafe(ReferencedWorld));
+	UE_CLOG(!bIsSuccess, LogWorldPartition, Warning, TEXT("[Cook] Errors while gathering packages to took from generators for owner object %s."), *GetFullNameSafe(World));
 
 	UE_LOG(LogWorldPartition, Display, TEXT("[Cook] Gathered %u packages to generate from %u Generators."), CookContext.NumPackageToGenerate(), CookContext.NumGenerators());
 
@@ -35,7 +36,7 @@ TArray<ICookPackageSplitter::FGeneratedPackage> FWorldPartitionCookPackageSplitt
 	BuildPackagesToGenerateList(PackagesToGenerate);
 
 	UE_LOG(LogWorldPartition, Display, TEXT("[Cook] Sending %u packages to be generated."), PackagesToGenerate.Num());
-	UE_LOG(LogWorldPartition, Display, TEXT("[Cook] Debug(GetGenerateList) : OwnerWorld=%s"), *GetFullNameSafe(NotNullGet(World)));
+	UE_LOG(LogWorldPartition, Display, TEXT("[Cook] Debug(GetGenerateList) : OwnerWorld=%s"), *GetFullNameSafe(World));
 
 	return PackagesToGenerate;
 }
@@ -100,13 +101,11 @@ bool FWorldPartitionCookPackageSplitter::PopulateGeneratorPackage(ICookPackageSp
 
 void FWorldPartitionCookPackageSplitter::Teardown(ICookPackageSplitter::ETeardown Status)
 {
-	UWorld* ReferencedWorld = const_cast<UWorld*>(NotNullGet(World));
-
-	UE_LOG(LogWorldPartition, Display, TEXT("[Cook] Debug(TearDown): OwnerWorld=%s"), *GetFullNameSafe(ReferencedWorld));
+	UE_LOG(LogWorldPartition, Display, TEXT("[Cook] Debug(TearDown): OwnerWorld=%s"), *GetFullNameSafe(World));
 
 	// We were destructed without GetGenerateList being called; nothing to teardown
 	// Assume that the world is partitioned as per FWorldPartitionCookPackageSplitter::ShouldSplit
-	UWorldPartition* WorldPartition = ReferencedWorld->PersistentLevel->GetWorldPartition();
+	UWorldPartition* WorldPartition = World->PersistentLevel->GetWorldPartition();
 	check(WorldPartition);
 
 	WorldPartition->EndCook(CookContext);
diff --git a/Engine/Source/Runtime/Engine/Private/WorldPartition/WorldPartitionCookPackageSplitter.h b/Engine/Source/Runtime/Engine/Private/WorldPartition/WorldPartitionCookPackageSplitter.h
--- a/Engine/Source/Runtime/Engine/Private/WorldPartition/WorldPartitionCookPackageSplitter.h
+++ b/Engine/Source/Runtime/Engine/Private/WorldPartition/WorldPartitionCookPackageSplitter.h
@@ -7,23 +7,20 @@
 #include "WorldCookPackageSplitter.h"
 #include "WorldPartition/Cook/WorldPartitionCookPackageContext.h"
 
-class FWorldPartitionCookPackageSplitter : public FWorldCookPackageSplitter::FSubSplitter
+class FWorldPartitionCookPackageSplitter : public FWorldCookPackageSplitter::ISubSplitter
 {
 public:
-	FWorldPartitionCookPackageSplitter(TNotNull<const UWorld*> World)
-		: FWorldCookPackageSplitter::FSubSplitter(World)
-	{}
-
-	//~Begin FWorldCookPackageSplitter::FSubSplitter interface
+	//~Begin FWorldCookPackageSplitter::ISubSplitter interface
 	virtual TArray<ICookPackageSplitter::FGeneratedPackage> GetGenerateList(const UPackage* OwnerPackage) override;
 	virtual bool PopulateGeneratedPackage(ICookPackageSplitter::FPopulateContext& PopulateContext) override;
 	virtual bool PopulateGeneratorPackage(ICookPackageSplitter::FPopulateContext& PopulateContext) override;
 	virtual void Teardown(ICookPackageSplitter::ETeardown Status) override;
-	//~End FWorldCookPackageSplitter::FSubSplitter interface
+	//~End FWorldCookPackageSplitter::ISubSplitter interface
 
 private:
 	void BuildPackagesToGenerateList(TArray<ICookPackageSplitter::FGeneratedPackage>& PackagesToGenerate) const;
 	
+	UWorld* World = nullptr;
 	FWorldPartitionCookPackageContext CookContext;
 };
 
diff --git a/Engine/Source/Runtime/Engine/Public/WorldCookPackageSplitter.h b/Engine/Source/Runtime/Engine/Public/WorldCookPackageSplitter.h
--- a/Engine/Source/Runtime/Engine/Public/WorldCookPackageSplitter.h
+++ b/Engine/Source/Runtime/Engine/Public/WorldCookPackageSplitter.h
@@ -19,29 +19,20 @@ class FWorldCookPackageSplitter : public FGCObject, ICookPackageSplitter
 public:
 	~FWorldCookPackageSplitter();
 
-	struct FSubSplitter
+	struct ISubSplitter
 	{
-		FSubSplitter(TNotNull<const UWorld*> InWorld)
-			: World(InWorld)
-		{}
-
-		virtual ~FSubSplitter()
-		{}
-
-		//~Begin FSubSplitter interface
+		virtual ~ISubSplitter()	{}
 		virtual TArray<ICookPackageSplitter::FGeneratedPackage> GetGenerateList(const UPackage* OwnerPackage) = 0;
 		virtual bool PopulateGeneratedPackage(ICookPackageSplitter::FPopulateContext& PopulateContext) = 0;
 		virtual bool PopulateGeneratorPackage(ICookPackageSplitter::FPopulateContext& PopulateContext) = 0;
 		virtual void Teardown(ICookPackageSplitter::ETeardown Status) = 0;
-		//~End FSubSplitter interface
-
-		TNotNull<const UWorld*> World;
 	};
 
 	struct FSubSplitterFactory
 	{
 		TFunction<bool(TNotNull<const UWorld*>)> ShouldSplit;
-		TFunction<FSubSplitter*(TNotNull<const UWorld*>)> MakeInstance;
+		TFunction<ISubSplitter*(TNotNull<const UWorld*>)> MakeInstance;
+		TFunction<void(TNotNull<const UWorld*>, TNotNull<ISubSplitter*>)> ReleaseInstance;
 	};
 
 	static UE_API void RegisterCookPackageSubSplitterFactory(UClass* Class, FSubSplitterFactory Factory);
@@ -83,15 +74,15 @@ protected:
 	virtual void Teardown(ETeardown Status) override;
 	//~ End of ICookPackageSplitter
 
-	TArray<TUniquePtr<FSubSplitter>> CookPackageSubSplitters;
+	TArray<TPair<const FWorldCookPackageSplitter::FSubSplitterFactory*, ISubSplitter*>> CookPackageSubSplitters;
 
 	template<typename FuncType>
 	bool ForEachCookPackageSubSplitters(FuncType Func)
 	{
 		TOverrideVoidReturnInvoker Invoker(true, Func);
-		for (const TUniquePtr<FWorldCookPackageSplitter::FSubSplitter>& CookPackageSubSplitter : CookPackageSubSplitters)
+		for (const TPair<const FWorldCookPackageSplitter::FSubSplitterFactory*, ISubSplitter*>& CookPackageSubSplitterPair : CookPackageSubSplitters)
 		{
-			if (!Invoker(CookPackageSubSplitter.Get()))
+			if (!Invoker(CookPackageSubSplitterPair.Value))
 			{
 				return false;
 			}
