Iris
* Send subobject archetype and spawn subobjects using it as the template instead of always using the class CDO

#rnx

diff --git a/Engine/Source/Runtime/Engine/Private/Net/Iris/ReplicationSystem/NetSubObjectFactory.cpp b/Engine/Source/Runtime/Engine/Private/Net/Iris/ReplicationSystem/NetSubObjectFactory.cpp
--- a/Engine/Source/Runtime/Engine/Private/Net/Iris/ReplicationSystem/NetSubObjectFactory.cpp
+++ b/Engine/Source/Runtime/Engine/Private/Net/Iris/ReplicationSystem/NetSubObjectFactory.cpp
@@ -26,6 +26,13 @@
 
 #if UE_WITH_IRIS
 
+static bool bUseArchetypeInSubObjectCreation = true;
+static FAutoConsoleVariableRef CVarUseArchetype(
+	TEXT("net.Iris.SubObjectSupportArchetype"),
+	bUseArchetypeInSubObjectCreation,
+	TEXT("Support dynamic subobjects that have an archetype instead of just the class CDO.")
+);
+
 TUniquePtr<UE::Net::FNetObjectCreationHeader> UNetSubObjectFactory::CreateAndFillHeader(UE::Net::FNetRefHandle Handle)
 {
 	using namespace UE::Net;
@@ -53,7 +60,8 @@ TUniquePtr<UE::Net::FNetObjectCreationHeader> UNetSubObjectFactory::CreateAndFil
 	check(SubObject->NeedsLoadForClient() ); // We have no business sending this unless the client can load
 	check(SubObject->GetClass()->NeedsLoadForClient());	// We have no business sending this unless the client can load
 
-	SubObjectHeader->ObjectClassReference = Bridge->GetOrCreateObjectReference(SubObject->GetClass());
+	const UObject* Template = bUseArchetypeInSubObjectCreation ? SubObject->GetArchetype() : SubObject->GetClass()->GetDefaultObject();
+	SubObjectHeader->ArchetypeReference = Bridge->GetOrCreateObjectReference(Template);
 
 	// Find the right Outer
 	UObject* OuterObject = SubObject->GetOuter();	
@@ -127,12 +135,12 @@ UNetSubObjectFactory::FInstantiateResult UNetSubObjectFactory::InstantiateReplic
 
 		if (!SubObject)
 		{
-			UE_LOG(LogIris, Error, TEXT("UNetSubObjectFactory::InstantiateNetObjectFromHeader %s: Failed to find static or stable name object referenced by SubObject: %s, Owner: %s, RootObject: %s"), 
+			UE_LOG(LogIris, Error, TEXT("UNetSubObjectFactory::InstantiateReplicatedObjectFromHeader %s: Failed to find static or stable name object referenced by SubObject: %s, Owner: %s, RootObject: %s"), 
 			*Context.Handle.ToString(), *Bridge->DescribeObjectReference(SubObjectHeader->ObjectReference, Context.ResolveContext), *Bridge->PrintObjectFromNetRefHandle(Context.RootObjectOfSubObject), *GetPathNameSafe(Bridge->GetReplicatedObject(Context.RootObjectOfSubObject)));
 			return FInstantiateResult();
 		}
 
-		UE_LOG(LogIris, Verbose, TEXT("UNetSubObjectFactory::InstantiateNetObjectFromHeader %s: Found static or stable name SubObject using path %s"), *Context.Handle.ToString(), ToCStr(SubObject->GetPathName()));
+		UE_LOG(LogIris, Verbose, TEXT("UNetSubObjectFactory::InstantiateReplicatedObjectFromHeader %s: Found static or stable name SubObject using path %s"), *Context.Handle.ToString(), ToCStr(SubObject->GetPathName()));
 
 		FInstantiateResult Result { .Instance = SubObject };
 		return Result;
@@ -160,31 +168,29 @@ UNetSubObjectFactory::FInstantiateResult UNetSubObjectFactory::InstantiateReplic
 
 		if (!OuterObject)
 		{
-			UE_LOG(LogIris, Error, TEXT("BeginInstantiateFromRemote Failed to find Outer %s for dynamic subobject %s"), *Bridge->DescribeObjectReference(SubObjectHeader->OuterReference, Context.ResolveContext), *Context.Handle.ToString())
+			UE_LOG(LogIris, Error, TEXT("UNetSubObjectFactory::InstantiateReplicatedObjectFromHeader Failed to find Outer %s for dynamic subobject %s"), *Bridge->DescribeObjectReference(SubObjectHeader->OuterReference, Context.ResolveContext), *Context.Handle.ToString())
 
 			// Fallback to the rootobject instead
 			OuterObject = RootObject;
 		}
 	}
 
-	// Find the class of the subobject
-	UObject* SubObjClassObj = Bridge->ResolveObjectReference(SubObjectHeader->ObjectClassReference, Context.ResolveContext);
-	UClass * SubObjClass = Cast<UClass>(SubObjClassObj);
-
-	if (!SubObjClass)
+	// Find the archetype of the subobject
+	UObject* Archetype = Bridge->ResolveObjectReference(SubObjectHeader->ArchetypeReference, Context.ResolveContext);
+	if (!Archetype)
 	{
-		UE_LOG(LogIris, Error, TEXT("UNetSubObjectFactory::InstantiateNetObjectFromHeader could not find UClass via reference: %s. Cannot spawn subobject for handle: %s"), 
-			*Bridge->DescribeObjectReference(SubObjectHeader->ObjectClassReference, Context.ResolveContext), *Context.Handle.ToString());
-		ensure(SubObjClass);
+		UE_LOG(LogIris, Error, TEXT("UNetSubObjectFactory::InstantiateReplicatedObjectFromHeader Unable to resolve archetype: %s, cannot instantiate dynamic subobject %s"), *Bridge->DescribeObjectReference(SubObjectHeader->ArchetypeReference, Context.ResolveContext), *Context.Handle.ToString());
+		ensure(Archetype);
 		return FInstantiateResult();
 	}
 
-	// Try to spawn SubObject
-	UObject* SubObj = NewObject<UObject>(OuterObject, SubObjClass);
+	// Instantiate the SubObject from it's archetype
+	UClass* SubObjClass = Archetype->GetClass();
+	UObject* SubObj = NewObject<UObject>(OuterObject, SubObjClass, NAME_None, RF_NoFlags, Archetype);
 
 	// Sanity check some things
-	checkf(SubObj != nullptr, TEXT("UNetSubObjectFactory::InstantiateNetObjectFromHeader: Subobject is NULL after instantiating. Class: %s, Outer %s, Actor %s"), *GetNameSafe(SubObjClass), *GetNameSafe(OuterObject), *GetNameSafe(RootObject));
-	checkf(!OuterObject || SubObj->IsIn(OuterObject), TEXT("UNetSubObjectFactory::InstantiateNetObjectFromHeader: Subobject is not in Outer. SubObject: %s, Outer %s, Actor %s"), *SubObj->GetName(), *GetNameSafe(OuterObject), *GetNameSafe(RootObject));
+	checkf(SubObj != nullptr, TEXT("UNetSubObjectFactory::InstantiateReplicatedObjectFromHeader: Subobject is NULL after instantiating. Class: %s, Outer %s, Actor %s"), *GetNameSafe(SubObjClass), *GetNameSafe(OuterObject), *GetNameSafe(RootObject));
+	checkf(!OuterObject || SubObj->IsIn(OuterObject), TEXT("UNetSubObjectFactory::InstantiateReplicatedObjectFromHeader: Subobject is not in Outer. SubObject: %s, Outer %s, Actor %s"), *SubObj->GetName(), *GetNameSafe(OuterObject), *GetNameSafe(RootObject));
 
 	FInstantiateResult Result { .Instance = SubObj };
 
@@ -286,13 +292,13 @@ bool FNetStaticSubObjectCreationHeader::Deserialize(const FCreationHeaderContext
 FString FNetDynamicSubObjectCreationHeader::ToString() const
 {
 	return FString::Printf(TEXT("\n\tFNetDynamicSubObjectCreationHeader (ProtocolId:0x%x):\n\t"
-								"ObjectClassReference=%s\n\t"
+								"ArchetypeReference=%s\n\t"
 								"OuterReference=%s\n\t"
 								"bUsePersistenLevel=%u\n\t"
 								"bOuterIsTransientLevel=%u\n\t"
 								"bOuterIsRootObject=%u\n\t"),
 								GetProtocolId(),
-								*ObjectClassReference.ToString(),
+								*ArchetypeReference.ToString(),
 								*OuterReference.ToString(),
 								bUsePersistentLevel,
 								bOuterIsTransientLevel,
@@ -304,7 +310,7 @@ bool FNetDynamicSubObjectCreationHeader::Serialize(const FCreationHeaderContext&
 {
 	FNetBitStreamWriter* Writer = Context.Serialization.GetBitStreamWriter();
 
-	WriteFullNetObjectReference(Context.Serialization, ObjectClassReference);
+	WriteFullNetObjectReference(Context.Serialization, ArchetypeReference);
 
 	if (!Writer->WriteBool(bOuterIsTransientLevel))
 	{
@@ -321,7 +327,7 @@ bool FNetDynamicSubObjectCreationHeader::Deserialize(const FCreationHeaderContex
 {
 	FNetBitStreamReader* Reader = Context.Serialization.GetBitStreamReader();
 
-	ReadFullNetObjectReference(Context.Serialization, ObjectClassReference);
+	ReadFullNetObjectReference(Context.Serialization, ArchetypeReference);
 
 	bOuterIsTransientLevel = Reader->ReadBool();
 	if (!bOuterIsTransientLevel)
diff --git a/Engine/Source/Runtime/Engine/Public/Net/Iris/ReplicationSystem/NetSubObjectFactory.h b/Engine/Source/Runtime/Engine/Public/Net/Iris/ReplicationSystem/NetSubObjectFactory.h
--- a/Engine/Source/Runtime/Engine/Public/Net/Iris/ReplicationSystem/NetSubObjectFactory.h
+++ b/Engine/Source/Runtime/Engine/Public/Net/Iris/ReplicationSystem/NetSubObjectFactory.h
@@ -104,7 +104,7 @@ public:
 
 	virtual FString ToString() const override;
 
-	FNetObjectReference ObjectClassReference;
+	FNetObjectReference ArchetypeReference;
 	FNetObjectReference OuterReference;
 	uint8 bUsePersistentLevel : 1 = false;
 	uint8 bOuterIsTransientLevel : 1 = false;	// When set the OuterReference was not sent because the Outer is the default transient level.
