[Engine] Exposed BuildCylinderVerts in SceneManagement.h, similar to BuildConeVerts etc

diff --git a/Engine/Source/Runtime/Engine/Public/SceneManagement.h b/Engine/Source/Runtime/Engine/Public/SceneManagement.h
--- a/Engine/Source/Runtime/Engine/Public/SceneManagement.h
+++ b/Engine/Source/Runtime/Engine/Public/SceneManagement.h
@@ -1726,7 +1726,7 @@ public:
 extern ENGINE_API FVector CalcConeVert(float Angle1, float Angle2, float AzimuthAngle);
 extern ENGINE_API void BuildConeVerts(float Angle1, float Angle2, float Scale, float XOffset, uint32 NumSides, TArray<FDynamicMeshVertex>& OutVerts, TArray<uint32>& OutIndices);
 
-void BuildCylinderVerts(const FVector& Base, const FVector& XAxis, const FVector& YAxis, const FVector& ZAxis, double Radius, double HalfHeight, uint32 Sides, TArray<FDynamicMeshVertex>& OutVerts, TArray<uint32>& OutIndices);
+extern ENGINE_API void BuildCylinderVerts(const FVector& Base, const FVector& XAxis, const FVector& YAxis, const FVector& ZAxis, double Radius, double HalfHeight, uint32 Sides, TArray<FDynamicMeshVertex>& OutVerts, TArray<uint32>& OutIndices);
 
 
 /**
