[Engine] Added PDI DrawTorus overload with vertex color

diff --git a/Engine/Source/Runtime/Engine/Private/PrimitiveDrawingUtils.cpp b/Engine/Source/Runtime/Engine/Private/PrimitiveDrawingUtils.cpp
--- a/Engine/Source/Runtime/Engine/Private/PrimitiveDrawingUtils.cpp
+++ b/Engine/Source/Runtime/Engine/Private/PrimitiveDrawingUtils.cpp
@@ -733,7 +733,6 @@ void GetCapsuleMesh(const FVector& Origin, const FVector& XAxis, const FVector&
 	GetOrientedHalfSphereMesh(BottomEnd, FRotationMatrix::MakeFromXY(XAxis, YAxis).Rotator(), FVector(Radius), NumSides, NumSides, UE_PI / 2, UE_PI, MaterialRenderProxy, DepthPriority, bDisableBackfaceCulling, ViewIndex, Collector, false, HitProxy);
 }
 
-
 void DrawCylinder(FPrimitiveDrawInterface* PDI,const FVector& Base, const FVector& XAxis, const FVector& YAxis, const FVector& ZAxis,
 	double Radius, double HalfHeight, uint32 Sides, const FMaterialRenderProxy* MaterialRenderProxy, uint8 DepthPriority)
 {
@@ -767,10 +766,14 @@ void DrawCylinder(class FPrimitiveDrawInterface* PDI, const FVector& Start, cons
 
 		DrawCylinder(PDI, FMatrix::Identity, Z * Length*0.5 + Start, X, Y, Z, Radius, Length * 0.5f, Sides, MaterialInstance, DepthPriority);
 	}
+}
 
+void DrawTorus(FPrimitiveDrawInterface* PDI, const FMatrix& Transform, const FVector& XAxis, const FVector& YAxis, double OuterRadius, double InnerRadius, int32 OuterSegments, int32 InnerSegments, const FMaterialRenderProxy* MaterialRenderProxy, uint8 DepthPriority, bool bPartial, float Angle, bool bEndCaps)
+{
+	DrawTorus(PDI, Transform, XAxis, YAxis, FColor::White, OuterRadius, InnerRadius, OuterSegments, InnerSegments, MaterialRenderProxy, DepthPriority, bPartial, Angle, bEndCaps);
 }
 
-void DrawTorus(FPrimitiveDrawInterface* PDI, const FMatrix& Transform, const FVector& XAxis, const FVector& YAxis, 
+void DrawTorus(FPrimitiveDrawInterface* PDI, const FMatrix& Transform, const FVector& XAxis, const FVector& YAxis, FColor Color,
 			   double OuterRadius, double InnerRadius, int32 OuterSegments, int32 InnerSegments, const FMaterialRenderProxy* MaterialRenderProxy,
 			   uint8 DepthPriority, bool bPartial, float Angle, bool bEndCaps)
 {
@@ -824,6 +827,7 @@ void DrawTorus(FPrimitiveDrawInterface* PDI, const FMatrix& Transform, const FVe
 			}
 
 			FDynamicMeshVertex MeshVertex;
+			MeshVertex.Color = Color;
 			MeshVertex.Position = FVector3f(Vertex);
 			MeshVertex.TextureCoordinate[0] = TC;
 			MeshVertex.TextureCoordinate[0].X += TCStepX * InnerIndex;
diff --git a/Engine/Source/Runtime/Engine/Public/PrimitiveDrawingUtils.h b/Engine/Source/Runtime/Engine/Public/PrimitiveDrawingUtils.h
--- a/Engine/Source/Runtime/Engine/Public/PrimitiveDrawingUtils.h
+++ b/Engine/Source/Runtime/Engine/Public/PrimitiveDrawingUtils.h
@@ -64,6 +64,27 @@ extern ENGINE_API void GetCapsuleMesh(const FVector& Origin, const FVector& XAxi
 	const FMaterialRenderProxy* MaterialRenderProxy, uint8 DepthPriority, bool bDisableBackfaceCulling, int32 ViewIndex, FMeshElementCollector& Collector, HHitProxy* HitProxy = nullptr);
 
 
+/**
+ * Draws a torus using triangles.
+ *
+ * @param	PDI						Draw interface.
+ * @param	Transform				Generic transform to apply (ex. a local-to-world transform).
+ * @param	XAxis					Normalized X alignment axis.
+ * @param	YAxis					Normalized Y alignment axis.
+ * @param	OuterRadius				Radius of the torus center-line. Viewed from above, the outside of the torus has a radius 
+ *                                  of OuterRadius + InnerRadius and the hole of the torus has a radius of OuterRadius - InnerRadius.
+ * @param	InnerRadius				Radius of the torus's cylinder.
+ * @param	OuterSegments			Numbers of segment divisions for outer circle.
+ * @param	InnerSegments			Numbers of segment divisions for inner circle.
+ * @param	MaterialRenderProxy		Material to use for render
+ * @param	DepthPriority			Depth priority for the circle.
+ * @param	bPartial				Whether full or partial torus should be rendered.
+ * @param	Angle					If partial, angle in radians of the arc clockwise beginning at the XAxis.
+ * @param	bEndCaps				If partial, whether the ends should be capped with triangles.
+ */
+extern ENGINE_API void DrawTorus(FPrimitiveDrawInterface* PDI, const FMatrix& Transform, const FVector& XAxis, const FVector& YAxis,
+								 double OuterRadius, double InnerRadius, int32 OuterSegments, int32 InnerSegments, const FMaterialRenderProxy* MaterialRenderProxy, uint8 DepthPriority, bool bPartial, float Angle, bool bEndCaps);
+
 /**
  * Draws a torus using triangles.
  *
@@ -83,7 +104,7 @@ extern ENGINE_API void GetCapsuleMesh(const FVector& Origin, const FVector& XAxi
  * @param	Angle					If partial, angle in radians of the arc clockwise beginning at the XAxis.
  * @param	bEndCaps				If partial, whether the ends should be capped with triangles.
  */
-extern ENGINE_API void DrawTorus(FPrimitiveDrawInterface* PDI, const FMatrix& Transform, const FVector& XAxis, const FVector& YAxis, 
+extern ENGINE_API void DrawTorus(FPrimitiveDrawInterface* PDI, const FMatrix& Transform, const FVector& XAxis, const FVector& YAxis, FColor Color,
 								 double OuterRadius, double InnerRadius, int32 OuterSegments, int32 InnerSegments, const FMaterialRenderProxy* MaterialRenderProxy, uint8 DepthPriority, bool bPartial, float Angle, bool bEndCaps);
 
 /**
