Update fallback memory estimate to 1GB for asset compilation

#rnx

diff --git a/Engine/Source/Runtime/Engine/Private/AssetCompilingManager.cpp b/Engine/Source/Runtime/Engine/Private/AssetCompilingManager.cpp
--- a/Engine/Source/Runtime/Engine/Private/AssetCompilingManager.cpp
+++ b/Engine/Source/Runtime/Engine/Private/AssetCompilingManager.cpp
@@ -55,7 +55,7 @@ static AsyncCompilationHelpers::FAsyncCompilationStandardCVars CVarAsyncCompilat
 //	default is 4 GB
 static TAutoConsoleVariable<int32> CVarAsyncAssetCompilationMemoryPerCore(
 	TEXT("Editor.AsyncAssetCompilationMemoryPerCore"),
-	4,
+	1,
 	TEXT("How much memory (in GB) should tasks reserve that report a required memory amount Unknown (-1).\n"),
 	ECVF_Default);
 
