OFPA: added World.CheckExternalActorsConsistency console command that will parse all worlds in all active mount points and report the ones that have inconsistent external actors.
Useful to gather the list of worlds that have external actors disabled but still have external actors files.

diff --git a/Engine/Source/Runtime/Engine/Private/World.cpp b/Engine/Source/Runtime/Engine/Private/World.cpp
--- a/Engine/Source/Runtime/Engine/Private/World.cpp
+++ b/Engine/Source/Runtime/Engine/Private/World.cpp
@@ -158,6 +158,68 @@ TRACE_DECLARE_INT_COUNTER(NumStreamingLevelsToConsider, TEXT("LevelStreamingProf
 
 #define LOCTEXT_NAMESPACE "World"
 
+#if WITH_EDITOR
+static FAutoConsoleCommand CheckExternalActorsConsistency(
+	TEXT("World.CheckExternalActorsConsistency"),
+	TEXT("Make sure worlds that don't use one file per actors don't have external actors."),
+	FConsoleCommandWithArgsDelegate::CreateLambda([](const TArray<FString>& Args)
+	{
+		IAssetRegistry& AssetRegistry = FModuleManager::LoadModuleChecked<FAssetRegistryModule>(TEXT("AssetRegistry")).Get();
+		AssetRegistry.WaitForCompletion();
+
+		TArray<FAssetData> WorldAssets;
+		{
+			FARFilter Filter;
+			Filter.bIncludeOnlyOnDiskAssets = true;
+			Filter.ClassPaths.Add(FTopLevelAssetPath(UWorld::StaticClass()));
+			Filter.bRecursiveClasses = true;
+
+			IAssetRegistry::GetChecked().GetAssets(Filter, WorldAssets);
+		}
+
+		TArray<TPair<FName, uint32>> InconsistentWorlds;
+
+		int32 Percentage = 0;
+		int32 WorldAssetIndex = 0;
+		for (const FAssetData& WorldAsset : WorldAssets)
+		{
+			if (!ULevel::GetIsLevelUsingExternalActorsFromAsset(WorldAsset))
+			{
+				FARFilter Filter;
+				Filter.bIncludeOnlyOnDiskAssets = true;
+				Filter.PackagePaths.Add(*ULevel::GetExternalActorsPath(WorldAsset.PackageName.ToString()));
+				Filter.bRecursivePaths = true;
+	
+				TArray<FAssetData> ActorAssets;
+				IAssetRegistry::GetChecked().GetAssets(Filter, ActorAssets);
+
+				if (!ActorAssets.IsEmpty())
+				{
+					InconsistentWorlds.Add({ WorldAsset.PackageName, ActorAssets.Num()});
+				}
+			}
+
+			const int32 CurrentPercentage = (++WorldAssetIndex * 100) / WorldAssets.Num();
+			if (Percentage != CurrentPercentage)
+			{
+				Percentage = CurrentPercentage;
+				UE_LOG(LogWorld, Log, TEXT("Processing world packages (%02d%%)"), Percentage);
+			}
+		}
+
+		if (InconsistentWorlds.Num())
+		{
+			UE_LOG(LogWorld, Warning, TEXT("Found %d inconsistent worlds:"), InconsistentWorlds.Num());
+
+			for (TPair<FName, uint32> InconsistentWorld : InconsistentWorlds)
+			{
+				UE_LOG(LogWorld, Warning, TEXT("\t %s (%d)"), *InconsistentWorld.Key.ToString(), InconsistentWorld.Value);
+			}
+		}
+	})
+);
+#endif //  WITH_EDITOR
+
 static int32 GMaximumMakingVisibleLevels = 1;
 static FAutoConsoleVariableRef CVarMaximumMakingVisibleLevels(
 	TEXT("LevelStreaming.MaximumMakingVisibleLevels"),
