WorldPartition - HLOD: Changed module type of WorldPartitionHLODUtilities from UncookedOnly to Editor
* UHLODLayer HLODBuilderClass & HLODBuilderSettings members are now editor only data, ensuring no reference to the WorldPartitionHLODUtilities module is kept in cooked builds

diff --git a/Engine/Plugins/Editor/WorldPartitionHLODUtilities/WorldPartitionHLODUtilities.uplugin b/Engine/Plugins/Editor/WorldPartitionHLODUtilities/WorldPartitionHLODUtilities.uplugin
--- a/Engine/Plugins/Editor/WorldPartitionHLODUtilities/WorldPartitionHLODUtilities.uplugin
+++ b/Engine/Plugins/Editor/WorldPartitionHLODUtilities/WorldPartitionHLODUtilities.uplugin
@@ -15,7 +15,7 @@
 	[
 		{
 			"Name" : "WorldPartitionHLODUtilities",
-			"Type" : "UncookedOnly",
+			"Type" : "Editor",
 			"LoadingPhase" : "Default"
 		}
 	]
diff --git a/Engine/Source/Runtime/Engine/Private/WorldPartition/HLOD/HLODLayer.cpp b/Engine/Source/Runtime/Engine/Private/WorldPartition/HLOD/HLODLayer.cpp
--- a/Engine/Source/Runtime/Engine/Private/WorldPartition/HLOD/HLODLayer.cpp
+++ b/Engine/Source/Runtime/Engine/Private/WorldPartition/HLOD/HLODLayer.cpp
@@ -29,11 +29,14 @@ DEFINE_LOG_CATEGORY_STATIC(LogHLODLayer, Log, All);
 
 UHLODLayer::UHLODLayer(const FObjectInitializer& ObjectInitializer)
 	: Super(ObjectInitializer)
+#if WITH_EDITORONLY_DATA
 	, bIsSpatiallyLoaded(true)
 	, CellSize(25600)
 	, LoadingRange(51200)
 	, HLODActorClass(AWorldPartitionHLOD::StaticClass())
-{}
+#endif
+{
+}
 
 #if WITH_EDITOR
 bool UHLODLayer::DoesRequireWarmup() const
