New material translator:
- Added support for named reroute expressions.
- Some refactoring and improvement to the logic that flows final values into SetMaterialAttribute instructions.

#rnx

diff --git a/Engine/Source/Runtime/Engine/Private/Materials/MaterialExpressionsIR.cpp b/Engine/Source/Runtime/Engine/Private/Materials/MaterialExpressionsIR.cpp
--- a/Engine/Source/Runtime/Engine/Private/Materials/MaterialExpressionsIR.cpp
+++ b/Engine/Source/Runtime/Engine/Private/Materials/MaterialExpressionsIR.cpp
@@ -1026,8 +1026,21 @@ void UMaterialExpressionTime::Build(MIR::FEmitter& Em)
 
 void UMaterialExpressionReroute::Build(MIR::FEmitter& Em)
 {
-	FValueRef Value = Em.Input(&Input);
-	Em.Output(0, Value);
+	Em.Output(0, Em.TryInput(&Input));
+}
+
+void UMaterialExpressionNamedRerouteDeclaration::Build(MIR::FEmitter& Em)
+{
+	Em.Output(0, Em.TryInput(&Input));
+}
+
+void UMaterialExpressionNamedRerouteUsage::Build(MIR::FEmitter& Em)
+{
+	if (!IsDeclarationValid())
+	{
+		Em.Error(TEXT("Named reroute expression does not have a valid declaration."));
+	}
+	Em.Output(0, Em.TryInput(&Declaration->Input));
 }
 
 void UMaterialExpressionClamp::Build(MIR::FEmitter& Em)
diff --git a/Engine/Source/Runtime/Engine/Private/Materials/MaterialIREmitter.cpp b/Engine/Source/Runtime/Engine/Private/Materials/MaterialIREmitter.cpp
--- a/Engine/Source/Runtime/Engine/Private/Materials/MaterialIREmitter.cpp
+++ b/Engine/Source/Runtime/Engine/Private/Materials/MaterialIREmitter.cpp
@@ -431,7 +431,10 @@ FEmitter& FEmitter::Output(int32 OutputIndex, FValueRef Value)
 
 FEmitter& FEmitter::Output(const FExpressionOutput* Output, FValueRef Value)
 {
-	MIR::Internal::BindValueToExpressionOutput(BuilderImpl, Output, Value);
+	if (Value)
+	{
+		MIR::Internal::BindValueToExpressionOutput(BuilderImpl, Output, Value);
+	}
 	return *this;
 }
 
diff --git a/Engine/Source/Runtime/Engine/Private/Materials/MaterialIRModuleBuilder.cpp b/Engine/Source/Runtime/Engine/Private/Materials/MaterialIRModuleBuilder.cpp
--- a/Engine/Source/Runtime/Engine/Private/Materials/MaterialIRModuleBuilder.cpp
+++ b/Engine/Source/Runtime/Engine/Private/Materials/MaterialIRModuleBuilder.cpp
@@ -15,6 +15,7 @@
 #include "Materials/MaterialExpressionFunctionInput.h"
 #include "Materials/MaterialExpressionFunctionOutput.h"
 #include "Materials/MaterialExpressionMaterialFunctionCall.h"
+#include "Materials/MaterialExpressionNamedReroute.h"
 #include "Materials/Material.h"
 #include "MaterialDomain.h"
 #include "MaterialExpressionIO.h"
@@ -99,6 +100,20 @@ struct FMaterialIRModuleBuilderImpl
 
 	void Step_GenerateOutputInstructions()
 	{
+		// Helper local function to process a single material attribute output.
+		auto PrepareSingleMaterialAttribute = [this](EMaterialProperty Property)
+		{
+			// Read the material input associated to this property
+			FMaterialInputDescription Input;
+			GetExpressionInputDescription(Property, Input);
+
+			// Push the connected expression to this material attribute input as a dependency, if any.
+			PushDependency(Input.Input->Expression);
+
+			// Create the set material output instruction (leave it unset for now; it is set in Step_FlowValuesIntoMaterialOutputs)
+			Emitter->SetMaterialOutput(Property, nullptr);
+		};
+		
 		// The normal input is read back from the value set in the material attribute.
 		// For this reason, the normal attribute is evaluated and set first, ensuring that
 		// other inputs can read its value.
@@ -120,41 +135,6 @@ struct FMaterialIRModuleBuilderImpl
 		}
 	}
 
-	bool GetExpressionInputDescription(EMaterialProperty Property, FMaterialInputDescription& Input)
-	{
-		if (PreviewInput.IsConnected() && Property == MP_EmissiveColor)
-		{
-			Input.Type = UE::Shader::EValueType::Float3;
-			Input.Input = &PreviewInput;
-			return true;
-		}
-		else
-		{
-			return Builder->Material->GetExpressionInputDescription(Property, Input);
-		}
-	}
-
-	void PrepareSingleMaterialAttribute(EMaterialProperty Property)
-	{
-		FMaterialInputDescription Input;
-		GetExpressionInputDescription(Property, Input);
-
-		MIR::FSetMaterialOutput* Output = Emitter->SetMaterialOutput(Property, nullptr);
-
-		if (Input.bUseConstant)
-		{
-			Output->Arg = Emitter->ConstantFromShaderValue(Input.ConstantValue);
-		}
-		else if (!Input.Input->IsConnected())
-		{
-			Output->Arg = MIR::Internal::CreateMaterialAttributeDefaultValue(*Emitter, Builder->Material, Property);
-		}
-		else
-		{
-			AnalysisContextStack.Last().ExpressionStack.Add(Input.Input->Expression);
-		}
-	}
-
 	void Step_BuildMaterialExpressionsToIRGraph()
 	{
 		while (true)
@@ -199,13 +179,17 @@ struct FMaterialIRModuleBuilderImpl
 		// Push to the expression stack all dependencies that still need to be analyzed.
 		for (FExpressionInputIterator It{ Emitter->Expression }; It; ++It)
 		{
-			// Ignore disconnected inputs and connected expressions already built.
-			if (!It->IsConnected() || CurrContext.BuiltExpressions.Contains(It->Expression))
+			PushDependency(It.Input->Expression);
+		}
+
+		// Named reroute usage nodes should make add their declaration as dependency, as they
+		// do not have any direct input (they simply forward to the declaration input).
+		if (auto NamedReroute = Cast<UMaterialExpressionNamedRerouteUsage>(Emitter->Expression))
+		{
+			if (NamedReroute->IsDeclarationValid())
 			{
-				continue;
+				PushDependency(NamedReroute->Declaration.Get());
 			}
-
-			CurrContext.ExpressionStack.Push(It->Expression);
 		}
 
 		// If on top of the stack there's a different expression, we have a dependency to analyze first.
@@ -221,16 +205,16 @@ struct FMaterialIRModuleBuilderImpl
 		// Flow the value into this expression's inputs from their connected outputs.
 		for (FExpressionInputIterator It{ Emitter->Expression}; It; ++It)
 		{
-			// Fetch the value flowing through connected output.
-			FExpressionOutput* ConnectedOutput = It->GetConnectedOutput();
-			if (ConnectedOutput)
+			FlowValueThroughConnection(It.Input);
+		}
+		
+		// Usage reroute nodes should forward the value coming in from their declaration input (since
+		// they don't have an input of their own).
+		if (auto NamedReroute = Cast<UMaterialExpressionNamedRerouteUsage>(Emitter->Expression))
+		{
+			if (NamedReroute->IsDeclarationValid())
 			{
-				MIR::FValue** ValuePtr = CurrContext.OutputValues.Find(ConnectedOutput);
-				if (ValuePtr)
-				{
-					// ...and flow it into this input.
-					CurrContext.InputValues.Add(It.Input, *ValuePtr);
-				}
+				FlowValueThroughConnection(&NamedReroute->Declaration->Input);
 			}
 		}
 
@@ -249,6 +233,36 @@ struct FMaterialIRModuleBuilderImpl
 			AddExpressionConnectionInsights(Emitter->Expression);
 		}
 	}
+	
+	// Pushes an expression dependency on this context expression stack.
+	void PushDependency(UMaterialExpression* Expression)
+	{
+		FAnalysisContext& CurrContext = AnalysisContextStack.Last();
+
+		// Ignore disconnected inputs and connected expressions already built.
+		if (Expression && !CurrContext.BuiltExpressions.Contains(Expression))
+		{
+			CurrContext.ExpressionStack.Push(Expression);
+		}
+	}
+
+	// Flows the value from connected output to specified input.
+	void FlowValueThroughConnection(FExpressionInput* Input)
+	{
+		FAnalysisContext& CurrContext = AnalysisContextStack.Last();
+
+		// Fetch the value flowing through connected output.
+		FExpressionOutput* ConnectedOutput = Input->GetConnectedOutput();
+		if (ConnectedOutput)
+		{
+			MIR::FValue** ValuePtr = CurrContext.OutputValues.Find(ConnectedOutput);
+			if (ValuePtr)
+			{
+				// ...and flow it into this input.
+				CurrContext.InputValues.Add(Input, *ValuePtr);
+			}
+		}
+	}
 
 	void PushFunctionCall(UMaterialExpressionMaterialFunctionCall* Call)
 	{
@@ -338,60 +352,84 @@ struct FMaterialIRModuleBuilderImpl
 
 	void Step_FlowValuesIntoMaterialOutputs()
 	{
-		FAnalysisContext& Context = AnalysisContextStack.Last();
-
 		for (uint32 StageIndex = 0; StageIndex < MIR::NumStages; ++StageIndex)
 		{
 			for (MIR::FSetMaterialOutput* Output : Module->Outputs[StageIndex])
 			{
+				// Get the material attribute input info.
 				FMaterialInputDescription Input;
 				ensure(GetExpressionInputDescription(Output->Property, Input));
 
-				// Connect output properties with substrate parameter if a FrontMaterial property is provided
-				if (Output->Property == MP_FrontMaterial && Output->Arg)
+				// Flow, if any, the value coming in through the connection to the material attribute pin.
+				FlowValueThroughConnection(Input.Input);
+
+				if (Input.bUseConstant)
 				{
-					if (MIR::FSubstrateParameter* SubstrateParameter = Output->Arg->As<MIR::FSubstrateParameter>())
-					{
-						for (int32 PropertyIndex = 0; PropertyIndex < MP_MAX; ++PropertyIndex)
-						{
-							SubstrateParameter->PropertyArgs[PropertyIndex] = FindIRModuleOutputValue(*Module, MIR::EStage::Stage_Pixel, static_cast<EMaterialProperty>(PropertyIndex));
-						}
-					}
+					// If input is marked to use constant, assign this output to the specified constant value.
+					Output->Arg = Emitter->ConstantFromShaderValue(Input.ConstantValue);
 				}
+				else if (MIR::FValue* Value = MIR::Internal::FetchValueFromExpressionInput(this, Input.Input))
+				{
+					// If a value is flowing in through the connection, cast it to this material attribute type and assign it.
+					const MIR::FType* OutputArgType = MIR::FType::FromShaderType(Input.Type);
+					Output->Arg = Emitter->Cast(Value, OutputArgType);
+				}
+				else
+				{
+					// Otherwise, fallback to assigning this material output to its default value.
+					Output->Arg = MIR::Internal::CreateMaterialAttributeDefaultValue(*Emitter, Builder->Material, Output->Property);
+				}
+
+				// The value being set to this material output is now valid.
+				check(Output->Arg);
 
-				if (!Output->Arg)
+				// If this property is the emissive color and we're previewing the material, apply gamma correction to the previewed value.
+				if (Output->Property == MP_EmissiveColor && PreviewInput.IsConnected())
 				{
-					MIR::FValue** ValuePtr = Context.OutputValues.Find(Input.Input->GetConnectedOutput());
-					check(ValuePtr && *ValuePtr);
-					MIR::FValue* Value = *ValuePtr;
+					MIR::FValue* Zero = Emitter->ConstantZero(MIR::EScalarKind::ScalarKind_Float);
 
-					if (PreviewInput.IsConnected() && Output->Property == MP_EmissiveColor)
-					{
-						MIR::FValue* Zero = Emitter->ConstantZero(MIR::EScalarKind::ScalarKind_Float);
+					// Get preview expression back into gamma corrected space, as DrawTile does not do this adjustment.
+					Output->Arg = Emitter->Pow(Emitter->Max(Output->Arg, Zero), Emitter->ConstantFloat(1.0f / 2.2f));
 
-						// Get preview expression back into gamma corrected space, as DrawTile does not do this adjustment.
-						Value = Emitter->Operator(MIR::EOperator::BO_Pow, Emitter->Max(Value, Zero), Emitter->ConstantFloat(1.f / 2.2f));
+					// Preview should display scalars as red, so if this is a scalar, create a vector padded with zeroes
+					if (Output->Arg->IsScalar())
+					{
+						Output->Arg = Emitter->Vector3(Output->Arg, Zero, Zero);
+					}
+				}
 
-						// Preview should display scalars as red, so if this is a scalar, create a vector padded with zeroes
-						if (Value->IsScalar())
+				// Connect output properties with substrate parameter if a FrontMaterial property is provided
+				if (Output->Property == MP_FrontMaterial)
+				{
+					if (MIR::FSubstrateParameter* SubstrateParameter = Output->Arg->As<MIR::FSubstrateParameter>())
+					{
+						for (int32 PropertyIndex = 0; PropertyIndex < MP_MAX; ++PropertyIndex)
 						{
-							Value = Emitter->Vector3(Value, Zero, Zero);
+							SubstrateParameter->PropertyArgs[PropertyIndex] = FindIRModuleOutputValue(*Module, MIR::EStage::Stage_Pixel, static_cast<EMaterialProperty>(PropertyIndex));
 						}
 					}
-					
-					MIR::Internal::BindValueToExpressionInput(this, Input.Input, Value);
-
-					const MIR::FType* OutputArgType = MIR::FType::FromShaderType(Input.Type);
-					Output->Arg = Emitter->Cast(Value, OutputArgType);
 				}
 
-				// Push this connection insight
-				check(Output->Arg);
+				// Finally, push this connection insight
 				PushConnectionInsight(Builder->Material, (int)Output->Property, Input.Input->Expression, Input.Input->OutputIndex, Output->Arg->Type);
 			}
 		}
 	}
 
+	bool GetExpressionInputDescription(EMaterialProperty Property, FMaterialInputDescription& Input)
+	{
+		if (PreviewInput.IsConnected() && Property == MP_EmissiveColor)
+		{
+			Input.Type = UE::Shader::EValueType::Float3;
+			Input.Input = &PreviewInput;
+			return true;
+		}
+		else
+		{
+			return Builder->Material->GetExpressionInputDescription(Property, Input);
+		}
+	}
+
 	void Step_AnalyzeIRGraph()
 	{
 		TArray<MIR::FValue*> ValueStack{};
diff --git a/Engine/Source/Runtime/Engine/Public/Materials/MaterialExpressionNamedReroute.h b/Engine/Source/Runtime/Engine/Public/Materials/MaterialExpressionNamedReroute.h
--- a/Engine/Source/Runtime/Engine/Public/Materials/MaterialExpressionNamedReroute.h
+++ b/Engine/Source/Runtime/Engine/Public/Materials/MaterialExpressionNamedReroute.h
@@ -62,6 +62,7 @@ public:
 
 	//~ Begin UMaterialExpression Interface
 #if WITH_EDITOR
+	ENGINE_API virtual void Build(MIR::FEmitter& Emitter) override;
 	ENGINE_API virtual int32 Compile(class FMaterialCompiler* Compiler, int32 OutputIndex) override;
 	ENGINE_API virtual int32 CompilePreview(FMaterialCompiler* Compiler, int32 OutputIndex) override;
 	ENGINE_API virtual void GetCaption(TArray<FString>& OutCaptions) const override;
@@ -106,8 +107,12 @@ public:
 	UPROPERTY()
 	FGuid DeclarationGuid;
 
+	// Check that the declaration isn't deleted
+	bool IsDeclarationValid() const;
+
 	//~ Begin UMaterialExpression Interface
 #if WITH_EDITOR
+	ENGINE_API virtual void Build(MIR::FEmitter& Emitter) override;
 	ENGINE_API virtual int32 Compile(class FMaterialCompiler* Compiler, int32 OutputIndex) override;
 	ENGINE_API virtual int32 CompilePreview(FMaterialCompiler* Compiler, int32 OutputIndex) override;
 
@@ -124,10 +129,6 @@ protected:
 	//~ Begin UMaterialExpressionRerouteBase Interface
 	ENGINE_API virtual bool GetRerouteInput(FExpressionInput& OutInput) const override;
 	//~ End UMaterialExpressionRerouteBase Interface
-
-private:
-	// Check that the declaration isn't deleted
-	bool IsDeclarationValid() const;
 };
 
 template<typename ExpressionsArrayType>
