[NewMaterialTranslator] - Adds the ability to debug break when a Poison value is encountered when building the IR.

- Can be enabled by `r.Material.Translator.DebugBreakOnPoison` which defaults to false.
- Added a commented out cvar to ConsoleVariables.ini to raise visilbity on the setting.

todo: Add a material translator development cvar, similar to shader development, which automatically enables a bunch of features to make development on the translator easier.

#rnx

diff --git a/Engine/Config/ConsoleVariables.ini b/Engine/Config/ConsoleVariables.ini
--- a/Engine/Config/ConsoleVariables.ini
+++ b/Engine/Config/ConsoleVariables.ini
@@ -61,6 +61,8 @@ r.DumpShaderDebugInfo=2
 ;r.Material.Translator.EnableNew=1
 ; When this is enabled the new translator will dump out the IR graph to a dotviz graph to your project's saved directory.
 ;r.Material.Translator.DebugDump=2
+; When this is enabled the new translator will break in the debugger when a poison value is emitted in the IR.
+;r.Material.Translator.DebugBreakOnPoison=True
 
 ; Uncomment to disable parallel rendering
 ;r.RHICmdBypass=1
diff --git a/Engine/Source/Runtime/Engine/Private/Materials/MaterialIREmitter.cpp b/Engine/Source/Runtime/Engine/Private/Materials/MaterialIREmitter.cpp
--- a/Engine/Source/Runtime/Engine/Private/Materials/MaterialIREmitter.cpp
+++ b/Engine/Source/Runtime/Engine/Private/Materials/MaterialIREmitter.cpp
@@ -13,6 +13,12 @@
 
 #if WITH_EDITOR
 
+static TAutoConsoleVariable<bool> CVarMaterialIRDebugBreakOnPoison(
+	TEXT("r.Material.Translator.DebugBreakOnPoison"),
+	0,
+	TEXT("Whether the material translator break in the debugger when hitting a poison value in the IR.\n"),
+	ECVF_Default);
+
 namespace MIR
 {
 
@@ -574,6 +580,10 @@ FValueRef FEmitter::ConstantFloat4(UE::Math::TVector4<TFloat> InValue)
 
 FValueRef FEmitter::Poison()
 {
+ 	if (CVarMaterialIRDebugBreakOnPoison.GetValueOnGameThread())
+ 	{
+ 		UE_DEBUG_BREAK();
+ 	}
 	return FPoison::Get();
 }
 
