CurveLinearColor : fix UV stretch in Atlas so that [0,1] parameter goes from first to last pixel centers

diff --git a/Engine/Source/Runtime/Core/Private/UObject/FortniteMainBranchObjectVersion.cpp b/Engine/Source/Runtime/Core/Private/UObject/FortniteMainBranchObjectVersion.cpp
--- a/Engine/Source/Runtime/Core/Private/UObject/FortniteMainBranchObjectVersion.cpp
+++ b/Engine/Source/Runtime/Core/Private/UObject/FortniteMainBranchObjectVersion.cpp
@@ -14,6 +14,6 @@ TMap<FGuid, FGuid> FFortniteMainBranchObjectVersion::GetSystemGuids()
 	SystemGuids.Add(DevGuids.Niagara_LatestScriptCompileVersion, FGuid("97F0FD7149D90D42B65D6557AF477D1A"));
 	SystemGuids.Add(DevGuids.SkeletalMeshDerivedDataVersion, FGuid("716424BCEC314E969ACFE36D2C2ECE49"));
 	SystemGuids.Add(DevGuids.STATICMESH_DERIVEDDATA_VER, FGuid("3ABF32149F9D4D838750368AB95573BE"));
-	SystemGuids.Add(DevGuids.MaterialTranslationDDCVersion, FGuid("AAAFC530186D4149A06F027D4A060F5E"));
+	SystemGuids.Add(DevGuids.MaterialTranslationDDCVersion, FGuid("2F6A9DD0B2674CF5B7564F4A14CC40E4"));
 	return SystemGuids;
 }
diff --git a/Engine/Source/Runtime/Engine/Private/Curves/CurveLinearColor.cpp b/Engine/Source/Runtime/Engine/Private/Curves/CurveLinearColor.cpp
--- a/Engine/Source/Runtime/Engine/Private/Curves/CurveLinearColor.cpp
+++ b/Engine/Source/Runtime/Engine/Private/Curves/CurveLinearColor.cpp
@@ -307,7 +307,7 @@ void UCurveLinearColor::DrawThumbnail(FCanvas* Canvas, FVector2D StartXY, FVecto
 	{
 		DrawItem.Position.X = StartXY.X + x;
 
-		FLinearColor Color = GetLinearColorValue(x / Width);
+		FLinearColor Color = GetLinearColorValue(x * (1.f / (Width-1)) );
 
 		DrawItem.SetColor(Color);
 		DrawItem.Draw(Canvas);
@@ -323,7 +323,7 @@ void UCurveLinearColor::PushToSourceData(TArray<FFloat16Color> &SrcData, int32 S
 		// Create base mip for the texture we created.
 		for (uint32 x = 0; x < SizeXY.X; x++)
 		{
-			SrcData[Start + x + y*SizeXY.X] = GetLinearColorValue(x / SizeXY.X);
+			SrcData[Start + x + y*SizeXY.X] = GetLinearColorValue(x * (1.f / (SizeXY.X-1)) );
 		}
 	}
 }
@@ -336,7 +336,7 @@ void UCurveLinearColor::PushUnadjustedToSourceData(TArray<FFloat16Color>& SrcDat
 		// Create base mip for the texture we created.
 		for (uint32 x = 0; x < SizeXY.X; x++)
 		{
-			SrcData[Start + x + y * SizeXY.X] = GetUnadjustedLinearColorValue(x / SizeXY.X);
+			SrcData[Start + x + y * SizeXY.X] = GetUnadjustedLinearColorValue(x * (1.f / (SizeXY.X-1)));
 		}
 	}
 }
diff --git a/Engine/Source/Runtime/Engine/Private/Materials/MaterialExpressions.cpp b/Engine/Source/Runtime/Engine/Private/Materials/MaterialExpressions.cpp
--- a/Engine/Source/Runtime/Engine/Private/Materials/MaterialExpressions.cpp
+++ b/Engine/Source/Runtime/Engine/Private/Materials/MaterialExpressions.cpp
@@ -21887,11 +21887,31 @@ int32 UMaterialExpressionCurveAtlasRowParameter::Compile(class FMaterialCompiler
 	int32 AtlasCode = Compiler->Texture(Atlas, AtlasRef, SAMPLERTYPE_LinearColor, SSM_Clamp_WorldGroupSettings, TMVM_None);
 	if (AtlasCode != INDEX_NONE)
 	{
-		int32 AtlasHeight = Compiler->ComponentMask(Compiler->TextureProperty(AtlasCode, TMTM_TextureSize), false, true, false, false);
+		int32 AtlasSize = Compiler->TextureProperty(AtlasCode, TMTM_TextureSize);
+		int32 AtlasWidth  = Compiler->ComponentMask(AtlasSize, true, false, false, false);
+		int32 AtlasHeight = Compiler->ComponentMask(AtlasSize, false, true, false, false);
 
 		// Calculate UVs from height and slot
 		// if the input is hooked up, use it, otherwise use the internal constant
-		int32 Arg1 = InputTime.GetTracedInput().Expression ? InputTime.Compile(Compiler) : Compiler->Constant(0);
+		int32 ArgT;
+		if ( InputTime.GetTracedInput().Expression )
+		{
+			ArgT = InputTime.Compile(Compiler);
+		}
+		else
+		{
+			ArgT = Compiler->Constant(0);
+		}
+
+		// ArgT is [t] the parameter on the curve
+		//	adjust so that [0,1] in t corresponds to the first to last pixel center
+		//		(CurveLinearColorAtlas has no mips, so that's not an issue)
+		//	Arg1 = (t*(width-1) + 0.5)/width
+		int32 Arg1 = Compiler->Div(Compiler->Add( Compiler->Mul(ArgT,Compiler->Sub(AtlasWidth,Compiler->Constant(1.0))) , Compiler->Constant(0.5)),AtlasWidth);
+		
+		// Arg2 is the curve index
+		//	each row of the atlas texture can be a different curve
+		//	Arg2 = (index + 0.5)/height so it is centered vertically on a pixel row
 		int32 Arg2 = Compiler->Div(Compiler->Add(Slot, Compiler->Constant(0.5)), AtlasHeight);
 
 		int32 UV = Compiler->AppendVector(Arg1, Arg2);
