Iris
* Revert NetSubObjectFactory to only use the object class to instantiate object instead of the original template
* Temp while investigating why some archetypes do not exist on the client

#rnx

diff --git a/Engine/Source/Runtime/Engine/Private/Net/Iris/ReplicationSystem/NetSubObjectFactory.cpp b/Engine/Source/Runtime/Engine/Private/Net/Iris/ReplicationSystem/NetSubObjectFactory.cpp
--- a/Engine/Source/Runtime/Engine/Private/Net/Iris/ReplicationSystem/NetSubObjectFactory.cpp
+++ b/Engine/Source/Runtime/Engine/Private/Net/Iris/ReplicationSystem/NetSubObjectFactory.cpp
@@ -60,8 +60,9 @@ TUniquePtr<UE::Net::FNetObjectCreationHeader> UNetSubObjectFactory::CreateAndFil
 	check(SubObject->NeedsLoadForClient() ); // We have no business sending this unless the client can load
 	check(SubObject->GetClass()->NeedsLoadForClient());	// We have no business sending this unless the client can load
 
-	const UObject* Template = bUseArchetypeInSubObjectCreation ? SubObject->GetArchetype() : SubObject->GetClass()->GetDefaultObject();
-	SubObjectHeader->ArchetypeReference = Bridge->GetOrCreateObjectReference(Template);
+	// Only use the class as the template while investigating why some archetypes are missing from the client
+	//const UObject* Template = bUseArchetypeInSubObjectCreation ? SubObject->GetArchetype() : SubObject->GetClass()->GetDefaultObject();
+	SubObjectHeader->ArchetypeReference = Bridge->GetOrCreateObjectReference(SubObject->GetClass());
 
 	// Find the right Outer
 	UObject* OuterObject = SubObject->GetOuter();	
@@ -185,8 +186,9 @@ UNetSubObjectFactory::FInstantiateResult UNetSubObjectFactory::InstantiateReplic
 	}
 
 	// Instantiate the SubObject from it's archetype
-	UClass* SubObjClass = Archetype->GetClass();
-	UObject* SubObj = NewObject<UObject>(OuterObject, SubObjClass, NAME_None, RF_NoFlags, Archetype);
+	// Archetype is always the class temporarely
+	UClass* SubObjClass = Cast<UClass>(Archetype);
+	UObject* SubObj = NewObject<UObject>(OuterObject, SubObjClass);
 
 	// Sanity check some things
 	checkf(SubObj != nullptr, TEXT("UNetSubObjectFactory::InstantiateReplicatedObjectFromHeader: Subobject is NULL after instantiating. Class: %s, Outer %s, Actor %s"), *GetNameSafe(SubObjClass), *GetNameSafe(OuterObject), *GetNameSafe(RootObject));
