LinearColor::HSVToLinearRGB do the required mods and clamps in the function so it's impossible for callers to fail to do that
CurveLinearColor : un-code-dupe the color adjustment function
TextureCompressorModule : AdjustColorsNew clean up since HSVToLinearRGB does clamps now

diff --git a/Engine/Source/Developer/TextureCompressor/Private/TextureCompressorModule.cpp b/Engine/Source/Developer/TextureCompressor/Private/TextureCompressorModule.cpp
--- a/Engine/Source/Developer/TextureCompressor/Private/TextureCompressorModule.cpp
+++ b/Engine/Source/Developer/TextureCompressor/Private/TextureCompressorModule.cpp
@@ -2873,21 +2873,10 @@ static inline void AdjustColorsNew(FLinearColor* Colors, int64 Count, const FTex
 			OriginalColor.R = FMath::Clamp(OriginalColor.R, 0.0f, 1.f);
 			OriginalColor.G = FMath::Clamp(OriginalColor.G, 0.0f, 1.f);
 			OriginalColor.B = FMath::Clamp(OriginalColor.B, 0.0f, 1.f);
+			// A will be clamped later when used
 		}
 		else
-		{
-			/*
-			if ( OriginalColor.R < 0 || OriginalColor.G < 0 || OriginalColor.B < 0 )
-			{
-				// need to log texture name for this to be of any use :
-				UE_CALL_ONCE( [&](){
-					UE_LOG(LogTextureCompressor, Warning,
-						TEXT("Negative pixel values (%f, %f, %f) are not expected"),
-						OriginalColor.R, OriginalColor.G, OriginalColor.B);
-				} );
-			}
-			*/
-			
+		{			
 			// yes HDR source, but we don't support negatives in Adjust
 			// clamp to >= 0 :
 			OriginalColor.R = FMath::Max(OriginalColor.R, 0.0f);
@@ -2896,10 +2885,10 @@ static inline void AdjustColorsNew(FLinearColor* Colors, int64 Count, const FTex
 		}
 
 		// Convert to HSV
-		FLinearColor HSVColor = OriginalColor.LinearRGBToHSV();
-		float& PixelHue = HSVColor.R;
-		float& PixelSaturation = HSVColor.G;
-		float& PixelValue = HSVColor.B;
+		const FLinearColor OriginalHSVColor = OriginalColor.LinearRGBToHSV();
+		float PixelHue = OriginalHSVColor.R;
+		float PixelSaturation = OriginalHSVColor.G;
+		float PixelValue = OriginalHSVColor.B;
 
 		if ( bAdjustValue )
 		{
@@ -2941,23 +2930,15 @@ static inline void AdjustColorsNew(FLinearColor* Colors, int64 Count, const FTex
 
 			// Apply saturation adjustment
 			PixelSaturation *= Params.AdjustSaturation;
-			PixelSaturation = FMath::Clamp(PixelSaturation, 0.0f, 1.0f);
 		}
 
 		// Apply hue adjustment
-		if ( AdjustHue != 0.f )
-		{
-			// PixelHue is [0,360) but AdjustHue is [0,360]
-			PixelHue += AdjustHue;
-
-			// Clamp HSV values
-			if ( PixelHue >= 360.f )
-			{
-				PixelHue -= 360.f;
-			}
-		}
+		// PixelHue is [0,360) but AdjustHue is [0,360]
+		PixelHue += AdjustHue;
 
 		// Convert back to a linear color
+		FLinearColor HSVColor(PixelHue,PixelSaturation,PixelValue);
+		// HSVToLinearRGB does the clamping/fmod for you
 		FLinearColor LinearColor = HSVColor.HSVToLinearRGB();
 
 		// Apply RGB curve adjustment (linear space)
diff --git a/Engine/Source/Runtime/Core/Private/Math/Color.cpp b/Engine/Source/Runtime/Core/Private/Math/Color.cpp
--- a/Engine/Source/Runtime/Core/Private/Math/Color.cpp
+++ b/Engine/Source/Runtime/Core/Private/Math/Color.cpp
@@ -403,6 +403,7 @@ FLinearColor FLinearColor::LinearRGBToHSV() const
 	{
 		return FLinearColor(0.f,0.f,0.f,A);
 	}
+	// RGBMin can be negative ; while that doesn't work great, it also doesn't break the math below
 
 	const float Hue = (RGBMax == RGBMin ? 0.0f :
 					   RGBMax == R    ? FMath::Fmod((((G - B) / RGBRange) * 60.0f) + 360.0f, 360.0f) :
@@ -423,13 +424,13 @@ FLinearColor FLinearColor::LinearRGBToHSV() const
 FLinearColor FLinearColor::HSVToLinearRGB() const
 {
 	// In this color, R = H, G = S, B = V  ; H in [0,360] , S&V in [0,1]
-	const float Hue = R;
-	const float Saturation = G;
-	const float Value = B;
+	const float Hue = fmodf(R,360.f) + 360.f;
+	const float Saturation = FMath::Clamp(G,0.f,1.f);
+	const float Value = B; // Value can be out of [0,1]
 
 	const float HDiv60 = Hue / 60.0f;
-	const float HDiv60_Floor = floorf(HDiv60);
-	const float HDiv60_Fraction = HDiv60 - HDiv60_Floor;
+	const uint32 HDiv60_Floor = (uint32)(HDiv60);
+	const float HDiv60_Fraction = HDiv60 - (float)HDiv60_Floor;
 
 	const float RGBValues[4] = {
 		Value,
@@ -445,7 +446,7 @@ FLinearColor FLinearColor::HSVToLinearRGB() const
 		{3, 1, 0},
 		{0, 1, 2},
 	};
-	const uint32 SwizzleIndex = ((uint32)HDiv60_Floor) % 6;  // use of uint32 ensures SwizzleIndex can't be negative even if Hue was
+	const uint32 SwizzleIndex = HDiv60_Floor % 6;
 
 	return FLinearColor(RGBValues[RGBSwizzle[SwizzleIndex][0]],
 						RGBValues[RGBSwizzle[SwizzleIndex][1]],
diff --git a/Engine/Source/Runtime/Engine/Classes/Curves/CurveLinearColor.h b/Engine/Source/Runtime/Engine/Classes/Curves/CurveLinearColor.h
--- a/Engine/Source/Runtime/Engine/Classes/Curves/CurveLinearColor.h
+++ b/Engine/Source/Runtime/Engine/Classes/Curves/CurveLinearColor.h
@@ -40,12 +40,17 @@ class UCurveLinearColor : public UCurveBase
 	ENGINE_API virtual TArray<FRichCurveEditInfo> GetCurves() override;
 	virtual bool IsLinearColorCurve() const override { return true; }
 
+	// GetLinearColorValue allows RGB > 1 for HDR
+	//	 if the input curves are LDR (<= 1) then the output is clamped to stay LDR, even if Adjustments would have changed it
 	UFUNCTION(BlueprintCallable, Category="Math|Curves")
 	ENGINE_API virtual FLinearColor GetLinearColorValue(float InTime) const override;
 
+	// GetClampedLinearColorValue always clamps RGB in [0,1] , eg. returns LDR colors
 	UFUNCTION(BlueprintCallable, Category = "Math|Curves")
 	ENGINE_API virtual FLinearColor GetClampedLinearColorValue(float InTime) const override;
 
+	// GetUnadjustedLinearColorValue returns the raw curve values without color adjustments
+	//	also does NOT clamp in [0,1] , beware how the RGBA is used, you may want clamping, at least >= 0
 	UFUNCTION(BlueprintCallable, Category = "Math|Curves")
 	ENGINE_API FLinearColor GetUnadjustedLinearColorValue(float InTime) const;
 
@@ -101,5 +106,11 @@ public:
 	UE_DEPRECATED(4.26, "OnUpdateGradient is deprecated. Please use UCurveBase::OnUpdateCurve instead")
 	FOnUpdateGradient OnUpdateGradient;
 #endif
+
+private:
+
+	// if bClampOver1Always is true, clamping in [0,1] is always done
+	// if bClampOver1Always is false, clamp in [0,1] is still done if the source RGB is in [0,1] , but NOT if source RGB is > 1
+	FLinearColor GetAdjustedColorValue(float InTime, bool bClampOver1Always) const;
 };
 
diff --git a/Engine/Source/Runtime/Engine/Private/Curves/CurveLinearColor.cpp b/Engine/Source/Runtime/Engine/Private/Curves/CurveLinearColor.cpp
--- a/Engine/Source/Runtime/Engine/Private/Curves/CurveLinearColor.cpp
+++ b/Engine/Source/Runtime/Engine/Private/Curves/CurveLinearColor.cpp
@@ -80,88 +80,28 @@ static FLinearColor ClampNonNegative(const FLinearColor & InColor)
 
 FLinearColor UCurveLinearColor::GetLinearColorValue( float InTime ) const
 {
-	// Logic copied from .\Engine\Source\Developer\TextureCompressor\Private\TextureCompressorModule.cpp
-	FLinearColor OriginalColor = GetUnadjustedLinearColorValue(InTime);
-	
-	// allow > 1.0 at this point but not < 0.0
-	//	we're just about to convert this RGB to HSV, that doesn't support negatives
-	OriginalColor = ClampNonNegative(OriginalColor);
-
-	// Only clamp value if the color is RGB < 1
-	const bool bShouldClampValue = (OriginalColor.R <= 1.0f && OriginalColor.G <= 1.0f && OriginalColor.B <= 1.0f);
-
-	// Convert to HSV
-	FLinearColor HSVColor = OriginalColor.LinearRGBToHSV();
-	float& PixelHue = HSVColor.R;
-	float& PixelSaturation = HSVColor.G;
-	float& PixelValue = HSVColor.B;
-
-	// Apply brightness adjustment
-	PixelValue *= AdjustBrightness;
-
-	// Apply brightness power adjustment
-	if (!FMath::IsNearlyEqual(AdjustBrightnessCurve, 1.0f, (float)UE_KINDA_SMALL_NUMBER) && AdjustBrightnessCurve != 0.0f)
-	{
-		// Raise HSV.V to the specified power
-		PixelValue = FMath::Pow(PixelValue, AdjustBrightnessCurve);
+	return GetAdjustedColorValue(InTime, false);
 	}
 
-	// Apply "vibrancy" adjustment
-	if (!FMath::IsNearlyZero(AdjustVibrance, (float)UE_KINDA_SMALL_NUMBER))
-	{
-		const float SatRaisePow = 5.0f;
-		const float InvSatRaised = FMath::Pow(1.0f - PixelSaturation, SatRaisePow);
-
-		const float ClampedVibrance = FMath::Clamp(AdjustVibrance, 0.0f, 1.0f);
-		const float HalfVibrance = ClampedVibrance * 0.5f;
-
-		const float SatProduct = HalfVibrance * InvSatRaised;
-
-		PixelSaturation += SatProduct;
-	}
-
-	// Apply saturation adjustment
-	PixelSaturation *= AdjustSaturation;
-
-	// Apply hue adjustment
-	PixelHue += AdjustHue;
-
-	// Clamp HSV values
-	{
-		PixelHue = FMath::Fmod(PixelHue, 360.0f);
-		if (PixelHue < 0.0f)
-		{
-			// Keep the hue value positive as HSVToLinearRGB prefers that
-			PixelHue += 360.0f;
-		}
-		PixelSaturation = FMath::Clamp(PixelSaturation, 0.0f, 1.0f);
-
-		if (bShouldClampValue)
-		{
-			PixelValue = FMath::Clamp(PixelValue, 0.0f, 1.0f);
-		}
-	}
-
-	// Convert back to a linear color
-	FLinearColor LinearColor = HSVColor.HSVToLinearRGB();
-
-	// Remap the alpha channel
-	LinearColor.A = FMath::Lerp(AdjustMinAlpha, AdjustMaxAlpha, OriginalColor.A);
-	return LinearColor;
+FLinearColor UCurveLinearColor::GetClampedLinearColorValue(float InTime) const
+{
+	return GetAdjustedColorValue(InTime, true);
 }
 
-
-FLinearColor UCurveLinearColor::GetClampedLinearColorValue(float InTime) const
+FLinearColor UCurveLinearColor::GetAdjustedColorValue(float InTime, bool bClampOver1Always) const
 {
-	// Logic copied from .\Engine\Source\Developer\TextureCompressor\Private\TextureCompressorModule.cpp
+	// Logic copied from .\Engine\Source\Developer\TextureCompressor\Private\TextureCompressorModule.cpp , from AdjustColorsOld
+	//  @todo : see fixes in AdjustColorsNew
 	FLinearColor OriginalColor = GetUnadjustedLinearColorValue(InTime);
 	
 	// allow > 1.0 at this point but not < 0.0
 	//	we're just about to convert this RGB to HSV, that doesn't support negatives
 	OriginalColor = ClampNonNegative(OriginalColor);
 
-	// Always clamp value
-	const bool bShouldClampValue = true;
+	// if bClampOver1Always , always clamp to stay in[ 0,1]
+	// else, Only clamp value if the color is input RGB <= 1 (fuzzy)
+	const bool bShouldClampValue = bClampOver1Always || 
+		( FMath::Max3(OriginalColor.R,OriginalColor.G,OriginalColor.B) < 1.0f+UE_KINDA_SMALL_NUMBER );
 
 	// Convert to HSV
 	FLinearColor HSVColor = OriginalColor.LinearRGBToHSV();
@@ -200,20 +140,11 @@ FLinearColor UCurveLinearColor::GetClampedLinearColorValue(float InTime) const
 	PixelHue += AdjustHue;
 
 	// Clamp HSV values
-	{
-		PixelHue = FMath::Fmod(PixelHue, 360.0f);
-		if (PixelHue < 0.0f)
-		{
-			// Keep the hue value positive as HSVToLinearRGB prefers that
-			PixelHue += 360.0f;
-		}
-		PixelSaturation = FMath::Clamp(PixelSaturation, 0.0f, 1.0f);
-
+	// HSVToLinearRGB now does the needed clamps, except for Value
 		if (bShouldClampValue)
 		{
 			PixelValue = FMath::Clamp(PixelValue, 0.0f, 1.0f);
 		}
-	}
 
 	// Convert back to a linear color
 	FLinearColor LinearColor = HSVColor.HSVToLinearRGB();
