Allow CVarOverrideTimeMaterialExpressions in Test builds.

#rnx

diff --git a/Engine/Source/Runtime/Engine/Private/SceneView.cpp b/Engine/Source/Runtime/Engine/Private/SceneView.cpp
--- a/Engine/Source/Runtime/Engine/Private/SceneView.cpp
+++ b/Engine/Source/Runtime/Engine/Private/SceneView.cpp
@@ -376,7 +376,7 @@ bool GetHairStrandsDepthOfFieldUseHairDepth()
 	return GHairStrandsComposeDOFDepth > 0 ? 1 : 0;
 }
 
-#if !(UE_BUILD_SHIPPING || UE_BUILD_TEST)
+#if !(UE_BUILD_SHIPPING)
 
 static TAutoConsoleVariable<float> CVarOverrideTimeMaterialExpressions(
 	TEXT("r.Test.OverrideTimeMaterialExpressions"), -1.0f,
@@ -2792,7 +2792,7 @@ void SetupCommonViewUniformBufferParameters(
 	ViewUniformShaderParameters.UnlitViewmodeMask = Inputs.EngineShowFlags && ((!Inputs.EngineShowFlags->Lighting && Inputs.EngineShowFlags->UnlitViewmode) || Inputs.EngineShowFlags->PathTracing) ? 1 : 0;
 	ViewUniformShaderParameters.OutOfBoundsMask = Inputs.EngineShowFlags && Inputs.EngineShowFlags->VisualizeOutOfBoundsPixels ? 1 : 0;
 
-#if !(UE_BUILD_SHIPPING || UE_BUILD_TEST)
+#if !(UE_BUILD_SHIPPING)
 	float OverrideTimeMaterialExpression = CVarOverrideTimeMaterialExpressions.GetValueOnRenderThread();
 	if (OverrideTimeMaterialExpression >= 0.0f)
 	{
