Expose the DisableOcclusionQueries ShowFlag in shipping builds so that it can be toggled for simple SceneCaptures that do not benefit from occlusion culling

#tests PIE - Disabled occlusion queries on a SceneCapture that captures every frame, Preflight - including ReplayRun showing no perf regression in the absence of a SceneCapture
#rnx

diff --git a/Engine/Source/Runtime/Engine/Public/ShowFlagsValues.inl b/Engine/Source/Runtime/Engine/Public/ShowFlagsValues.inl
--- a/Engine/Source/Runtime/Engine/Public/ShowFlagsValues.inl
+++ b/Engine/Source/Runtime/Engine/Public/ShowFlagsValues.inl
@@ -400,7 +400,7 @@ SHOWFLAG_FIXED_IN_SHIPPING(0, OcclusionMeshes, SFG_Visualize, NSLOCTEXT("UnrealE
 /** Visualize Occlusion Query bounding meshes */
 SHOWFLAG_FIXED_IN_SHIPPING(0, VisualizeInstanceOcclusionQueries, SFG_Visualize, NSLOCTEXT("UnrealEd", "VisualizeInstanceOcclusionQueriesSF", "Visualize Instance Culling Occlusion Queries"))
 /** Disable hardware occlusion queries, similar to setting r.AllowOcclusionQueries=0, but just for this scene. */
-SHOWFLAG_FIXED_IN_SHIPPING(0, DisableOcclusionQueries, SFG_Developer, NSLOCTEXT("UnrealEd", "DisableOcclusionQueries", "Disable Hardware Occlusion Queries"))
+SHOWFLAG_ALWAYS_ACCESSIBLE(DisableOcclusionQueries, SFG_Hidden, NSLOCTEXT("UnrealEd", "DisableOcclusionQueries", "Disable Hardware Occlusion Queries"))
 
 // RHI_RAYTRACING begin
 SHOWFLAG_ALWAYS_ACCESSIBLE(PathTracing, SFG_Hidden, NSLOCTEXT("UnrealEd", "PathTracing", "Path tracing"))
