UCurveLinearColorAtlas : clean up the way it uses UTexture API
no need to store SrcData member at all, UTexture already has a Source
no need to manually cache platformdata, UTexture does that for you better
use correct PreEdit/PostEdit pattern for source changes
remove unused member variables

diff --git a/Engine/Source/Editor/TextureEditor/Private/TextureEditorToolkit.cpp b/Engine/Source/Editor/TextureEditor/Private/TextureEditorToolkit.cpp
--- a/Engine/Source/Editor/TextureEditor/Private/TextureEditorToolkit.cpp
+++ b/Engine/Source/Editor/TextureEditor/Private/TextureEditorToolkit.cpp
@@ -2386,7 +2386,10 @@ void FTextureEditorToolkit::FillToolbar(FToolBarBuilder& ToolbarBuilder)
 	TSharedRef<SWidget> View3DControl = MakeView3DControlWidget();
 
 	UCurveLinearColorAtlas* Atlas = Cast<UCurveLinearColorAtlas>(GetTexture());
-	if (!Atlas)
+	bool bCanRecompress = CanRecompressTexture(GetTexture());
+	bool bIsAtlas = (Atlas != nullptr);
+	
+	if ( bCanRecompress && ! bIsAtlas )
 	{
 		// These actions don't make sense in read-only mode
 		if(GetOpenMethod() == EAssetOpenMethod::Edit)
@@ -2398,8 +2401,10 @@ void FTextureEditorToolkit::FillToolbar(FToolBarBuilder& ToolbarBuilder)
 			}
 			ToolbarBuilder.EndSection();
 		}
-		
+	}
 
+	if ( 1 )
+	{
 		ToolbarBuilder.BeginSection("Channels");
 		{
 			ToolbarBuilder.AddWidget(ChannelControl);
diff --git a/Engine/Source/Editor/UnrealEd/Private/Factories/EditorFactories.cpp b/Engine/Source/Editor/UnrealEd/Private/Factories/EditorFactories.cpp
--- a/Engine/Source/Editor/UnrealEd/Private/Factories/EditorFactories.cpp
+++ b/Engine/Source/Editor/UnrealEd/Private/Factories/EditorFactories.cpp
@@ -2302,23 +2302,13 @@ UObject* UCurveLinearColorAtlasFactory::FactoryCreateNew(UClass* Class, UObject*
 	}
 
 	UCurveLinearColorAtlas* Object = NewObject<UCurveLinearColorAtlas>(InParent, Class, Name, Flags);
-	Object->Source.Init(Width, Height, 1, 1, TSF_RGBA16F);
-	const int32 TextureDataSize = IntCastChecked<int32>(Object->Source.CalcMipSize(0));
-	Object->SrcData.AddUninitialized(TextureDataSize);
-	uint32* TextureData = (uint32*)Object->Source.LockMip(0);
-	FFloat16Color InitColor(FLinearColor::White);
-	for (uint32 y = 0; y < Object->TextureSize; y++)
-	{
-		// Create base mip for the texture we created.
-		for (uint32 x = 0; x < Object->TextureSize; x++)
-		{
-			Object->SrcData[x*Object->TextureSize + y] = InitColor;
-		}
-	}
-	FMemory::Memcpy(TextureData, Object->SrcData.GetData(), TextureDataSize);
-	Object->Source.UnlockMip(0);
 
-	Object->UpdateResource();
+	Object->TextureSize = Width;
+	Object->TextureHeight = Height;
+	Object->bSquareResolution = (Width == Height);
+
+	Object->UpdateTextures();
+
 	return Object;
 }
 
diff --git a/Engine/Source/Runtime/Engine/Classes/Curves/CurveLinearColor.h b/Engine/Source/Runtime/Engine/Classes/Curves/CurveLinearColor.h
--- a/Engine/Source/Runtime/Engine/Classes/Curves/CurveLinearColor.h
+++ b/Engine/Source/Runtime/Engine/Classes/Curves/CurveLinearColor.h
@@ -69,9 +69,9 @@ public:
 
 	ENGINE_API void DrawThumbnail(class FCanvas* Canvas, FVector2D StartXY, FVector2D SizeXY);
 
-	ENGINE_API void PushToSourceData(TArray<FFloat16Color> &SrcData, int32 StartXY, FVector2D SizeXY);
+	ENGINE_API void PushToSourceData(TArrayView<FFloat16Color> &SrcData, int32 Start, int32 Width);
 
-	ENGINE_API void PushUnadjustedToSourceData(TArray<FFloat16Color>& SrcData, int32 StartXY, FVector2D SizeXY);
+	ENGINE_API void PushUnadjustedToSourceData(TArrayView<FFloat16Color>& SrcData, int32 Start, int32 Width);
 
 	ENGINE_API virtual void OnCurveChanged(const TArray<FRichCurveEditInfo>& ChangedCurveEditInfos) override;
 #endif
diff --git a/Engine/Source/Runtime/Engine/Classes/Curves/CurveLinearColorAtlas.h b/Engine/Source/Runtime/Engine/Classes/Curves/CurveLinearColorAtlas.h
--- a/Engine/Source/Runtime/Engine/Classes/Curves/CurveLinearColorAtlas.h
+++ b/Engine/Source/Runtime/Engine/Classes/Curves/CurveLinearColorAtlas.h
@@ -98,15 +98,6 @@ class UCurveLinearColorAtlas : public UTexture2D
 	UFUNCTION(BlueprintCallable, Category = "Math|Curves")
 	ENGINE_API bool GetCurvePosition(UCurveLinearColor* InCurve, float& Position);
 
-#if WITH_EDITORONLY_DATA
-	UPROPERTY()
-	uint32 bIsDirty : 1;
-
-	uint32	bHasAnyDirtyTextures : 1;
-	uint32	bShowDebugColorsForNullGradients : 1;	// Renders alternate blue/yellow lines for empty gradients. Good for debugging, but turns off optimization for selective updates to gradients.
-
-	TArray<FFloat16Color> SrcData;
-#endif
 	UPROPERTY(EditAnywhere, Category = "Curves", DisplayName = "Texture Width")
 	uint32	TextureSize;						// Width of the lookup textures
 
@@ -139,7 +130,4 @@ protected:
 	ENGINE_API void RestoreCachedColorAdjustments();
 #endif
 
-#if WITH_EDITORONLY_DATA
-	FVector2D SizeXY;
-#endif
 };
diff --git a/Engine/Source/Runtime/Engine/Private/Curves/CurveLinearColor.cpp b/Engine/Source/Runtime/Engine/Private/Curves/CurveLinearColor.cpp
--- a/Engine/Source/Runtime/Engine/Private/Curves/CurveLinearColor.cpp
+++ b/Engine/Source/Runtime/Engine/Private/Curves/CurveLinearColor.cpp
@@ -230,10 +230,14 @@ void UCurveLinearColor::PostEditChangeProperty(struct FPropertyChangedEvent& Pro
 
 void UCurveLinearColor::DrawThumbnail(FCanvas* Canvas, FVector2D StartXY, FVector2D SizeXY)
 {
+	// this draws the curve into the thumbnail texture
+	//	typically start == {0,0} and size == {256,256}
+	// this updates continuously when your mouse hovers over the thumbnail
 	FCanvasTileItem DrawItem(StartXY, FVector2D(1.f, SizeXY.Y), FLinearColor::Black);
 
 	// Render the gradient
 	float Width = SizeXY.X;
+	check( Width > 1.f );
 	for (float x = 0.f; x < Width; ++x)
 	{
 		DrawItem.Position.X = StartXY.X + x;
@@ -246,29 +250,21 @@ void UCurveLinearColor::DrawThumbnail(FCanvas* Canvas, FVector2D StartXY, FVecto
 }
 
 
-void UCurveLinearColor::PushToSourceData(TArray<FFloat16Color> &SrcData, int32 StartXY, FVector2D SizeXY)
+void UCurveLinearColor::PushToSourceData(TArrayView<FFloat16Color> &SrcData, int32 Start, int32 Width)
 {
-	int32 Start = StartXY;
-	for (uint32 y = 0; y < SizeXY.Y; y++)
+	check( Width > 1 );
+	for (int32 x = 0; x < Width; x++)
 	{
-		// Create base mip for the texture we created.
-		for (uint32 x = 0; x < SizeXY.X; x++)
-		{
-			SrcData[Start + x + y*SizeXY.X] = GetLinearColorValue(x * (1.f / (SizeXY.X-1)) );
-		}
+		SrcData[Start + x] = GetLinearColorValue(x * (1.f / (Width-1)) );
 	}
 }
 
-void UCurveLinearColor::PushUnadjustedToSourceData(TArray<FFloat16Color>& SrcData, int32 StartXY, FVector2D SizeXY)
+void UCurveLinearColor::PushUnadjustedToSourceData(TArrayView<FFloat16Color>& SrcData, int32 Start, int32 Width)
 {
-	int32 Start = StartXY;
-	for (uint32 y = 0; y < SizeXY.Y; y++)
+	check( Width > 1 );
+	for (int32 x = 0; x < Width; x++)
 	{
-		// Create base mip for the texture we created.
-		for (uint32 x = 0; x < SizeXY.X; x++)
-		{
-			SrcData[Start + x + y * SizeXY.X] = GetUnadjustedLinearColorValue(x * (1.f / (SizeXY.X-1)));
-		}
+		SrcData[Start + x] = GetUnadjustedLinearColorValue(x * (1.f / (Width-1)) );
 	}
 }
 
@@ -277,6 +273,7 @@ void UCurveLinearColor::OnCurveChanged(const TArray<FRichCurveEditInfo>& Changed
 	Super::OnCurveChanged(ChangedCurveEditInfos);
 
 	PRAGMA_DISABLE_DEPRECATION_WARNINGS
+	// this redraws any UCurveLinearColorAtlas that holds me
 	OnUpdateGradient.Broadcast(this);
 	PRAGMA_ENABLE_DEPRECATION_WARNINGS
 }
diff --git a/Engine/Source/Runtime/Engine/Private/Curves/CurveLinearColorAtlas.cpp b/Engine/Source/Runtime/Engine/Private/Curves/CurveLinearColorAtlas.cpp
--- a/Engine/Source/Runtime/Engine/Private/Curves/CurveLinearColorAtlas.cpp
+++ b/Engine/Source/Runtime/Engine/Private/Curves/CurveLinearColorAtlas.cpp
@@ -18,9 +18,6 @@ UCurveLinearColorAtlas::UCurveLinearColorAtlas(const FObjectInitializer& ObjectI
 	bSquareResolution = true;
 	TextureHeight = 256;
 #if WITH_EDITORONLY_DATA
-	bHasAnyDirtyTextures = false;
-	bShowDebugColorsForNullGradients = false;
-	SizeXY = { (float)TextureSize,  1.0f };
 	MipGenSettings = TMGS_NoMipmaps;
 #endif
 	Filter = TextureFilter::TF_Bilinear;
@@ -75,6 +72,9 @@ void UCurveLinearColorAtlas::PostEditChangeProperty(struct FPropertyChangedEvent
 			PropertyName == GET_MEMBER_NAME_CHECKED(UCurveLinearColorAtlas, bSquareResolution) ||
 			PropertyName == GET_MEMBER_NAME_CHECKED(UCurveLinearColorAtlas, TextureHeight))
 		{
+			TextureSize = FMath::Max<uint32>(TextureSize,2);
+			TextureHeight = FMath::Max<uint32>(TextureHeight,1);
+
 			if (bSquareResolution)
 			{
 				TextureHeight = TextureSize;
@@ -86,9 +86,6 @@ void UCurveLinearColorAtlas::PostEditChangeProperty(struct FPropertyChangedEvent
 				GradientCurves.RemoveAt(TextureHeight, OldCurveCount - TextureHeight);
 			}
 
-			Source.Init(TextureSize, TextureHeight, 1, 1, TSF_RGBA16F);
-
-			SizeXY = { (float)TextureSize, 1.0f };
 			UpdateTextures();
 			bRequiresNotifyMaterials = true;
 		}
@@ -227,11 +224,6 @@ void UCurveLinearColorAtlas::RestoreCachedColorAdjustments()
 void UCurveLinearColorAtlas::PostLoad()
 {
 #if WITH_EDITOR
-	if (FApp::CanEverRender())
-	{
-		FinishCachePlatformData();
-	}
-
 	for (int32 i = 0; i < GradientCurves.Num(); ++i)
 	{
 		if (GradientCurves[i] != nullptr)
@@ -240,31 +232,31 @@ void UCurveLinearColorAtlas::PostLoad()
 		}
 	}
 	
-	if (bSquareResolution)
+	// re-draw into the Texture Source on load
+	//	when the code is stable this should be an unnecessary nop
+	//	but it lets us refresh the data when the code changes
+	if ( 1 )
 	{
-		TextureHeight = TextureSize;
+		UpdateTextures();
 	}
-	Source.Init(TextureSize, TextureHeight, 1, 1, TSF_RGBA16F);
-	SizeXY = { (float)TextureSize, 1.0f };
-	UpdateTextures();
-#endif
+#endif // #if WITH_EDITOR
+
+	// Super is UTexture2D
+	//	will do CachePlatformData for us
 
 	Super::PostLoad();
 }
 
 #if WITH_EDITOR
-static void RenderGradient(TArray<FFloat16Color>& InSrcData, UObject* Gradient, int32 StartXY, FVector2D SizeXY, bool bUseUnadjustedColor)
+static void RenderGradient(TArrayView<FFloat16Color>& InSrcData, UObject* Gradient, int32 Start, int32 Width, bool bUseUnadjustedColor)
 {
 	if (Gradient == nullptr)
 	{
-		int32 Start = StartXY;
-		for (uint32 y = 0; y < SizeXY.Y; y++)
+		FFloat16Color White16 = FLinearColor::White;
+
+		for (int32 x = 0; x < Width; x++)
 		{
-			// Create base mip for the texture we created.
-			for (uint32 x = 0; x < SizeXY.X; x++)
-			{
-				InSrcData[Start + x + y * SizeXY.X] = FLinearColor::White;
-			}
+			InSrcData[Start + x] = White16;
 		}
 	}
 	else if (Gradient->IsA(UCurveLinearColor::StaticClass()))
@@ -273,32 +265,15 @@ static void RenderGradient(TArray<FFloat16Color>& InSrcData, UObject* Gradient,
 		UCurveLinearColor* GradientCurve = CastChecked<UCurveLinearColor>(Gradient);
 		if (bUseUnadjustedColor)
 		{
-			GradientCurve->PushUnadjustedToSourceData(InSrcData, StartXY, SizeXY);
+			GradientCurve->PushUnadjustedToSourceData(InSrcData, Start, Width);
 		}
 		else
 		{
-			GradientCurve->PushToSourceData(InSrcData, StartXY, SizeXY);
+			GradientCurve->PushToSourceData(InSrcData, Start, Width);
 		}
 	}
 }
 
-static void UpdateTexture(UCurveLinearColorAtlas& Atlas)
-{
-	const int32 TextureDataSize = Atlas.Source.CalcMipSize(0);
-
-	FGuid MD5Guid;
-	FMD5 MD5;
-	MD5.Update(reinterpret_cast<const uint8*>(Atlas.SrcData.GetData()), TextureDataSize);
-	MD5.Final(reinterpret_cast<uint8*>(&MD5Guid));
-
-	uint32* TextureData = reinterpret_cast<uint32*>(Atlas.Source.LockMip(0));
-	FMemory::Memcpy(TextureData, Atlas.SrcData.GetData(), TextureDataSize);
-	Atlas.Source.UnlockMip(0);
-
-	Atlas.Source.SetId(MD5Guid, /*bInGuidIsHash*/ true);
-	Atlas.UpdateResource();
-}
-
 // Immediately render a new material to the specified slot index (SlotIndex must be within this section's range)
 void UCurveLinearColorAtlas::OnCurveUpdated(UCurveBase* Curve, EPropertyChangeType::Type ChangeType)
 {
@@ -310,47 +285,60 @@ void UCurveLinearColorAtlas::OnCurveUpdated(UCurveBase* Curve, EPropertyChangeTy
 
 		if (SlotIndex != INDEX_NONE && (uint32)SlotIndex < MaxSlotsPerTexture())
 		{
-			// Determine the position of the gradient
-			int32 StartXY = SlotIndex * TextureSize;
-
-			// Render the single gradient to the render target
-			RenderGradient(SrcData, Gradient, StartXY, SizeXY, bDisableAllAdjustments);
+			// @@ for efficiency, could update just the one gradient instead of all?
 
-			UpdateTexture(*this);
+			UpdateTextures();
 		}
-	}	
+	}
 }
 
 // Render any textures
 void UCurveLinearColorAtlas::UpdateTextures()
 {
 	LLM_SCOPE(ELLMTag::Textures);
+	
+	if (bSquareResolution)
+	{
+		TextureHeight = TextureSize;
+	}
+	
+	check( TextureSize > 1 );
+	check( TextureHeight >= 1 );
+
+	PreEditChange(nullptr);
+
+	if (Source.GetSizeX() != TextureSize ||
+		Source.GetSizeY() != TextureHeight ||
+		Source.GetFormat() != TSF_RGBA16F )
+	{
+		Source.Init(TextureSize, TextureHeight, 1, 1, TSF_RGBA16F);
+	}
+
+	const int32 TextureNumPixels = TextureSize * TextureHeight;
+	FFloat16Color * SrcData = (FFloat16Color *) Source.LockMip(0);
 
-	// Save off the data needed to render each gradient.
-	// Callback into the section owner to get the Gradients array
-	const int32 TextureDataSize = Source.CalcMipSize(0);
-	SrcData.Empty();
-	SrcData.AddUninitialized(TextureDataSize);
+	TArrayView<FFloat16Color> SrcDataArrayView(SrcData,TextureNumPixels);
+	check( Source.CalcMipSize(0) == SrcDataArrayView.NumBytes() );
 
-	int32 NumSlotsToRender = FMath::Min(GradientCurves.Num(), (int32)MaxSlotsPerTexture());
-	for (int32 i = 0; i < NumSlotsToRender; ++i)
+	uint32 NumSlotsToRender = FMath::Min((uint32)GradientCurves.Num(), (uint32)TextureHeight);
+	for (uint32 i = 0; i < NumSlotsToRender; ++i)
 	{
-		int32 StartXY = i * TextureSize;
-		RenderGradient(SrcData, GradientCurves[i], StartXY, SizeXY, bDisableAllAdjustments);
+		RenderGradient(SrcDataArrayView, GradientCurves[i], i*TextureSize, TextureSize, bDisableAllAdjustments);
 	}
 
-	for (uint32 y = GradientCurves.Num(); y < TextureHeight; y++)
+	FFloat16Color White16 = FLinearColor::White;
+
+	for (uint32 y = NumSlotsToRender; y < TextureHeight; y++)
 	{
 		// Create base mip for the texture we created.
 		for (uint32 x = 0; x < TextureSize; x++)
 		{
-			SrcData[y*TextureSize + x] = FLinearColor::White;
+			SrcData[y*TextureSize + x] = White16;
 		}
 	}
 
-	UpdateTexture(*this);
-
-	bIsDirty = false;
+	Source.UnlockMip(0);
+	PostEditChange();
 }
 
 #endif
