Reduce size of FPrimaryAssetLoadState by packing its internal state into a buffer that is only allocated when the state is set.

diff --git a/Engine/Source/Runtime/Engine/Private/AssetManager.cpp b/Engine/Source/Runtime/Engine/Private/AssetManager.cpp
--- a/Engine/Source/Runtime/Engine/Private/AssetManager.cpp
+++ b/Engine/Source/Runtime/Engine/Private/AssetManager.cpp
@@ -57,31 +57,300 @@ LLM_DEFINE_TAG(AssetManager);
 
 DEFINE_LOG_CATEGORY(LogAssetManager);
 
-// Temp hack to claw back memory. Proper fix to come later
-// TODO: FIXME: JMarcus
-#define UE_ASSETMANAGER_DOWNLOADPARMS_IN_LOADSTATE 1
+#define UE_ASSETMANAGER_SMALL_LOADSTATE 1
+
+#if UE_ASSETMANAGER_SMALL_LOADSTATE
+
+/** Structure defining the current loading state of an asset */
+struct FPrimaryAssetLoadState
+{
+private:
+	// Note: this implementation reduces memory by moving as much of the state as possible into a single buffer.
+	// If the state is empty, this buffer is not allocated so we don't pay the memory for it.
+	// Layout is [ (TSharedPtr<FStreamableHandle>) | (NumBundleNames * sizeof(FName)) | (optional FAssetManagerDownloadParams)]
+
+	/** The storage buffer */
+	uint8* Storage = nullptr;
+	/** Number of bundles in the storage buffer */
+	uint16 NumBundleNames = 0;
+	/** If true, the buffer contains download params */
+	bool bWithDownloadParams = 0;
+
+	constexpr static SIZE_T GetBundleNamesOffset()
+	{ 
+		return Align(sizeof(TSharedPtr<FStreamableHandle>), alignof(FName));
+	}
+
+	static SIZE_T GetDownloadParamsOffset(uint16 InNumBundles)
+	{
+		return Align(GetBundleNamesOffset() + sizeof(FName) * InNumBundles, alignof(FAssetManagerDownloadParams));
+	}
+
+	void ReallocStorage(uint16 InNumBundles, bool bInWithDownloadParams)
+	{
+		FMemory::Free(Storage);
+
+		const SIZE_T StorageSize = bInWithDownloadParams
+			? (GetDownloadParamsOffset(InNumBundles) + sizeof(FAssetManagerDownloadParams))
+			: (GetBundleNamesOffset() + sizeof(FName) * InNumBundles);
+
+		Storage = reinterpret_cast<uint8*>(FMemory::Malloc(StorageSize));
+	}
+
+	bool NeedsRealloc(uint16 InNumBundles, bool bInWithDownloadParams) const
+	{
+		return !Storage || InNumBundles != NumBundleNames || bInWithDownloadParams != bWithDownloadParams;
+	}
+
+	void DestroyStoredItems()
+	{
+		check(Storage);
+		DestructItem(reinterpret_cast<TSharedPtr<FStreamableHandle>*>(Storage));
+		DestructItems(GetBundleNamesInternal(), NumBundleNames);
+		if (FAssetManagerDownloadParams* DownloadParams = GetDownloadParamsInternal())
+		{
+			DestructItem(DownloadParams);
+		}
+	}
+
+	/** The handle to the streamable state for this asset, this keeps the objects in memory. If handle is invalid, not in memory at all */
+	TSharedPtr<FStreamableHandle>& GetHandleInternal()
+	{
+		check(Storage);
+		return *reinterpret_cast<TSharedPtr<FStreamableHandle>*>(Storage);
+	}
+
+	const TSharedPtr<FStreamableHandle>& GetHandleInternal() const
+	{
+		return const_cast<FPrimaryAssetLoadState*>(this)->GetHandleInternal();
+	}
+
+	/** The set of bundles to be loaded by the handle */
+	FName* GetBundleNamesInternal()
+	{
+		check(Storage);
+		return reinterpret_cast<FName*>(Storage + GetBundleNamesOffset());
+	}
+
+	const FName* GetBundleNamesInternal() const
+	{
+		return const_cast<FPrimaryAssetLoadState*>(this)->GetBundleNamesInternal();
+	}
+
+	/** Params for installing assets before loading */
+	FAssetManagerDownloadParams* GetDownloadParamsInternal()
+	{
+		if (bWithDownloadParams)
+		{
+			check(Storage);
+			return reinterpret_cast<FAssetManagerDownloadParams*>(Storage + GetDownloadParamsOffset(NumBundleNames));
+		}
+
+		return nullptr;
+	}
+
+	const FAssetManagerDownloadParams* GetDownloadParamsInternal() const
+	{
+		return const_cast<FPrimaryAssetLoadState*>(this)->GetDownloadParamsInternal();
+	}
+
+	inline static const TSharedPtr<FStreamableHandle> InvalidHandle = {};
+
+public:
+	FPrimaryAssetLoadState() = default;
+	FPrimaryAssetLoadState(const FPrimaryAssetLoadState&) = delete;
+
+	FPrimaryAssetLoadState(FPrimaryAssetLoadState&& Other) 
+		: Storage(Other.Storage)
+		, NumBundleNames(Other.NumBundleNames)
+		, bWithDownloadParams(Other.bWithDownloadParams)
+	{
+		Other.Storage = nullptr;
+		Other.NumBundleNames = 0;
+		Other.bWithDownloadParams = false;
+	}
+
+	FPrimaryAssetLoadState& operator=(const FPrimaryAssetLoadState&) = delete;
+
+	FPrimaryAssetLoadState& operator=(FPrimaryAssetLoadState&& Other)
+	{
+		if (Storage)
+		{
+			DestroyStoredItems();
+			FMemory::Free(Storage);
+		}
+
+		Storage = Other.Storage;
+		NumBundleNames = Other.NumBundleNames;
+		bWithDownloadParams = Other.bWithDownloadParams;
+
+		Other.Storage = nullptr;
+		Other.NumBundleNames = 0;
+		Other.bWithDownloadParams = false;
+
+		return *this;
+	}
+
+	~FPrimaryAssetLoadState()
+	{
+		if (Storage)
+		{
+			DestroyStoredItems();
+			FMemory::Free(Storage);
+			Storage = nullptr;
+		}
+	}
+
+	bool IsSame(TConstArrayView<FName> InBundleNames, const TOptional<FAssetManagerDownloadParams>& InDownloadParams) const
+	{
+		if (NumBundleNames == InBundleNames.Num() && bWithDownloadParams == InDownloadParams.IsSet())
+		{
+			// bWithDownloadParams can only be true if storage is allocated
+			if (bWithDownloadParams && *GetDownloadParamsInternal() != InDownloadParams)
+			{
+				return false;
+			}
+
+			return Algo::Compare(InBundleNames, GetBundleNames());
+		}
+
+		return false;
+	}
+
+	void SetState(TSharedPtr<FStreamableHandle> InHandle, TConstArrayView<FName> InBundleNames, TOptional<FAssetManagerDownloadParams> InDownloadParams)
+	{
+		const uint16 InNumBundleNames = IntCastChecked<uint16>(InBundleNames.Num());
+
+		if (NeedsRealloc(InNumBundleNames, InDownloadParams.IsSet()))
+		{
+			if (Storage)
+			{
+				DestroyStoredItems();
+			}
+			ReallocStorage(InNumBundleNames, InDownloadParams.IsSet());
+			NumBundleNames = InNumBundleNames;
+			bWithDownloadParams = InDownloadParams.IsSet();
+
+			new(Storage) TSharedPtr<FStreamableHandle>(MoveTemp(InHandle));
+
+			ConstructItems<FName>(GetBundleNamesInternal(), InBundleNames.GetData(), InNumBundleNames);
+
+			if (InDownloadParams)
+			{
+				new(GetDownloadParamsInternal()) FAssetManagerDownloadParams(MoveTemp(*InDownloadParams));
+			}
+		}
+		else
+		{
+			GetHandleInternal() = MoveTemp(InHandle);
+
+			CopyAssignItems(GetBundleNamesInternal(), InBundleNames.GetData(), InNumBundleNames);
+
+			if (InDownloadParams)
+			{
+				*GetDownloadParamsInternal() = MoveTemp(*InDownloadParams);
+			}
+		}
+	}
+
+	void MoveStateFrom(TSharedPtr<FStreamableHandle> InHandle, FPrimaryAssetLoadState&& InOther)
+	{
+		*this = MoveTemp(InOther);
+		if (Storage)
+		{
+			GetHandleInternal() = MoveTemp(InHandle);
+		}
+	}
+
+	const TSharedPtr<FStreamableHandle>& GetHandle() const UE_LIFETIMEBOUND
+	{
+		if (Storage)
+		{
+			return GetHandleInternal();
+		}
+		return InvalidHandle;
+	}
+
+	TConstArrayView<FName> GetBundleNames() const UE_LIFETIMEBOUND
+	{ 
+		if (Storage)
+		{
+			return MakeArrayView(GetBundleNamesInternal(), NumBundleNames);
+		}
+		return {};
+	}
+
+	/** If this state is keeping things in memory */
+	bool IsValid() const 
+	{
+		const TSharedPtr<FStreamableHandle>& Handle = GetHandle();
+		return Handle.IsValid() && Handle->IsActive();
+	}
+
+	/** Reset this state */
+	void Reset(bool bCancelHandle)
+	{
+		const TSharedPtr<FStreamableHandle>& Handle = GetHandle();
+		if (Handle.IsValid())
+		{
+			if (Handle->IsActive() && bCancelHandle)
+			{
+				// This will call the cancel callback if set
+				Handle->CancelHandle();
+			}
+		}
+		
+		if (Storage)
+		{
+			DestroyStoredItems();
+			FMemory::Free(Storage);
+			Storage = nullptr;
+		}
+
+		NumBundleNames = 0;
+		bWithDownloadParams = false;
+	}
+};
+
+#else
 
 /** Structure defining the current loading state of an asset */
 struct FPrimaryAssetLoadState
 {
+private:
 	/** The handle to the streamable state for this asset, this keeps the objects in memory. If handle is invalid, not in memory at all */
 	TSharedPtr<FStreamableHandle> Handle;
 
 	/** The set of bundles to be loaded by the handle */
 	TArray<FName> BundleNames;
-#if UE_ASSETMANAGER_DOWNLOADPARMS_IN_LOADSTATE
+
+	/** Params for installing assets before loading */
 	TOptional<FAssetManagerDownloadParams> DownloadParams;
 
+public:
 	bool IsSame(TConstArrayView<FName> InBundleNames, const TOptional<FAssetManagerDownloadParams>& InDownloadParams) const
 	{
 		return Algo::Compare(InBundleNames, BundleNames) && (InDownloadParams == DownloadParams);
 	}
-#else
-	bool IsSame(TConstArrayView<FName> InBundleNames, const TOptional<FAssetManagerDownloadParams>& InDownloadParams) const
+
+	void SetState(TSharedPtr<FStreamableHandle> InHandle, TConstArrayView<FName> InBundleNames, TOptional<FAssetManagerDownloadParams> InDownloadParams)
 	{
-		return Algo::Compare(InBundleNames, BundleNames);
+		Handle = MoveTemp(InHandle);
+		BundleNames = InBundleNames;
+		DownloadParams = MoveTemp(InDownloadParams);
 	}
-#endif // UE_ASSETMANAGER_DOWNLOADPARMS_IN_LOADSTATE
+
+	void MoveStateFrom(TSharedPtr<FStreamableHandle> InHandle, FPrimaryAssetLoadState&& InOther)
+	{
+		Handle = MoveTemp(InHandle);
+		BundleNames = MoveTemp(InOther.BundleNames);
+		DownloadParams = MoveTemp(InOther.DownloadParams);
+	}
+
+	const TSharedPtr<FStreamableHandle>& GetHandle() const { return Handle; }
+
+	const TArray<FName>& GetBundleNames() const { return BundleNames; }
+
 
 	/** If this state is keeping things in memory */
 	bool IsValid() const { return Handle.IsValid() && Handle->IsActive(); }
@@ -100,12 +369,12 @@ struct FPrimaryAssetLoadState
 			Handle = nullptr;
 		}
 		BundleNames.Reset();
-#if UE_ASSETMANAGER_DOWNLOADPARMS_IN_LOADSTATE
 		DownloadParams.Reset();
-#endif // UE_ASSETMANAGER_DOWNLOADPARMS_IN_LOADSTATE
 	}
 };
 
+#endif // UE_ASSETMANAGER_SMALL_LOADSTATE
+
 struct FPrimaryAssetTypeData;
 
 /** Structure representing data about a specific asset */
@@ -1788,7 +2057,7 @@ bool UAssetManager::GetPrimaryAssetIdList(FPrimaryAssetType PrimaryAssetType, TA
 
 		for (const TPair<FName, FPrimaryAssetData>& Pair : TypeData.GetAssets())
 		{
-			if ((!(Filter & EAssetManagerFilter::UnloadedOnly)) || ((Pair.Value.CurrentState.BundleNames.Num() == 0) && (Pair.Value.PendingState.BundleNames.Num() == 0)))
+			if ((!(Filter & EAssetManagerFilter::UnloadedOnly)) || ((Pair.Value.CurrentState.GetBundleNames().Num() == 0) && (Pair.Value.PendingState.GetBundleNames().Num() == 0)))
 			{
 				PrimaryAssetIdList.Add(FPrimaryAssetId(PrimaryAssetType, Pair.Key));
 			}
@@ -1856,7 +2125,7 @@ TSharedPtr<FStreamableHandle> UAssetManager::ChangeBundleStateForPrimaryAssets(
 
 			if (!bRemoveAllBundles)
 			{
-				NewBundleState = NameData->PendingState.IsValid() ? NameData->PendingState.BundleNames : NameData->CurrentState.BundleNames;
+				NewBundleState = NameData->PendingState.IsValid() ? NameData->PendingState.GetBundleNames() : NameData->CurrentState.GetBundleNames();
 
 				for (const FName& RemoveBundle : RemoveBundles)
 				{
@@ -1877,12 +2146,13 @@ TSharedPtr<FStreamableHandle> UAssetManager::ChangeBundleStateForPrimaryAssets(
 				if (NameData->PendingState.IsSame(NewBundleState, LoadParams.DownloadParams))
 				{
 					// This will wait on any existing handles to finish
-					AllHandles.Add(NameData->PendingState.Handle);
+					AllHandles.Add(NameData->PendingState.GetHandle());
 					continue;
 				}
 
 				// Clear pending state
-				NameData->PendingState.Reset(true);
+				constexpr const bool bCancelHandle = true;
+				NameData->PendingState.Reset(bCancelHandle);
 			}
 			else if (NameData->CurrentState.IsValid() && NameData->CurrentState.IsSame(NewBundleState, LoadParams.DownloadParams))
 			{
@@ -1921,16 +2191,10 @@ TSharedPtr<FStreamableHandle> UAssetManager::ChangeBundleStateForPrimaryAssets(
 			if (PathsToLoad.Num() == 0)
 			{
 				// New state has no assets to load. Set the CurrentState's bundles and clear the handle
-				NameData->CurrentState.BundleNames = NewBundleState;
-#if UE_ASSETMANAGER_DOWNLOADPARMS_IN_LOADSTATE
-				NameData->CurrentState.DownloadParams = CopyTemp(*LoadParams.DownloadParams);
-#endif // UE_ASSETMANAGER_DOWNLOADPARMS_IN_LOADSTATE
-				NameData->CurrentState.Handle.Reset();
+				NameData->CurrentState.SetState({}, NewBundleState, LoadParams.DownloadParams);
 				continue;
 			}
 
-			TSharedPtr<FStreamableHandle> NewHandle;
-
 			TStringBuilder<1024> DebugName;
 			PrimaryAssetId.AppendString(DebugName);
 
@@ -1950,28 +2214,27 @@ TSharedPtr<FStreamableHandle> UAssetManager::ChangeBundleStateForPrimaryAssets(
 				DebugName << TEXT(")");
 			}
 
-			// Set pending bundle state before issuing the request, because it can be resolved immediately.
-			// OnAssetStateChangeCompleted requires this data to handle state transition correctly.
-			NameData->PendingState.BundleNames = NewBundleState;
-#if UE_ASSETMANAGER_DOWNLOADPARMS_IN_LOADSTATE
-			NameData->PendingState.DownloadParams = CopyTemp(LoadParams.DownloadParams);
-#endif // UE_ASSETMANAGER_DOWNLOADPARMS_IN_LOADSTATE
-
 			FAssetManagerLoadParams AssetListLoadParams;
 			AssetListLoadParams.Priority = LoadParams.Priority;
 			AssetListLoadParams.DownloadParams = CopyTemp(LoadParams.DownloadParams);
 			
-			NewHandle = LoadAssetList(MoveTemp(PathsToLoad), MoveTemp(AssetListLoadParams), FString(DebugName), CopyTemp(Location));
+			TSharedPtr<FStreamableHandle> NewHandle = LoadAssetList(MoveTemp(PathsToLoad), MoveTemp(AssetListLoadParams), FString(DebugName), CopyTemp(Location));
 			if (NewHandle && !NewHandle->HasLoadCompleted())
 			{
 				// save handle so possible duplicate requests will not override this one
 				// set only if loading is not completed yet because completion callback might already reset the PendingState
-				NameData->PendingState.Handle = NewHandle;
+
+				// OnAssetStateChangeCompleted requires this data to handle state transition correctly.
+				NameData->PendingState.SetState(NewHandle, NewBundleState, LoadParams.DownloadParams);
+
 				NewHandle->BindCompleteDelegate(FStreamableDelegateWithHandle::CreateUObject(this, &UAssetManager::OnAssetStateChangeCompleted, PrimaryAssetId));
 				NewHandle->BindCancelDelegate(FStreamableDelegateWithHandle::CreateUObject(this, &UAssetManager::OnAssetStateChangeCompleted, PrimaryAssetId));
 			}
 			else
 			{
+				// OnAssetStateChangeCompleted requires this data to handle state transition correctly.
+				NameData->PendingState.SetState({}, NewBundleState, LoadParams.DownloadParams);
+
 				OnAssetStateChangeCompleted(NewHandle, PrimaryAssetId);
 			}
 
@@ -2125,7 +2388,7 @@ void UAssetManager::OnAssetStateChangeCompleted(TSharedPtr<FStreamableHandle> Bo
 			return;
 		}
 
-		if (!NameData->PendingState.Handle.IsValid() || NameData->PendingState.Handle == BoundHandle)
+		if (!NameData->PendingState.GetHandle().IsValid() || NameData->PendingState.GetHandle() == BoundHandle)
 		{
 			// If an error occured, don't update the current state with a bad state
 			if (BoundHandle->HasError())
@@ -2148,18 +2411,14 @@ void UAssetManager::OnAssetStateChangeCompleted(TSharedPtr<FStreamableHandle> Bo
 			}
 			else
 			{
-				NameData->CurrentState.Handle = BoundHandle;
-#if UE_ASSETMANAGER_DOWNLOADPARMS_IN_LOADSTATE
-				NameData->CurrentState.DownloadParams = NameData->PendingState.DownloadParams;
-#endif // UE_ASSETMANAGER_DOWNLOADPARMS_IN_LOADSTATE
-				NameData->CurrentState.BundleNames = NameData->PendingState.BundleNames;
+				NameData->CurrentState.MoveStateFrom(BoundHandle, MoveTemp(NameData->PendingState));
 				
 				NotifyOnAssetStateChangeCompleted(PrimaryAssetId);
 			}
 
 			NameData->PendingState.Reset(false);
 		}
-		else if (NameData->PendingState.Handle.IsValid())
+		else if (NameData->PendingState.GetHandle().IsValid())
 		{
 			UE_LOG(LogAssetManager, Verbose, TEXT("OnAssetStateChangeCompleted: Received after pending data changed, ignoring (%s)"), *PrimaryAssetId.ToString());
 		}
@@ -2200,9 +2459,9 @@ TSharedPtr<FStreamableHandle> UAssetManager::GetPrimaryAssetHandle(const FPrimar
 
 	if (Bundles)
 	{
-		*Bundles = LoadState.BundleNames;
+		*Bundles = LoadState.GetBundleNames();
 	}
-	return LoadState.Handle;
+	return LoadState.GetHandle();
 }
 
 bool UAssetManager::GetPrimaryAssetsWithBundleState(TArray<FPrimaryAssetId>& PrimaryAssetList, const TArray<FPrimaryAssetType>& ValidTypes, const TArray<FName>& RequiredBundles, const TArray<FName>& ExcludedBundles, bool bForceCurrent) const
@@ -2236,7 +2495,7 @@ bool UAssetManager::GetPrimaryAssetsWithBundleState(TArray<FPrimaryAssetId>& Pri
 			// Check bundle requirements
 			for (const FName& RequiredName : RequiredBundles)
 			{
-				if (!LoadState.BundleNames.Contains(RequiredName))
+				if (!LoadState.GetBundleNames().Contains(RequiredName))
 				{
 					bFailedTest = true;
 					break;
@@ -2245,7 +2504,7 @@ bool UAssetManager::GetPrimaryAssetsWithBundleState(TArray<FPrimaryAssetId>& Pri
 
 			for (const FName& ExcludedName : ExcludedBundles)
 			{
-				if (LoadState.BundleNames.Contains(ExcludedName))
+				if (LoadState.GetBundleNames().Contains(ExcludedName))
 				{
 					bFailedTest = true;
 					break;
@@ -2284,7 +2543,7 @@ void UAssetManager::GetPrimaryAssetBundleStateMap(TMap<FPrimaryAssetId, TArray<F
 
 			FPrimaryAssetId AssetID(TypePair.Key, NamePair.Key);
 
-			BundleStateMap.Add(AssetID, LoadState.BundleNames);
+			BundleStateMap.Add(AssetID, TArray<FName>(LoadState.GetBundleNames()));
 		}
 	}
 }
@@ -3296,7 +3555,7 @@ void UAssetManager::DumpLoadedAssetState()
 
 				FString BundleString;
 
-				for (const FName& BundleName : LoadState.BundleNames)
+				for (const FName& BundleName : LoadState.GetBundleNames())
 				{
 					if (!BundleString.IsEmpty())
 					{
