Fixes instability and flickering in the Cascaded Shadow Maps (CSM) for the UDirectionalLightComponent.

The issue was caused by using the view's projection matrix with Temporal AA jitter when computing the CSM split depth bounds. Since the jittered projection matrix changes every frame, it leads to instability in the resulting shadow cascade bounding spheres, causing visible flickering in the shadow cascades.

Possibly introduced in CL 38648400

#tests cinder watch+ fn br

diff --git a/Engine/Source/Runtime/Engine/Private/Components/DirectionalLightComponent.cpp b/Engine/Source/Runtime/Engine/Private/Components/DirectionalLightComponent.cpp
--- a/Engine/Source/Runtime/Engine/Private/Components/DirectionalLightComponent.cpp
+++ b/Engine/Source/Runtime/Engine/Private/Components/DirectionalLightComponent.cpp
@@ -800,7 +800,7 @@ private:
 		const FViewMatrices& ViewMatrices = View.GetShadowViewMatrices();
 
 		const FMatrix& ViewMatrix = ViewMatrices.GetViewMatrix();
-		const FMatrix& ProjectionMatrix = ViewMatrices.GetProjectionMatrix();
+		const FMatrix& ProjectionMatrix = ViewMatrices.GetProjectionNoAAMatrix();
 
 		const FVector CameraDirection = ViewMatrix.GetColumn(2);
 		const FVector LightDirection = -GetDirection();
