Make FStreamableDownloadPriority an int16.  This should still provide plenty of granularity but will reduce the size of FAssetManagerDownloadParams

diff --git a/Engine/Source/Runtime/Engine/Classes/Engine/StreamableManager.h b/Engine/Source/Runtime/Engine/Classes/Engine/StreamableManager.h
--- a/Engine/Source/Runtime/Engine/Classes/Engine/StreamableManager.h
+++ b/Engine/Source/Runtime/Engine/Classes/Engine/StreamableManager.h
@@ -34,7 +34,7 @@ using FStreamableDelegateWithHandle = TDelegate<void(TSharedPtr<struct FStreamab
 using FStreamableUpdateDelegate = TDelegate<void(TSharedRef<struct FStreamableHandle>)>;
 
 /** EXPERIMENTAL - download priority */
-using FStreamableDownloadPriority = int32;
+using FStreamableDownloadPriority = int16;
 
 namespace UE::StreamableManager::Private
 {
