CSV stats for tracking DSTM actor migrations.

diff --git a/Engine/Source/Runtime/Engine/Private/Actor.cpp b/Engine/Source/Runtime/Engine/Private/Actor.cpp
--- a/Engine/Source/Runtime/Engine/Private/Actor.cpp
+++ b/Engine/Source/Runtime/Engine/Private/Actor.cpp
@@ -92,6 +92,8 @@ DEFINE_LOG_CATEGORY(LogActor);
 
 DECLARE_CYCLE_STAT(TEXT("PostActorConstruction"), STAT_PostActorConstruction, STATGROUP_Engine);
 
+CSV_DEFINE_CATEGORY(ActorMigration, UE_WITH_REMOTE_OBJECT_HANDLE);
+
 #if UE_BUILD_SHIPPING
 #define DEBUG_CALLSPACE(Format, ...)
 #else
@@ -1257,6 +1259,14 @@ static USceneComponent* FixupNativeActorComponents(AActor* Actor);
 
 void AActor::PostMigrate(const FUObjectMigrationContext& MigrationContext)
 {
+	CSV_CUSTOM_STAT(ActorMigration, Total, 1, ECsvCustomStatOp::Accumulate);
+	CSV_CUSTOM_STAT(ActorMigration, Loan, MigrationContext.MigrationSide == EObjectMigrationSide::Send && MigrationContext.GetObjectMigrationSendType(this) == EObjectMigrationSendType::Loan, ECsvCustomStatOp::Accumulate);
+	CSV_CUSTOM_STAT(ActorMigration, ReturnBorrowed, MigrationContext.MigrationSide == EObjectMigrationSide::Send && MigrationContext.GetObjectMigrationSendType(this) == EObjectMigrationSendType::ReturnBorrowed, ECsvCustomStatOp::Accumulate);
+	CSV_CUSTOM_STAT(ActorMigration, ReassignOwnership, MigrationContext.MigrationSide == EObjectMigrationSide::Send && MigrationContext.GetObjectMigrationSendType(this) == EObjectMigrationSendType::ReassignOwnership, ECsvCustomStatOp::Accumulate);
+	CSV_CUSTOM_STAT(ActorMigration, Borrowed, MigrationContext.MigrationSide == EObjectMigrationSide::Receive && MigrationContext.GetObjectMigrationRecvType(this) == EObjectMigrationRecvType::Borrowed, ECsvCustomStatOp::Accumulate);
+	CSV_CUSTOM_STAT(ActorMigration, ReturnedLoan, MigrationContext.MigrationSide == EObjectMigrationSide::Receive && MigrationContext.GetObjectMigrationRecvType(this) == EObjectMigrationRecvType::ReturnedLoan, ECsvCustomStatOp::Accumulate);
+	CSV_CUSTOM_STAT(ActorMigration, AssignedOwnership, MigrationContext.MigrationSide == EObjectMigrationSide::Receive && MigrationContext.GetObjectMigrationRecvType(this) == EObjectMigrationRecvType::AssignedOwnership, ECsvCustomStatOp::Accumulate);
+
 	// We want to be able to store "annotations" (metadata) about this Actor while we loan it out, so that
 	// when we have the loan returned to us, we can check what has changed...
 	struct FActorLoanedState
