[NewMaterialTranslator] - Implement UMaterialExpressionStaticComponentMaskParameter Build.

- The Build function implements a compile time mask/swizzle.
- In FEmitter::Parameter implement the case statement to handle EMaterialParameterType::StaticComponentMask.  This isn't technically a uniform parameter, but can be configured via a MI so the value must be come from the StaticParameterSet->EditorOnly.StaticComponentMaskParameters.
- The FValue we emit from the Parameter function is a constant bool4 which matches how the data is stored on the FStaticComponentMaskParameter.
- Then the Build function takes this FValue bool4 and creates a swizzle mask from it.
- FEmitter::Swizzle should emit poison if we don't supply any components to mask.  This gives the user an error if their ComponentMask contains no components checked.
- Add ConstantBool2, ConstantBool3, and ConstantBool4 to the FEmitter API.

#rnx

diff --git a/Engine/Source/Runtime/Engine/Private/Materials/MaterialExpressionsIR.cpp b/Engine/Source/Runtime/Engine/Private/Materials/MaterialExpressionsIR.cpp
--- a/Engine/Source/Runtime/Engine/Private/Materials/MaterialExpressionsIR.cpp
+++ b/Engine/Source/Runtime/Engine/Private/Materials/MaterialExpressionsIR.cpp
@@ -397,6 +397,40 @@ void UMaterialExpressionStaticBoolParameter::Build(MIR::FEmitter& Em)
 	Em.Output(0, Value);
 }
 
+void UMaterialExpressionStaticComponentMaskParameter::Build(MIR::FEmitter& Em)
+{
+	FValueRef Value = BuildMaterialExpressionParameter(Em, this);
+	const MIR::FDimensional* Dimensional = Value->As<MIR::FDimensional>();
+
+	// BuildMaterialExpressionParameter should return this as a constant 4 dimensional bool.
+	check(Dimensional);
+	check(Dimensional->AreComponentsConstant());
+
+	MIR::FSwizzleMask Mask;
+
+	TConstArrayView<MIR::FValue*> Components = Dimensional->GetComponents();
+	check(Components.Num() == 4);
+
+	if (Components[0]->IsTrue())
+	{
+		Mask.Append(MIR::EVectorComponent::X);
+	}
+	if (Components[1]->IsTrue())
+	{
+		Mask.Append(MIR::EVectorComponent::Y);
+	}
+	if (Components[2]->IsTrue())
+	{
+		Mask.Append(MIR::EVectorComponent::Z);
+	}
+	if (Components[3]->IsTrue())
+	{
+		Mask.Append(MIR::EVectorComponent::W);
+	}
+	
+	Em.Output(0, Em.Swizzle(Em.Input(&Input), Mask));
+}
+
 void UMaterialExpressionStaticSwitch::Build(MIR::FEmitter& Em)
 {
 	bool bCondition = Em.ConstantBool(Em.InputDefaultBool(&Value, DefaultValue));
@@ -569,19 +603,19 @@ void UMaterialExpressionComponentMask::Build(MIR::FEmitter& Em)
 	MIR::FSwizzleMask Mask;
 	if (R)
 	{
-		Mask.Components[Mask.NumComponents++] = MIR::EVectorComponent::X;
+		Mask.Append(MIR::EVectorComponent::X);
 	}
 	if (G)
 	{
-		Mask.Components[Mask.NumComponents++] = MIR::EVectorComponent::Y;
+		Mask.Append(MIR::EVectorComponent::Y);
 	}
 	if (B)
 	{
-		Mask.Components[Mask.NumComponents++] = MIR::EVectorComponent::Z;
+		Mask.Append(MIR::EVectorComponent::Z);
 	}
 	if (A)
 	{
-		Mask.Components[Mask.NumComponents++] = MIR::EVectorComponent::W;
+		Mask.Append(MIR::EVectorComponent::W);
 	}
 
 	Em.Output(0, Em.Swizzle(Value, Mask));
diff --git a/Engine/Source/Runtime/Engine/Private/Materials/MaterialIREmitter.cpp b/Engine/Source/Runtime/Engine/Private/Materials/MaterialIREmitter.cpp
--- a/Engine/Source/Runtime/Engine/Private/Materials/MaterialIREmitter.cpp
+++ b/Engine/Source/Runtime/Engine/Private/Materials/MaterialIREmitter.cpp
@@ -514,6 +514,30 @@ FValueRef FEmitter::ConstantBool(bool InX)
 	return InX ? ConstantTrue() : ConstantFalse();
 }
 
+FValueRef FEmitter::ConstantBool2(bool InX, bool InY)
+{
+	FValueRef X = ConstantBool(InX);
+	FValueRef Y = ConstantBool(InY);
+	return Vector2(X, Y);
+}
+
+FValueRef FEmitter::ConstantBool3(bool InX, bool InY, bool InZ)
+{
+	FValueRef X = ConstantBool(InX);
+	FValueRef Y = ConstantBool(InY);
+	FValueRef Z = ConstantBool(InZ);
+	return Vector3(X, Y, Z);
+}
+
+FValueRef FEmitter::ConstantBool4(bool InX, bool InY, bool InZ, bool InW)
+{
+	FValueRef X = ConstantBool(InX);
+	FValueRef Y = ConstantBool(InY);
+	FValueRef Z = ConstantBool(InZ);
+	FValueRef W = ConstantBool(InW);
+	return Vector4(X, Y, Z, W);
+}
+
 FValueRef FEmitter::ConstantInt(TInteger InX)
 {
 	FConstant Scalar = MakePrototype<FConstant>(FPrimitiveType::GetScalar(ScalarKind_Int));
@@ -691,6 +715,24 @@ FValueRef FEmitter::Parameter(FName Name, FMaterialParameterMetadata& Metadata,
 			return ConstantBool(Metadata.Value.Bool[0]);
 		}
 
+		case EMaterialParameterType::StaticComponentMask:
+		{
+			// Apply eventual parameter override
+			for (const FStaticComponentMaskParameter& Param : StaticParameterSet->EditorOnly.StaticComponentMaskParameters)
+			{
+				if (Param.IsOverride() && Param.ParameterInfo.Name == Name)
+				{
+					Metadata.Value.Bool[0] = Param.R;
+					Metadata.Value.Bool[1] = Param.G;
+					Metadata.Value.Bool[2] = Param.B;
+					Metadata.Value.Bool[3] = Param.A;
+					break;
+				}
+			}
+
+			return ConstantBool4(Metadata.Value.Bool[0], Metadata.Value.Bool[1], Metadata.Value.Bool[2], Metadata.Value.Bool[3]);
+		}
+
 		default:
 			UE_MIR_TODO();
 	}
@@ -1932,7 +1974,11 @@ FValueRef FEmitter::Swizzle(FValueRef Value, FSwizzleMask Mask)
 	}
 
 	// At least one component must have been specified.
-	check(Mask.NumComponents > 0);
+	if (Mask.NumComponents <= 0)
+	{
+		Errorf(Value, TEXT("Swizzle mask has no components."));
+		return Value.ToPoison();
+	}
 
 	// We can only swizzle on non-matrix primitive types.
 	const FPrimitiveType* PrimitiveType = Value->Type->AsVector();
diff --git a/Engine/Source/Runtime/Engine/Public/Materials/MaterialExpressionStaticComponentMaskParameter.h b/Engine/Source/Runtime/Engine/Public/Materials/MaterialExpressionStaticComponentMaskParameter.h
--- a/Engine/Source/Runtime/Engine/Public/Materials/MaterialExpressionStaticComponentMaskParameter.h
+++ b/Engine/Source/Runtime/Engine/Public/Materials/MaterialExpressionStaticComponentMaskParameter.h
@@ -36,6 +36,7 @@ public:
 
 	//~ Begin UMaterialExpression Interface
 #if WITH_EDITOR
+	virtual void Build(MIR::FEmitter& Emitter) override;
 	virtual int32 Compile(class FMaterialCompiler* Compiler, int32 OutputIndex) override;
 	virtual void GetCaption(TArray<FString>& OutCaptions) const override;
 	virtual bool GetParameterValue(FMaterialParameterMetadata& OutMeta) const override
diff --git a/Engine/Source/Runtime/Engine/Public/Materials/MaterialIREmitter.h b/Engine/Source/Runtime/Engine/Public/Materials/MaterialIREmitter.h
--- a/Engine/Source/Runtime/Engine/Public/Materials/MaterialIREmitter.h
+++ b/Engine/Source/Runtime/Engine/Public/Materials/MaterialIREmitter.h
@@ -264,6 +264,15 @@ public:
 	// Returns the given constant boolean scalar.
 	FValueRef ConstantBool(bool InX);
 
+	// Emits a constant bool 2D column vector and returns it.
+	FValueRef ConstantBool2(bool InX, bool InY);
+
+	// Emits a constant bool 3D column vector and returns it.
+	FValueRef ConstantBool3(bool InX, bool InY, bool InZ);
+
+	// Emits a constant bool 4D column vector and returns it.
+	FValueRef ConstantBool4(bool InX, bool InY, bool InZ, bool InW);
+
 	// Returns the given constant integer scalar.
 	FValueRef ConstantInt(TInteger InX);
 	
