Do not set the owner name on material resources during the serialization as it can race with game thread.

- Material resources can be shared between materials or material instances, and there is a possible race with updating the owner, so remove the update from the Serialize call. PostLoad executes a path that will also set the owner name.

diff --git a/Engine/Source/Runtime/Engine/Private/Materials/MaterialShared.cpp b/Engine/Source/Runtime/Engine/Private/Materials/MaterialShared.cpp
--- a/Engine/Source/Runtime/Engine/Private/Materials/MaterialShared.cpp
+++ b/Engine/Source/Runtime/Engine/Private/Materials/MaterialShared.cpp
@@ -1460,7 +1460,6 @@ void FMaterial::SerializeInlineShaderMap(FArchive& Ar, const FName& SerializingA
 				if (LoadedShaderMap->Serialize(Ctx))
 				{
 					GameThreadShaderMap = MoveTemp(LoadedShaderMap);
-					GameThreadShaderMap->GetResource()->SetOwnerName(GetOwnerFName());
 #if WITH_EDITOR
 					GameThreadShaderMap->AssociateWithAsset(GetAssetPath());
 #endif
