UCurveLinearColorAtlas : make TextureHeight automatically from the number of curves
change width of new atlas textures to 64 instead of 256
fix bug with notification callback being added over and over
remove square texture option

diff --git a/Engine/Source/Editor/UnrealEd/Classes/Factories/CurveLinearColorAtlasFactory.h b/Engine/Source/Editor/UnrealEd/Classes/Factories/CurveLinearColorAtlasFactory.h
--- a/Engine/Source/Editor/UnrealEd/Classes/Factories/CurveLinearColorAtlasFactory.h
+++ b/Engine/Source/Editor/UnrealEd/Classes/Factories/CurveLinearColorAtlasFactory.h
@@ -16,18 +16,10 @@ class UCurveLinearColorAtlasFactory : public UFactory
 {
 	GENERATED_UCLASS_BODY()
 
-	/** Width of the texture render target */
-	UPROPERTY(meta=(ToolTip="Width of the texture render target"))
+	/** Width of the texture atlas */
+	UPROPERTY(meta=(ToolTip="Width of the texture atlas"))
 	int32 Width;
 
-	/** Height of the texture render target */
-	UPROPERTY(meta=(ToolTip="Height of the texture render target"))
-	int32 Height;
-
-	/** Pixel format of the texture render target */
-	UPROPERTY(meta=(ToolTip="Pixel format of the texture render target"))
-	uint8 Format;
-
 	virtual FText GetDisplayName() const override;
 
 	uint32 GetMenuCategories() const override;
diff --git a/Engine/Source/Editor/UnrealEd/Private/Factories/EditorFactories.cpp b/Engine/Source/Editor/UnrealEd/Private/Factories/EditorFactories.cpp
--- a/Engine/Source/Editor/UnrealEd/Private/Factories/EditorFactories.cpp
+++ b/Engine/Source/Editor/UnrealEd/Private/Factories/EditorFactories.cpp
@@ -2278,9 +2278,7 @@ UCurveLinearColorAtlasFactory::UCurveLinearColorAtlasFactory(const FObjectInitia
 	bEditAfterNew = true;
 	bEditorImport = false;
 
-	Width = 256;
-	Height = 256;
-	Format = 0;
+	Width = 64;
 }
 
 FText UCurveLinearColorAtlasFactory::GetDisplayName() const
@@ -2295,17 +2293,9 @@ uint32 UCurveLinearColorAtlasFactory::GetMenuCategories() const
 
 UObject* UCurveLinearColorAtlasFactory::FactoryCreateNew(UClass* Class, UObject* InParent, FName Name, EObjectFlags Flags, UObject* Context, FFeedbackContext* Warn)
 {
-	// Do not create a texture with bad dimensions.
-	if ((Width & (Width - 1)) || (Height & (Height - 1)))
-	{
-		return nullptr;
-	}
-
 	UCurveLinearColorAtlas* Object = NewObject<UCurveLinearColorAtlas>(InParent, Class, Name, Flags);
-
-	Object->TextureSize = Width;
-	Object->TextureHeight = Height;
-	Object->bSquareResolution = (Width == Height);
+	
+	Object->TextureSize = FMath::Max(Width,2);
 
 	Object->UpdateTextures();
 
diff --git a/Engine/Source/Runtime/Engine/Classes/Curves/CurveLinearColorAtlas.h b/Engine/Source/Runtime/Engine/Classes/Curves/CurveLinearColorAtlas.h
--- a/Engine/Source/Runtime/Engine/Classes/Curves/CurveLinearColorAtlas.h
+++ b/Engine/Source/Runtime/Engine/Classes/Curves/CurveLinearColorAtlas.h
@@ -78,12 +78,6 @@ class UCurveLinearColorAtlas : public UTexture2D
 
 	ENGINE_API virtual void PostEditChangeProperty(struct FPropertyChangedEvent& PropertyChangedEvent) override;
 
-	// How many slots are available per texture
-	FORCEINLINE uint32 MaxSlotsPerTexture()
-	{
-		return TextureHeight;
-	}
-
 	// Immediately render a new material to the specified slot index(SlotIndex must be within this section's range)
 	ENGINE_API void OnCurveUpdated(UCurveBase* Curve, EPropertyChangeType::Type ChangeType);
 
@@ -101,13 +95,6 @@ class UCurveLinearColorAtlas : public UTexture2D
 	UPROPERTY(EditAnywhere, Category = "Curves", DisplayName = "Texture Width")
 	uint32	TextureSize;						// Width of the lookup textures
 
-	UPROPERTY(EditAnywhere, Category = "Curves")
-	/** Set texture height equal to texture width. */
-	uint32 bSquareResolution : 1;
-
-	UPROPERTY(EditAnywhere, Category = "Curves", meta = (EditCondition = "!bSquareResolution"))
-	uint32	TextureHeight;						// Height of the lookup textures
-
 	UPROPERTY(EditAnywhere, Category = "Curves")
 	TArray<TObjectPtr<UCurveLinearColor>> GradientCurves;
 
diff --git a/Engine/Source/Runtime/Engine/Private/Curves/CurveLinearColorAtlas.cpp b/Engine/Source/Runtime/Engine/Private/Curves/CurveLinearColorAtlas.cpp
--- a/Engine/Source/Runtime/Engine/Private/Curves/CurveLinearColorAtlas.cpp
+++ b/Engine/Source/Runtime/Engine/Private/Curves/CurveLinearColorAtlas.cpp
@@ -14,9 +14,11 @@
 UCurveLinearColorAtlas::UCurveLinearColorAtlas(const FObjectInitializer& ObjectInitializer)
 	: Super(ObjectInitializer)
 {
+	// a better default value for TextureSize is 32 or 64
+	//	but that's not simple to change because of the way UProps serialize deltas
+	//	instead we change the value in UCurveLinearColorAtlasFactory
 	TextureSize = 256;
-	bSquareResolution = true;
-	TextureHeight = 256;
+
 #if WITH_EDITORONLY_DATA
 	MipGenSettings = TMGS_NoMipmaps;
 #endif
@@ -68,47 +70,31 @@ void UCurveLinearColorAtlas::PostEditChangeProperty(struct FPropertyChangedEvent
 	{
 		const FName PropertyName(PropertyChangedEvent.Property->GetFName());
 		// if Resizing
-		if (PropertyName == GET_MEMBER_NAME_CHECKED(UCurveLinearColorAtlas, TextureSize) || 
-			PropertyName == GET_MEMBER_NAME_CHECKED(UCurveLinearColorAtlas, bSquareResolution) ||
-			PropertyName == GET_MEMBER_NAME_CHECKED(UCurveLinearColorAtlas, TextureHeight))
+		if ( PropertyName == GET_MEMBER_NAME_CHECKED(UCurveLinearColorAtlas, TextureSize) )
 		{
 			TextureSize = FMath::Max<uint32>(TextureSize,2);
-			TextureHeight = FMath::Max<uint32>(TextureHeight,1);
-
-			if (bSquareResolution)
-			{
-				TextureHeight = TextureSize;
-			}
-
-			if ((uint32)GradientCurves.Num() > TextureHeight)
-			{
-				int32 OldCurveCount = GradientCurves.Num();
-				GradientCurves.RemoveAt(TextureHeight, OldCurveCount - TextureHeight);
-			}
 
 			UpdateTextures();
 			bRequiresNotifyMaterials = true;
 		}
 		else if (PropertyName == GET_MEMBER_NAME_CHECKED(UCurveLinearColorAtlas, GradientCurves))
 		{
-			if ((uint32)GradientCurves.Num() > TextureHeight)
-			{
-				int32 OldCurveCount = GradientCurves.Num();
-				GradientCurves.RemoveAt(TextureHeight, OldCurveCount - TextureHeight);
-			}
-			else
+			// note: any old curves that were removed do not get their OnUpdateCurve delegate removed! (it still will notify me)
+			//	but when we get a notification from them, we will then remove their OnUpdateCurve delegate
+
+			for (int32 i = 0; i < GradientCurves.Num(); ++i)
 			{
-				for (int32 i = 0; i < GradientCurves.Num(); ++i)
+				if (GradientCurves[i] != nullptr)
 				{
-					if (GradientCurves[i] != nullptr)
-					{
-						GradientCurves[i]->OnUpdateCurve.AddUObject(this, &UCurveLinearColorAtlas::OnCurveUpdated);
-					}
+					// AddDelegateInstance will just keep adding more copies of the same delegate to the callback list
+					//	first remove all notifications to me, then add :
+					GradientCurves[i]->OnUpdateCurve.RemoveAll(this);
+					GradientCurves[i]->OnUpdateCurve.AddUObject(this, &UCurveLinearColorAtlas::OnCurveUpdated);
 				}
-				UpdateTextures();
-				bRequiresNotifyMaterials = true;
 			}
-		}	
+			UpdateTextures();
+			bRequiresNotifyMaterials = true;
+		}
 		else if (PropertyName == GET_MEMBER_NAME_CHECKED(UCurveLinearColorAtlas, bDisableAllAdjustments))
 		{
 			if (bDisableAllAdjustments)
@@ -250,19 +236,10 @@ void UCurveLinearColorAtlas::PostLoad()
 #if WITH_EDITOR
 static void RenderGradient(TArrayView<FFloat16Color>& InSrcData, UObject* Gradient, int32 Start, int32 Width, bool bUseUnadjustedColor)
 {
-	if (Gradient == nullptr)
-	{
-		FFloat16Color White16 = FLinearColor::White;
-
-		for (int32 x = 0; x < Width; x++)
-		{
-			InSrcData[Start + x] = White16;
-		}
-	}
-	else if (Gradient->IsA(UCurveLinearColor::StaticClass()))
+	UCurveLinearColor* GradientCurve = Cast<UCurveLinearColor>(Gradient);
+	if ( GradientCurve )
 	{
 		// Render a gradient
-		UCurveLinearColor* GradientCurve = CastChecked<UCurveLinearColor>(Gradient);
 		if (bUseUnadjustedColor)
 		{
 			GradientCurve->PushUnadjustedToSourceData(InSrcData, Start, Width);
@@ -272,25 +249,52 @@ static void RenderGradient(TArrayView<FFloat16Color>& InSrcData, UObject* Gradie
 			GradientCurve->PushToSourceData(InSrcData, Start, Width);
 		}
 	}
+	else
+	{
+		FFloat16Color White16 = FLinearColor::White;
+
+		for (int32 x = 0; x < Width; x++)
+		{
+			InSrcData[Start + x] = White16;
+		}
+	}
 }
 
 // Immediately render a new material to the specified slot index (SlotIndex must be within this section's range)
 void UCurveLinearColorAtlas::OnCurveUpdated(UCurveBase* Curve, EPropertyChangeType::Type ChangeType)
 {
-	if (ChangeType != EPropertyChangeType::Interactive)
+	// @@ todo : this ("Interactive" branch) is broken (for dragging the handles at the top, it works for dragging points in the curve graph)
+	//	in theory we should see "Interactive" ChangeType during curve drags
+	//	but in fact UCurveBase::OnCurveChanged just always passes "ValueSet" even during mouse drags
+	//	so we never see Interactive, and are always re-rendering the Atlas
+	if (ChangeType == EPropertyChangeType::Interactive)
 	{
+		return;
+	}
+
+	/*
+	UE_LOG(LogCore,Display, TEXT("UCurveLinearColorAtlas::OnCurveUpdated : [%s] from [%s]"), 
+					*GetName(),*Curve->GetName()
+					);	
+	*/
+
 		UCurveLinearColor* Gradient = CastChecked<UCurveLinearColor>(Curve);
 
 		int32 SlotIndex = GradientCurves.Find(Gradient);
 
-		if (SlotIndex != INDEX_NONE && (uint32)SlotIndex < MaxSlotsPerTexture())
+	if ( SlotIndex == INDEX_NONE )
 		{
-			// @@ for efficiency, could update just the one gradient instead of all?
+		// this Curve is no longer in my list
+		//	eg. it was removed from my array
+		//	do not notify me any more:
+		Gradient->OnUpdateCurve.RemoveAll(this);
+		return;
+	}
+
+	// @todo : for efficiency, could update just the one gradient instead of all?
 
 			UpdateTextures();
 		}
-	}
-}
 
 // Render any textures
 void UCurveLinearColorAtlas::UpdateTextures()
@@ -298,15 +302,11 @@ void UCurveLinearColorAtlas::UpdateTextures()
 	LLM_SCOPE(ELLMTag::Textures);
 	
 	TextureSize = FMath::Max<uint32>(TextureSize,2);
-	TextureHeight = FMath::Max<uint32>(TextureHeight,1);
-
-	if (bSquareResolution)
-	{
-		TextureHeight = TextureSize;
-	}
-	
 	check( TextureSize > 1 );
-	check( TextureHeight >= 1 );
+	
+	// TextureHeight set automatically from NumCurves
+	uint32 NumCurves = GradientCurves.Num();
+	uint32 TextureHeight = FMath::Max<uint32>(NumCurves,1);
 
 	PreEditChange(nullptr);
 
@@ -323,17 +323,16 @@ void UCurveLinearColorAtlas::UpdateTextures()
 	TArrayView<FFloat16Color> SrcDataArrayView(SrcData,TextureNumPixels);
 	check( Source.CalcMipSize(0) == SrcDataArrayView.NumBytes() );
 
-	uint32 NumSlotsToRender = FMath::Min((uint32)GradientCurves.Num(), (uint32)TextureHeight);
-	for (uint32 i = 0; i < NumSlotsToRender; ++i)
+	for (uint32 i = 0; i < NumCurves; ++i)
 	{
 		RenderGradient(SrcDataArrayView, GradientCurves[i], i*TextureSize, TextureSize, bDisableAllAdjustments);
 	}
 
 	FFloat16Color White16 = FLinearColor::White;
 
-	for (uint32 y = NumSlotsToRender; y < TextureHeight; y++)
+	for (uint32 y = NumCurves; y < TextureHeight; y++)
 	{
-		// Create base mip for the texture we created.
+		// this is only used when NumCurves == 1 , then TextureHeight == 1
 		for (uint32 x = 0; x < TextureSize; x++)
 		{
 			SrcData[y*TextureSize + x] = White16;
