Iris
* Renable archetype support in subobjects

#rnx

diff --git a/Engine/Source/Runtime/Engine/Private/Net/Iris/ReplicationSystem/NetSubObjectFactory.cpp b/Engine/Source/Runtime/Engine/Private/Net/Iris/ReplicationSystem/NetSubObjectFactory.cpp
--- a/Engine/Source/Runtime/Engine/Private/Net/Iris/ReplicationSystem/NetSubObjectFactory.cpp
+++ b/Engine/Source/Runtime/Engine/Private/Net/Iris/ReplicationSystem/NetSubObjectFactory.cpp
@@ -60,9 +60,23 @@ TUniquePtr<UE::Net::FNetObjectCreationHeader> UNetSubObjectFactory::CreateAndFil
 	check(SubObject->NeedsLoadForClient() ); // We have no business sending this unless the client can load
 	check(SubObject->GetClass()->NeedsLoadForClient());	// We have no business sending this unless the client can load
 
-	// Only use the class as the template while investigating why some archetypes are missing from the client
-	//const UObject* Template = bUseArchetypeInSubObjectCreation ? SubObject->GetArchetype() : SubObject->GetClass()->GetDefaultObject();
-	SubObjectHeader->ArchetypeReference = Bridge->GetOrCreateObjectReference(SubObject->GetClass());
+	// Set the archetype to spawn the object from
+	{
+		UObject* Archetype = nullptr;
+		if (bUseArchetypeInSubObjectCreation)
+		{
+			Archetype = SubObject->GetArchetype();
+		}
+		else
+		{
+			Archetype = SubObject->GetClass();
+		}
+
+		SubObjectHeader->ArchetypeReference = Bridge->GetOrCreateObjectReference(Archetype);
+
+		ensureMsgf(SubObjectHeader->ArchetypeReference.GetRefHandle().IsStatic(), TEXT("Archetype %s for subobject %s (%s) should be a stable name otherwise the client can't use it"), 
+			*GetNameSafe(Archetype), *Bridge->PrintObjectFromNetRefHandle(Handle), *GetNameSafe(SubObject->GetClass()->GetAuthoritativeClass()));
+	}
 
 	// Find the right Outer
 	UObject* OuterObject = SubObject->GetOuter();	
@@ -185,10 +199,22 @@ UNetSubObjectFactory::FInstantiateResult UNetSubObjectFactory::InstantiateReplic
 		return FInstantiateResult();
 	}
 
-	// Instantiate the SubObject from it's archetype
-	// Archetype is always the class temporarely
-	UClass* SubObjClass = Cast<UClass>(Archetype);
-	UObject* SubObj = NewObject<UObject>(OuterObject, SubObjClass);
+	// Instantiate the SubObject
+	 UObject* SubObj = nullptr;
+	 UClass* SubObjClass = nullptr;
+
+	// Only if bUseArchetypeInSubObjectCreation is false on the server. Remove along with the cvar later
+	if (Archetype->IsA<UClass>())
+	{
+		SubObjClass = CastChecked<UClass>(Archetype);
+		SubObj = NewObject<UObject>(OuterObject, SubObjClass);
+	}
+	else
+	{
+		SubObjClass = Archetype->GetClass();
+		SubObj = NewObject<UObject>(OuterObject, SubObjClass, NAME_None, RF_NoFlags, Archetype);
+	}
+	
 
 	// Sanity check some things
 	checkf(SubObj != nullptr, TEXT("UNetSubObjectFactory::InstantiateReplicatedObjectFromHeader: Subobject is NULL after instantiating. Class: %s, Outer %s, Actor %s"), *GetNameSafe(SubObjClass), *GetNameSafe(OuterObject), *GetNameSafe(RootObject));
