Fix for a crash when performing DSTM push migrations with remote dependencies.

diff --git a/Engine/Source/Runtime/Engine/Classes/GameFramework/Pawn.h b/Engine/Source/Runtime/Engine/Classes/GameFramework/Pawn.h
--- a/Engine/Source/Runtime/Engine/Classes/GameFramework/Pawn.h
+++ b/Engine/Source/Runtime/Engine/Classes/GameFramework/Pawn.h
@@ -175,6 +175,8 @@ public:
 	/** If Pawn is possessed by a player, returns its Player State.  Needed for network play as controllers are not replicated to clients. */
 	APlayerState* GetPlayerState() const { return PlayerState; }
 
+	TObjectPtr<APlayerState> GetPlayerStateObjPtr() const { return PlayerState; }
+
 	/** Templated convenience version of GetPlayerState. */
 	template<class T>
 	T* GetPlayerState() const { return Cast<T>(PlayerState); }
