Move SoundConcurrency MaxCount to per-platform int. Minor clean-up moving functional bits to implementation files

#rnx

diff --git a/Engine/Source/Runtime/Engine/Classes/Sound/SoundConcurrency.h b/Engine/Source/Runtime/Engine/Classes/Sound/SoundConcurrency.h
--- a/Engine/Source/Runtime/Engine/Classes/Sound/SoundConcurrency.h
+++ b/Engine/Source/Runtime/Engine/Classes/Sound/SoundConcurrency.h
@@ -5,6 +5,8 @@
 #include "CoreMinimal.h"
 #include "UObject/ObjectMacros.h"
 #include "UObject/Object.h"
+#include "UObject/PerPlatformProperties.h"
+
 #include "SoundConcurrency.generated.h"
 
 class FAudioDevice;
@@ -72,29 +74,29 @@ enum class EConcurrencyVolumeScaleMode
 USTRUCT(BlueprintType)
 struct FSoundConcurrencySettings
 {
-	GENERATED_USTRUCT_BODY()
+	GENERATED_BODY()
 
 	/** The max number of allowable concurrent active voices for voices playing in this concurrency group. */
 	UPROPERTY(EditAnywhere, BlueprintReadWrite, Category = Concurrency, meta = (UIMin = "1", ClampMin = "1"))
-	int32 MaxCount;
+	FPerPlatformInt MaxCount;
 
 	/* Whether or not to limit the concurrency to per sound owner (i.e. the actor that plays the sound). If the sound doesn't have an owner, it falls back to global concurrency. */
 	UPROPERTY(EditAnywhere, BlueprintReadWrite, Category = Concurrency)
-	uint8 bLimitToOwner:1;
+	uint8 bLimitToOwner : 1 = 0;
 
 	/**
 	 * Whether or not volume scaling can recover volume ducking behavior when concurrency group sounds stop (default scale mode only).
 	 */
 	UPROPERTY(EditAnywhere, BlueprintReadWrite, Category = "Volume Scaling", meta = (DisplayName = "Can Recover", EditCondition = "VolumeScaleMode == EConcurrencyVolumeScaleMode::Default"))
-	uint8 bVolumeScaleCanRelease:1;
+	uint8 bVolumeScaleCanRelease : 1 = 0;
 
 	/** Which concurrency resolution policy to use if max voice count is reached. */
 	UPROPERTY(EditAnywhere, BlueprintReadWrite, Category = Concurrency)
-	TEnumAsByte<EMaxConcurrentResolutionRule::Type> ResolutionRule;
+	TEnumAsByte<EMaxConcurrentResolutionRule::Type> ResolutionRule = EMaxConcurrentResolutionRule::StopFarthestThenOldest;
 
 	/** Amount of time to wait (in seconds) between different sounds which play with this concurrency. Sounds rejected from this will ignore virtualization settings. */
 	UPROPERTY(EditAnywhere, BlueprintReadWrite, Category = Concurrency, meta = (UIMin = "0.0", ClampMin = "0.0"))
-	float RetriggerTime;
+	float RetriggerTime = 0.0f;
 
 private:
 	/**
@@ -106,44 +108,42 @@ private:
 	 * AppliedVolumeScale = Math.Pow(DuckingScale, VoiceGeneration)
 	 */
 	UPROPERTY(EditAnywhere, Category = "Volume Scaling", meta = (UIMin = "0.0", UIMax = "1.0", ClampMin = "0.0", ClampMax = "1.0"))
-	float VolumeScale;
+	float VolumeScale = 1.0f;
 
 public:
 	/** Volume Scale mode designating how to scale voice volume based on number of member sounds active in group. */
 	UPROPERTY(EditAnywhere, BlueprintReadWrite, Category = "Volume Scaling")
-	EConcurrencyVolumeScaleMode VolumeScaleMode;
+	EConcurrencyVolumeScaleMode VolumeScaleMode = EConcurrencyVolumeScaleMode::Default;
 
 	/**
 	 * Time taken to apply duck using volume scalar.
 	 */
 	UPROPERTY(EditAnywhere, BlueprintReadWrite, Category = "Volume Scaling", meta = (DisplayName = "Duck Time", UIMin = "0.0", ClampMin = "0.0", UIMax = "10.0", ClampMax = "1000000.0"))
-	float VolumeScaleAttackTime;
+	float VolumeScaleAttackTime = 0.01f;
 
 	/**
 	 * Time taken to recover volume scalar duck.
 	 */
 	UPROPERTY(EditAnywhere, BlueprintReadWrite, Category = "Volume Scaling", meta = (DisplayName = "Recover Time", EditCondition = "bVolumeScaleCanRelease && VolumeScaleMode == EConcurrencyVolumeScaleMode::Default", UIMin = "0.0", ClampMin = "0.0", UIMax = "10.0", ClampMax="1000000.0"))
-	float VolumeScaleReleaseTime;
+	float VolumeScaleReleaseTime = 0.5f;
 
 	/**
 	 * Time taken to fade out if voice is evicted or culled due to another voice in the group starting.
 	 */
 	UPROPERTY(EditAnywhere, BlueprintReadWrite, Category = "Voice Stealing", meta = (UIMin = "0.0", ClampMin = "0.0", UIMax = "1.0", ClampMax="1000000.0"))
-	float VoiceStealReleaseTime;
-
-
-	FSoundConcurrencySettings()
-		: MaxCount(16)
-		, bLimitToOwner(0)
-		, bVolumeScaleCanRelease(0)
-		, ResolutionRule(EMaxConcurrentResolutionRule::StopFarthestThenOldest)
-		, RetriggerTime(0.0f)
-		, VolumeScale(1.0f)
-		, VolumeScaleMode(EConcurrencyVolumeScaleMode::Default)
-		, VolumeScaleAttackTime(0.01f)
-		, VolumeScaleReleaseTime(0.5f)
-		, VoiceStealReleaseTime(0.0f)
-	{}
+	float VoiceStealReleaseTime = 0.0f;
+
+	/**
+	  * Returns the MaxCount as a per-platform integer value.
+	  */
+	ENGINE_API FPerPlatformInt GetMaxCount() const;
+
+	/**
+	  * Sets the MaxCount as a per-platform integer value.
+	  */
+	ENGINE_API void SetMaxCount(FPerPlatformInt InMaxCount);
+
+	FSoundConcurrencySettings();
 
 	/**
 	 * Retrieves the volume scale
@@ -244,25 +244,25 @@ public:
 	static FConcurrencyGroupID GenerateNewID();
 
 	/** Returns the active sounds array. */
-	const TArray<FActiveSound*>& GetActiveSounds() const { return ActiveSounds; }
+	const TArray<FActiveSound*>& GetActiveSounds() const;
 
 	/** Returns the id of the concurrency group */
-	FConcurrencyGroupID GetGroupID() const { return GroupID; }
+	FConcurrencyGroupID GetGroupID() const;
 
 	/** Returns the current generation (effectively, the number of concurrency sound instances active) */
-	const int32 GetNextGeneration() const { return ActiveSounds.Num(); }
+	const int32 GetNextGeneration() const;
 
 	/** Returns the settings associated with the group */
-	const FSoundConcurrencySettings& GetSettings() const { return Settings; }
+	const FSoundConcurrencySettings& GetSettings() const;
 
 	/** Returns the parent object ID */
-	FConcurrencyObjectID GetObjectID() const { return ObjectID; }
+	FConcurrencyObjectID GetObjectID() const;
 
 	/** Determines if the group is full. */
-	bool IsEmpty() const { return ActiveSounds.Num() == 0; }
+	bool IsEmpty() const;
 
 	/** Determines if the group is full. */
-	bool IsFull() const { return Settings.MaxCount <= ActiveSounds.Num(); }
+	bool IsFull() const;
 
 	/** Adds an active sound to the active sound array. */
 	void AddActiveSound(FActiveSound& ActiveSound);
diff --git a/Engine/Source/Runtime/Engine/Private/SoundBase.cpp b/Engine/Source/Runtime/Engine/Private/SoundBase.cpp
--- a/Engine/Source/Runtime/Engine/Private/SoundBase.cpp
+++ b/Engine/Source/Runtime/Engine/Private/SoundBase.cpp
@@ -244,7 +244,7 @@ void USoundBase::PostLoad()
 	{
 		bOverrideConcurrency = true;
 		ConcurrencyOverrides.bLimitToOwner = false;
-		ConcurrencyOverrides.MaxCount = FMath::Max(MaxConcurrentPlayCount_DEPRECATED, 1);
+		ConcurrencyOverrides.SetMaxCount(FPerPlatformInt(FMath::Max(MaxConcurrentPlayCount_DEPRECATED, 1)));
 		ConcurrencyOverrides.ResolutionRule = MaxConcurrentResolutionRule_DEPRECATED;
 	}
 
diff --git a/Engine/Source/Runtime/Engine/Private/SoundConcurrency.cpp b/Engine/Source/Runtime/Engine/Private/SoundConcurrency.cpp
--- a/Engine/Source/Runtime/Engine/Private/SoundConcurrency.cpp
+++ b/Engine/Source/Runtime/Engine/Private/SoundConcurrency.cpp
@@ -56,6 +56,27 @@ USoundConcurrency::USoundConcurrency(const FObjectInitializer& ObjectInitializer
 {
 }
 
+FSoundConcurrencySettings::FSoundConcurrencySettings()
+PRAGMA_DISABLE_DEPRECATION_WARNINGS
+	: MaxCount(16)
+PRAGMA_ENABLE_DEPRECATION_WARNINGS
+{
+}
+
+FPerPlatformInt FSoundConcurrencySettings::GetMaxCount() const
+{
+PRAGMA_DISABLE_DEPRECATION_WARNINGS
+	return MaxCount;
+PRAGMA_ENABLE_DEPRECATION_WARNINGS
+}
+
+void FSoundConcurrencySettings::SetMaxCount(FPerPlatformInt InMaxCount)
+{
+PRAGMA_DISABLE_DEPRECATION_WARNINGS
+	MaxCount = MoveTemp(InMaxCount);
+PRAGMA_ENABLE_DEPRECATION_WARNINGS
+}
+
 float FSoundConcurrencySettings::GetVolumeScale() const
 {
 	return FMath::Clamp(VolumeScale, 0.0f, 1.0f);;
@@ -178,6 +199,43 @@ FConcurrencyGroupID FConcurrencyGroup::GenerateNewID()
 	return ++ConcurrencyGroupIDs;
 }
 
+const TArray<FActiveSound*>& FConcurrencyGroup::GetActiveSounds() const
+{
+	return ActiveSounds;
+}
+
+FConcurrencyGroupID FConcurrencyGroup::GetGroupID() const
+{
+	return GroupID;
+}
+
+const int32 FConcurrencyGroup::GetNextGeneration() const
+{
+	return ActiveSounds.Num();
+}
+
+const FSoundConcurrencySettings& FConcurrencyGroup::GetSettings() const
+{
+	return Settings;
+}
+
+FConcurrencyObjectID FConcurrencyGroup::GetObjectID() const
+{
+	return ObjectID;
+}
+
+bool FConcurrencyGroup::IsEmpty() const
+{
+	return ActiveSounds.IsEmpty();
+}
+
+bool FConcurrencyGroup::IsFull() const
+{
+PRAGMA_DISABLE_DEPRECATION_WARNINGS
+	return Settings.MaxCount.GetValue() <= ActiveSounds.Num();
+PRAGMA_ENABLE_DEPRECATION_WARNINGS
+}
+
 void FConcurrencyGroup::AddActiveSound(FActiveSound& ActiveSound)
 {
 	check(GroupID != 0);
@@ -319,10 +377,12 @@ void FConcurrencyGroup::CullSoundsDueToMaxConcurrency()
 {
 	// Nothing to do if our active sound count is less than or equal to our max active sounds
 	// or if eviction is supported
-	if (ActiveSounds.Num() <= Settings.MaxCount)
+PRAGMA_DISABLE_DEPRECATION_WARNINGS
+	if (ActiveSounds.Num() <= Settings.MaxCount.GetValue())
 	{
 		return;
 	}
+PRAGMA_ENABLE_DEPRECATION_WARNINGS
 
 	// Comparator for sorting group's ActiveSounds according to their dynamic concurrency nature.
 	// Quieter/lower priority sounds will be at the front of the array.
@@ -385,7 +445,10 @@ void FConcurrencyGroup::CullSoundsDueToMaxConcurrency()
 
 	// Need to make a new list when stopping the sounds since the process of stopping an active sound
 	// will remove the sound from this concurrency group's ActiveSounds array.
-	const int32 NumSoundsToStop = ActiveSounds.Num() - Settings.MaxCount;
+PRAGMA_DISABLE_DEPRECATION_WARNINGS
+	const int32 NumSoundsToStop = ActiveSounds.Num() - Settings.MaxCount.GetValue();
+PRAGMA_ENABLE_DEPRECATION_WARNINGS
+
 	check(NumSoundsToStop > 0);
 	for (int32 i = 0; i < ActiveSounds.Num(); ++i)
 	{
@@ -524,7 +587,9 @@ FActiveSound* FSoundConcurrencyManager::CreateNewActiveSound(const FActiveSound&
 #if !UE_BUILD_SHIPPING
 	for (const FConcurrencyHandle& ConcurrencyHandle : ConcurrencyHandles)
 	{
-		check(ConcurrencyHandle.Settings.MaxCount > 0);
+PRAGMA_DISABLE_DEPRECATION_WARNINGS
+		check(ConcurrencyHandle.Settings.MaxCount.GetValue() > 0);
+PRAGMA_ENABLE_DEPRECATION_WARNINGS
 	}
 #endif
 
