[Backout] - CL43048705
#rnx
Original CL Desc
-----------------------------------------------------------------
Move SoundConcurrency MaxCount to per-platform int. Minor clean-up moving functional bits to implementation files

#rnx

diff --git a/Engine/Source/Runtime/Engine/Classes/Sound/SoundConcurrency.h b/Engine/Source/Runtime/Engine/Classes/Sound/SoundConcurrency.h
--- a/Engine/Source/Runtime/Engine/Classes/Sound/SoundConcurrency.h
+++ b/Engine/Source/Runtime/Engine/Classes/Sound/SoundConcurrency.h
@@ -5,8 +5,6 @@
 #include "CoreMinimal.h"
 #include "UObject/ObjectMacros.h"
 #include "UObject/Object.h"
-#include "UObject/PerPlatformProperties.h"
-
 #include "SoundConcurrency.generated.h"
 
 class FAudioDevice;
@@ -74,29 +72,29 @@ enum class EConcurrencyVolumeScaleMode
 USTRUCT(BlueprintType)
 struct FSoundConcurrencySettings
 {
-	GENERATED_BODY()
+	GENERATED_USTRUCT_BODY()
 
 	/** The max number of allowable concurrent active voices for voices playing in this concurrency group. */
 	UPROPERTY(EditAnywhere, BlueprintReadWrite, Category = Concurrency, meta = (UIMin = "1", ClampMin = "1"))
-	FPerPlatformInt MaxCount;
+	int32 MaxCount;
 
 	/* Whether or not to limit the concurrency to per sound owner (i.e. the actor that plays the sound). If the sound doesn't have an owner, it falls back to global concurrency. */
 	UPROPERTY(EditAnywhere, BlueprintReadWrite, Category = Concurrency)
-	uint8 bLimitToOwner : 1 = 0;
+	uint8 bLimitToOwner:1;
 
 	/**
 	 * Whether or not volume scaling can recover volume ducking behavior when concurrency group sounds stop (default scale mode only).
 	 */
 	UPROPERTY(EditAnywhere, BlueprintReadWrite, Category = "Volume Scaling", meta = (DisplayName = "Can Recover", EditCondition = "VolumeScaleMode == EConcurrencyVolumeScaleMode::Default"))
-	uint8 bVolumeScaleCanRelease : 1 = 0;
+	uint8 bVolumeScaleCanRelease:1;
 
 	/** Which concurrency resolution policy to use if max voice count is reached. */
 	UPROPERTY(EditAnywhere, BlueprintReadWrite, Category = Concurrency)
-	TEnumAsByte<EMaxConcurrentResolutionRule::Type> ResolutionRule = EMaxConcurrentResolutionRule::StopFarthestThenOldest;
+	TEnumAsByte<EMaxConcurrentResolutionRule::Type> ResolutionRule;
 
 	/** Amount of time to wait (in seconds) between different sounds which play with this concurrency. Sounds rejected from this will ignore virtualization settings. */
 	UPROPERTY(EditAnywhere, BlueprintReadWrite, Category = Concurrency, meta = (UIMin = "0.0", ClampMin = "0.0"))
-	float RetriggerTime = 0.0f;
+	float RetriggerTime;
 
 private:
 	/**
@@ -108,42 +106,44 @@ private:
 	 * AppliedVolumeScale = Math.Pow(DuckingScale, VoiceGeneration)
 	 */
 	UPROPERTY(EditAnywhere, Category = "Volume Scaling", meta = (UIMin = "0.0", UIMax = "1.0", ClampMin = "0.0", ClampMax = "1.0"))
-	float VolumeScale = 1.0f;
+	float VolumeScale;
 
 public:
 	/** Volume Scale mode designating how to scale voice volume based on number of member sounds active in group. */
 	UPROPERTY(EditAnywhere, BlueprintReadWrite, Category = "Volume Scaling")
-	EConcurrencyVolumeScaleMode VolumeScaleMode = EConcurrencyVolumeScaleMode::Default;
+	EConcurrencyVolumeScaleMode VolumeScaleMode;
 
 	/**
 	 * Time taken to apply duck using volume scalar.
 	 */
 	UPROPERTY(EditAnywhere, BlueprintReadWrite, Category = "Volume Scaling", meta = (DisplayName = "Duck Time", UIMin = "0.0", ClampMin = "0.0", UIMax = "10.0", ClampMax = "1000000.0"))
-	float VolumeScaleAttackTime = 0.01f;
+	float VolumeScaleAttackTime;
 
 	/**
 	 * Time taken to recover volume scalar duck.
 	 */
 	UPROPERTY(EditAnywhere, BlueprintReadWrite, Category = "Volume Scaling", meta = (DisplayName = "Recover Time", EditCondition = "bVolumeScaleCanRelease && VolumeScaleMode == EConcurrencyVolumeScaleMode::Default", UIMin = "0.0", ClampMin = "0.0", UIMax = "10.0", ClampMax="1000000.0"))
-	float VolumeScaleReleaseTime = 0.5f;
+	float VolumeScaleReleaseTime;
 
 	/**
 	 * Time taken to fade out if voice is evicted or culled due to another voice in the group starting.
 	 */
 	UPROPERTY(EditAnywhere, BlueprintReadWrite, Category = "Voice Stealing", meta = (UIMin = "0.0", ClampMin = "0.0", UIMax = "1.0", ClampMax="1000000.0"))
-	float VoiceStealReleaseTime = 0.0f;
-
-	/**
-	  * Returns the MaxCount as a per-platform integer value.
-	  */
-	ENGINE_API FPerPlatformInt GetMaxCount() const;
-
-	/**
-	  * Sets the MaxCount as a per-platform integer value.
-	  */
-	ENGINE_API void SetMaxCount(FPerPlatformInt InMaxCount);
-
-	FSoundConcurrencySettings();
+	float VoiceStealReleaseTime;
+
+
+	FSoundConcurrencySettings()
+		: MaxCount(16)
+		, bLimitToOwner(0)
+		, bVolumeScaleCanRelease(0)
+		, ResolutionRule(EMaxConcurrentResolutionRule::StopFarthestThenOldest)
+		, RetriggerTime(0.0f)
+		, VolumeScale(1.0f)
+		, VolumeScaleMode(EConcurrencyVolumeScaleMode::Default)
+		, VolumeScaleAttackTime(0.01f)
+		, VolumeScaleReleaseTime(0.5f)
+		, VoiceStealReleaseTime(0.0f)
+	{}
 
 	/**
 	 * Retrieves the volume scale
@@ -244,25 +244,25 @@ public:
 	static FConcurrencyGroupID GenerateNewID();
 
 	/** Returns the active sounds array. */
-	const TArray<FActiveSound*>& GetActiveSounds() const;
+	const TArray<FActiveSound*>& GetActiveSounds() const { return ActiveSounds; }
 
 	/** Returns the id of the concurrency group */
-	FConcurrencyGroupID GetGroupID() const;
+	FConcurrencyGroupID GetGroupID() const { return GroupID; }
 
 	/** Returns the current generation (effectively, the number of concurrency sound instances active) */
-	const int32 GetNextGeneration() const;
+	const int32 GetNextGeneration() const { return ActiveSounds.Num(); }
 
 	/** Returns the settings associated with the group */
-	const FSoundConcurrencySettings& GetSettings() const;
+	const FSoundConcurrencySettings& GetSettings() const { return Settings; }
 
 	/** Returns the parent object ID */
-	FConcurrencyObjectID GetObjectID() const;
+	FConcurrencyObjectID GetObjectID() const { return ObjectID; }
 
 	/** Determines if the group is full. */
-	bool IsEmpty() const;
+	bool IsEmpty() const { return ActiveSounds.Num() == 0; }
 
 	/** Determines if the group is full. */
-	bool IsFull() const;
+	bool IsFull() const { return Settings.MaxCount <= ActiveSounds.Num(); }
 
 	/** Adds an active sound to the active sound array. */
 	void AddActiveSound(FActiveSound& ActiveSound);
diff --git a/Engine/Source/Runtime/Engine/Private/SoundBase.cpp b/Engine/Source/Runtime/Engine/Private/SoundBase.cpp
--- a/Engine/Source/Runtime/Engine/Private/SoundBase.cpp
+++ b/Engine/Source/Runtime/Engine/Private/SoundBase.cpp
@@ -244,7 +244,7 @@ void USoundBase::PostLoad()
 	{
 		bOverrideConcurrency = true;
 		ConcurrencyOverrides.bLimitToOwner = false;
-		ConcurrencyOverrides.SetMaxCount(FPerPlatformInt(FMath::Max(MaxConcurrentPlayCount_DEPRECATED, 1)));
+		ConcurrencyOverrides.MaxCount = FMath::Max(MaxConcurrentPlayCount_DEPRECATED, 1);
 		ConcurrencyOverrides.ResolutionRule = MaxConcurrentResolutionRule_DEPRECATED;
 	}
 
diff --git a/Engine/Source/Runtime/Engine/Private/SoundConcurrency.cpp b/Engine/Source/Runtime/Engine/Private/SoundConcurrency.cpp
--- a/Engine/Source/Runtime/Engine/Private/SoundConcurrency.cpp
+++ b/Engine/Source/Runtime/Engine/Private/SoundConcurrency.cpp
@@ -56,27 +56,6 @@ USoundConcurrency::USoundConcurrency(const FObjectInitializer& ObjectInitializer
 {
 }
 
-FSoundConcurrencySettings::FSoundConcurrencySettings()
-PRAGMA_DISABLE_DEPRECATION_WARNINGS
-	: MaxCount(16)
-PRAGMA_ENABLE_DEPRECATION_WARNINGS
-{
-}
-
-FPerPlatformInt FSoundConcurrencySettings::GetMaxCount() const
-{
-PRAGMA_DISABLE_DEPRECATION_WARNINGS
-	return MaxCount;
-PRAGMA_ENABLE_DEPRECATION_WARNINGS
-}
-
-void FSoundConcurrencySettings::SetMaxCount(FPerPlatformInt InMaxCount)
-{
-PRAGMA_DISABLE_DEPRECATION_WARNINGS
-	MaxCount = MoveTemp(InMaxCount);
-PRAGMA_ENABLE_DEPRECATION_WARNINGS
-}
-
 float FSoundConcurrencySettings::GetVolumeScale() const
 {
 	return FMath::Clamp(VolumeScale, 0.0f, 1.0f);;
@@ -199,43 +178,6 @@ FConcurrencyGroupID FConcurrencyGroup::GenerateNewID()
 	return ++ConcurrencyGroupIDs;
 }
 
-const TArray<FActiveSound*>& FConcurrencyGroup::GetActiveSounds() const
-{
-	return ActiveSounds;
-}
-
-FConcurrencyGroupID FConcurrencyGroup::GetGroupID() const
-{
-	return GroupID;
-}
-
-const int32 FConcurrencyGroup::GetNextGeneration() const
-{
-	return ActiveSounds.Num();
-}
-
-const FSoundConcurrencySettings& FConcurrencyGroup::GetSettings() const
-{
-	return Settings;
-}
-
-FConcurrencyObjectID FConcurrencyGroup::GetObjectID() const
-{
-	return ObjectID;
-}
-
-bool FConcurrencyGroup::IsEmpty() const
-{
-	return ActiveSounds.IsEmpty();
-}
-
-bool FConcurrencyGroup::IsFull() const
-{
-PRAGMA_DISABLE_DEPRECATION_WARNINGS
-	return Settings.MaxCount.GetValue() <= ActiveSounds.Num();
-PRAGMA_ENABLE_DEPRECATION_WARNINGS
-}
-
 void FConcurrencyGroup::AddActiveSound(FActiveSound& ActiveSound)
 {
 	check(GroupID != 0);
@@ -377,12 +319,10 @@ void FConcurrencyGroup::CullSoundsDueToMaxConcurrency()
 {
 	// Nothing to do if our active sound count is less than or equal to our max active sounds
 	// or if eviction is supported
-PRAGMA_DISABLE_DEPRECATION_WARNINGS
-	if (ActiveSounds.Num() <= Settings.MaxCount.GetValue())
+	if (ActiveSounds.Num() <= Settings.MaxCount)
 	{
 		return;
 	}
-PRAGMA_ENABLE_DEPRECATION_WARNINGS
 
 	// Comparator for sorting group's ActiveSounds according to their dynamic concurrency nature.
 	// Quieter/lower priority sounds will be at the front of the array.
@@ -445,10 +385,7 @@ PRAGMA_ENABLE_DEPRECATION_WARNINGS
 
 	// Need to make a new list when stopping the sounds since the process of stopping an active sound
 	// will remove the sound from this concurrency group's ActiveSounds array.
-PRAGMA_DISABLE_DEPRECATION_WARNINGS
-	const int32 NumSoundsToStop = ActiveSounds.Num() - Settings.MaxCount.GetValue();
-PRAGMA_ENABLE_DEPRECATION_WARNINGS
-
+	const int32 NumSoundsToStop = ActiveSounds.Num() - Settings.MaxCount;
 	check(NumSoundsToStop > 0);
 	for (int32 i = 0; i < ActiveSounds.Num(); ++i)
 	{
@@ -587,9 +524,7 @@ FActiveSound* FSoundConcurrencyManager::CreateNewActiveSound(const FActiveSound&
 #if !UE_BUILD_SHIPPING
 	for (const FConcurrencyHandle& ConcurrencyHandle : ConcurrencyHandles)
 	{
-PRAGMA_DISABLE_DEPRECATION_WARNINGS
-		check(ConcurrencyHandle.Settings.MaxCount.GetValue() > 0);
-PRAGMA_ENABLE_DEPRECATION_WARNINGS
+		check(ConcurrencyHandle.Settings.MaxCount > 0);
 	}
 #endif
 
