Fix incorrect Substrate settings and shading code causing incorrect quality settings.

diff --git a/Engine/Shaders/Private/Substrate/SubstrateEvaluation.ush b/Engine/Shaders/Private/Substrate/SubstrateEvaluation.ush
--- a/Engine/Shaders/Private/Substrate/SubstrateEvaluation.ush
+++ b/Engine/Shaders/Private/Substrate/SubstrateEvaluation.ush
@@ -269,9 +269,9 @@ half Substrate_D_GGX_Aniso(float ax, float ay, float NoH, float XoH, float YoH)
 half Substrate_Vis_GGX(half Roughness, half a2, half NoV, half NoL)
 {
 #if SUBSTRATE_SHADING_QUALITY > 1
-	return Vis_SmithJoint(a2, NoV, NoL);
-#else
 	return Vis_SmithJointApprox(a2, NoV, NoL);
+#else
+	return Vis_SmithJoint(a2, NoV, NoL);
 #endif
 }
 
diff --git a/Engine/Source/Runtime/Engine/Private/ShaderCompiler/ShaderCompiler.cpp b/Engine/Source/Runtime/Engine/Private/ShaderCompiler/ShaderCompiler.cpp
--- a/Engine/Source/Runtime/Engine/Private/ShaderCompiler/ShaderCompiler.cpp
+++ b/Engine/Source/Runtime/Engine/Private/ShaderCompiler/ShaderCompiler.cpp
@@ -3886,7 +3886,7 @@ void GlobalBeginCompileShader(
 			const bool bLowQuality = SubstrateShadingQuality > 1;
 			SET_SHADER_DEFINE(Input.Environment, USE_ACHROMATIC_BXDF_ENERGY, bLowQuality ? 1u : 0u);
 
-			const uint32 SubstrateSheenQuality = Substrate::GetSheenQuality();
+			const uint32 SubstrateSheenQuality = Substrate::GetSheenQuality(Target.GetPlatform());
 			Input.Environment.SetDefine(TEXT("SUBSTRATE_SHEEN_QUALITY"), bLowQuality ? 2 : SubstrateSheenQuality);
 
 			const uint32 SubstrateNormalQuality = Substrate::GetNormalQuality();
@@ -3938,7 +3938,7 @@ void GlobalBeginCompileShader(
 	{
 		static IConsoleVariable* CVar = IConsoleManager::Get().FindConsoleVariable(TEXT("r.Material.RoughDiffuse"));
 		const bool bMaterialRoughDiffuse = CVar && CVar->GetInt() != 0;
-		const bool bSubstrateRoughDiffuse = Substrate::IsRoughDiffuseEnabled() && !Substrate::IsBackCompatibilityEnabled();
+		const bool bSubstrateRoughDiffuse = Substrate::IsRoughDiffuseEnabled(Target.GetPlatform()) && !Substrate::IsBackCompatibilityEnabled();
 		SET_SHADER_DEFINE(Input.Environment, MATERIAL_ROUGHDIFFUSE, (bSubstrate ? bSubstrateRoughDiffuse : bMaterialRoughDiffuse) ? 1 : 0);
 	}
 
diff --git a/Engine/Source/Runtime/RenderCore/Private/RenderUtils.cpp b/Engine/Source/Runtime/RenderCore/Private/RenderUtils.cpp
--- a/Engine/Source/Runtime/RenderCore/Private/RenderUtils.cpp
+++ b/Engine/Source/Runtime/RenderCore/Private/RenderUtils.cpp
@@ -2094,6 +2094,13 @@ namespace Substrate
 		return CVarSubstrateRoughDiffuse.GetValueOnAnyThread() > 0;
 	}
 
+	bool IsRoughDiffuseEnabled(EShaderPlatform InPlatform)
+	{
+		// Variant for shader compilation per platform
+		static FShaderPlatformCachedIniValue<int32> CVar(TEXT("r.Substrate.RoughDiffuse"));
+		return CVar.Get(InPlatform) > 0;
+	}
+
 	bool IsGlintEnabled()
 	{
 		return IsSubstrateEnabled() && CVarSubstrateGlints.GetValueOnAnyThread() > 0;
diff --git a/Engine/Source/Runtime/RenderCore/Public/RenderUtils.h b/Engine/Source/Runtime/RenderCore/Public/RenderUtils.h
--- a/Engine/Source/Runtime/RenderCore/Public/RenderUtils.h
+++ b/Engine/Source/Runtime/RenderCore/Public/RenderUtils.h
@@ -594,6 +594,7 @@ namespace Substrate
 	RENDERCORE_API bool IsSubstrateBlendableGBufferEnabled(EShaderPlatform InPlatform);
 	RENDERCORE_API bool IsStochasticLightingEnabled(EShaderPlatform InPlatform);
 	RENDERCORE_API bool IsRoughDiffuseEnabled();
+	RENDERCORE_API bool IsRoughDiffuseEnabled(EShaderPlatform InPlatform);
 	RENDERCORE_API bool IsGlintEnabled();
 	RENDERCORE_API bool IsGlintEnabled(EShaderPlatform InPlatform);
 	RENDERCORE_API uint32 GlintLUTIndex();
