Fixed TSAN data race error on LevelStreaming CurrentState when using wp.runtime.UpdateStreaming.EnableAsyncUpdate.

diff --git a/Engine/Source/Runtime/Engine/Classes/Engine/LevelStreaming.h b/Engine/Source/Runtime/Engine/Classes/Engine/LevelStreaming.h
--- a/Engine/Source/Runtime/Engine/Classes/Engine/LevelStreaming.h
+++ b/Engine/Source/Runtime/Engine/Classes/Engine/LevelStreaming.h
@@ -10,6 +10,7 @@
 #include "LatentActions.h"
 #include "ProfilingDebugging/ProfilingHelpers.h"
 #include "GameFramework/UpdateLevelVisibilityLevelInfo.h"
+#include "Misc/TransactionallySafeRWLock.h"
 #include "Misc/Timeout.h"
 
 #if WITH_EDITOR
@@ -278,6 +279,9 @@ protected:
 	TArray<int32> AsyncRequestIDs;
 
 private:
+	/** Used to get the CurrentState from an async task while it can be modified by the gamethread */
+	mutable FTransactionallySafeRWLock CurrentStateLock;
+
 	/** What the current streamed state of the streaming level is */
 	ELevelStreamingState CurrentState;
 
@@ -361,10 +365,10 @@ public:
 	/** Returns the current loaded/visible state of the streaming level. */
 PRAGMA_DISABLE_DEPRECATION_WARNINGS
 	UE_DEPRECATED(5.2, "ULevelStreaming::ECurrentState has been replaced by ELevelStreamingState. Use GetLevelStreamingState instead.")
-	ECurrentState GetCurrentState() const { return (ECurrentState)CurrentState; }
+	ECurrentState GetCurrentState() const { return (ECurrentState)GetLevelStreamingState(); }
 PRAGMA_ENABLE_DEPRECATION_WARNINGS
 
-	ELevelStreamingState GetLevelStreamingState() const { return CurrentState; }
+	ENGINE_API ELevelStreamingState GetLevelStreamingState() const;
 
 protected:	
 	/** Updates the current state of the streaming level and notifies any observers. */
diff --git a/Engine/Source/Runtime/Engine/Private/LevelStreaming.cpp b/Engine/Source/Runtime/Engine/Private/LevelStreaming.cpp
--- a/Engine/Source/Runtime/Engine/Private/LevelStreaming.cpp
+++ b/Engine/Source/Runtime/Engine/Private/LevelStreaming.cpp
@@ -34,6 +34,7 @@
 #include "Streaming/LevelStreamingDelegates.h"
 #include "AssetRegistry/AssetRegistryModule.h"
 #include "WorldPartition/WorldPartition.h"
+#include "Misc/ScopeRWLock.h"
 
 #include UE_INLINE_GENERATED_CPP_BY_NAME(LevelStreaming)
 
@@ -589,10 +590,20 @@ void ULevelStreaming::OnLevelRemoved()
 	SetCurrentState(ELevelStreamingState::Removed);
 }
 
+ELevelStreamingState ULevelStreaming::GetLevelStreamingState() const
+{
+	UE::TReadScopeLock ReadLock(CurrentStateLock);
+	return CurrentState;
+}
+
 void ULevelStreaming::SetCurrentState(ELevelStreamingState NewState)
 {
-	ELevelStreamingState OldState = CurrentState;
-	CurrentState = NewState;
+	ELevelStreamingState OldState;
+	{
+		UE::TWriteScopeLock WriteLock(CurrentStateLock);
+		OldState = CurrentState;
+		CurrentState = NewState;
+	}
 
 	if (OldState != NewState)
 	{
@@ -2402,21 +2413,23 @@ bool ULevelStreaming::IsValidStreamingLevel() const
 
 EStreamingStatus ULevelStreaming::GetLevelStreamingStatus() const
 {
-	if (CurrentState == ELevelStreamingState::FailedToLoad)
+	// This function can be called from an async task (when wp.runtime.UpdateStreaming.EnableAsyncUpdate = true)
+	ELevelStreamingState LocalState = GetLevelStreamingState();
+	if (LocalState == ELevelStreamingState::FailedToLoad)
 	{
 		return LEVEL_FailedToLoad;
 	}
-	else if (CurrentState == ELevelStreamingState::MakingInvisible)
+	else if (LocalState == ELevelStreamingState::MakingInvisible)
 	{
 		return LEVEL_MakingInvisible;
 	}
 	else if (LoadedLevel)
 	{
-		if (CurrentState == ELevelStreamingState::LoadedVisible)
+		if (LocalState == ELevelStreamingState::LoadedVisible)
 		{
 			return LEVEL_Visible;
 		}
-		else if ((CurrentState == ELevelStreamingState::MakingVisible) && (GetWorld()->IsLevelMakingVisible(LoadedLevel)))
+		else if ((LocalState == ELevelStreamingState::MakingVisible) && (GetWorld()->IsLevelMakingVisible(LoadedLevel)))
 		{
 			return LEVEL_MakingVisible;
 		}
@@ -2424,12 +2437,12 @@ EStreamingStatus ULevelStreaming::GetLevelStreamingStatus() const
 	}
 	else
 	{
-		if (CurrentState == ELevelStreamingState::Loading)
+		if (LocalState == ELevelStreamingState::Loading)
 		{
 			return LEVEL_Loading;
 		}
 
-		check(CurrentState == ELevelStreamingState::Removed || CurrentState == ELevelStreamingState::Unloaded);
+		check(LocalState == ELevelStreamingState::Removed || LocalState == ELevelStreamingState::Unloaded);
 
 		// We can't use FindObjectFast() while the GC is running... Let's treat the level as unloaded at this point.
 		if (IsGarbageCollecting())
