[VisualLogger] fixed ever growing buffer in FVisualLoggerBinaryFileDevice when VisualLogger is not set to record to file

diff --git a/Engine/Source/Runtime/Engine/Private/VisualLogger/VisualLoggerBinaryFileDevice.cpp b/Engine/Source/Runtime/Engine/Private/VisualLogger/VisualLoggerBinaryFileDevice.cpp
--- a/Engine/Source/Runtime/Engine/Private/VisualLogger/VisualLoggerBinaryFileDevice.cpp
+++ b/Engine/Source/Runtime/Engine/Private/VisualLogger/VisualLoggerBinaryFileDevice.cpp
@@ -100,14 +100,17 @@ void FVisualLoggerBinaryFileDevice::SetFileName(const FString& InFileName)
 
 void FVisualLoggerBinaryFileDevice::Serialize(const UObject* InLogOwner, const FName& InOwnerName, const FName& InOwnerDisplayName, const FName& InOwnerClassName, const FVisualLogEntry& InLogEntry)
 {
-	const int32 NumEntries = FrameCache.Num();
-	if (NumEntries> 0 && LastLogTimeStamp + FrameCacheLenght <= InLogEntry.TimeStamp && FileArchive)
+	if (FileArchive)
 	{
-		FVisualLoggerHelpers::Serialize(*FileArchive, FrameCache);
-		FrameCache.Reset();
-		LastLogTimeStamp = InLogEntry.TimeStamp;
+		const int32 NumEntries = FrameCache.Num();
+		if (NumEntries> 0 && LastLogTimeStamp + FrameCacheLenght <= InLogEntry.TimeStamp)
+		{
+			FVisualLoggerHelpers::Serialize(*FileArchive, FrameCache);
+			FrameCache.Reset();
+			LastLogTimeStamp = InLogEntry.TimeStamp;
+		}
+
+		FrameCache.Add(FVisualLogEntryItem(InOwnerName, InOwnerDisplayName, InOwnerClassName, InLogEntry));
 	}
-
-	FrameCache.Add(FVisualLogEntryItem(InOwnerName, InOwnerDisplayName, InOwnerClassName, InLogEntry));
 }
 #endif
