World Partition Cell Transformer ISM: when removing the root component for batching, and the actor cannot be removed because it still has unbatched components, make sure to create a root scene component to replace the one removed.

diff --git a/Engine/Source/Runtime/Engine/Private/WorldPartition/WorldPartitionRuntimeCellTransformerISM.cpp b/Engine/Source/Runtime/Engine/Private/WorldPartition/WorldPartitionRuntimeCellTransformerISM.cpp
--- a/Engine/Source/Runtime/Engine/Private/WorldPartition/WorldPartitionRuntimeCellTransformerISM.cpp
+++ b/Engine/Source/Runtime/Engine/Private/WorldPartition/WorldPartitionRuntimeCellTransformerISM.cpp
@@ -101,6 +101,12 @@ void UWorldPartitionRuntimeCellTransformerISM::Transform(ULevel* InLevel)
 				{
 					*Actor = nullptr;
 				}
+				else if (USceneComponent* OldRootComponent = (*Actor)->GetRootComponent(); OldRootComponent && !IsValid(OldRootComponent))
+				{
+					USceneComponent* NewRootComponent = NewObject<USceneComponent>(*Actor);
+					NewRootComponent->SetRelativeTransform(OldRootComponent->GetRelativeTransform());
+					(*Actor)->SetRootComponent(NewRootComponent);
+				}
 			}
 		}
 	}
