[MIR] Add support for virtual textures in the new material translator.

- Add VTPageTableRead MIR instruction to encode VT page table read operations; This separates them from texture read operations in order to optimize the number of page tables.
- Refactore FSubstrateParameter value to FPromoteSubstrateParameter instruction; This must also be an instruction and not a simple value and naming align with the Substrate HLSL code.
- Move some code from the old translator into new namespace UE::MaterialTranslatorUtils to be shared between old and new translator.
- Cache VT stack in MIR analyzer stage to find page table loads that that share the same parameters; This still emits individual instructions resulting in one page table load for each texture access, which needs to be optimized in a later stage (separate CL).

#rnx

diff --git a/Engine/Source/Runtime/Engine/Private/Materials/HLSLMaterialTranslator.cpp b/Engine/Source/Runtime/Engine/Private/Materials/HLSLMaterialTranslator.cpp
--- a/Engine/Source/Runtime/Engine/Private/Materials/HLSLMaterialTranslator.cpp
+++ b/Engine/Source/Runtime/Engine/Private/Materials/HLSLMaterialTranslator.cpp
@@ -8,6 +8,7 @@
 #include "VT/VirtualTextureScalability.h"
 #include "Engine/Engine.h"
 #include "MaterialDomain.h"
+#include "MaterialSharedPrivate.h"
 #include "Engine/Texture.h"
 #include "Engine/TextureCollection.h"
 #include "Materials/MaterialAttributeDefinitionMap.h"
@@ -2965,7 +2966,6 @@ void FHLSLMaterialTranslator::GetMaterialEnvironment(EShaderPlatform InPlatform,
 		else
 		{
 			// Setup page table defines to map each VT stack to either 1 or 2 page table textures, depending on how many layers it uses
-			const FMaterialVirtualTextureStack& Stack = MaterialCompilationOutput.UniformExpressionSet.VTStacks[i];
 			FString PageTableValue = FString::Printf(TEXT("Material.VirtualTexturePageTable0_%d"), i);
 			if (EnvironmentDefines->VirtualPageTypes[i] & FEnvironmentDefines::LOCAL_TABLE1)
 			{
@@ -7100,17 +7100,6 @@ void FHLSLMaterialTranslator::SetPotentiallyManipulateTexCoords()
 	bPotentiallyManipulateTexCoords = true;
 }
 
-static const TCHAR* GetVTAddressMode(TextureAddress Address)
-{
-	switch (Address)
-	{
-	case TA_Wrap: return TEXT("VTADDRESSMODE_WRAP");
-	case TA_Clamp: return TEXT("VTADDRESSMODE_CLAMP");
-	case TA_Mirror: return TEXT("VTADDRESSMODE_MIRROR");
-	default: checkNoEntry(); return nullptr;
-	}
-}
-
 uint32 FHLSLMaterialTranslator::AcquireVTStackIndex(
 	ETextureMipValueMode MipValueMode, 
 	TextureAddress AddressU, TextureAddress AddressV, 
@@ -7123,6 +7112,8 @@ uint32 FHLSLMaterialTranslator::AcquireVTStackIndex(
 	const FString& UV_Ddy,
 	bool bAdaptive, bool bGenerateFeedback, bool bMeshPaint, bool bMaterialCache)
 {
+	using namespace UE::MaterialTranslatorUtils;
+
 	const uint64 CoordinateHash = GetParameterHash(CoordinateIndex);
 	const uint64 MipValue0Hash = GetParameterHash(MipValue0Index);
 	const uint64 MipValue1Hash = GetParameterHash(MipValue1Index);
@@ -7413,6 +7404,8 @@ int32 FHLSLMaterialTranslator::TextureSample(
 	bool EnableFeedback
 )
 {
+	using namespace UE::MaterialTranslatorUtils;
+
 	if (TextureIndex == INDEX_NONE || CoordinateIndex == INDEX_NONE)
 	{
 		return INDEX_NONE;
@@ -7638,47 +7631,10 @@ int32 FHLSLMaterialTranslator::TextureSample(
 		}
 	}
 
-	// StaticAddress mode at time of compile
-	// This may not be 100% correct, if SamplerSource is set to 'SSM_FromTextureAsset', as texture parameter may change the address mode in a derived instance
 	TextureAddress StaticAddressX = TA_Wrap;
 	TextureAddress StaticAddressY = TA_Wrap;
 	TextureAddress StaticAddressZ = TA_Wrap;
-	if (Texture && Texture->Source.GetNumBlocks() > 1)
-	{
-		// UDIM (multi-block) texture are forced to use wrap address mode
-		// This is important for supporting VT stacks made from UDIMs with differing number of blocks, as this requires wrapping vAddress for certain layers
-		StaticAddressX = TA_Wrap;
-		StaticAddressY = TA_Wrap;
-		StaticAddressZ = TA_Wrap;
-	}
-	else
-	{
-		switch (SamplerSource)
-		{
-		case SSM_FromTextureAsset:
-			if (Texture)
-			{
-				StaticAddressX = Texture->GetTextureAddressX();
-				StaticAddressY = Texture->GetTextureAddressY();
-				StaticAddressZ = Texture->GetTextureAddressZ();
-			}
-			break;
-		case SSM_Wrap_WorldGroupSettings:
-			StaticAddressX = TA_Wrap;
-			StaticAddressY = TA_Wrap;
-			StaticAddressZ = TA_Wrap;
-			break;
-		case SSM_Clamp_WorldGroupSettings:
-		case SSM_TerrainWeightmapGroupSettings:
-			StaticAddressX = TA_Clamp;
-			StaticAddressY = TA_Clamp;
-			StaticAddressZ = TA_Clamp;
-			break;
-		default:
-			checkNoEntry();
-			break;
-		}
-	}
+	GetTextureAddressForSamplerSource(Texture, SamplerSource, StaticAddressX, StaticAddressY, StaticAddressZ);
 
 	EMaterialValueType UVsType;
 	switch (TextureType)
diff --git a/Engine/Source/Runtime/Engine/Private/Materials/HLSLMaterialTranslator.h b/Engine/Source/Runtime/Engine/Private/Materials/HLSLMaterialTranslator.h
--- a/Engine/Source/Runtime/Engine/Private/Materials/HLSLMaterialTranslator.h
+++ b/Engine/Source/Runtime/Engine/Private/Materials/HLSLMaterialTranslator.h
@@ -169,28 +169,6 @@ struct FShaderCodeChunk
 	{}
 };
 
-struct FMaterialVTStackEntry
-{
-	uint64 ScopeID;
-	uint64 CoordinateHash;
-	uint64 MipValue0Hash;
-	uint64 MipValue1Hash;
-	ETextureMipValueMode MipValueMode;
-	TextureAddress AddressU;
-	TextureAddress AddressV;
-	int32 DebugCoordinateIndex;
-	int32 DebugMipValue0Index;
-	int32 DebugMipValue1Index;
-	int32 PreallocatedStackTextureIndex;
-	bool bAdaptive;
-	bool bGenerateFeedback;
-	bool bMeshPaint;
-	bool bMaterialCache;
-	float AspectRatio;
-
-	int32 CodeIndex;
-};
-
 struct FMaterialCustomExpressionEntry
 {
 	uint64 ScopeID;
@@ -256,6 +234,29 @@ class FHLSLMaterialTranslator : public FMaterialCompiler
 {
 	friend class FMaterialDerivativeAutogen;
 protected:
+	/** Translator specific structure during generation of VT stacks. This is later converted to FMaterialVirtualTextureStack entries in FUniformExpressionSet. */
+	struct FMaterialVTStackEntry
+	{
+		uint64 ScopeID;
+		uint64 CoordinateHash;
+		uint64 MipValue0Hash;
+		uint64 MipValue1Hash;
+		ETextureMipValueMode MipValueMode;
+		TextureAddress AddressU;
+		TextureAddress AddressV;
+		int32 DebugCoordinateIndex;
+		int32 DebugMipValue0Index;
+		int32 DebugMipValue1Index;
+		int32 PreallocatedStackTextureIndex;
+		bool bAdaptive;
+		bool bGenerateFeedback;
+		bool bMeshPaint;
+		bool bMaterialCache;
+		float AspectRatio;
+
+		int32 CodeIndex;
+	};
+
 	/** Data that is different for the finite and anlytical partial derivative options. **/
 	FMaterialDerivativeVariation DerivativeVariations[CompiledPDV_MAX];
 
diff --git a/Engine/Source/Runtime/Engine/Private/Materials/MaterialIR.cpp b/Engine/Source/Runtime/Engine/Private/Materials/MaterialIR.cpp
--- a/Engine/Source/Runtime/Engine/Private/Materials/MaterialIR.cpp
+++ b/Engine/Source/Runtime/Engine/Private/Materials/MaterialIR.cpp
@@ -42,7 +42,6 @@ const TCHAR* LexToString(EValueKind Kind)
 		case VK_ExternalInput: return TEXT("ExternalInput");
 		case VK_TextureObject: return TEXT("TextureObject");
 		case VK_UniformParameter: return TEXT("UniformParameter");
-		case VK_SubstrateParameter: return TEXT("SubstrateParameter");
 		case VK_Dimensional: return TEXT("Dimensional");
 		case VK_SetMaterialOutput: return TEXT("SetMaterialOutput");
 		case VK_Operator: return TEXT("Operator");
@@ -50,7 +49,9 @@ const TCHAR* LexToString(EValueKind Kind)
 		case VK_Subscript: return TEXT("Subscript");
 		case VK_Scalar: return TEXT("Scalar");
 		case VK_TextureRead: return TEXT("TextureRead");
+		case VK_VTPageTableRead: return TEXT("VTPageTableRead");
 		case VK_InlineHLSL: return TEXT("InlineHLSL");
+		case VK_PromoteSubstrateParameter: return TEXT("PromoteSubstrateParameter");
 		case VK_StageSwitch: return TEXT("StageSwitch");
 		case VK_HardwarePartialDerivative: return TEXT("HardwarePartialDerivative");
 		/* invalid entries */
@@ -100,7 +101,6 @@ uint32 FValue::GetSizeInBytes() const
 		case VK_ExternalInput: return sizeof(FExternalInput);
 		case VK_TextureObject: return sizeof(FTextureObject);
 		case VK_UniformParameter: return sizeof(FUniformParameter);
-		case VK_SubstrateParameter: return sizeof(VK_SubstrateParameter);
 		case VK_Dimensional: return sizeof(FDimensional) + sizeof(FValue*) * static_cast<const FDimensional*>(this)->GetComponents().Num();
 		case VK_SetMaterialOutput: return sizeof(FSetMaterialOutput);
 		case VK_Operator: return sizeof(FOperator);
@@ -108,7 +108,9 @@ uint32 FValue::GetSizeInBytes() const
 		case VK_Subscript: return sizeof(FSubscript);
 		case VK_Scalar: return sizeof(FScalar);
 		case VK_TextureRead: return sizeof(FTextureRead);
+		case VK_VTPageTableRead: return sizeof(FVTPageTableRead);
 		case VK_InlineHLSL: return sizeof(FInlineHLSL);
+		case VK_PromoteSubstrateParameter: return sizeof(FPromoteSubstrateParameter);
 		case VK_StageSwitch: return sizeof(FStageSwitch);
 		case VK_HardwarePartialDerivative: return sizeof(FHardwarePartialDerivative);
 		/* invalid entries */
@@ -199,12 +201,24 @@ TConstArrayView<FValue*> FValue::GetUses() const
 			return { &This->TextureObject, FTextureRead::NumStaticUses };
 		}
 
+		case VK_VTPageTableRead:
+		{
+			auto This = static_cast<const FVTPageTableRead*>(this);
+			return { &This->TextureObject, FVTPageTableRead::NumStaticUses };
+		}
+
 		case VK_InlineHLSL:
 		{
 			auto This = static_cast<const FInlineHLSL*>(this);
 			return { This->Arguments, This->NumArguments };
 		}
 
+		case VK_PromoteSubstrateParameter:
+		{
+			auto This = static_cast<const FPromoteSubstrateParameter*>(this);
+			return { This->WorldSpaceTangentsAndNormals, FPromoteSubstrateParameter::NumStaticUses };
+		}
+
 		case VK_StageSwitch:
 		{
 			auto This = static_cast<const FStageSwitch*>(this);
diff --git a/Engine/Source/Runtime/Engine/Private/Materials/MaterialIREmitter.cpp b/Engine/Source/Runtime/Engine/Private/Materials/MaterialIREmitter.cpp
--- a/Engine/Source/Runtime/Engine/Private/Materials/MaterialIREmitter.cpp
+++ b/Engine/Source/Runtime/Engine/Private/Materials/MaterialIREmitter.cpp
@@ -7,6 +7,7 @@
 #include "Materials/MaterialExternalCodeRegistry.h"
 #include "Shader/ShaderTypes.h"
 #include "MaterialShared.h"
+#include "MaterialSharedPrivate.h"
 #include "MaterialExpressionIO.h"
 #include "Engine/Texture.h"
 #include "Engine/Texture2D.h"
@@ -620,16 +621,25 @@ FValueRef FEmitter::ExternalInput(EExternalInput Id)
 
 FValueRef FEmitter::TextureObject(UTexture* Texture, EMaterialSamplerType SamplerType)
 {
-	if (Texture->GetMaterialType() != MCT_Texture2D)
+	switch (Texture->GetMaterialType())
 	{
-		Error(TEXT("Only Textures 2D are supported for now."));
-		return Poison();
-	}
+		case MCT_Texture2D:
+		case MCT_TextureVirtual:
+		{
+			FTextureObject Proto = MakePrototype<FTextureObject>(FObjectType::GetTexture2D());
+			Proto.Texture = Texture;
+			Proto.SamplerType = SamplerType;
+			return EmitPrototype(*this, Proto);
+		}
+		break;
 
-	FTextureObject Proto = MakePrototype<FTextureObject>(FObjectType::GetTexture2D());
-	Proto.Texture = Texture;
-	Proto.SamplerType = SamplerType;
-	return EmitPrototype(*this, Proto);
+		default:
+		{
+			Error(TEXT("Only 2D and Virtual Textures are supported for now."));
+		}
+		break;
+	}
+	return Poison();
 }
 
 FValueRef FEmitter::Parameter(FName Name, FMaterialParameterMetadata& Metadata, EMaterialSamplerType SamplerType)
@@ -689,16 +699,25 @@ FValueRef FEmitter::Parameter(FName Name, FMaterialParameterMetadata& Metadata,
 
 		case EMaterialParameterType::Texture:
 		{
-			if (Metadata.Value.Texture->GetMaterialType() != MCT_Texture2D)
+			switch (Metadata.Value.Texture->GetMaterialType())
 			{
-				Errorf(TEXT("Unsupported texture type"));
-				return MIR::FPoison::Get();
-			}
+				case MCT_Texture2D:
+				case MCT_TextureVirtual:
+				{
+					MIR::FUniformParameter Prototype = MakePrototype<MIR::FUniformParameter>(FObjectType::GetTexture2D());
+					Prototype.ParameterIdInModule = RegisterParameter(Module, Info, Metadata);
+					Prototype.SamplerType = SamplerType;
+					return EmitPrototype(*this, Prototype);
+				}
+				break;
 
-			MIR::FUniformParameter Prototype = MakePrototype<MIR::FUniformParameter>(FObjectType::GetTexture2D());
-			Prototype.ParameterIdInModule = RegisterParameter(Module, Info, Metadata);
-			Prototype.SamplerType = SamplerType;
-			return EmitPrototype(*this, Prototype);
+				default:
+				{
+					Errorf(TEXT("Unsupported texture type"));
+				}
+				break;
+			}
+			return MIR::FPoison::Get();
 		}
 
 		case EMaterialParameterType::StaticSwitch:
@@ -744,12 +763,6 @@ FValueRef FEmitter::CustomPrimitiveData(uint32 PrimitiveDataIndex)
 	return {};
 }
 
-FValueRef FEmitter::SubstrateParameter()
-{
-	MIR::FSubstrateParameter Prototype = MakePrototype<MIR::FSubstrateParameter>(FObjectType::GetSubstrateData());
-	return EmitPrototype(*this, Prototype);
-}
-
 /*------------------------------ Instruction emission ------------------------------*/
 
 FSetMaterialOutput* FEmitter::SetMaterialOutput(EMaterialProperty InProperty, FValue* Arg)
@@ -2280,6 +2293,19 @@ FValueRef FEmitter::StageSwitch(const FType* Type, TConstArrayView<FValueRef> Va
 	return EmitPrototype(*this, Prototype);
 }
 
+static FValueRef VTPageTableLoadFromSamplerSource(FEmitter& Em, FValueRef Texture, ESamplerSourceMode SamplerSourceMode, FValueRef TexCoord, FValueRef TexCoordDdx, FValueRef TexCoordDdy)
+{
+	MIR::FTextureObject* TextureObject = Texture->As<MIR::FTextureObject>();
+	check(TextureObject);
+
+	TextureAddress StaticAddressX = TA_Wrap;
+	TextureAddress StaticAddressY = TA_Wrap;
+	TextureAddress StaticAddressZ = TA_Wrap;
+	UE::MaterialTranslatorUtils::GetTextureAddressForSamplerSource(TextureObject->Texture, SamplerSourceMode, StaticAddressX, StaticAddressY, StaticAddressZ);
+
+	return Em.VTPageTableLoad(Texture, StaticAddressX, StaticAddressY, TexCoord, TexCoordDdx, TexCoordDdy);
+}
+
 FValueRef FEmitter::TextureGather(FValueRef Texture, FValueRef TexCoord, ETextureReadMode GatherMode, ESamplerSourceMode SamplerSourceMode, EMaterialSamplerType SamplerType)
 {
 	check(GatherMode >= ETextureReadMode::GatherRed && GatherMode <= ETextureReadMode::GatherAlpha);
@@ -2296,6 +2322,11 @@ FValueRef FEmitter::TextureGather(FValueRef Texture, FValueRef TexCoord, ETextur
 	Prototype.SamplerSourceMode = SamplerSourceMode;
 	Prototype.SamplerType = SamplerType;
 
+	if (IsVirtualSamplerType(SamplerType))
+	{
+		Prototype.VTPageTable = VTPageTableLoadFromSamplerSource(*this, Texture, SamplerSourceMode, Prototype.TexCoord, {}, {});
+	}
+
 	return EmitPrototype(*this, Prototype);
 }
 
@@ -2329,6 +2360,12 @@ FValueRef FEmitter::TextureSample(FValueRef Texture, FValueRef TexCoord, bool bA
 		PrototypeAn.TexCoordDdy = Operator(BO_Multiply, PrototypeAn.TexCoordDdy, Exp2ViewMaterialTextureMipBias);
 	}
 
+	if (IsVirtualSamplerType(SamplerType))
+	{
+		PrototypeHw.VTPageTable = VTPageTableLoadFromSamplerSource(*this, Texture, SamplerSourceMode, PrototypeHw.TexCoord, {}, {});
+		PrototypeAn.VTPageTable = VTPageTableLoadFromSamplerSource(*this, Texture, SamplerSourceMode, PrototypeAn.TexCoord, PrototypeAn.TexCoordDdx, PrototypeAn.TexCoordDdy);
+	}
+
 	FStageSwitch StageSwitch = MakePrototype<FStageSwitch>(PrototypeHw.Type);
 	StageSwitch.SetArgs(EmitPrototype(*this, PrototypeHw), EmitPrototype(*this, PrototypeAn));
 
@@ -2355,6 +2392,11 @@ FValueRef FEmitter::TextureSampleLevel(FValueRef Texture, FValueRef TexCoord, FV
 		Prototype.MipValue = Operator(BO_Add, MipLevel, ExternalInput(EExternalInput::ViewMaterialTextureMipBias));
 	}
 
+	if (IsVirtualSamplerType(SamplerType))
+	{
+		Prototype.VTPageTable = VTPageTableLoadFromSamplerSource(*this, Texture, SamplerSourceMode, Prototype.TexCoord, {}, {});
+	}
+
 	return EmitPrototype(*this, Prototype);
 }
 
@@ -2385,6 +2427,12 @@ FValueRef FEmitter::TextureSampleBias(FValueRef Texture, FValueRef TexCoord, FVa
 	PrototypeAn.TexCoordDdx = Operator(BO_Multiply, AnalyticalPartialDerivative(TexCoord, EDerivativeAxis::X), Exp2MipBias);
 	PrototypeAn.TexCoordDdy = Operator(BO_Multiply, AnalyticalPartialDerivative(TexCoord, EDerivativeAxis::Y), Exp2MipBias);	
 
+	if (IsVirtualSamplerType(SamplerType))
+	{
+		PrototypeHw.VTPageTable = VTPageTableLoadFromSamplerSource(*this, Texture, SamplerSourceMode, PrototypeHw.TexCoord, {}, {});
+		PrototypeAn.VTPageTable = VTPageTableLoadFromSamplerSource(*this, Texture, SamplerSourceMode, PrototypeAn.TexCoord, PrototypeAn.TexCoordDdx, PrototypeAn.TexCoordDdy);
+	}
+
 	FStageSwitch StageSwitch = MakePrototype<FStageSwitch>(PrototypeHw.Type);
 	StageSwitch.SetArgs(EmitPrototype(*this, PrototypeHw), EmitPrototype(*this, PrototypeAn));
 
@@ -2414,6 +2462,25 @@ FValueRef FEmitter::TextureSampleGrad(FValueRef Texture, FValueRef TexCoord, FVa
 		Prototype.TexCoordDdy = Operator(BO_Multiply, Prototype.TexCoordDdy, ViewMaterialTextureDerivativeMultiply);
 	}
 
+	if (IsVirtualSamplerType(SamplerType))
+	{
+		Prototype.VTPageTable = VTPageTableLoadFromSamplerSource(*this, Texture, SamplerSourceMode, Prototype.TexCoord, Prototype.TexCoordDdx, Prototype.TexCoordDdy);
+	}
+
+	return EmitPrototype(*this, Prototype);
+}
+
+FValueRef FEmitter::VTPageTableLoad(FValueRef Texture, TextureAddress AddressU, TextureAddress AddressV, FValueRef TexCoord, FValueRef TexCoordDdx, FValueRef TexCoordDdy, bool bEnableFeedback)
+{
+	// VT stack and layer indices are initializes during IR analysis
+	FVTPageTableRead Prototype = MakePrototype<MIR::FVTPageTableRead>(FObjectType::GetVTPageTableResult());
+	Prototype.TextureObject = Texture;
+	Prototype.TexCoord = TexCoord;
+	Prototype.TexCoordDdx = TexCoordDdx;
+	Prototype.TexCoordDdy = TexCoordDdy;
+	Prototype.AddressU = AddressU;
+	Prototype.AddressV = AddressV;
+	Prototype.bEnableFeedback = bEnableFeedback;
 	return EmitPrototype(*this, Prototype);
 }
 
@@ -2828,6 +2895,12 @@ FValueRef FEmitter::InlineHLSL(const FMaterialExternalCodeDeclaration* InExterna
 	return EmitInlineHLSL(*this, ReturnType, InExternalCodeDeclaration, nullptr, InputValues, ValueFlags, UsedGraphProperties);
 }
 
+FValueRef FEmitter::PromoteSubstrateParameter()
+{
+	MIR::FPromoteSubstrateParameter Prototype = MakePrototype<MIR::FPromoteSubstrateParameter>(FObjectType::GetSubstrateData());
+	return EmitPrototype(*this, Prototype);
+}
+
 const FType* FEmitter::TryGetCommonType(const FType* A, const FType* B)
 {
 	// Trivial case: types are equal
diff --git a/Engine/Source/Runtime/Engine/Private/Materials/MaterialIRInternal.cpp b/Engine/Source/Runtime/Engine/Private/Materials/MaterialIRInternal.cpp
--- a/Engine/Source/Runtime/Engine/Private/Materials/MaterialIRInternal.cpp
+++ b/Engine/Source/Runtime/Engine/Private/Materials/MaterialIRInternal.cpp
@@ -78,7 +78,7 @@ MIR::FValue* CreateMaterialAttributeDefaultValue(FEmitter& Emitter, const UMater
 		case MCT_UInt:
 		case MCT_UInt4: return Emitter.ConstantInt4({ (TInteger)DefaultValue.X, (TInteger)DefaultValue.Y, (TInteger)DefaultValue.Z, (TInteger)DefaultValue.W });
 
-		case MCT_Substrate: return Emitter.SubstrateParameter();
+		case MCT_Substrate: return Emitter.PromoteSubstrateParameter();
 
 		default: UE_MIR_UNREACHABLE();
 	}
@@ -93,6 +93,7 @@ EMaterialTextureParameterType TextureMaterialValueTypeToParameterType(EMaterialV
 		case MCT_TextureCube: return EMaterialTextureParameterType::Cube;
 		case MCT_TextureCubeArray: return EMaterialTextureParameterType::ArrayCube;
 		case MCT_VolumeTexture: return EMaterialTextureParameterType::Volume;
+		case MCT_TextureVirtual: return EMaterialTextureParameterType::Virtual;
 		default: UE_MIR_UNREACHABLE();
 	}
 }
diff --git a/Engine/Source/Runtime/Engine/Private/Materials/MaterialIRModuleBuilder.cpp b/Engine/Source/Runtime/Engine/Private/Materials/MaterialIRModuleBuilder.cpp
--- a/Engine/Source/Runtime/Engine/Private/Materials/MaterialIRModuleBuilder.cpp
+++ b/Engine/Source/Runtime/Engine/Private/Materials/MaterialIRModuleBuilder.cpp
@@ -401,12 +401,20 @@ struct FMaterialIRModuleBuilderImpl
 				// Connect output properties with substrate parameter if a FrontMaterial property is provided
 				if (Output->Property == MP_FrontMaterial)
 				{
-					if (MIR::FSubstrateParameter* SubstrateParameter = Output->Arg->As<MIR::FSubstrateParameter>())
+					if (MIR::FPromoteSubstrateParameter* SubstrateParameter = Output->Arg->As<MIR::FPromoteSubstrateParameter>())
 					{
 						for (int32 PropertyIndex = 0; PropertyIndex < MP_MAX; ++PropertyIndex)
 						{
 							SubstrateParameter->PropertyArgs[PropertyIndex] = FindIRModuleOutputValue(*Module, MIR::EStage::Stage_Pixel, static_cast<EMaterialProperty>(PropertyIndex));
 						}
+
+						// Store separate value for normal vector transformed into world-space for Substrate
+						MIR::FValue* TangentSpaceNormal = SubstrateParameter->PropertyArgs[MP_Normal];
+						MIR::FValue* WorldSpaceNormal = Emitter->InlineHLSL(TangentSpaceNormal->Type, TEXT("TransformTangentNormalToWorld(Parameters.TangentToWorld, $0)"), { TangentSpaceNormal });
+						SubstrateParameter->WorldSpaceTangentsAndNormals[0] = WorldSpaceNormal;
+						SubstrateParameter->WorldSpaceTangentsAndNormals[1] = WorldSpaceNormal;
+						SubstrateParameter->WorldSpaceTangentsAndNormals[2] = WorldSpaceNormal;
+						SubstrateParameter->WorldSpaceTangentsAndNormals[3] = WorldSpaceNormal;
 					}
 				}
 
@@ -520,10 +528,27 @@ struct FMaterialIRModuleBuilderImpl
 		KeepDefineConditionally(TEXT("NEEDS_PER_INSTANCE_RANDOM_PS"), ValueAnalyzer.Material->bUsedWithInstancedStaticMeshes);
 		KeepDefineConditionally(TEXT("USES_PER_INSTANCE_FADE_AMOUNT"), ValueAnalyzer.Material->bUsedWithInstancedStaticMeshes);
 
+		// Derive additional defines from the final state of the IR analysis
+		DeriveEnvironmentDefinesFromAnalysisState();
+
 		// Move final environemnt defines from analyzer into output module
 		Module->EnvironmentDefines = MoveTemp(ValueAnalyzer.EnvironmentDefines);
 	}
 
+	void DeriveEnvironmentDefinesFromAnalysisState()
+	{
+		// Virtual texturing defines
+		int32 NumVirtualTextureFeedbackRequests = 0;
+		for (const FMaterialIRValueAnalyzer::FVTStackEntry& VTStack : ValueAnalyzer.VTStacks)
+		{
+			if (VTStack.bGenerateFeedback)
+			{
+				++NumVirtualTextureFeedbackRequests;
+			}
+		}
+		Module->CompilationOutput.NumVirtualTextureFeedbackRequests = NumVirtualTextureFeedbackRequests;
+	}
+
 	void Step_AnalyzeBuiltinDefines()
 	{
 		// Match various defines against the material configuration
diff --git a/Engine/Source/Runtime/Engine/Private/Materials/MaterialIRToHLSLTranslator.cpp b/Engine/Source/Runtime/Engine/Private/Materials/MaterialIRToHLSLTranslator.cpp
--- a/Engine/Source/Runtime/Engine/Private/Materials/MaterialIRToHLSLTranslator.cpp
+++ b/Engine/Source/Runtime/Engine/Private/Materials/MaterialIRToHLSLTranslator.cpp
@@ -6,6 +6,7 @@
 #include "Materials/MaterialIR.h"
 #include "Materials/MaterialExternalCodeRegistry.h"
 #include "MaterialIRInternal.h"
+#include "MaterialSharedPrivate.h"
 
 #include "ShaderCore.h"
 #include "MaterialShared.h"
@@ -192,11 +193,31 @@ static const TCHAR* GetShadingModelParameterName(EMaterialShadingModel InModel)
 	}
 }
 
-static bool IsFoldable(const MIR::FInstruction* Instr, MIR::EStage Stage)
+static bool InstructionUsesPhiValue(const MIR::FInstruction* Instr, MIR::EStage Stage)
 {
 	if (auto Branch = Instr->As<MIR::FBranch>())
 	{
-		return !Branch->TrueBlock[Stage].Instructions && !Branch->FalseBlock[Stage].Instructions;
+		return Branch->TrueBlock[Stage].Instructions
+			|| Branch->FalseBlock[Stage].Instructions;
+	}
+	return false;
+}
+
+static bool IsFoldable(const MIR::FInstruction* Instr, MIR::EStage Stage)
+{
+	// Instructions that use Phi values cannot be folded/inlined.
+	if (InstructionUsesPhiValue(Instr, Stage))
+	{
+		return false;
+	}
+
+	// Don't fold instructions that become very long to improve readability
+	switch (Instr->Kind)
+	{
+	case MIR::VK_VTPageTableRead:
+		return false;
+	default:
+		break;
 	}
 
 	return true;
@@ -277,7 +298,7 @@ struct FTranslator : FMaterialIRToHLSLTranslation
 			if (Substrate::IsSubstrateEnabled())
 			{
 				Printer << TEXT("// Substrate normal and tangent") << NewLine;
-				Printer << TEXT("Parameters.SharedLocalBases.Normals[0] = Parameters.TangentToWorld[2]") << EndOfStatement;
+				Printer << TEXT("Parameters.SharedLocalBases.Normals[0] = TransformTangentNormalToWorld(Parameters.TangentToWorld, PixelMaterialInputs.Normal)") << EndOfStatement;
 				Printer << TEXT("#if MATERIAL_TANGENTSPACENORMAL") << NewLine;
 				Printer << TEXT("Parameters.SharedLocalBases.Normals[0] *= Parameters.TwoSidedSign") << EndOfStatement;
 				Printer << TEXT("#endif") << NewLine;
@@ -322,9 +343,9 @@ struct FTranslator : FMaterialIRToHLSLTranslation
                 Printer << LowerType(Instr->Type) << TEXT(" ") << LocalStr;
 
                 LocalIdentifier.Add(Instr, MoveTemp(LocalStr));
-                if (IsFoldable(Instr, CurrentStage))
+                if (!InstructionUsesPhiValue(Instr, CurrentStage))
 				{
-                    Printer << TEXT(" = ";)
+                    Printer << TEXT(" = ");
                 }
             }
 
@@ -375,25 +396,26 @@ struct FTranslator : FMaterialIRToHLSLTranslation
 			case MIR::VK_ExternalInput: LowerExternalInput(static_cast<const MIR::FExternalInput*>(InValue)); break;
 			case MIR::VK_TextureObject: LowerTextureObject(static_cast<const MIR::FTextureObject*>(InValue)); break;
 			case MIR::VK_UniformParameter: LowerUniformParameter(static_cast<const MIR::FUniformParameter*>(InValue)); break;
-			case MIR::VK_SubstrateParameter: LowerSubstrateParameter(static_cast<const MIR::FSubstrateParameter*>(InValue)); break;
 			default: UE_MIR_UNREACHABLE();
 		}
 
 		return NoOp;
 	}
 
-	ENoOp LowerInstruction(MIR::FInstruction* Instr)
+	ENoOp LowerInstruction(const MIR::FInstruction* Instr)
 	{
 		switch (Instr->Kind)
 		{
-			case MIR::VK_Dimensional: LowerDimensional(static_cast<MIR::FDimensional*>(Instr)); break;
-			case MIR::VK_SetMaterialOutput: LowerSetMaterialOutput(static_cast<MIR::FSetMaterialOutput*>(Instr)); break;
-			case MIR::VK_Operator: LowerOperator(static_cast<MIR::FOperator*>(Instr)); break;
-			case MIR::VK_Branch: LowerBranch(static_cast<MIR::FBranch*>(Instr)); break;
-			case MIR::VK_Subscript: LowerSubscript(static_cast<MIR::FSubscript*>(Instr)); break;
-			case MIR::VK_Scalar: LowerScalar(static_cast<MIR::FScalar*>(Instr)); break;
-			case MIR::VK_TextureRead: LowerTextureRead(static_cast<MIR::FTextureRead*>(Instr)); break;
+			case MIR::VK_Dimensional: LowerDimensional(static_cast<const MIR::FDimensional*>(Instr)); break;
+			case MIR::VK_SetMaterialOutput: LowerSetMaterialOutput(static_cast<const MIR::FSetMaterialOutput*>(Instr)); break;
+			case MIR::VK_Operator: LowerOperator(static_cast<const MIR::FOperator*>(Instr)); break;
+			case MIR::VK_Branch: LowerBranch(static_cast<const MIR::FBranch*>(Instr)); break;
+			case MIR::VK_Subscript: LowerSubscript(static_cast<const MIR::FSubscript*>(Instr)); break;
+			case MIR::VK_Scalar: LowerScalar(static_cast<const MIR::FScalar*>(Instr)); break;
+			case MIR::VK_TextureRead: LowerTextureRead(static_cast<const MIR::FTextureRead*>(Instr)); break;
+			case MIR::VK_VTPageTableRead: LowerVTPageTableRead(static_cast<const MIR::FVTPageTableRead*>(Instr)); break;
 			case MIR::VK_InlineHLSL: LowerInlineHLSL(static_cast<const MIR::FInlineHLSL*>(Instr)); break;
+			case MIR::VK_PromoteSubstrateParameter: LowerPromoteSubstrateParameter(static_cast<const MIR::FPromoteSubstrateParameter*>(Instr)); break;
 			case MIR::VK_StageSwitch: LowerStageSwitch(static_cast<const MIR::FStageSwitch*>(Instr)); break;
 			case MIR::VK_HardwarePartialDerivative: LowerHardwarePartialDerivative(static_cast<const MIR::FHardwarePartialDerivative*>(Instr)); break;
 			default: UE_MIR_UNREACHABLE();
@@ -643,9 +665,9 @@ struct FTranslator : FMaterialIRToHLSLTranslation
 			{ MIR::FPrimitiveType::GetFloat3(), TEXT("ColorScaleBehindWater"), MP_MAX, TEXT("(float3)1.0f") },
 			{ MIR::FPrimitiveType::GetInt(), TEXT("ShadingModel"), MP_MAX, TEXT("1") },
 			{ MIR::FPrimitiveType::GetFloat3(), TEXT("RawNormal"), MP_Normal, TEXT("Parameters.TangentToWorld[2]") },
-			{ MIR::FPrimitiveType::GetFloat3(), TEXT("RawTangent"), MP_Tangent, TEXT("Parameters.TangentToWorld[2]") },
+			{ MIR::FPrimitiveType::GetFloat3(), TEXT("RawTangent"), MP_Tangent, TEXT("Parameters.TangentToWorld[0]") },
 			{ MIR::FPrimitiveType::GetFloat3(), TEXT("RawClearCoatNormal"), MP_Normal, TEXT("Parameters.TangentToWorld[2]") },
-			{ MIR::FPrimitiveType::GetFloat3(), TEXT("RawCustomTangent"), MP_Tangent, TEXT("Parameters.TangentToWorld[2]") },
+			{ MIR::FPrimitiveType::GetFloat3(), TEXT("RawCustomTangent"), MP_Tangent, TEXT("Parameters.TangentToWorld[0]") },
 			{ MIR::FPrimitiveType::GetInt(), TEXT("SharedLocalBasisIndex"), MP_MAX, nullptr },
 			{ MIR::FPrimitiveType::GetInt(), TEXT("ClearCoatBottomNormal_SharedLocalBasisIndex"), MP_MAX, nullptr },
 			{ MIR::FPrimitiveType::GetInt(), TEXT("SharedLocalBasisTypes"), MP_MAX, TEXT("Parameters.SharedLocalBases.Types") },
@@ -664,7 +686,7 @@ struct FTranslator : FMaterialIRToHLSLTranslation
 		return LegacyArguments;
 	}
 
-	void LowerSubstrateParameter(const MIR::FSubstrateParameter* SubstrateParameter)
+	void LowerPromoteSubstrateParameter(const MIR::FPromoteSubstrateParameter* SubstrateParameter)
 	{
 		Printer << TEXT("Parameters.SubstrateTree.PromoteParameterBlendedBSDFToOperator") << BeginArgs;
 		Printer.Tabs++;
@@ -678,34 +700,53 @@ struct FTranslator : FMaterialIRToHLSLTranslation
 		{
 			Printer << ListSeparator << NewLine << TEXT("/*") << Argument.ParameterName << TEXT("*/ ");
 
-			if (Argument.MaterialProperty != MP_MAX && SubstrateParameter->PropertyArgs[Argument.MaterialProperty])
-			{
-				LowerValue(SubstrateParameter->PropertyArgs[Argument.MaterialProperty]);
-			}
-			else if (Argument.DefaultInlineHLSL)
-			{
-				Printer << Argument.DefaultInlineHLSL;
-			}
-			else if (const MIR::FPrimitiveType* PrimitiveType = Argument.ParameterType->AsPrimitive())
+			switch (Argument.MaterialProperty)
 			{
-				if (const MIR::FPrimitiveType* ScalarType = Argument.ParameterType->AsScalar())
+				case MP_Normal:
 				{
-					Printer << TEXT("0.0f");
+					LowerValue(SubstrateParameter->WorldSpaceTangentsAndNormals[0]);
 				}
-				else
+				break;
+
+				case MP_Tangent:
 				{
-					LowerType(Argument.ParameterType);
-					Printer << BeginArgs;
-					for (int32 ComponentIndex = 0; ComponentIndex < PrimitiveType->NumRows; ++ComponentIndex)
+					LowerValue(SubstrateParameter->WorldSpaceTangentsAndNormals[1]);
+				}
+				break;
+
+				default:
+				{
+					if (Argument.MaterialProperty != MP_MAX && SubstrateParameter->PropertyArgs[Argument.MaterialProperty])
+					{
+						LowerValue(SubstrateParameter->PropertyArgs[Argument.MaterialProperty]);
+					}
+					else if (Argument.DefaultInlineHLSL)
+					{
+						Printer << Argument.DefaultInlineHLSL;
+					}
+					else if (const MIR::FPrimitiveType* PrimitiveType = Argument.ParameterType->AsPrimitive())
 					{
-						Printer << ListSeparator << TEXT("0.0f");
+						if (const MIR::FPrimitiveType* ScalarType = Argument.ParameterType->AsScalar())
+						{
+							Printer << TEXT("0.0f");
+						}
+						else
+						{
+							LowerType(Argument.ParameterType);
+							Printer << BeginArgs;
+							for (int32 ComponentIndex = 0; ComponentIndex < PrimitiveType->NumRows; ++ComponentIndex)
+							{
+								Printer << ListSeparator << TEXT("0.0f");
+							}
+							Printer << EndArgs;
+						}
+					}
+					else
+					{
+						UE_MIR_UNREACHABLE();
 					}
-					Printer << EndArgs;
 				}
-			}
-			else
-			{
-				UE_MIR_UNREACHABLE();
+				break;
 			}
 		}
 
@@ -816,7 +857,7 @@ struct FTranslator : FMaterialIRToHLSLTranslation
 				case MP_Specular:			Printer << TEXT("0"); break;
 				case MP_Roughness:			Printer << TEXT("0"); break;
 				case MP_Anisotropy:			Printer << TEXT("0"); break;
-				case MP_Normal:				Printer << TEXT("float3(0.0f, 0.0f, 1.0f)"); break;
+				case MP_Normal:				LowerValue(Output->Arg); break;
 				case MP_Tangent:			Printer << TEXT("0"); break;
 				case MP_SubsurfaceColor:	Printer << TEXT("0"); break;
 				case MP_AmbientOcclusion:	Printer << TEXT("1.0f"); break;
@@ -974,9 +1015,10 @@ struct FTranslator : FMaterialIRToHLSLTranslation
 		Printer << LowerType(Scalar->Type) << "(" << LowerValue(Scalar->Arg) << ")";
 	}
 
-	void LowerTextureRead(const MIR::FTextureRead* TextureRead)
+	void LowerStandardTextureRead(const MIR::FTextureRead* TextureRead)
 	{
-		bool bSamplerNeedsBrackets = LowerSamplerType(TextureRead->SamplerType);
+		bool bSamplerNeedsBrackets = false;
+		LowerSamplerType(TextureRead->SamplerType, bSamplerNeedsBrackets);
 		if (bSamplerNeedsBrackets)
 		{
 			Printer << TEXT("(");
@@ -1022,7 +1064,132 @@ struct FTranslator : FMaterialIRToHLSLTranslation
 		}
 	}
 
-	bool LowerSamplerType(EMaterialSamplerType SamplerType)
+	void LowerVirtualTextureRead(const MIR::FTextureRead* TextureRead)
+	{
+		checkf(
+			TextureRead->TextureObject->Type->AsObject()->ObjectKind == MIR::ObjectKind_Texture2D,
+			TEXT("Texture object type must be a 2D texture for virtual texture sampling in new material translator")
+		);
+
+		checkf(TextureRead->VTPageTable, TEXT("Missing page table for virtual texture read instruction"));
+		const MIR::FVTPageTableRead* VTPageTableResult = TextureRead->VTPageTable->As<MIR::FVTPageTableRead>();
+
+		const MIR::FTextureObject* TextureObject = TextureRead->TextureObject->As<MIR::FTextureObject>();
+
+		const int32 VirtualTextureIndex = TextureObject->Analysis_UniformParameterIndex;
+		//const int32 VTStackIndex = VTPageTableResult->VTStackIndex;
+		const int32 VTLayerIndex = VTPageTableResult->VTLayerIndex;
+
+		// Sampling function
+		Printer << TEXT("TextureVirtualSample");
+
+		// Sampling function argument list
+		FString SamplerExpression;
+		if (TextureRead->SamplerSourceMode != SSM_FromTextureAsset)
+		{
+			// VT doesn't care if the shared sampler is wrap or clamp. It only cares if it is aniso or not.
+			// The wrap/clamp/mirror operation is handled in the shader explicitly.
+			const TCHAR* SharedSamplerName = TextureRead->bUseAnisoSampler ? TEXT("View.SharedBilinearAnisoClampedSampler") : TEXT("View.SharedBilinearClampedSampler");
+			SamplerExpression = FString::Printf(TEXT("GetMaterialSharedSampler(Material.VirtualTexturePhysical_%dSampler, %s)"), VirtualTextureIndex, SharedSamplerName);
+		}
+		else
+		{
+			SamplerExpression = FString::Printf(TEXT("Material.VirtualTexturePhysical_%dSampler"), VirtualTextureIndex);
+		}
+
+		// VT page information
+		FString VTPageTableResultExpression;
+		FString VTPackedUniformName;
+
+		VTPackedUniformName = FString::Printf(TEXT("Material.VTPackedUniform[%d]"), VirtualTextureIndex);
+
+		// 'Texture name/sampler', 'PageTableResult', 'LayerIndex', 'PackedUniform'
+		Printer
+			<< BeginArgs
+			<< ListSeparator << LowerValue(TextureRead->TextureObject)
+			<< ListSeparator << SamplerExpression
+			<< ListSeparator << LowerValue(TextureRead->VTPageTable)
+			<< ListSeparator << VTLayerIndex
+			<< ListSeparator << TEXT("VTUniform_Unpack(") << VTPackedUniformName << TEXT(")")
+			<< EndArgs
+			;
+	}
+
+	void LowerTextureRead(const MIR::FTextureRead* TextureRead)
+	{
+		bool bSamplerNeedsBrackets = false;
+		LowerSamplerType(TextureRead->SamplerType, bSamplerNeedsBrackets);
+		if (bSamplerNeedsBrackets)
+		{
+			Printer << TEXT("(");
+		}
+
+		if (IsVirtualSamplerType(TextureRead->SamplerType))
+		{
+			LowerVirtualTextureRead(TextureRead);
+		}
+		else
+		{
+			LowerStandardTextureRead(TextureRead);
+		}
+
+		if (bSamplerNeedsBrackets)
+		{
+			Printer << TEXT(")");
+		}
+	}
+
+	void LowerVTPageTableRead(const MIR::FVTPageTableRead* VTPageTableRead)
+	{
+		using namespace UE::MaterialTranslatorUtils;
+
+		const TCHAR* VTPageTableReadFunctionName = TEXT("TextureLoadVirtualPageTable");
+
+		const bool bHasDerivativeTexCoords = VTPageTableRead->TexCoordDdx != nullptr && VTPageTableRead->TexCoordDdy != nullptr;
+
+		if (bHasDerivativeTexCoords)
+		{
+			Printer
+				<< VTPageTableReadFunctionName << TEXT("Grad")
+				<< BeginArgs
+				<< ListSeparator << FString::Printf(TEXT("VIRTUALTEXTURE_PAGETABLE_%d"), VTPageTableRead->VTStackIndex)
+				<< ListSeparator << FString::Printf(TEXT("VTPageTableUniform_Unpack(VIRTUALTEXTURE_PAGETABLE_UNIFORM_%d)"), VTPageTableRead->VTStackIndex)
+				<< ListSeparator << LowerValue(VTPageTableRead->TexCoord)
+				<< ListSeparator << GetVTAddressMode(VTPageTableRead->AddressU)
+				<< ListSeparator << GetVTAddressMode(VTPageTableRead->AddressV)
+				<< ListSeparator << LowerValue(VTPageTableRead->TexCoordDdx)
+				<< ListSeparator << LowerValue(VTPageTableRead->TexCoordDdy)
+				;
+		}
+		else
+		{
+			Printer
+				<< VTPageTableReadFunctionName
+				<< BeginArgs
+				<< ListSeparator << FString::Printf(TEXT("VIRTUALTEXTURE_PAGETABLE_%d"), VTPageTableRead->VTStackIndex)
+				<< ListSeparator << FString::Printf(TEXT("VTPageTableUniform_Unpack(VIRTUALTEXTURE_PAGETABLE_UNIFORM_%d)"), VTPageTableRead->VTStackIndex)
+				<< ListSeparator << LowerValue(VTPageTableRead->TexCoord)
+				<< ListSeparator << GetVTAddressMode(VTPageTableRead->AddressU)
+				<< ListSeparator << GetVTAddressMode(VTPageTableRead->AddressV)
+				;
+		}
+
+		// VT feedback
+		const bool bGenerateFeedback = VTPageTableRead->VTFeedbackRequestIndex != INDEX_NONE;
+
+		if (bGenerateFeedback)
+		{
+			Printer
+				<< ListSeparator << 0
+				<< ListSeparator << TEXT("Parameters.SvPosition.xy")
+				<< ListSeparator << VTPageTableRead->VTFeedbackRequestIndex << TEXT("u + LIGHTMAP_VT_ENABLED")
+				<< ListSeparator << TEXT("Parameters.VirtualTextureFeedback")
+				<< EndArgs
+				;
+		}
+	}
+
+	void LowerSamplerType(EMaterialSamplerType SamplerType, bool& bOutSamplerNeedsBrackets)
 	{
 		switch (SamplerType)
 		{
@@ -1068,13 +1235,14 @@ struct FTranslator : FMaterialIRToHLSLTranslation
 			case SAMPLERTYPE_Masks:
 			case SAMPLERTYPE_VirtualMasks:
 			case SAMPLERTYPE_Data:
-				return false;
+				bOutSamplerNeedsBrackets = false;
+				break;
 
 			default:
 				UE_MIR_UNREACHABLE();
 		}
 
-		return true;
+		bOutSamplerNeedsBrackets = true;
 	}
 
 	ENoOp LowerTextureSamplerReference(MIR::FValue* TextureValue, ESamplerSourceMode SamplerSource)
@@ -1110,6 +1278,7 @@ struct FTranslator : FMaterialIRToHLSLTranslation
 		switch (TextureType)
 		{
 			case MCT_Texture2D: Printer << TEXT("Texture2D_"); break;
+			case MCT_TextureVirtual: Printer << TEXT("VirtualTexturePhysical_"); break;
 			default: UE_MIR_UNREACHABLE();
 		}
 
@@ -1258,6 +1427,7 @@ struct FTranslator : FMaterialIRToHLSLTranslation
 			switch (ObjectType->ObjectKind)
 			{
 				case MIR::ObjectKind_SubstrateData: Printer << TEXT("FSubstrateData"); break;
+				case MIR::ObjectKind_VTPageTableResult: Printer << TEXT("VTPageTableResult"); break;
 				default: UE_MIR_UNREACHABLE(); break;
 			}
 		}
@@ -1299,9 +1469,11 @@ struct FTranslator : FMaterialIRToHLSLTranslation
 		OutEnvironment.SetDefine(TEXT("MATERIAL_FULLY_ROUGH"), false);
 		OutEnvironment.SetDefine(TEXT("MATERIAL_USES_ANISOTROPY"), false);
 		OutEnvironment.SetDefine(TEXT("MATERIAL_NEURAL_POST_PROCESS"), (CompilationOutput.bUsedWithNeuralNetworks || Material->IsUsedWithNeuralNetworks()) && Material->IsPostProcessMaterial());
-		OutEnvironment.SetDefine(TEXT("NUM_VIRTUALTEXTURE_SAMPLES"), 0);
-		OutEnvironment.SetDefine(TEXT("MATERIAL_VIRTUALTEXTURE_FEEDBACK"), false);
+		OutEnvironment.SetDefine(TEXT("NUM_VIRTUALTEXTURE_SAMPLES"), CompilationOutput.UniformExpressionSet.GetVTStacks().Num());
+		OutEnvironment.SetDefine(TEXT("NUM_VIRTUALTEXTURE_FEEDBACK_REQUESTS"), CompilationOutput.NumVirtualTextureFeedbackRequests);
+		OutEnvironment.SetDefine(TEXT("MATERIAL_VIRTUALTEXTURE_FEEDBACK"), CompilationOutput.NumVirtualTextureFeedbackRequests > 0);
 		OutEnvironment.SetDefine(TEXT("IS_MATERIAL_SHADER"), true);
+		OutEnvironment.SetDefine(TEXT("VIRTUAL_TEXTURE_OUTPUT"), CompilationOutput.bHasRuntimeVirtualTextureOutputNode != 0);
 
 		// Set all defines that are defined by the module.
 		// Any conditional exemption via material properties, such as 'Material->IsUsedWithInstancedStaticMeshes()', are handled during the material IR analysis.
@@ -1401,6 +1573,19 @@ struct FTranslator : FMaterialIRToHLSLTranslation
 			OutEnvironment.SetDefine(TEXT("SUBSTRATE_USE_FULLYSIMPLIFIED_MATERIAL"), false);
 			OutEnvironment.SetDefine(TEXT("SUBSTRATE_CLAMPED_CLOSURE_COUNT"), 1);
 		}
+
+		for (int32 VTStackIndex = 0; VTStackIndex < CompilationOutput.UniformExpressionSet.GetVTStacks().Num(); ++VTStackIndex)
+		{
+			const FMaterialVirtualTextureStack& VTStack = CompilationOutput.UniformExpressionSet.GetVTStack(VTStackIndex);
+
+			// Setup page table defines to map each VT stack to either 1 or 2 page table textures, depending on how many layers it uses
+			FString PageTableValue = FString::Printf(TEXT("Material.VirtualTexturePageTable0_%d"), VTStackIndex);
+			OutEnvironment.SetDefine(*FString::Printf(TEXT("VIRTUALTEXTURE_PAGETABLE_%d"), VTStackIndex), *PageTableValue);
+
+			// Setup page table uniform defines.
+			FString PageTableUniformValue = FString::Printf(TEXT("Material.VTPackedPageTableUniform[%d*2], Material.VTPackedPageTableUniform[%d*2+1]"), VTStackIndex, VTStackIndex);
+			OutEnvironment.SetDefine(*FString::Printf(TEXT("VIRTUALTEXTURE_PAGETABLE_UNIFORM_%d"), VTStackIndex), *PageTableUniformValue);
+		}
 	}
 };
 
diff --git a/Engine/Source/Runtime/Engine/Private/Materials/MaterialIRTypes.cpp b/Engine/Source/Runtime/Engine/Private/Materials/MaterialIRTypes.cpp
--- a/Engine/Source/Runtime/Engine/Private/Materials/MaterialIRTypes.cpp
+++ b/Engine/Source/Runtime/Engine/Private/Materials/MaterialIRTypes.cpp
@@ -129,6 +129,7 @@ FStringView FType::GetSpelling() const
 		{
 			case ObjectKind_Texture2D: return TEXT("Texture2D");
 			case ObjectKind_SubstrateData: return TEXT("SubstrateData");
+			case ObjectKind_VTPageTableResult: return TEXT("VTPageTableResult");
 		}
 	}
 	UE_MIR_UNREACHABLE();
@@ -216,7 +217,7 @@ bool FType::IsTexture() const
 {
 	if (const FObjectType* ObjectPtr = AsObject())
 	{
-		return true;
+		return ObjectPtr->ObjectKind == ObjectKind_Texture2D;
 	}
 	return false; 
 }
@@ -407,6 +408,12 @@ const FType* FObjectType::GetSubstrateData()
 	return &Instance;
 }
 
+const FType* FObjectType::GetVTPageTableResult()
+{
+	static const FObjectType Instance{ { TypeKind_Object }, ObjectKind_VTPageTableResult };
+	return &Instance;
+}
+
 } // namespace MIR
 
 #endif // #if WITH_EDITOR
diff --git a/Engine/Source/Runtime/Engine/Private/Materials/MaterialIRValueAnalyzer.cpp b/Engine/Source/Runtime/Engine/Private/Materials/MaterialIRValueAnalyzer.cpp
--- a/Engine/Source/Runtime/Engine/Private/Materials/MaterialIRValueAnalyzer.cpp
+++ b/Engine/Source/Runtime/Engine/Private/Materials/MaterialIRValueAnalyzer.cpp
@@ -75,6 +75,75 @@ static void AnalyzeInlineHLSL(FMaterialIRValueAnalyzer& Analyzer, MIR::FInlineHL
 	}
 }
 
+static int32 AcquireVTStackIndex(FMaterialIRValueAnalyzer& Analyzer, MIR::FVTPageTableRead& VTPageTableRead, int32 TextureUniformIndex, bool bGenerateFeedback)
+{
+	// Try to find matching VT stack entry
+	const uint64 CoordinateHash = CityHash64(reinterpret_cast<const char*>(&VTPageTableRead.TexCoord), sizeof(VTPageTableRead.TexCoord));
+
+	for (int32 VTStackIndex = 0; VTStackIndex < Analyzer.VTStacks.Num(); ++VTStackIndex)
+	{
+		const FMaterialIRValueAnalyzer::FVTStackEntry& VTStack = Analyzer.VTStacks[VTStackIndex];
+		if (VTStack.TexCoord == VTPageTableRead.TexCoord &&
+			VTStack.bGenerateFeedback == bGenerateFeedback &&
+			VTStack.AddressU == VTPageTableRead.AddressU &&
+			VTStack.AddressV == VTPageTableRead.AddressV)
+		{
+			return VTStackIndex;
+		}
+	}
+
+	// Add new VT stack entry
+	FMaterialIRValueAnalyzer::FVTStackEntry& VTStack = Analyzer.VTStacks.AddDefaulted_GetRef();
+	VTStack.TexCoord = VTPageTableRead.TexCoord;
+	VTStack.bGenerateFeedback = bGenerateFeedback;
+	VTStack.AddressU = VTPageTableRead.AddressU;
+	VTStack.AddressV = VTPageTableRead.AddressV;
+
+	return Analyzer.CompilationOutput->UniformExpressionSet.AddVTStack(INDEX_NONE);
+}
+
+static void PropagateStateInStageVTPageTableRead(FMaterialIRValueAnalyzer& Analyzer, MIR::FVTPageTableRead* VTPageTableRead, MIR::EStage Stage)
+{
+	check(VTPageTableRead->TextureObject);
+	MIR::FTextureObject* Texture = VTPageTableRead->TextureObject->As<MIR::FTextureObject>();
+	check(Texture);
+
+	const int32 TextureUniformIndex = (int32)Texture->Analysis_UniformParameterIndex;
+	check(TextureUniformIndex >= 0);
+
+	// Only support GPU feedback from pixel shader
+	const bool bGenerateFeedback = VTPageTableRead->bEnableFeedback && Stage == MIR::Stage_Pixel;
+
+	const int32 VTStackIndex = AcquireVTStackIndex(Analyzer, *VTPageTableRead, TextureUniformIndex, bGenerateFeedback);
+
+	VTPageTableRead->VTStackIndex = VTStackIndex;
+	VTPageTableRead->VTFeedbackRequestIndex = INDEX_NONE;
+
+	// Acquire VT layer
+	int32 VTLayerIndex = Analyzer.CompilationOutput->UniformExpressionSet.GetVTStack(VTStackIndex).FindLayer(TextureUniformIndex);
+	if (VTLayerIndex == INDEX_NONE)
+	{
+		VTLayerIndex = Analyzer.CompilationOutput->UniformExpressionSet.AddVTLayer(VTStackIndex, TextureUniformIndex);
+	}
+	VTPageTableRead->VTLayerIndex = VTLayerIndex;
+
+	// Acquire VT feedback request index
+	if (bGenerateFeedback)
+	{
+		// Exclude current stack
+		int32 VTFeedbackRequestIndex = 0;
+		for (int32 VTStackSearchIndex = 0; VTStackSearchIndex + 1 < Analyzer.VTStacks.Num(); ++VTStackSearchIndex)
+		{
+			const FMaterialIRValueAnalyzer::FVTStackEntry& VTStack = Analyzer.VTStacks[VTStackSearchIndex];
+			if (VTStack.bGenerateFeedback)
+			{
+				++VTFeedbackRequestIndex;
+			}
+		}
+		VTPageTableRead->VTFeedbackRequestIndex = VTFeedbackRequestIndex;
+	}
+}
+
 static void AnalyzeTextureObject(FMaterialIRValueAnalyzer& Analyzer, MIR::FTextureObject* TextureObject)
 {
 	EMaterialTextureParameterType ParamType = MIR::Internal::TextureMaterialValueTypeToParameterType(TextureObject->Texture->GetMaterialType());
@@ -285,6 +354,7 @@ void FMaterialIRValueAnalyzer::PropagateStateInStage(MIR::FValue* Value, MIR::ES
 	switch (Value->Kind)
 	{
 		EXPAND_CASE(InlineHLSL);
+		EXPAND_CASE(VTPageTableRead);
 		default: break;
 	}
 
diff --git a/Engine/Source/Runtime/Engine/Private/Materials/MaterialIRValueAnalyzer.h b/Engine/Source/Runtime/Engine/Private/Materials/MaterialIRValueAnalyzer.h
--- a/Engine/Source/Runtime/Engine/Private/Materials/MaterialIRValueAnalyzer.h
+++ b/Engine/Source/Runtime/Engine/Private/Materials/MaterialIRValueAnalyzer.h
@@ -15,6 +15,16 @@
 
 struct FMaterialIRValueAnalyzer
 {
+	// Translator specific structure during generation of VT stacks. This is later converted to FMaterialVirtualTextureStack entries in FUniformExpressionSet.
+	// The equivalent in the old translator is FHLSLMaterialTranslator::FMaterialVTStackEntry.
+	struct FVTStackEntry
+	{
+		MIR::FValue* TexCoord{};
+		bool bGenerateFeedback{};
+		TextureAddress AddressU{ TA_MAX };
+		TextureAddress AddressV{ TA_MAX };
+	};
+
 	//
 	void Setup(UMaterial* InMaterial, FMaterialIRModule* InModule, FMaterialCompilationOutput* InCompilationOutput, FMaterialInsights* InInsights = nullptr);
 		
@@ -44,6 +54,9 @@ struct FMaterialIRValueAnalyzer
 		
 	//
 	FMaterialInsights* Insights{}; // Optional
+
+	// VT stack entries during analysis. This is later converted to FMaterialVirtualTextureStack entries in FUniformExpressionSet.
+	TArray<FVTStackEntry> VTStacks;
 };
 
 #endif // #if WITH_EDITOR
diff --git a/Engine/Source/Runtime/Engine/Private/Materials/MaterialShader.cpp b/Engine/Source/Runtime/Engine/Private/Materials/MaterialShader.cpp
--- a/Engine/Source/Runtime/Engine/Private/Materials/MaterialShader.cpp
+++ b/Engine/Source/Runtime/Engine/Private/Materials/MaterialShader.cpp
@@ -542,6 +542,68 @@ void RecordMaterialDependenciesForCook(UE::Cook::FCookEventContext CookContext,
 
 } // namespace UE::MaterialInterface::Private
 
+namespace UE::MaterialTranslatorUtils
+{
+
+const TCHAR* GetVTAddressMode(TextureAddress Address)
+{
+	switch (Address)
+	{
+		case TA_Wrap: return TEXT("VTADDRESSMODE_WRAP");
+		case TA_Clamp: return TEXT("VTADDRESSMODE_CLAMP");
+		case TA_Mirror: return TEXT("VTADDRESSMODE_MIRROR");
+		default: checkNoEntry(); return nullptr;
+	}
+}
+
+void GetTextureAddressForSamplerSource(const UTexture* InTexture, ESamplerSourceMode InSamplerSource, TextureAddress& OutStaticAddressX, TextureAddress& OutStaticAddressY, TextureAddress& OutStaticAddressZ)
+{
+	// StaticAddress mode at time of compile
+	// This may not be 100% correct, if SamplerSource is set to 'SSM_FromTextureAsset', as texture parameter may change the address mode in a derived instance
+	OutStaticAddressX = TA_Wrap;
+	OutStaticAddressY = TA_Wrap;
+	OutStaticAddressZ = TA_Wrap;
+
+	if (InTexture && InTexture->Source.GetNumBlocks() > 1)
+	{
+		// UDIM (multi-block) texture are forced to use wrap address mode
+		// This is important for supporting VT stacks made from UDIMs with differing number of blocks, as this requires wrapping vAddress for certain layers
+		OutStaticAddressX = TA_Wrap;
+		OutStaticAddressY = TA_Wrap;
+		OutStaticAddressZ = TA_Wrap;
+	}
+	else
+	{
+		switch (InSamplerSource)
+		{
+			case SSM_FromTextureAsset:
+				if (InTexture)
+				{
+					OutStaticAddressX = InTexture->GetTextureAddressX();
+					OutStaticAddressY = InTexture->GetTextureAddressY();
+					OutStaticAddressZ = InTexture->GetTextureAddressZ();
+				}
+				break;
+			case SSM_Wrap_WorldGroupSettings:
+				OutStaticAddressX = TA_Wrap;
+				OutStaticAddressY = TA_Wrap;
+				OutStaticAddressZ = TA_Wrap;
+				break;
+			case SSM_Clamp_WorldGroupSettings:
+			case SSM_TerrainWeightmapGroupSettings:
+				OutStaticAddressX = TA_Clamp;
+				OutStaticAddressY = TA_Clamp;
+				OutStaticAddressZ = TA_Clamp;
+				break;
+			default:
+				checkNoEntry();
+				break;
+			}
+	}
+}
+
+} // namespace UE::MaterialTranslatorUtils
+
 #endif // WITH_EDITOR
 
 /** Called for every material shader to update the appropriate stats. */
diff --git a/Engine/Source/Runtime/Engine/Private/Materials/MaterialSharedPrivate.h b/Engine/Source/Runtime/Engine/Private/Materials/MaterialSharedPrivate.h
--- a/Engine/Source/Runtime/Engine/Private/Materials/MaterialSharedPrivate.h
+++ b/Engine/Source/Runtime/Engine/Private/Materials/MaterialSharedPrivate.h
@@ -6,6 +6,7 @@
 #include "Containers/Array.h"
 #include "Containers/ArrayView.h"
 #include "Containers/UnrealString.h"
+#include "Engine/TextureDefines.h"
 #include "MaterialShared.h"
 #include "RHIFeatureLevel.h"
 #include "RHIShaderPlatform.h"
@@ -49,6 +50,13 @@ struct FRecordedMaterialResourceForCooking
 
 } // namespace UE::MaterialInterface::Private
 
+/** Utility functions shared between the old and new material translator. */
+namespace UE::MaterialTranslatorUtils
+{
+	const TCHAR* GetVTAddressMode(TextureAddress Address);
+	void GetTextureAddressForSamplerSource(const UTexture* InTexture, ESamplerSourceMode InSamplerSource, TextureAddress& OutStaticAddressX, TextureAddress& OutStaticAddressY, TextureAddress& OutStaticAddressZ);
+} // /namespace UE::MaterialTranslatorUtils
+
 ///////////////////////////////////////////////////////
 // CompactBinary interface for FRecordedMaterialResourceForCooking
 ///////////////////////////////////////////////////////
diff --git a/Engine/Source/Runtime/Engine/Public/MaterialShared.h b/Engine/Source/Runtime/Engine/Public/MaterialShared.h
--- a/Engine/Source/Runtime/Engine/Public/MaterialShared.h
+++ b/Engine/Source/Runtime/Engine/Public/MaterialShared.h
@@ -770,6 +770,7 @@ public:
 		EstimatedNumVirtualTextureLookups(0),
 		NumUsedUVScalars(0),
 		NumUsedCustomInterpolatorScalars(0),
+		NumVirtualTextureFeedbackRequests(0),
 #endif
 		SubstrateMaterialCompilationOutput(),
 		UsedDBufferTextures(0),
@@ -886,6 +887,9 @@ public:
 	/** Number of used custom vertex interpolation scalars. */
 	LAYOUT_FIELD_EDITORONLY(uint8, NumUsedCustomInterpolatorScalars);
 
+	/** Number of feedback requests used for virtual texturing. */
+	LAYOUT_FIELD_EDITORONLY(uint8, NumVirtualTextureFeedbackRequests);
+
 	/** The Substrate material layout */
 	LAYOUT_FIELD(FSubstrateMaterialCompilationOutput, SubstrateMaterialCompilationOutput);
 
diff --git a/Engine/Source/Runtime/Engine/Public/Materials/MaterialIR.h b/Engine/Source/Runtime/Engine/Public/Materials/MaterialIR.h
--- a/Engine/Source/Runtime/Engine/Public/Materials/MaterialIR.h
+++ b/Engine/Source/Runtime/Engine/Public/Materials/MaterialIR.h
@@ -3,6 +3,7 @@
 
 #include "Materials/MaterialIRCommon.h"
 #include "Materials/MaterialParameters.h"
+#include "Engine/TextureDefines.h"
 
 #if WITH_EDITOR
 
@@ -89,7 +90,6 @@ enum EValueKind
 	VK_ExternalInput,
 	VK_TextureObject,
 	VK_UniformParameter,
-	VK_SubstrateParameter,
 
 	// Instructions
 
@@ -102,7 +102,9 @@ enum EValueKind
 	VK_Subscript,
 	VK_Scalar,
 	VK_TextureRead,
+	VK_VTPageTableRead,
 	VK_InlineHLSL,
+	VK_PromoteSubstrateParameter,
 	VK_StageSwitch,
 	VK_HardwarePartialDerivative,
 
@@ -425,14 +427,6 @@ struct FUniformParameter: TValue<VK_UniformParameter>
 	uint32 Analysis_UniformParameterIndex;
 };
 
-// Material IR representation of the FSubstrateData shader parameter to transition legacy materials into Substrate.
-struct FSubstrateParameter : TValue<VK_SubstrateParameter>
-{
-	FValue* PropertyArgs[MP_MAX];
-
-	bool bIsUnlit;
-};
-
 /*------------------------------------ Instructions ------------------------------------*/
 
 // A block of instructions.
@@ -555,6 +549,19 @@ struct FInlineHLSL : TInstruction<VK_InlineHLSL, 0>
 	};
 };
 
+// Material IR representation of the FSubstrateData shader parameter to transition legacy materials into Substrate.
+struct FPromoteSubstrateParameter : TInstruction<VK_PromoteSubstrateParameter, 4>
+{
+	// Extra IR values for tangent and normal vectors that are transformed into world-space specifically for Substrate conversion.
+	FValue* WorldSpaceTangentsAndNormals[4];
+
+	// IR values for all property arguments that forwarded to the Substrate conversion function.
+	FValue* PropertyArgs[MP_MAX];
+
+	// Specifies what Substrate BSDF to select, default or unlit.
+	bool bIsUnlit;
+};
+
 // Operator enumeration.
 //
 // Note: If you modify this enum, update the implementations of the helper functions below.
@@ -764,7 +771,7 @@ enum class ETextureReadMode
 const TCHAR* LexToString(ETextureReadMode Mode);
 
 // This instruction performs texture read operaation (sample or gather).
-struct FTextureRead : TInstruction<VK_TextureRead, 5>
+struct FTextureRead : TInstruction<VK_TextureRead, 6>
 {
 	// The texture object to sample.
 	FValue* TextureObject;
@@ -781,6 +788,9 @@ struct FTextureRead : TInstruction<VK_TextureRead, 5>
 	// Optional. The analytical partial derivative of the coordinates along the Y axis.
 	FValue* TexCoordDdy;
 
+	// Only required for virtual textures. Must point to FVTPageTableRead value.
+	FValue* VTPageTable;
+
 	// The mip value mode to use for sampling.
 	ETextureReadMode Mode : 8;
 
@@ -789,6 +799,42 @@ struct FTextureRead : TInstruction<VK_TextureRead, 5>
 
 	// The sampler type to use for sampling.
 	EMaterialSamplerType SamplerType : 8;
+
+	// Used for VTs whether anisotropic filtering is used or not.
+	bool bUseAnisoSampler : 1;
+};
+
+struct FVTPageTableRead : TInstruction<VK_VTPageTableRead, 4>
+{
+	// The texture object to load the page table for.
+	FValue* TextureObject;
+
+	// The texture coordinate at which to sample.
+	FValue* TexCoord;
+
+	// Optional. The analytical partial derivative of the coordinates along the X axis.
+	FValue* TexCoordDdx;
+
+	// Optional. The analytical partial derivative of the coordinates along the Y axis.
+	FValue* TexCoordDdy;
+
+	// Index of the virtual texture slot. This is determined in the analysis stage.
+	int32 VTStackIndex : 16;
+
+	// Index of the VT layer within the stack. This is determined in the analysis stage.
+	int32 VTLayerIndex : 16;
+
+	// Index of the VT feedback request. This is determined in the analysis stage.
+	int32 VTFeedbackRequestIndex : 16;
+
+	// Texture address mode for U axis.
+	TextureAddress AddressU : 8;
+
+	// Texture address mode for V axis.
+	TextureAddress AddressV : 8;
+
+	// Enables VT texture sampling feedback
+	bool bEnableFeedback : 1;
 };
 
 // Utility value for selecting a different value based on the execution stage.
diff --git a/Engine/Source/Runtime/Engine/Public/Materials/MaterialIREmitter.h b/Engine/Source/Runtime/Engine/Public/Materials/MaterialIREmitter.h
--- a/Engine/Source/Runtime/Engine/Public/Materials/MaterialIREmitter.h
+++ b/Engine/Source/Runtime/Engine/Public/Materials/MaterialIREmitter.h
@@ -313,9 +313,6 @@ public:
 
 	// Emits a custom primitive data value with given index.
 	FValueRef CustomPrimitiveData(uint32 PrimitiveDataIndex);
-
-	// Emits a parameter that encodes MCT_Substrate as IR value.
-	FValueRef SubstrateParameter();
 	
 	/*----------------------------- Instructions emission -----------------------------*/
 
@@ -422,11 +419,22 @@ public:
 	// @param Texture The texture to sample. Must be of type MIR::FObjectType and represent a valid texture.
 	// @param TexCoord The UV coordinates at which to sample the texture.
 	// @param TexCoordDdx The partial derivative of the texture coordinates along the screen space x axis.
-	// @param TexCoordDdx The partial derivative of the texture coordinates along the screen space y axis.
+	// @param TexCoordDdy The partial derivative of the texture coordinates along the screen space y axis.
 	// @param bAddViewMaterialMipBias Whether to add the view material texture mip bias to the mip level.
 	// @param SamplerSourceMode Which sampler source mode to use.
 	// @param SamplerType Which sampler type to use.
 	FValueRef TextureSampleGrad(FValueRef Texture, FValueRef TexCoord, FValueRef TexCoordDdx, FValueRef TexCoordDdy, bool bAddViewMaterialMipBias, ESamplerSourceMode SamplerSourceMode, EMaterialSamplerType SamplerType);
+
+	// Emits the virtual texture page load instruction.
+	//
+	// @param Texture The texture to sample. Must be of type MIR::FObjectType and represent a valid virtual texture.
+	// @param AddressU Texture address mode for the U-coordinate.
+	// @param AddressV Texture address mode for the V-coordinate.
+	// @param TexCoord The UV coordinates at which to sample the virtual texture.
+	// @param TexCoordDdx The partial derivative of the texture coordinates along the screen space x axis.
+	// @param TexCoordDdy The partial derivative of the texture coordinates along the screen space y axis.
+	// @param bEnableFeedback Specifies whether to enable virtual texture sampling feedback. This is only used for pixel shaders. Enabled by default.
+	FValueRef VTPageTableLoad(FValueRef Texture, TextureAddress AddressU, TextureAddress AddressV, FValueRef TexCoord, FValueRef TexCoordDdx = {}, FValueRef TexCoordDdy = {}, bool bEnableFeedback = true);
 	
 	// Emits the partial derivative of given value.
 	// This function differentiates the given value with respect to the given axis.
@@ -459,6 +467,9 @@ public:
 	// @param UsedGraphProperties Properties used during graph processing (default is None).
 	FValueRef InlineHLSL(const FMaterialExternalCodeDeclaration* InExternalCodeDeclaration, TConstArrayView<FValueRef> InputValues, EValueFlags ValueFlags = EValueFlags::None, EGraphProperties UsedGraphProperties = EGraphProperties::None);
 
+	// Emits an instruction that promotes a Substrate parameter, encoded as MCT_Substrate.
+	FValueRef PromoteSubstrateParameter();
+
 	/*--------------------------- Unary operators shortcuts ----------------------------*/
 
 	FValueRef BitwiseNot(FValueRef A) { return Operator(UO_BitwiseNot, A); }
diff --git a/Engine/Source/Runtime/Engine/Public/Materials/MaterialIRTypes.h b/Engine/Source/Runtime/Engine/Public/Materials/MaterialIRTypes.h
--- a/Engine/Source/Runtime/Engine/Public/Materials/MaterialIRTypes.h
+++ b/Engine/Source/Runtime/Engine/Public/Materials/MaterialIRTypes.h
@@ -23,6 +23,7 @@ enum EObjectKind
 {
 	ObjectKind_Texture2D,
 	ObjectKind_SubstrateData,
+	ObjectKind_VTPageTableResult,
 };
 
 const TCHAR* TypeKindToString(ETypeKind Kind);
@@ -182,6 +183,7 @@ struct FObjectType : FType
 
 	static const FType* GetTexture2D();
 	static const FType* GetSubstrateData();
+	static const FType* GetVTPageTableResult();
 };
 
 } // namespace MIR
