Add #include for FStringOutputDevice now that it is separated into its own header.

diff --git a/Engine/Plugins/Media/MediaFrameworkUtilities/Source/MediaFrameworkUtilitiesEditor/Private/AssetEditor/SMediaProfileSourcesTreeView.cpp b/Engine/Plugins/Media/MediaFrameworkUtilities/Source/MediaFrameworkUtilitiesEditor/Private/AssetEditor/SMediaProfileSourcesTreeView.cpp
--- a/Engine/Plugins/Media/MediaFrameworkUtilities/Source/MediaFrameworkUtilitiesEditor/Private/AssetEditor/SMediaProfileSourcesTreeView.cpp
+++ b/Engine/Plugins/Media/MediaFrameworkUtilities/Source/MediaFrameworkUtilitiesEditor/Private/AssetEditor/SMediaProfileSourcesTreeView.cpp
@@ -20,6 +20,7 @@
 #include "MediaAssets/ProxyMediaOutput.h"
 #include "MediaAssets/ProxyMediaSource.h"
 #include "Misc/MessageDialog.h"
+#include "Misc/StringOutputDevice.h"
 #include "Profile/MediaProfile.h"
 #include "Profile/MediaProfileSettings.h"
 #include "Styling/SlateIconFinder.h"
diff --git a/Engine/Source/Runtime/Engine/Private/UnrealExporter.cpp b/Engine/Source/Runtime/Engine/Private/UnrealExporter.cpp
--- a/Engine/Source/Runtime/Engine/Private/UnrealExporter.cpp
+++ b/Engine/Source/Runtime/Engine/Private/UnrealExporter.cpp
@@ -8,16 +8,17 @@
 #include "UnrealExporter.h"
 #include "UObject/UnrealType.h"
 #include "Exporters/Exporter.h"
+#include "Misc/AsciiSet.h"
+#include "Misc/FeedbackContext.h"
 #include "Misc/FileHelper.h"
-#include "Misc/Paths.h"
 #include "Misc/OutputDeviceFile.h"
+#include "Misc/Paths.h"
+#include "Misc/StringOutputDevice.h"
 #include "Serialization/BufferArchive.h"
 #include "UObject/UObjectIterator.h"
 #include "Model.h"
-#include "Misc/FeedbackContext.h"
 #include "AssetExportTask.h"
 #include "DataTableUtils.h"
-#include "Misc/AsciiSet.h"
 #include "Serialization/ArchiveSerializedPropertyChain.h"
 #include "UObject/GCObjectScopeGuard.h"
 #include "UObject/OverridableManager.h"
