[MIR] Fix static analysis warning by using explicitly pointer type in sizeof() expression.

#rnx

diff --git a/Engine/Source/Runtime/Engine/Private/Materials/MaterialIRValueAnalyzer.cpp b/Engine/Source/Runtime/Engine/Private/Materials/MaterialIRValueAnalyzer.cpp
--- a/Engine/Source/Runtime/Engine/Private/Materials/MaterialIRValueAnalyzer.cpp
+++ b/Engine/Source/Runtime/Engine/Private/Materials/MaterialIRValueAnalyzer.cpp
@@ -78,7 +78,7 @@ static void AnalyzeInlineHLSL(FMaterialIRValueAnalyzer& Analyzer, MIR::FInlineHL
 static int32 AcquireVTStackIndex(FMaterialIRValueAnalyzer& Analyzer, MIR::FVTPageTableRead& VTPageTableRead, int32 TextureUniformIndex, bool bGenerateFeedback)
 {
 	// Try to find matching VT stack entry
-	const uint64 CoordinateHash = CityHash64(reinterpret_cast<const char*>(&VTPageTableRead.TexCoord), sizeof(VTPageTableRead.TexCoord));
+	const uint64 CoordinateHash = CityHash64(reinterpret_cast<const char*>(&VTPageTableRead.TexCoord), sizeof(void*));
 
 	for (int32 VTStackIndex = 0; VTStackIndex < Analyzer.VTStacks.Num(); ++VTStackIndex)
 	{
