Cleaned up calls to GetMorphTargets.
- Removed unused MorphTargets param from FSkeletalMeshRenderData::InitResources.
- Removed unused MorphTargets param from FSkeletalMeshLODRenderData::InitResources.
- Removed unused bHasMorphTargets from FSkeletalMeshObjectGPUSkin::GetUsedVertexFactoryData.

diff --git a/Engine/Plugins/ChaosClothAsset/Source/ChaosClothAssetEngine/Private/ChaosClothAsset/ClothAssetBase.cpp b/Engine/Plugins/ChaosClothAsset/Source/ChaosClothAssetEngine/Private/ChaosClothAsset/ClothAssetBase.cpp
--- a/Engine/Plugins/ChaosClothAsset/Source/ChaosClothAssetEngine/Private/ChaosClothAsset/ClothAssetBase.cpp
+++ b/Engine/Plugins/ChaosClothAsset/Source/ChaosClothAssetEngine/Private/ChaosClothAsset/ClothAssetBase.cpp
@@ -443,8 +443,7 @@ void UChaosClothAssetBase::InitResources()
 
 	if (SkeletalMeshRenderData.IsValid())
 	{
-		TArray<UMorphTarget*> DummyMorphTargets;  // Not even used by InitResources() at the moment
-		SkeletalMeshRenderData->InitResources(GetHasVertexColors(), DummyMorphTargets, this);
+		SkeletalMeshRenderData->InitResources(GetHasVertexColors(), this);
 	}
 }
 
diff --git a/Engine/Source/Runtime/Engine/Private/SkeletalMesh.cpp b/Engine/Source/Runtime/Engine/Private/SkeletalMesh.cpp
--- a/Engine/Source/Runtime/Engine/Private/SkeletalMesh.cpp
+++ b/Engine/Source/Runtime/Engine/Private/SkeletalMesh.cpp
@@ -984,7 +984,7 @@ void USkeletalMesh::InitResources()
 
 		// TODO : Update RenderData->CurrentFirstLODIdx based on whether IStreamingManager::Get().IsRenderAssetStreamingEnabled(EStreamableRenderAssetType::SkeletalMesh).
 
-		SkelMeshRenderData->InitResources(GetHasVertexColors(), MutableView(GetMorphTargets()), this);
+		SkelMeshRenderData->InitResources(GetHasVertexColors(), this);
 		CachedSRRState.bHasPendingInitHint = true;
 
 		// For now in the editor force all LODs to stream to make sure tools have all LODs available
diff --git a/Engine/Source/Runtime/Engine/Private/SkeletalMeshLODRenderData.cpp b/Engine/Source/Runtime/Engine/Private/SkeletalMeshLODRenderData.cpp
--- a/Engine/Source/Runtime/Engine/Private/SkeletalMeshLODRenderData.cpp
+++ b/Engine/Source/Runtime/Engine/Private/SkeletalMeshLODRenderData.cpp
@@ -211,7 +211,7 @@ void FSkeletalMeshLODRenderData::InitMorphResources()
 	}
 }
 
-void FSkeletalMeshLODRenderData::InitResources(bool bNeedsVertexColors, int32 LODIndex, TArray<UMorphTarget*>& InMorphTargets, USkinnedAsset* Owner)
+void FSkeletalMeshLODRenderData::InitResources(bool bNeedsVertexColors, int32 LODIndex, USkinnedAsset* Owner)
 {
 	const FName OwnerName(USkinnedAsset::GetLODPathName(Owner, LODIndex));
 
@@ -252,12 +252,16 @@ void FSkeletalMeshLODRenderData::InitResources(bool bNeedsVertexColors, int32 LO
 		}
 	}
 
-	// Always make sure the morph target resources are in an initialized state. This should not get hit since the data should come with the load.
-	bool bNeedsMorphTargetRenderData = Owner && Owner->GetMorphTargets().Num() > 0 && !MorphTargetVertexInfoBuffers.IsMorphResourcesInitialized();
-	if (bNeedsMorphTargetRenderData)
+	if (Owner)
 	{
-		const FSkeletalMeshLODInfo* SkeletalMeshLODInfo = Owner->GetLODInfo(LODIndex);
-		MorphTargetVertexInfoBuffers.InitMorphResources(GMaxRHIShaderPlatform, RenderSections, Owner->GetMorphTargets(), StaticVertexBuffers.StaticMeshVertexBuffer.GetNumVertices(), LODIndex, SkeletalMeshLODInfo->MorphTargetPositionErrorTolerance);
+		const TArray<TObjectPtr<UMorphTarget>>& MorphTargets = Owner->GetMorphTargets();
+		// Always make sure the morph target resources are in an initialized state. This should not get hit since the data should come with the load.
+		bool bNeedsMorphTargetRenderData = !MorphTargets.IsEmpty() && !MorphTargetVertexInfoBuffers.IsMorphResourcesInitialized();
+		if (bNeedsMorphTargetRenderData)
+		{
+			const FSkeletalMeshLODInfo* SkeletalMeshLODInfo = Owner->GetLODInfo(LODIndex);
+			MorphTargetVertexInfoBuffers.InitMorphResources(GMaxRHIShaderPlatform, RenderSections, MorphTargets, StaticVertexBuffers.StaticMeshVertexBuffer.GetNumVertices(), LODIndex, SkeletalMeshLODInfo->MorphTargetPositionErrorTolerance);
+		}
 	}
 	
 	if (!MorphTargetVertexInfoBuffers.IsRHIInitialized() && MorphTargetVertexInfoBuffers.IsMorphCPUDataValid() && MorphTargetVertexInfoBuffers.NumTotalBatches > 0)
diff --git a/Engine/Source/Runtime/Engine/Private/SkeletalMeshRenderData.cpp b/Engine/Source/Runtime/Engine/Private/SkeletalMeshRenderData.cpp
--- a/Engine/Source/Runtime/Engine/Private/SkeletalMeshRenderData.cpp
+++ b/Engine/Source/Runtime/Engine/Private/SkeletalMeshRenderData.cpp
@@ -614,7 +614,7 @@ PRAGMA_DISABLE_DEPRECATION_WARNINGS
 PRAGMA_ENABLE_DEPRECATION_WARNINGS
 }
 
-void FSkeletalMeshRenderData::InitResources(bool bNeedsVertexColors, TArray<UMorphTarget*>& InMorphTargets, USkinnedAsset* Owner)
+void FSkeletalMeshRenderData::InitResources(bool bNeedsVertexColors, USkinnedAsset* Owner)
 {
 	if (!bInitialized)
 	{
@@ -636,7 +636,7 @@ void FSkeletalMeshRenderData::InitResources(bool bNeedsVertexColors, TArray<UMor
 
 			if(RenderData.GetNumVertices() > 0)
 			{
-				RenderData.InitResources(bNeedsVertexColors, LODIndex, InMorphTargets, Owner);
+				RenderData.InitResources(bNeedsVertexColors, LODIndex, Owner);
 			}
 		}
 
diff --git a/Engine/Source/Runtime/Engine/Private/SkeletalRenderGPUSkin.cpp b/Engine/Source/Runtime/Engine/Private/SkeletalRenderGPUSkin.cpp
--- a/Engine/Source/Runtime/Engine/Private/SkeletalRenderGPUSkin.cpp
+++ b/Engine/Source/Runtime/Engine/Private/SkeletalRenderGPUSkin.cpp
@@ -2270,7 +2270,6 @@ void FSkeletalMeshObjectGPUSkin::GetUsedVertexFactoryData(
 	USkinnedMeshComponent* SkinnedMeshComponent,
 	FSkelMeshRenderSection& RenderSection,
 	ERHIFeatureLevel::Type InFeatureLevel,
-	bool bHasMorphTargets,
 	FPSOPrecacheVertexFactoryDataList& VertexFactoryDataList)
 {
 	FSkeletalMeshLODRenderData& LODRenderData = SkelMeshRenderData->LODRenderData[LODIndex];
diff --git a/Engine/Source/Runtime/Engine/Private/SkeletalRenderGPUSkin.h b/Engine/Source/Runtime/Engine/Private/SkeletalRenderGPUSkin.h
--- a/Engine/Source/Runtime/Engine/Private/SkeletalRenderGPUSkin.h
+++ b/Engine/Source/Runtime/Engine/Private/SkeletalRenderGPUSkin.h
@@ -493,7 +493,7 @@ public:
 	virtual void RefreshClothingTransforms(const FMatrix& InNewLocalToWorld, uint32 FrameNumber) override;
 	virtual void UpdateSkinWeightBuffer(const TArrayView<const FSkelMeshComponentLODInfo> InLODInfo) override;
 
-	static void GetUsedVertexFactoryData(FSkeletalMeshRenderData* SkelMeshRenderData, int32 InLOD, USkinnedMeshComponent* SkinnedMeshComponent, FSkelMeshRenderSection& RenderSection, ERHIFeatureLevel::Type InFeatureLevel, bool bHasMorphTargets, FPSOPrecacheVertexFactoryDataList& VertexFactoryDataList);
+	static void GetUsedVertexFactoryData(FSkeletalMeshRenderData* SkelMeshRenderData, int32 InLOD, USkinnedMeshComponent* SkinnedMeshComponent, FSkelMeshRenderSection& RenderSection, ERHIFeatureLevel::Type InFeatureLevel, FPSOPrecacheVertexFactoryDataList& VertexFactoryDataList);
 
 	inline ESkeletalMeshGPUSkinTechnique GetGPUSkinTechnique(int32 LODIndex)
 	{
diff --git a/Engine/Source/Runtime/Engine/Private/SkinnedAsset.cpp b/Engine/Source/Runtime/Engine/Private/SkinnedAsset.cpp
--- a/Engine/Source/Runtime/Engine/Private/SkinnedAsset.cpp
+++ b/Engine/Source/Runtime/Engine/Private/SkinnedAsset.cpp
@@ -230,8 +230,7 @@ void USkinnedAsset::GetVertexFactoryTypesPerMaterialIndex(FPSOPrecacheVertexFact
 			else if (!SkelMeshRenderData->RequiresCPUSkinning(FeatureLevel, LODIndex))
 			{
 				// Add all the vertex factories which can be used for gpu skinning
-				bool bHasMorphTargets = GetMorphTargets().Num() > 0;
-				FSkeletalMeshObjectGPUSkin::GetUsedVertexFactoryData(SkelMeshRenderData, LODIndex, SkinnedMeshComponent, RenderSection, FeatureLevel, bHasMorphTargets, VFsPerMaterial->VertexFactoryDataList);
+				FSkeletalMeshObjectGPUSkin::GetUsedVertexFactoryData(SkelMeshRenderData, LODIndex, SkinnedMeshComponent, RenderSection, FeatureLevel, VFsPerMaterial->VertexFactoryDataList);
 			}
 		}
 	}
diff --git a/Engine/Source/Runtime/Engine/Public/Rendering/SkeletalMeshLODRenderData.h b/Engine/Source/Runtime/Engine/Public/Rendering/SkeletalMeshLODRenderData.h
--- a/Engine/Source/Runtime/Engine/Public/Rendering/SkeletalMeshLODRenderData.h
+++ b/Engine/Source/Runtime/Engine/Public/Rendering/SkeletalMeshLODRenderData.h
@@ -188,7 +188,7 @@ public:
 	*
 	* @param Parent Parent mesh
 	*/
-	void InitResources(bool bNeedsVertexColors, int32 LODIndex, TArray<class UMorphTarget*>& InMorphTargets, USkinnedAsset* Owner);
+	void InitResources(bool bNeedsVertexColors, int32 LODIndex, USkinnedAsset* Owner);
 
 	void InitMorphResources();
 
diff --git a/Engine/Source/Runtime/Engine/Public/Rendering/SkeletalMeshRenderData.h b/Engine/Source/Runtime/Engine/Public/Rendering/SkeletalMeshRenderData.h
--- a/Engine/Source/Runtime/Engine/Public/Rendering/SkeletalMeshRenderData.h
+++ b/Engine/Source/Runtime/Engine/Public/Rendering/SkeletalMeshRenderData.h
@@ -11,7 +11,6 @@ struct FMeshUVChannelInfo;
 class USkeletalMesh;
 class USkinnedAsset;
 struct FSkeletalMaterial;
-class UMorphTarget;
 struct FResourceSizeEx;
 
 class FSkeletalMeshRenderData
@@ -73,7 +72,7 @@ public:
 	ENGINE_API void Serialize(FArchive& Ar, USkinnedAsset* Owner);
 
 	/** Initializes rendering resources. */
-	ENGINE_API void InitResources(bool bNeedsVertexColors, TArray<UMorphTarget*>& InMorphTargets, USkinnedAsset* Owner);
+	ENGINE_API void InitResources(bool bNeedsVertexColors, USkinnedAsset* Owner);
 
 	/** Releases rendering resources. */
 	ENGINE_API void ReleaseResources();
