Do not create ODSC manager if it is not enabled, this saves significant time on Material creation/destruction in Test and Development builds

diff --git a/Engine/Source/Runtime/Engine/Private/ODSC/ODSCManager.cpp b/Engine/Source/Runtime/Engine/Private/ODSC/ODSCManager.cpp
--- a/Engine/Source/Runtime/Engine/Private/ODSC/ODSCManager.cpp
+++ b/Engine/Source/Runtime/Engine/Private/ODSC/ODSCManager.cpp
@@ -99,11 +99,11 @@ void FODSCManager::ResetCurrentPrimitiveSceneInfo()
 FODSCManager::FODSCManager()
 	: FTSTickerObjectBase(0.0f, FTSBackgroundableTicker::GetCoreTicker())
 {
-	FString Host;
-	const bool bODSCEnabled = FParse::Value(FCommandLine::Get(), TEXT("-odschost="), Host);
-
-	if (IsRunningCookOnTheFly() || bODSCEnabled)
+	if (IsODSCEnabled())
 	{
+		FString Host;
+		const bool bHostSet = FParse::Value(FCommandLine::Get(), TEXT("-odschost="), Host);
+
 		FCoreDelegates::OnEnginePreExit.AddRaw(this, &FODSCManager::OnEnginePreExit);
 		Thread = new FODSCThread(Host);
 		Thread->StartThread();
@@ -270,6 +270,18 @@ void FODSCManager::AddThreadedShaderPipelineRequest(
 #endif
 }
 
+bool FODSCManager::IsODSCEnabled()
+{
+	FString Host;
+	const bool bODSCEnabled = FParse::Value(FCommandLine::Get(), TEXT("-odschost="), Host);
+
+	if (IsRunningCookOnTheFly() || bODSCEnabled)
+	{
+		return true;
+	}
+	return false;
+}
+
 void FODSCManager::RegisterMaterialInstance(const UMaterialInstance* MaterialInstance)
 {
 	if (IsODSCActive() && MaterialInstance->HasAnyInternalFlags(EInternalObjectFlags_AsyncLoading))
diff --git a/Engine/Source/Runtime/Engine/Public/ODSC/ODSCManager.h b/Engine/Source/Runtime/Engine/Public/ODSC/ODSCManager.h
--- a/Engine/Source/Runtime/Engine/Public/ODSC/ODSCManager.h
+++ b/Engine/Source/Runtime/Engine/Public/ODSC/ODSCManager.h
@@ -110,6 +110,10 @@ public:
 	static void RegisterMaterialInstance(const UMaterialInstance* MI);
 	static void UnregisterMaterialInstance(const UMaterialInstance* MI);
 
+	/** Check options to see if ODSC could be active in this session */
+	static bool IsODSCEnabled();
+
+	/** Check if it is currently running */
 	static inline bool IsODSCActive();
 	static inline bool ShouldForceRecompile(const FMaterialShaderMap* MaterialShaderMap, const FMaterial* Material);
 
diff --git a/Engine/Source/Runtime/Launch/Private/LaunchEngineLoop.cpp b/Engine/Source/Runtime/Launch/Private/LaunchEngineLoop.cpp
--- a/Engine/Source/Runtime/Launch/Private/LaunchEngineLoop.cpp
+++ b/Engine/Source/Runtime/Launch/Private/LaunchEngineLoop.cpp
@@ -3016,7 +3016,10 @@ int32 FEngineLoop::PreInitPreStartupScreen(const TCHAR* CmdLine)
 
 #if WITH_ODSC
 	check(!GODSCManager);
-	GODSCManager = new FODSCManager();
+	if (FODSCManager::IsODSCEnabled())
+	{
+		GODSCManager = new FODSCManager();
+	}
 #endif
 
 	if (!FPlatformProperties::RequiresCookedData())
