Fix OnBlendComplete not being broadcast in Standalone (non-networked) games

diff --git a/Engine/Source/Runtime/Engine/Private/PlayerCameraManager.cpp b/Engine/Source/Runtime/Engine/Private/PlayerCameraManager.cpp
--- a/Engine/Source/Runtime/Engine/Private/PlayerCameraManager.cpp
+++ b/Engine/Source/Runtime/Engine/Private/PlayerCameraManager.cpp
@@ -154,7 +154,8 @@ void APlayerCameraManager::SetViewTarget(class AActor* NewTarget, struct FViewTa
 			AssignViewTarget(NewTarget, PendingViewTarget, TransitionParams);
 			PendingViewTarget.CheckViewTarget(PCOwner);
 
-			if (bUseClientSideCameraUpdates && GetNetMode() != NM_Client)
+			const ENetMode CurNetMode = GetNetMode();
+			if (bUseClientSideCameraUpdates && (CurNetMode != NM_Client && CurNetMode != NM_Standalone))
 			{
 				if (UWorld* World = GetWorld())
 				{
