Actor Coloration
* Apply same lighting and post process logic when choosing a sub category of this mode.

#rnx

diff --git a/Engine/Source/Runtime/Engine/Private/ShowFlags.cpp b/Engine/Source/Runtime/Engine/Private/ShowFlags.cpp
--- a/Engine/Source/Runtime/Engine/Private/ShowFlags.cpp
+++ b/Engine/Source/Runtime/Engine/Private/ShowFlags.cpp
@@ -329,6 +329,7 @@ void ApplyViewMode(EViewModeIndex ViewModeIndex, bool bPerspective, FEngineShowF
 		case VMI_RequiredTextureResolution:
 		case VMI_LODColoration:
 		case VMI_HLODColoration:
+		case VMI_VisualizeActorColoration:
 			bPostProcessing = false;
 			break;
 		case VMI_StationaryLightOverlap:
@@ -549,6 +550,7 @@ void EngineShowFlagOverride(EShowFlagInitMode ShowFlagInitMode, EViewModeIndex V
 			ViewModeIndex == VMI_RequiredTextureResolution ||
 			ViewModeIndex == VMI_LODColoration ||
 			ViewModeIndex == VMI_HLODColoration ||
+			ViewModeIndex == VMI_VisualizeActorColoration ||
 			ViewModeIndex == VMI_VisualizeGPUSkinCache ||
 			ViewModeIndex == VMI_LWCComplexity ||
 			ViewModeIndex == VMI_LightmapDensity)
