Fixed many log categories having a default verbosity higher than their compile-time verbosity.

diff --git a/Engine/Plugins/Experimental/PhysicsControl/Source/ControlRigPhysics/Private/RigPhysicsSimulation.h b/Engine/Plugins/Experimental/PhysicsControl/Source/ControlRigPhysics/Private/RigPhysicsSimulation.h
--- a/Engine/Plugins/Experimental/PhysicsControl/Source/ControlRigPhysics/Private/RigPhysicsSimulation.h
+++ b/Engine/Plugins/Experimental/PhysicsControl/Source/ControlRigPhysics/Private/RigPhysicsSimulation.h
@@ -28,7 +28,7 @@ struct FRigPhysicsControlComponent;
 struct FRigPhysicsJointComponent;
 
 #if UE_BUILD_SHIPPING || UE_BUILD_TEST
-DECLARE_LOG_CATEGORY_EXTERN(LogRigPhysics, Log, Warning);
+DECLARE_LOG_CATEGORY_EXTERN(LogRigPhysics, Warning, Warning);
 #else
 DECLARE_LOG_CATEGORY_EXTERN(LogRigPhysics, Log, All);
 #endif
diff --git a/Engine/Plugins/Experimental/PhysicsControl/Source/ControlRigPhysics/Public/ControlRigPhysicsModule.h b/Engine/Plugins/Experimental/PhysicsControl/Source/ControlRigPhysics/Public/ControlRigPhysicsModule.h
--- a/Engine/Plugins/Experimental/PhysicsControl/Source/ControlRigPhysics/Public/ControlRigPhysicsModule.h
+++ b/Engine/Plugins/Experimental/PhysicsControl/Source/ControlRigPhysics/Public/ControlRigPhysicsModule.h
@@ -7,7 +7,7 @@
 #include "Modules/ModuleManager.h"
 
 #if UE_BUILD_SHIPPING || UE_BUILD_TEST
-DECLARE_LOG_CATEGORY_EXTERN(LogControlRigPhysics, Log, Warning);
+DECLARE_LOG_CATEGORY_EXTERN(LogControlRigPhysics, Warning, Warning);
 #else
 DECLARE_LOG_CATEGORY_EXTERN(LogControlRigPhysics, Log, All);
 #endif
diff --git a/Engine/Plugins/Experimental/PhysicsControl/Source/PhysicsControl/Public/PhysicsControlLog.h b/Engine/Plugins/Experimental/PhysicsControl/Source/PhysicsControl/Public/PhysicsControlLog.h
--- a/Engine/Plugins/Experimental/PhysicsControl/Source/PhysicsControl/Public/PhysicsControlLog.h
+++ b/Engine/Plugins/Experimental/PhysicsControl/Source/PhysicsControl/Public/PhysicsControlLog.h
@@ -5,7 +5,7 @@
 #include "Logging/LogMacros.h"
 
 #if UE_BUILD_SHIPPING || UE_BUILD_TEST
-PHYSICSCONTROL_API DECLARE_LOG_CATEGORY_EXTERN(LogPhysicsControl, Log, Warning);
+PHYSICSCONTROL_API DECLARE_LOG_CATEGORY_EXTERN(LogPhysicsControl, Warning, Warning);
 #else
 PHYSICSCONTROL_API DECLARE_LOG_CATEGORY_EXTERN(LogPhysicsControl, Log, All);
 #endif
diff --git a/Engine/Source/Runtime/AnimGraphRuntime/Private/BoneControllers/AnimNode_RigidBody.cpp b/Engine/Source/Runtime/AnimGraphRuntime/Private/BoneControllers/AnimNode_RigidBody.cpp
--- a/Engine/Source/Runtime/AnimGraphRuntime/Private/BoneControllers/AnimNode_RigidBody.cpp
+++ b/Engine/Source/Runtime/AnimGraphRuntime/Private/BoneControllers/AnimNode_RigidBody.cpp
@@ -52,7 +52,7 @@ DEFINE_STAT(STAT_RigidBodyNodeInitTime_SetupSimulation);
 CSV_DECLARE_CATEGORY_MODULE_EXTERN(ENGINE_API, Animation);
 
 #if UE_BUILD_SHIPPING || UE_BUILD_TEST
-DECLARE_LOG_CATEGORY_EXTERN(LogRBAN, Log, Warning);
+DECLARE_LOG_CATEGORY_EXTERN(LogRBAN, Warning, Warning);
 #else
 DECLARE_LOG_CATEGORY_EXTERN(LogRBAN, Log, All);
 #endif
diff --git a/Engine/Source/Runtime/AnimGraphRuntime/Private/BoneControllers/AnimNode_SkeletalControlBase.cpp b/Engine/Source/Runtime/AnimGraphRuntime/Private/BoneControllers/AnimNode_SkeletalControlBase.cpp
--- a/Engine/Source/Runtime/AnimGraphRuntime/Private/BoneControllers/AnimNode_SkeletalControlBase.cpp
+++ b/Engine/Source/Runtime/AnimGraphRuntime/Private/BoneControllers/AnimNode_SkeletalControlBase.cpp
@@ -13,7 +13,7 @@
 #define LOCTEXT_NAMESPACE "AnimNode_SkeletalControlBase"
 
 #if UE_BUILD_SHIPPING || UE_BUILD_TEST
-	DECLARE_LOG_CATEGORY_EXTERN(LogSkeletalControlBase, Log, Warning);
+	DECLARE_LOG_CATEGORY_EXTERN(LogSkeletalControlBase, Warning, Warning);
 #else
 	DECLARE_LOG_CATEGORY_EXTERN(LogSkeletalControlBase, Log, All);
 #endif
diff --git a/Engine/Source/Runtime/Engine/Public/Audio/AudioTimingLog.h b/Engine/Source/Runtime/Engine/Public/Audio/AudioTimingLog.h
--- a/Engine/Source/Runtime/Engine/Public/Audio/AudioTimingLog.h
+++ b/Engine/Source/Runtime/Engine/Public/Audio/AudioTimingLog.h
@@ -6,10 +6,9 @@
 
 // Log AudioTiming. For verbose detailed timings diagnostics and debugging.
 #if UE_BUILD_SHIPPING
-	// Dial back the compiled verbosity in shipping
-	#define AUDIO_TIMING_COMPILED_VERBOSITY Error
+// Dial back the compiled verbosity in shipping
+ENGINE_API DECLARE_LOG_CATEGORY_EXTERN(LogAudioTiming, Error, Error);
 #else // UE_BUILD_SHIPPING
-	#define AUDIO_TIMING_COMPILED_VERBOSITY All
+ENGINE_API DECLARE_LOG_CATEGORY_EXTERN(LogAudioTiming, Log, All);
 #endif // UE_BUILD_SHIPPING
 
-ENGINE_API DECLARE_LOG_CATEGORY_EXTERN(LogAudioTiming, Log, AUDIO_TIMING_COMPILED_VERBOSITY);
diff --git a/Engine/Source/Runtime/Experimental/Chaos/Private/Chaos/Evolution/PBDMinEvolution.cpp b/Engine/Source/Runtime/Experimental/Chaos/Private/Chaos/Evolution/PBDMinEvolution.cpp
--- a/Engine/Source/Runtime/Experimental/Chaos/Private/Chaos/Evolution/PBDMinEvolution.cpp
+++ b/Engine/Source/Runtime/Experimental/Chaos/Private/Chaos/Evolution/PBDMinEvolution.cpp
@@ -23,7 +23,7 @@
 namespace Chaos
 {
 #if UE_BUILD_SHIPPING || UE_BUILD_TEST
-	CHAOS_API DECLARE_LOG_CATEGORY_EXTERN(LogChaosMinEvolution, Log, Warning);
+	CHAOS_API DECLARE_LOG_CATEGORY_EXTERN(LogChaosMinEvolution, Warning, Warning);
 #else
 	CHAOS_API DECLARE_LOG_CATEGORY_EXTERN(LogChaosMinEvolution, Log, All);
 #endif
diff --git a/Engine/Source/Runtime/Experimental/Chaos/Public/Chaos/Joint/ChaosJointLog.h b/Engine/Source/Runtime/Experimental/Chaos/Public/Chaos/Joint/ChaosJointLog.h
--- a/Engine/Source/Runtime/Experimental/Chaos/Public/Chaos/Joint/ChaosJointLog.h
+++ b/Engine/Source/Runtime/Experimental/Chaos/Public/Chaos/Joint/ChaosJointLog.h
@@ -5,7 +5,7 @@
 #include "CoreMinimal.h"
 
 #if UE_BUILD_SHIPPING || UE_BUILD_TEST
-CHAOS_API DECLARE_LOG_CATEGORY_EXTERN(LogChaosJoint, Log, Warning);
+CHAOS_API DECLARE_LOG_CATEGORY_EXTERN(LogChaosJoint, Warning, Warning);
 #else
 CHAOS_API DECLARE_LOG_CATEGORY_EXTERN(LogChaosJoint, Log, All);
 #endif
