Added Greater.h include to files which use TGreater, which will break with an upcoming change to remove that include from a header.

diff --git a/Engine/Plugins/Enterprise/DatasmithImporter/Source/DatasmithImporter/Private/DatasmithBlueprintLibrary.cpp b/Engine/Plugins/Enterprise/DatasmithImporter/Source/DatasmithImporter/Private/DatasmithBlueprintLibrary.cpp
--- a/Engine/Plugins/Enterprise/DatasmithImporter/Source/DatasmithImporter/Private/DatasmithBlueprintLibrary.cpp
+++ b/Engine/Plugins/Enterprise/DatasmithImporter/Source/DatasmithImporter/Private/DatasmithBlueprintLibrary.cpp
@@ -42,6 +42,7 @@
 #include "StaticMeshCompiler.h"
 #include "SourceUri.h"
 #include "Subsystems/AssetEditorSubsystem.h"
+#include "Templates/Greater.h"
 #include "Templates/SharedPointer.h"
 #include "UObject/Package.h"
 #include "UObject/StrongObjectPtr.h"
diff --git a/Engine/Plugins/Enterprise/DatasmithImporter/Source/DatasmithImporter/Private/UVTools/UVGenerationDataprepOperation.cpp b/Engine/Plugins/Enterprise/DatasmithImporter/Source/DatasmithImporter/Private/UVTools/UVGenerationDataprepOperation.cpp
--- a/Engine/Plugins/Enterprise/DatasmithImporter/Source/DatasmithImporter/Private/UVTools/UVGenerationDataprepOperation.cpp
+++ b/Engine/Plugins/Enterprise/DatasmithImporter/Source/DatasmithImporter/Private/UVTools/UVGenerationDataprepOperation.cpp
@@ -11,6 +11,7 @@
 #include "Engine/StaticMeshActor.h"
 #include "MeshDescription.h"
 #include "OverlappingCorners.h"
+#include "Templates/Greater.h"
 
 void UUVGenerationFlattenMappingOperation::OnExecution_Implementation(const FDataprepContext& InContext)
 {
diff --git a/Engine/Plugins/Experimental/MeshModelingToolsetExp/Source/MeshModelingToolsExp/Private/ExtractSplineTool.cpp b/Engine/Plugins/Experimental/MeshModelingToolsetExp/Source/MeshModelingToolsExp/Private/ExtractSplineTool.cpp
--- a/Engine/Plugins/Experimental/MeshModelingToolsetExp/Source/MeshModelingToolsExp/Private/ExtractSplineTool.cpp
+++ b/Engine/Plugins/Experimental/MeshModelingToolsetExp/Source/MeshModelingToolsExp/Private/ExtractSplineTool.cpp
@@ -43,6 +43,7 @@
 
 #include "CurveOps/GenerateCrossSectionOp.h"
 
+#include "Templates/Greater.h"
 
 #include UE_INLINE_GENERATED_CPP_BY_NAME(ExtractSplineTool)
 
diff --git a/Engine/Plugins/Experimental/NNERuntimeRDG/Source/NNERuntimeRDG/Private/Hlsl/NNERuntimeRDGReduce.cpp b/Engine/Plugins/Experimental/NNERuntimeRDG/Source/NNERuntimeRDG/Private/Hlsl/NNERuntimeRDGReduce.cpp
--- a/Engine/Plugins/Experimental/NNERuntimeRDG/Source/NNERuntimeRDG/Private/Hlsl/NNERuntimeRDGReduce.cpp
+++ b/Engine/Plugins/Experimental/NNERuntimeRDG/Source/NNERuntimeRDG/Private/Hlsl/NNERuntimeRDGReduce.cpp
@@ -5,6 +5,7 @@
 #include "Algo/Sort.h"
 #include "Helper/NNERuntimeRDGOperatorHelper.h"
 #include "Misc/EnumerateRange.h"
+#include "Templates/Greater.h"
 #include "NNEHlslShadersLog.h"
 #include "NNEHlslShadersReduceCS.h"
 #include "NNERuntimeRDGHlslHelper.h"
diff --git a/Engine/Plugins/MemoryUsageQueries/Source/MemoryUsageQueries/Private/MemoryUsageInfoProviderLLM.cpp b/Engine/Plugins/MemoryUsageQueries/Source/MemoryUsageQueries/Private/MemoryUsageInfoProviderLLM.cpp
--- a/Engine/Plugins/MemoryUsageQueries/Source/MemoryUsageQueries/Private/MemoryUsageInfoProviderLLM.cpp
+++ b/Engine/Plugins/MemoryUsageQueries/Source/MemoryUsageQueries/Private/MemoryUsageInfoProviderLLM.cpp
@@ -3,6 +3,7 @@
 #include "MemoryUsageInfoProviderLLM.h"
 #include "MemoryUsageQueries.h"
 #include "MemoryUsageQueriesPrivate.h"
+#include "Templates/Greater.h"
 
 #if ENABLE_LOW_LEVEL_MEM_TRACKER
 
diff --git a/Engine/Plugins/Runtime/GameplayAbilities/Source/GameplayAbilities/Private/Abilities/Tasks/AbilityTask.cpp b/Engine/Plugins/Runtime/GameplayAbilities/Source/GameplayAbilities/Private/Abilities/Tasks/AbilityTask.cpp
--- a/Engine/Plugins/Runtime/GameplayAbilities/Source/GameplayAbilities/Private/Abilities/Tasks/AbilityTask.cpp
+++ b/Engine/Plugins/Runtime/GameplayAbilities/Source/GameplayAbilities/Private/Abilities/Tasks/AbilityTask.cpp
@@ -4,6 +4,7 @@
 #include "AbilitySystemLog.h"
 #include "AbilitySystemComponent.h"
 #include "AbilitySystemStats.h"
+#include "Templates/Greater.h"
 
 #include UE_INLINE_GENERATED_CPP_BY_NAME(AbilityTask)
 
diff --git a/Engine/Plugins/Runtime/Metasound/Source/MetasoundStandardNodes/Private/MetasoundTriggerOnThresholdNode.cpp b/Engine/Plugins/Runtime/Metasound/Source/MetasoundStandardNodes/Private/MetasoundTriggerOnThresholdNode.cpp
--- a/Engine/Plugins/Runtime/Metasound/Source/MetasoundStandardNodes/Private/MetasoundTriggerOnThresholdNode.cpp
+++ b/Engine/Plugins/Runtime/Metasound/Source/MetasoundStandardNodes/Private/MetasoundTriggerOnThresholdNode.cpp
@@ -13,6 +13,7 @@
 #include "MetasoundTrigger.h"
 #include "MetasoundVertex.h"
 #include "MetasoundStandardNodesCategories.h"
+#include "Templates/Greater.h"
 
 #define LOCTEXT_NAMESPACE "MetasoundStandardNodes"
 
diff --git a/Engine/Plugins/Runtime/nDisplay/Source/DisplayCluster/Private/Render/Viewport/DisplayClusterViewport.cpp b/Engine/Plugins/Runtime/nDisplay/Source/DisplayCluster/Private/Render/Viewport/DisplayClusterViewport.cpp
--- a/Engine/Plugins/Runtime/nDisplay/Source/DisplayCluster/Private/Render/Viewport/DisplayClusterViewport.cpp
+++ b/Engine/Plugins/Runtime/nDisplay/Source/DisplayCluster/Private/Render/Viewport/DisplayClusterViewport.cpp
@@ -43,6 +43,8 @@
 
 #include "Misc/DisplayClusterLog.h"
 
+#include "Templates/Greater.h"
+
 int32 GDisplayClusterMultiGPUEnable = 1;
 static FAutoConsoleVariableRef CVarDisplayClusterMultiGPUEnable(
 	TEXT("DC.MultiGPU"),
diff --git a/Engine/Source/Developer/IoStoreUtilities/Private/IoStoreUtilities.cpp b/Engine/Source/Developer/IoStoreUtilities/Private/IoStoreUtilities.cpp
--- a/Engine/Source/Developer/IoStoreUtilities/Private/IoStoreUtilities.cpp
+++ b/Engine/Source/Developer/IoStoreUtilities/Private/IoStoreUtilities.cpp
@@ -88,6 +88,7 @@
 #include "Serialization/ZenPackageHeader.h"
 #include "String/ParseTokens.h"
 #include "Tasks/Task.h"
+#include "Templates/Greater.h"
 
 IMPLEMENT_MODULE(FDefaultModuleImpl, IoStoreUtilities);
 
diff --git a/Engine/Source/Developer/MaterialBaking/Private/MaterialBakingModule.cpp b/Engine/Source/Developer/MaterialBaking/Private/MaterialBakingModule.cpp
--- a/Engine/Source/Developer/MaterialBaking/Private/MaterialBakingModule.cpp
+++ b/Engine/Source/Developer/MaterialBaking/Private/MaterialBakingModule.cpp
@@ -31,6 +31,7 @@
 #if WITH_EDITOR
 #include "Misc/FileHelper.h"
 #endif
+#include "Templates/Greater.h"
 
 IMPLEMENT_MODULE(FMaterialBakingModule, MaterialBaking);
 
diff --git a/Engine/Source/Editor/CurveEditor/Private/SCurveEditorView.cpp b/Engine/Source/Editor/CurveEditor/Private/SCurveEditorView.cpp
--- a/Engine/Source/Editor/CurveEditor/Private/SCurveEditorView.cpp
+++ b/Engine/Source/Editor/CurveEditor/Private/SCurveEditorView.cpp
@@ -22,6 +22,7 @@
 #include "Misc/AssertionMacros.h"
 #include "SCurveEditorPanel.h"
 #include "Slate/SRetainerWidget.h"
+#include "Templates/Greater.h"
 #include "Templates/Tuple.h"
 #include "Templates/UnrealTemplate.h"
 
diff --git a/Engine/Source/Editor/PhysicsAssetEditor/Private/PhysicsAssetEditorSharedData.cpp b/Engine/Source/Editor/PhysicsAssetEditor/Private/PhysicsAssetEditorSharedData.cpp
--- a/Engine/Source/Editor/PhysicsAssetEditor/Private/PhysicsAssetEditorSharedData.cpp
+++ b/Engine/Source/Editor/PhysicsAssetEditor/Private/PhysicsAssetEditorSharedData.cpp
@@ -51,6 +51,7 @@
 #include "Factories.h"
 #include "HAL/PlatformApplicationMisc.h"
 #include "SPrimaryButton.h"
+#include "Templates/Greater.h"
 
 #define LOCTEXT_NAMESPACE "PhysicsAssetEditorShared"
 
diff --git a/Engine/Source/Editor/StaticMeshEditor/Private/StaticMeshEditorSubsystem.cpp b/Engine/Source/Editor/StaticMeshEditor/Private/StaticMeshEditorSubsystem.cpp
--- a/Engine/Source/Editor/StaticMeshEditor/Private/StaticMeshEditorSubsystem.cpp
+++ b/Engine/Source/Editor/StaticMeshEditor/Private/StaticMeshEditorSubsystem.cpp
@@ -35,6 +35,7 @@
 #include "Async/Async.h"
 #include "Misc/ScopedSlowTask.h"
 #include "Misc/FeedbackContext.h"
+#include "Templates/Greater.h"
 
 #include "StaticMeshResources.h"
 #include "UnrealEdGlobals.h"
diff --git a/Engine/Source/Editor/UnrealEd/Private/Commandlets/AssetRegistryGenerator.cpp b/Engine/Source/Editor/UnrealEd/Private/Commandlets/AssetRegistryGenerator.cpp
--- a/Engine/Source/Editor/UnrealEd/Private/Commandlets/AssetRegistryGenerator.cpp
+++ b/Engine/Source/Editor/UnrealEd/Private/Commandlets/AssetRegistryGenerator.cpp
@@ -52,6 +52,7 @@
 #include "String/Find.h"
 #include "String/ParseTokens.h"
 #include "TargetDomain/TargetDomainUtils.h"
+#include "Templates/Greater.h"
 #include "Templates/UniquePtr.h"
 #include "UObject/SoftObjectPath.h"
 #include "Logging/StructuredLog.h"
diff --git a/Engine/Source/Editor/UnrealEd/Private/Commandlets/DumpMaterialShaderTypes.cpp b/Engine/Source/Editor/UnrealEd/Private/Commandlets/DumpMaterialShaderTypes.cpp
--- a/Engine/Source/Editor/UnrealEd/Private/Commandlets/DumpMaterialShaderTypes.cpp
+++ b/Engine/Source/Editor/UnrealEd/Private/Commandlets/DumpMaterialShaderTypes.cpp
@@ -15,6 +15,7 @@
 #include "Materials/MaterialInstanceConstant.h"
 #include "Materials/MaterialExpressionParameter.h"
 #include "Materials/MaterialDependencyWalker.h"
+#include "Templates/Greater.h"
 #include "MaterialShared.h"
 #include "CollectionManagerTypes.h"
 #include "ICollectionManager.h"
diff --git a/Engine/Source/Editor/UnrealEd/Private/Cooker/CookProfiling.cpp b/Engine/Source/Editor/UnrealEd/Private/Cooker/CookProfiling.cpp
--- a/Engine/Source/Editor/UnrealEd/Private/Cooker/CookProfiling.cpp
+++ b/Engine/Source/Editor/UnrealEd/Private/Cooker/CookProfiling.cpp
@@ -23,6 +23,7 @@
 #include "Serialization/JsonSerializer.h"
 #include "Serialization/JsonWriter.h"
 #include "Templates/Casts.h"
+#include "Templates/Greater.h"
 #include "UObject/GCObject.h"
 #include "UObject/MetaData.h"
 #include "UObject/Package.h"
diff --git a/Engine/Source/Runtime/CoreUObject/Private/UObject/PropertyArray.cpp b/Engine/Source/Runtime/CoreUObject/Private/UObject/PropertyArray.cpp
--- a/Engine/Source/Runtime/CoreUObject/Private/UObject/PropertyArray.cpp
+++ b/Engine/Source/Runtime/CoreUObject/Private/UObject/PropertyArray.cpp
@@ -3,6 +3,7 @@
 #include "UObject/UnrealType.h"
 
 #include "Misc/ScopeExit.h"
+#include "Templates/Greater.h"
 #include "UObject/LinkerLoad.h"
 #include "UObject/OverriddenPropertySet.h"
 #include "UObject/PropertyHelper.h"
diff --git a/Engine/Source/Runtime/CoreUObject/Private/UObject/ReferenceChainSearch.cpp b/Engine/Source/Runtime/CoreUObject/Private/UObject/ReferenceChainSearch.cpp
--- a/Engine/Source/Runtime/CoreUObject/Private/UObject/ReferenceChainSearch.cpp
+++ b/Engine/Source/Runtime/CoreUObject/Private/UObject/ReferenceChainSearch.cpp
@@ -9,6 +9,7 @@
 #include "HAL/PlatformStackWalk.h"
 #include "HAL/ThreadHeartBeat.h"
 #include "Misc/FeedbackContext.h"
+#include "Templates/Greater.h"
 #include "Misc/StringOutputDevice.h"
 #include "UObject/FastReferenceCollector.h"
 #include "UObject/GCObject.h"
diff --git a/Engine/Source/Runtime/Engine/Classes/Animation/AnimCompress.h b/Engine/Source/Runtime/Engine/Classes/Animation/AnimCompress.h
--- a/Engine/Source/Runtime/Engine/Classes/Animation/AnimCompress.h
+++ b/Engine/Source/Runtime/Engine/Classes/Animation/AnimCompress.h
@@ -11,6 +11,7 @@
 #include "UObject/ObjectMacros.h"
 #include "UObject/Object.h"
 #include "Templates/Function.h"
+#include "Templates/Greater.h"
 #include "Animation/AnimSequence.h"
 #include "Animation/AnimBoneCompressionCodec.h"
 #include "AnimationUtils.h"
diff --git a/Engine/Source/Runtime/Engine/Private/KismetArrayLibrary.cpp b/Engine/Source/Runtime/Engine/Private/KismetArrayLibrary.cpp
--- a/Engine/Source/Runtime/Engine/Private/KismetArrayLibrary.cpp
+++ b/Engine/Source/Runtime/Engine/Private/KismetArrayLibrary.cpp
@@ -5,6 +5,7 @@
 #include "GameFramework/Actor.h"
 #include "Blueprint/BlueprintSupport.h"
 #include "Math/RandomStream.h"
+#include "Templates/Greater.h"
 
 #include UE_INLINE_GENERATED_CPP_BY_NAME(KismetArrayLibrary)
 
diff --git a/Engine/Source/Runtime/Engine/Private/LightMap.cpp b/Engine/Source/Runtime/Engine/Private/LightMap.cpp
--- a/Engine/Source/Runtime/Engine/Private/LightMap.cpp
+++ b/Engine/Source/Runtime/Engine/Private/LightMap.cpp
@@ -30,6 +30,7 @@
 #include "DataDrivenShaderPlatformInfo.h"
 #include "StaticMeshLight.h"
 #include "Engine/InstancedStaticMesh.h"
+#include "Templates/Greater.h"
 
 #define VISUALIZE_PACKING 0
 
diff --git a/Engine/Source/Runtime/Engine/Private/Particles/ParticleSystemRender.cpp b/Engine/Source/Runtime/Engine/Private/Particles/ParticleSystemRender.cpp
--- a/Engine/Source/Runtime/Engine/Private/Particles/ParticleSystemRender.cpp
+++ b/Engine/Source/Runtime/Engine/Private/Particles/ParticleSystemRender.cpp
@@ -38,6 +38,7 @@
 #include "DataDrivenShaderPlatformInfo.h"
 #include "PrimitiveUniformShaderParametersBuilder.h"
 #include "Experimental/ConcurrentLinearAllocator.h"
+#include "Templates/Greater.h"
 
 DECLARE_CYCLE_STAT(TEXT("DynamicSpriteEmitterData GetDynamicMeshElementsEmitter GetParticleOrderData RT"), STAT_FDynamicSpriteEmitterData_GetDynamicMeshElementsEmitter_GetParticleOrderData, STATGROUP_Particles);
 DECLARE_CYCLE_STAT(TEXT("DynamicSpriteEmitterData PerParticleWorkOrTasks RT"), STAT_FDynamicSpriteEmitterData_PerParticleWorkOrTasks, STATGROUP_Particles);
diff --git a/Engine/Source/Runtime/Engine/Private/SceneViewExtension.cpp b/Engine/Source/Runtime/Engine/Private/SceneViewExtension.cpp
--- a/Engine/Source/Runtime/Engine/Private/SceneViewExtension.cpp
+++ b/Engine/Source/Runtime/Engine/Private/SceneViewExtension.cpp
@@ -2,6 +2,7 @@
 
 #include "SceneViewExtension.h"
 #include "Engine/Engine.h"
+#include "Templates/Greater.h"
 
 //
 // FSceneViewExtensionBase
diff --git a/Engine/Source/Runtime/Engine/Private/ShadowMap.cpp b/Engine/Source/Runtime/Engine/Private/ShadowMap.cpp
--- a/Engine/Source/Runtime/Engine/Private/ShadowMap.cpp
+++ b/Engine/Source/Runtime/Engine/Private/ShadowMap.cpp
@@ -16,6 +16,7 @@
 #include "Misc/QueuedThreadPool.h"
 #include "Modules/ModuleManager.h"
 #include "GameFramework/WorldSettings.h"
+#include "Templates/Greater.h"
 
 #if WITH_EDITOR
 
diff --git a/Engine/Source/Runtime/GameplayTags/Private/GameplayTagsManager.cpp b/Engine/Source/Runtime/GameplayTags/Private/GameplayTagsManager.cpp
--- a/Engine/Source/Runtime/GameplayTags/Private/GameplayTagsManager.cpp
+++ b/Engine/Source/Runtime/GameplayTags/Private/GameplayTagsManager.cpp
@@ -41,6 +41,7 @@ FSimpleMulticastDelegate UGameplayTagsManager::OnEditorRefreshGameplayTagTree;
 #include "HAL/RunnableThread.h"
 #include "Containers/Queue.h"
 #include "Templates/Function.h"
+#include "Templates/Greater.h"
 #include "UObject/StrongObjectPtr.h"
 #include "Async/Async.h"
 
diff --git a/Engine/Source/Runtime/Landscape/Private/LandscapeEdit.cpp b/Engine/Source/Runtime/Landscape/Private/LandscapeEdit.cpp
--- a/Engine/Source/Runtime/Landscape/Private/LandscapeEdit.cpp
+++ b/Engine/Source/Runtime/Landscape/Private/LandscapeEdit.cpp
@@ -104,6 +104,7 @@ LandscapeEdit.cpp: Landscape editing
 #include "Spatial/PointHashGrid3.h"
 #include "Engine/Texture2DArray.h"
 #include "Async/ParallelFor.h"
+#include "Templates/Greater.h"
 
 DEFINE_LOG_CATEGORY(LogLandscape);
 DEFINE_LOG_CATEGORY(LogLandscapeBP);
diff --git a/Engine/Source/Runtime/MovieScene/Private/EntitySystem/BuiltInComponentTypes.cpp b/Engine/Source/Runtime/MovieScene/Private/EntitySystem/BuiltInComponentTypes.cpp
--- a/Engine/Source/Runtime/MovieScene/Private/EntitySystem/BuiltInComponentTypes.cpp
+++ b/Engine/Source/Runtime/MovieScene/Private/EntitySystem/BuiltInComponentTypes.cpp
@@ -10,6 +10,7 @@
 #include "EntitySystem/MovieSceneEntityManager.h"
 #include "EntitySystem/MovieSceneEntitySystemLinker.h"
 #include "EntitySystem/MovieScenePropertyBinding.h"
+#include "Templates/Greater.h"
 
 #include UE_INLINE_GENERATED_CPP_BY_NAME(BuiltInComponentTypes)
 
diff --git a/Engine/Source/Runtime/MovieScene/Private/EntitySystem/MovieSceneEvaluationHookSystem.cpp b/Engine/Source/Runtime/MovieScene/Private/EntitySystem/MovieSceneEvaluationHookSystem.cpp
--- a/Engine/Source/Runtime/MovieScene/Private/EntitySystem/MovieSceneEvaluationHookSystem.cpp
+++ b/Engine/Source/Runtime/MovieScene/Private/EntitySystem/MovieSceneEvaluationHookSystem.cpp
@@ -10,6 +10,7 @@
 #include "Evaluation/PreAnimatedState/MovieScenePreAnimatedCaptureSource.h"
 #include "Evaluation/PreAnimatedState/MovieScenePreAnimatedCaptureSources.h"
 #include "Evaluation/MovieSceneEvaluationTemplateInstance.h"
+#include "Templates/Greater.h"
 
 #include UE_INLINE_GENERATED_CPP_BY_NAME(MovieSceneEvaluationHookSystem)
 
diff --git a/Engine/Source/Runtime/MovieScene/Private/Tests/MovieSceneEntitySystemTypesTests.cpp b/Engine/Source/Runtime/MovieScene/Private/Tests/MovieSceneEntitySystemTypesTests.cpp
--- a/Engine/Source/Runtime/MovieScene/Private/Tests/MovieSceneEntitySystemTypesTests.cpp
+++ b/Engine/Source/Runtime/MovieScene/Private/Tests/MovieSceneEntitySystemTypesTests.cpp
@@ -10,6 +10,7 @@
 #include "Containers/Set.h"
 #include "Containers/ArrayView.h"
 #include "Misc/AutomationTest.h"
+#include "Templates/Greater.h"
 
 #if WITH_DEV_AUTOMATION_TESTS
 
diff --git a/Engine/Source/Runtime/MovieSceneTracks/Private/Systems/MovieSceneEventSystems.cpp b/Engine/Source/Runtime/MovieSceneTracks/Private/Systems/MovieSceneEventSystems.cpp
--- a/Engine/Source/Runtime/MovieSceneTracks/Private/Systems/MovieSceneEventSystems.cpp
+++ b/Engine/Source/Runtime/MovieSceneTracks/Private/Systems/MovieSceneEventSystems.cpp
@@ -16,6 +16,7 @@
 #include "IMovieScenePlayer.h"
 #include "MovieSceneSequence.h"
 
+#include "Templates/Greater.h"
 #include "Templates/SharedPointer.h"
 #include "Templates/SubclassOf.h"
 #include "Engine/World.h"
diff --git a/Engine/Source/Runtime/MovieSceneTracks/Private/Systems/WeightAndEasingEvaluatorSystem.cpp b/Engine/Source/Runtime/MovieSceneTracks/Private/Systems/WeightAndEasingEvaluatorSystem.cpp
--- a/Engine/Source/Runtime/MovieSceneTracks/Private/Systems/WeightAndEasingEvaluatorSystem.cpp
+++ b/Engine/Source/Runtime/MovieSceneTracks/Private/Systems/WeightAndEasingEvaluatorSystem.cpp
@@ -17,6 +17,7 @@
 #include "Sections/MovieSceneSubSection.h"
 #include "Systems/FloatChannelEvaluatorSystem.h"
 #include "Systems/DoubleChannelEvaluatorSystem.h"
+#include "Templates/Greater.h"
 #include "IMovieScenePlayer.h"
 #include "MovieSceneSection.h"
 
diff --git a/Engine/Source/Runtime/Renderer/Private/DistanceFieldObjectManagement.cpp b/Engine/Source/Runtime/Renderer/Private/DistanceFieldObjectManagement.cpp
--- a/Engine/Source/Runtime/Renderer/Private/DistanceFieldObjectManagement.cpp
+++ b/Engine/Source/Runtime/Renderer/Private/DistanceFieldObjectManagement.cpp
@@ -24,6 +24,7 @@
 #include "UnrealEngine.h"
 #include "InstanceDataSceneProxy.h"
 #include "Lumen/Lumen.h"
+#include "Templates/Greater.h"
 
 DECLARE_GPU_STAT(DistanceFields);
 
