New Material Translator:  Support for Large World Coordinates.  All supported operators implemented, plus world space position, actor, camera, object, and view position expressions, and derivatives.  Construction, casting, component fetching / swizzling, and operators all require functions with LWC.  Some operators don't support LWC -- these need to down-cast their inputs to float in the emitter.  Other operators that accept LWC inputs generate float outputs (example: cosine where the output is -1 to 1) -- these need to inform downstream code that their output is float.  CVars related to LWC are supported, including r.MaterialEditor.LWCTruncateMode (allows automatic truncation when subtracting two values, or disabling the truncate expression), and r.MaterialEditor.LWCEnabled (immediately casts LWC inputs to float on load when disabled).  LWC inputs given bright colors to stand out better (pink vector, cyan scalar).

Unrelated minor fixes:  Added implementation of Max expression (Min was already present).  Fixed typo on bool types in FType::FromShaderType.  Changed case of BO_Atan2 to BO_ATan2, to match similar UO_ATan.  BO_Distance op fixed to return scalar.  "If" expression fixed to be consistent with original translator (handle equals pin not being set, support for vector conditionals).

#rnx

diff --git a/Engine/Shaders/Private/DoubleFloat.ush b/Engine/Shaders/Private/DoubleFloat.ush
--- a/Engine/Shaders/Private/DoubleFloat.ush
+++ b/Engine/Shaders/Private/DoubleFloat.ush
@@ -194,6 +194,12 @@ FDFVector4 DFSwizzle(FDFVector2 V, int C0, int C1, int C2, int C3) { return Make
 FDFVector4 DFSwizzle(FDFVector3 V, int C0, int C1, int C2, int C3) { return MakeDFVector(DFGetComponent(V, C0), DFGetComponent(V, C1), DFGetComponent(V, C2), DFGetComponent(V, C3)); }
 FDFVector4 DFSwizzle(FDFVector4 V, int C0, int C1, int C2, int C3) { return MakeDFVector(DFGetComponent(V, C0), DFGetComponent(V, C1), DFGetComponent(V, C2), DFGetComponent(V, C3)); }
 
+// Needed for new translator
+FDFScalar DFDistance(FDFScalar A, FDFScalar B) { return DFAbs(DFSubtract(A, B)); }
+FDFScalar DFDistance(FDFVector2 A, FDFVector2 B) { return DFLength(DFSubtract(A, B)); }
+FDFScalar DFDistance(FDFVector3 A, FDFVector3 B) { return DFLength(DFSubtract(A, B)); }
+FDFScalar DFDistance(FDFVector4 A, FDFVector4 B) { return DFLength(DFSubtract(A, B)); }
+
 #define DFHackToFloat DFDemote
 
 // Transforms *to* absolute world space
diff --git a/Engine/Shaders/Private/DoubleFloatOperations.ush b/Engine/Shaders/Private/DoubleFloatOperations.ush
--- a/Engine/Shaders/Private/DoubleFloatOperations.ush
+++ b/Engine/Shaders/Private/DoubleFloatOperations.ush
@@ -542,6 +542,20 @@ FBoolType DFEquals(FFloatType Lhs, FDFType Rhs)
 	return DFEquals(DFPromote(Lhs), Rhs);
 }
 
+// Needed for new translator
+FBoolType DFNotEquals(FDFType Lhs, FDFType Rhs)
+{
+#if UE_DF_FORCE_FP32_OPS
+	return Lhs.High != Rhs.High;
+#else
+#if COMPILER_SUPPORTS_HLSL2021
+	return or(Lhs.High != Rhs.High, Lhs.Low != Rhs.Low);
+#else
+	return Lhs.High != Rhs.High || Lhs.Low != Rhs.Low;
+#endif
+#endif
+}
+
 FDFType DFSelect(FBoolType S, FDFType Lhs, FDFType Rhs) { return DFConstructor(select(S, Lhs.High, Rhs.High), select(S, Lhs.Low, Rhs.Low)); }
 FDFType DFSelect(FBoolType S, FDFType Lhs, FFloatType Rhs) { return DFConstructor(select(S, Lhs.High, Rhs), select(S, Lhs.Low, (FFloatType)0)); }
 FDFType DFSelect(FBoolType S, FFloatType Lhs, FDFType Rhs) { return DFConstructor(select(S, Lhs, Rhs.High), select(S, (FFloatType)0, Rhs.Low)); }
@@ -639,6 +653,12 @@ FDFType DFMax(FFloatType Lhs, FDFType Rhs)
 	return DFSelect(DFLess(Lhs, Rhs), Rhs, DFPromote(Lhs));
 }
 
+// Needed for new translator
+FDFType DFClamp(FDFType X, FDFType A, FDFType B)
+{
+	return DFMin(DFMax(X, A), B);
+}
+
 /**********/
 /** MISC **/
 /**********/
diff --git a/Engine/Shaders/Private/LWCOperations.ush b/Engine/Shaders/Private/LWCOperations.ush
--- a/Engine/Shaders/Private/LWCOperations.ush
+++ b/Engine/Shaders/Private/LWCOperations.ush
@@ -82,6 +82,12 @@ FBoolType LWCEqualsApprox(FLWCType Lhs, FFloatType Rhs, float Threshold)
 	return abs(LWCGetTile(Lhs) * UE_LWC_RENDER_TILE_SIZE + (Lhs.Offset - Rhs)) < (FFloatType)Threshold;
 }
 
+// Needed for new translator
+FBoolType LWCNotEquals(FLWCType Lhs, FLWCType Rhs)
+{
+	return (LWCGetTile(Lhs) - LWCGetTile(Rhs)) * UE_LWC_RENDER_TILE_SIZE != Rhs.Offset - Lhs.Offset;
+}
+
 FLWCType LWCSelect(FBoolType S, FLWCType Lhs, FLWCType Rhs) { return LWCConstructor(select(S, LWCGetTile(Lhs), LWCGetTile(Rhs)), select(S, Lhs.Offset, Rhs.Offset)); }
 FLWCType LWCSelect(FBoolType S, FFloatType Lhs, FLWCType Rhs) { return LWCConstructor(select(S, (FFloatType)0, LWCGetTile(Rhs)), select(S, Lhs, Rhs.Offset)); }
 FLWCType LWCSelect(FBoolType S, FLWCType Lhs, FFloatType Rhs) { return LWCConstructor(select(S, LWCGetTile(Lhs), (FFloatType)0), select(S, Lhs.Offset, Rhs)); }
@@ -169,6 +175,12 @@ FLWCType LWCMax(FLWCType Lhs, FLWCType Rhs) { return LWCSelect(LWCGreater(Lhs, R
 FLWCType LWCMax(FFloatType Lhs, FLWCType Rhs) { return LWCSelect(LWCGreater(Lhs, Rhs), Lhs, Rhs); }
 FLWCType LWCMax(FLWCType Lhs, FFloatType Rhs) { return LWCSelect(LWCGreater(Lhs, Rhs), Lhs, Rhs); }
 
+// Needed for new translator
+FLWCType LWCClamp(FLWCType X, FLWCType A, FLWCType B)
+{
+    return LWCMin(LWCMax(X, A), B);
+}
+
 FLWCType LWCAbs(FLWCType V) { return LWCSelect(LWCLess(V, (FFloatType)0), LWCNegate(V), V); }
 
 FFloatType LWCStep(FLWCType Lhs, FLWCType Rhs)   { return select(LWCGreaterEqual(Rhs, Lhs), (FFloatType)1.0f, (FFloatType)0.0f); }
diff --git a/Engine/Shaders/Private/LargeWorldCoordinates.ush b/Engine/Shaders/Private/LargeWorldCoordinates.ush
--- a/Engine/Shaders/Private/LargeWorldCoordinates.ush
+++ b/Engine/Shaders/Private/LargeWorldCoordinates.ush
@@ -501,6 +501,12 @@ float2 LWCNormalize(FLWCVector2 V) { return LWCToFloat(LWCMultiply(V, LWCRcpLeng
 float3 LWCNormalize(FLWCVector3 V) { return LWCToFloat(LWCMultiply(V, LWCRcpLength(V))); }
 float4 LWCNormalize(FLWCVector4 V) { return LWCToFloat(LWCMultiply(V, LWCRcpLength(V))); }
 
+// Needed for new translator
+FLWCScalar LWCDistance(FLWCScalar A, FLWCScalar B) { return LWCAbs(LWCSubtract(A, B)); }
+FLWCScalar LWCDistance(FLWCVector2 A, FLWCVector2 B) { return LWCLength(LWCSubtract(A, B)); }
+FLWCScalar LWCDistance(FLWCVector3 A, FLWCVector3 B) { return LWCLength(LWCSubtract(A, B)); }
+FLWCScalar LWCDistance(FLWCVector4 A, FLWCVector4 B) { return LWCLength(LWCSubtract(A, B)); }
+
 // LWCHackToFloat is a marker for places where LWC quantities are transformed to float, without regard for range/precision
 // This will work correctly for content that's not using LWC scale values
 // 'LWC_HACK_TO_ZERO' is a big hammer to work around some internal compiler errors....using this will break the functionality of the shader, but can unblock testing of other shaders
diff --git a/Engine/Shaders/Private/WorldSpaceMath.ush b/Engine/Shaders/Private/WorldSpaceMath.ush
--- a/Engine/Shaders/Private/WorldSpaceMath.ush
+++ b/Engine/Shaders/Private/WorldSpaceMath.ush
@@ -97,6 +97,7 @@ FDFInverseMatrix WSToDF(FWSInverseMatrix In) { return DFFromTileOffset(In); }
 #define WSLessEqual LWCLessEqual
 #define WSEquals LWCEquals
 #define WSEqualsApprox LWCEqualsApprox
+#define WSNotEquals LWCNotEquals
 
 #define WSDdxDemote LWCDdx
 #define WSDdyDemote LWCDdy
@@ -134,12 +135,14 @@ FDFInverseMatrix WSToDF(FWSInverseMatrix In) { return DFFromTileOffset(In); }
 #define WSMax LWCMax
 #define WSAbs LWCAbs
 #define WSStep LWCStep
+#define WSClamp LWCClamp
 
 #define WSVectorSum LWCVectorSum
 #define WSDot LWCDot
 #define WSLength LWCLength
 #define WSRcpLengthDemote LWCRcpLength
 #define WSNormalizeDemote LWCNormalize
+#define WSDistance LWCDistance
 
 
 #elif WSVECTOR_IS_DOUBLEFLOAT
@@ -231,6 +234,7 @@ FDFInverseMatrix WSToDF(FWSInverseMatrix In) { return In; }
 #define WSLessEqual DFLessEqual
 #define WSEquals DFEquals
 #define WSEqualsApprox DFEqualsApprox
+#define WSNotEquals DFNotEquals
 
 #define WSDdxDemote DFDdxDemote
 #define WSDdyDemote DFDdyDemote
@@ -268,12 +272,14 @@ FDFInverseMatrix WSToDF(FWSInverseMatrix In) { return In; }
 #define WSMax DFMax
 #define WSAbs DFAbs
 #define WSStep DFStep
+#define WSClamp DFClamp
 
 #define WSVectorSum DFVectorSum
 #define WSDot DFDot
 #define WSLength DFLength
 #define WSRcpLengthDemote DFRcpLengthDemote
 #define WSNormalizeDemote DFNormalizeDemote
+#define WSDistance DFDistance
 
 
 #endif
diff --git a/Engine/Source/Editor/UnrealEd/Private/MaterialGraphSchema.cpp b/Engine/Source/Editor/UnrealEd/Private/MaterialGraphSchema.cpp
--- a/Engine/Source/Editor/UnrealEd/Private/MaterialGraphSchema.cpp
+++ b/Engine/Source/Editor/UnrealEd/Private/MaterialGraphSchema.cpp
@@ -637,7 +637,16 @@ FLinearColor GetColorForConnectionType(const UGraphEditorSettings* Settings, UE:
 		}
 		else if (TypeDesc.ComponentType == EValueComponentType::Double)
 		{
-			return Settings->DoublePinTypeColor;
+			// Tentatively use these two color pin settings for LWC inputs to make them stand out better in the graph, since LWC math is expensive.
+			// We could in the future add unique color pin settings for these.
+			if (TypeDesc.NumComponents == 1)
+			{
+				return Settings->SoftObjectPinTypeColor;	// Default cyan
+			}
+			else
+			{
+				return Settings->SoftClassPinTypeColor;		// Default pink
+			}
 		}
 		else if (TypeDesc.ComponentType == EValueComponentType::Bool)
 		{
diff --git a/Engine/Source/Runtime/Engine/Private/Materials/HLSLMaterialTranslator.cpp b/Engine/Source/Runtime/Engine/Private/Materials/HLSLMaterialTranslator.cpp
--- a/Engine/Source/Runtime/Engine/Private/Materials/HLSLMaterialTranslator.cpp
+++ b/Engine/Source/Runtime/Engine/Private/Materials/HLSLMaterialTranslator.cpp
@@ -139,13 +139,6 @@ static FCookStatsManager::FAutoRegisterCallback RegisterCookStats([](FCookStatsM
 
 #endif
 
-static int32 GetLWCTruncateMode()
-{
-	static const auto CVar = IConsoleManager::Get().FindTConsoleVariableDataInt(TEXT("r.MaterialEditor.LWCTruncateMode"));
-	const int32 LWCTruncateValue = CVar ? CVar->GetValueOnAnyThread() : 0;
-	return LWCTruncateValue;
-}
-
 #if WITH_EDITORONLY_DATA
 
 static int32 GAnalyticDerivEnabled = 1;
@@ -5229,7 +5222,7 @@ int32 FHLSLMaterialTranslator::CastShadingModelToFloat(int32 Code)
 
 int32 FHLSLMaterialTranslator::TruncateLWC(int32 Code)
 {
-	const int32 LWCTruncateMode = GetLWCTruncateMode();
+	const int32 LWCTruncateMode = UE::MaterialTranslatorUtils::GetLWCTruncateMode();
 	const bool bAllowLWCTruncate = LWCTruncateMode == 1 || LWCTruncateMode == 2;
 
 	if (!bAllowLWCTruncate)
@@ -9730,7 +9723,7 @@ int32 FHLSLMaterialTranslator::Sub(int32 A, int32 B)
 		const bool bTruncateToFloat = (AType == BType) && IsLWCType(AType) && IsLWCType(BType);
 		if (bTruncateToFloat)
 		{
-			const bool bAllowLWCTruncate = GetLWCTruncateMode() == 2; // Allow automatic truncate?
+			const bool bAllowLWCTruncate = UE::MaterialTranslatorUtils::GetLWCTruncateMode() == 2; // Allow automatic truncate?
 			if (bAllowLWCTruncate)
 			{
 				return TruncateLWC(ResultCode);
diff --git a/Engine/Source/Runtime/Engine/Private/Materials/MaterialExpressionsIR.cpp b/Engine/Source/Runtime/Engine/Private/Materials/MaterialExpressionsIR.cpp
--- a/Engine/Source/Runtime/Engine/Private/Materials/MaterialExpressionsIR.cpp
+++ b/Engine/Source/Runtime/Engine/Private/Materials/MaterialExpressionsIR.cpp
@@ -240,6 +240,7 @@
 #include "Materials/MaterialFunctionInterface.h"
 #include "Materials/MaterialParameterCollection.h"
 #include "Materials/MaterialExpressionOperator.h"
+#include "Materials/MaterialSharedPrivate.h"
 #include "MaterialShared.h"
 #include "Misc/MemStackUtility.h"
 #include "RenderUtils.h"
@@ -483,7 +484,8 @@ void UMaterialExpressionAppendVector::Build(MIR::FEmitter& Em)
 	check(Dimensions >= 2 && Dimensions <= 4);
 
 	// Construct the output vector type.
-	const MIR::FPrimitiveType* ResultType = MIR::FPrimitiveType::GetVector(MIR::ScalarKind_Float, Dimensions);
+	MIR::EScalarKind ResultKind = (AType->ScalarKind == MIR::ScalarKind_LWC || (BType && BType->ScalarKind == MIR::ScalarKind_LWC)) ? MIR::ScalarKind_LWC : MIR::ScalarKind_Float;
+	const MIR::FPrimitiveType* ResultType = MIR::FPrimitiveType::GetVector(ResultKind, Dimensions);
 
 	// Set up each output vector component 
 	FValueRef Components[4] = { nullptr, nullptr, nullptr, nullptr };
@@ -759,6 +761,11 @@ void UMaterialExpressionDivide::Build(MIR::FEmitter& Em)
 	BuildBinaryOperatorWithDefaults(Em, MIR::BO_Divide, &A, ConstA, &B, ConstB);
 }
 
+void UMaterialExpressionMax::Build(MIR::FEmitter& Em)
+{
+	BuildBinaryOperatorWithDefaults(Em, MIR::BO_Max, &A, ConstA, &B, ConstB);
+}
+
 void UMaterialExpressionMin::Build(MIR::FEmitter& Em)
 { 
 	BuildBinaryOperatorWithDefaults(Em, MIR::BO_Min, &A, ConstA, &B, ConstB);
@@ -843,20 +850,41 @@ static FValueRef EmitAlmostEquals(MIR::FEmitter& Em, FValueRef A, FValueRef B, f
 
 void UMaterialExpressionIf::Build(MIR::FEmitter& Em)
 {
-	// Get input values and check their types are what we expect.
-	FValueRef AValue = Em.CheckIsScalar(Em.InputDefaultFloat(&A, 0.f));
-	FValueRef BValue = Em.CheckIsScalar(Em.InputDefaultFloat(&B, ConstB));
+	// Get input values
+	FValueRef AValue = Em.InputDefaultFloat(&A, 0.f);
+	FValueRef BValue = Em.InputDefaultFloat(&B, ConstB);
 	FValueRef AGreaterThanBValue = Em.InputDefaultFloat(&AGreaterThanB, 0.f);
-	FValueRef AEqualsBValue = Em.InputDefaultFloat(&AEqualsB, 0.f);
+	FValueRef AEqualsBValue = Em.TryInput(&AEqualsB);
 	FValueRef ALessThanBValue = Em.InputDefaultFloat(&ALessThanB, 0.f);
 
-	// Emit the comparison expressions.
+	FValueRef OutputValue;
+
+	// Less than comparison -- if equals value isn't present (see below), AGreaterThanBValue will also be returned for the equal case.
 	FValueRef ALessThanBConditionValue = Em.LessThan(AValue, BValue);
-	FValueRef AEqualsBConditionValue = EmitAlmostEquals(Em, AValue, BValue, EqualsThreshold);
 
-	// And finally emit the full conditional expression.
-	FValueRef OutputValue = Em.Branch(AEqualsBConditionValue, AEqualsBValue, AGreaterThanBValue);
-	OutputValue = Em.Branch(ALessThanBConditionValue, ALessThanBValue, OutputValue);
+	if (ALessThanBConditionValue->Type->IsBoolScalar())
+	{
+		OutputValue = Em.Branch(ALessThanBConditionValue, ALessThanBValue, AGreaterThanBValue);
+	}
+	else
+	{
+		OutputValue = Em.Select(ALessThanBConditionValue, ALessThanBValue, AGreaterThanBValue);
+	}
+
+	// Equals value is optional -- if present, generate an additional conditional.
+	if (AEqualsBValue.IsValid())
+	{
+		FValueRef AEqualsBConditionValue = EmitAlmostEquals(Em, AValue, BValue, EqualsThreshold);
+
+		if (AEqualsBConditionValue->Type->IsBoolScalar())
+		{
+			OutputValue = Em.Branch(AEqualsBConditionValue, AEqualsBValue, OutputValue);
+		}
+		else
+		{
+			OutputValue = Em.Select(AEqualsBConditionValue, AEqualsBValue, OutputValue);
+		}
+	}
 
 	Em.Output(0, OutputValue);
 }
@@ -1357,16 +1385,22 @@ static FValueRef BuildViewProperty(MIR::FEmitter& Em, EMaterialExposedViewProper
 	FString HLSLCode = bInvProperty && bHasCustomInverseCode ? PropertyMeta.InvPropertyCode : PropertyMeta.PropertyCode;
 	const MIR::FType* HLSLCodeType = MIR::FType::FromMaterialValueType(PropertyMeta.Type);
 
-	if (bInvProperty && !bHasCustomInverseCode)
+	FValueRef Result = Em.InlineHLSL(HLSLCodeType, MoveTemp(HLSLCode), {}, MIR::EValueFlags::SubstituteTagsInInlineHLSL);
+
+	// CastToNonLWCIfDisabled
+	const MIR::FPrimitiveType* PrimitiveType = HLSLCodeType->AsPrimitive();
+	if (PrimitiveType && PrimitiveType->IsLWC() && !UE::MaterialTranslatorUtils::IsLWCEnabled())
 	{
-		// Fall back to compute the property's inverse from PropertyCode
-		return Em.Divide(Em.ConstantFloat(1.0f), Em.InlineHLSL(HLSLCodeType, MoveTemp(HLSLCode), {}, MIR::EValueFlags::SubstituteTagsInInlineHLSL));
+		Result = Em.CastToFloat(Result);
 	}
-	else
+
+	// Fall back to compute the property's inverse from PropertyCode, if no custom inverse
+	if (bInvProperty && !bHasCustomInverseCode)
 	{
-		// @todo-laura.hermanns - old translator used CastToNonLWCIfDisabled(), but LWC not supported in MIR yet
-		return Em.InlineHLSL(HLSLCodeType, MoveTemp(HLSLCode), {}, MIR::EValueFlags::SubstituteTagsInInlineHLSL);
+		Result = Em.Divide(Em.ConstantFloat(1.0f), Result);
 	}
+
+	return Result;
 }
 
 void UMaterialExpressionViewProperty::Build(MIR::FEmitter& Em)
@@ -1459,4 +1493,78 @@ void UMaterialExpressionUIntToFloat::Build(MIR::FEmitter& Em)
 	Em.Output(0, Em.CastToFloat(Em.Input(&Input)));
 }
 
+void UMaterialExpressionTruncateLWC::Build(MIR::FEmitter& Em)
+{
+	FValueRef InputValue = Em.Input(&Input);
+	int32 LWCTruncateMode = UE::MaterialTranslatorUtils::GetLWCTruncateMode();
+
+	if (LWCTruncateMode == 1 || LWCTruncateMode == 2)
+	{
+		const MIR::FPrimitiveType* Primitive = InputValue->Type->AsPrimitive();
+		if (Primitive && Primitive->ScalarKind == MIR::EScalarKind::ScalarKind_LWC)
+		{
+			Em.Output(0, Em.CastToFloat(Em.Input(&Input)));
+			return;
+		}
+	}
+
+	Em.Output(0, InputValue);
+}
+
+void UMaterialExpressionActorPositionWS::Build(MIR::FEmitter& Em)
+{
+	if (OriginType == EPositionOrigin::CameraRelative)
+	{
+		Em.Output(0, Em.ExternalInput(MIR::EExternalInput::ActorPosition_CameraRelative));
+	}
+	else if (!UE::MaterialTranslatorUtils::IsLWCEnabled())
+	{
+		// CastToNonLWCIfDisabled
+		Em.Output(0, Em.CastToFloat(Em.ExternalInput(MIR::EExternalInput::ActorPosition_Absolute)));
+	}
+	else
+	{
+		Em.Output(0, Em.ExternalInput(MIR::EExternalInput::ActorPosition_Absolute));
+	}
+}
+
+void UMaterialExpressionObjectPositionWS::Build(MIR::FEmitter& Em)
+{
+	if (OriginType == EPositionOrigin::CameraRelative)
+	{
+		Em.Output(0, Em.ExternalInput(MIR::EExternalInput::ObjectPosition_CameraRelative));
+	}
+	else if (!UE::MaterialTranslatorUtils::IsLWCEnabled())
+	{
+		// CastToNonLWCIfDisabled
+		Em.Output(0, Em.CastToFloat(Em.ExternalInput(MIR::EExternalInput::ObjectPosition_Absolute)));
+	}
+	else
+	{
+		Em.Output(0, Em.ExternalInput(MIR::EExternalInput::ObjectPosition_Absolute));
+	}
+}
+
+void UMaterialExpressionWorldPosition::Build(MIR::FEmitter& Em)
+{
+	// Make sure EWorldPositionIncludedOffsets and corresponding elements of EExternalInput stay in sync, so this enum addition is valid
+	static_assert(WPT_MAX == 4 && WPT_Default == 0);
+	static_assert((int32)MIR::EExternalInput::WorldPosition_Absolute - (int32)MIR::EExternalInput::WorldPosition_Absolute == WPT_Default);
+	static_assert((int32)MIR::EExternalInput::WorldPosition_AbsoluteNoOffsets - (int32)MIR::EExternalInput::WorldPosition_Absolute == WPT_ExcludeAllShaderOffsets);
+	static_assert((int32)MIR::EExternalInput::WorldPosition_CameraRelative - (int32)MIR::EExternalInput::WorldPosition_Absolute == WPT_CameraRelative);
+	static_assert((int32)MIR::EExternalInput::WorldPosition_CameraRelativeNoOffsets - (int32)MIR::EExternalInput::WorldPosition_Absolute == WPT_CameraRelativeNoOffsets);
+
+	FValueRef WorldPosition = Em.ExternalInput((MIR::EExternalInput)((int32)MIR::EExternalInput::WorldPosition_Absolute + FMath::Clamp(WorldPositionShaderOffset, 0, WPT_MAX - 1)));
+
+	// CastToNonLWCIfDisabled
+	if (!UE::MaterialTranslatorUtils::IsLWCEnabled())
+	{
+		WorldPosition = Em.CastToFloat(WorldPosition);
+	}
+
+	Em.Output(0, WorldPosition);
+	Em.Output(1, Em.Swizzle(WorldPosition, MIR::FSwizzleMask(MIR::EVectorComponent::X, MIR::EVectorComponent::Y)));
+	Em.Output(2, Em.Subscript(WorldPosition, 2));
+}
+
 #endif // WITH_EDITOR
diff --git a/Engine/Source/Runtime/Engine/Private/Materials/MaterialIR.cpp b/Engine/Source/Runtime/Engine/Private/Materials/MaterialIR.cpp
--- a/Engine/Source/Runtime/Engine/Private/Materials/MaterialIR.cpp
+++ b/Engine/Source/Runtime/Engine/Private/Materials/MaterialIR.cpp
@@ -432,6 +432,11 @@ const TCHAR* LexToString(EExternalInput Input)
 		case EExternalInput::TexCoord5:     return TEXT("TexCoord5");
 		case EExternalInput::TexCoord6:     return TEXT("TexCoord6");
 		case EExternalInput::TexCoord7:     return TEXT("TexCoord7");
+		case EExternalInput::WorldPosition_Absolute: return TEXT("WorldPosition_Absolute");
+		case EExternalInput::WorldPosition_AbsoluteNoOffsets: return TEXT("WorldPosition_AbsoluteNoOffsets");
+		case EExternalInput::WorldPosition_CameraRelative: return TEXT("WorldPosition_CameraRelative");
+		case EExternalInput::WorldPosition_CameraRelativeNoOffsets: return TEXT("WorldPosition_CameraRelativeNoOffsets");
+
 		case EExternalInput::TexCoord0_Ddx: return TEXT("TexCoord0_Ddx");
 		case EExternalInput::TexCoord1_Ddx: return TEXT("TexCoord1_Ddx");
 		case EExternalInput::TexCoord2_Ddx: return TEXT("TexCoord2_Ddx");
@@ -440,6 +445,11 @@ const TCHAR* LexToString(EExternalInput Input)
 		case EExternalInput::TexCoord5_Ddx: return TEXT("TexCoord5_Ddx");
 		case EExternalInput::TexCoord6_Ddx: return TEXT("TexCoord6_Ddx");
 		case EExternalInput::TexCoord7_Ddx: return TEXT("TexCoord7_Ddx");
+		case EExternalInput::WorldPosition_Absolute_Ddx: return TEXT("WorldPosition_Absolute_Ddx");
+		case EExternalInput::WorldPosition_AbsoluteNoOffsets_Ddx: return TEXT("WorldPosition_AbsoluteNoOffsets_Ddx");
+		case EExternalInput::WorldPosition_CameraRelative_Ddx: return TEXT("WorldPosition_CameraRelative_Ddx");
+		case EExternalInput::WorldPosition_CameraRelativeNoOffsets_Ddx: return TEXT("WorldPosition_CameraRelativeNoOffsets_Ddx");
+
 		case EExternalInput::TexCoord0_Ddy: return TEXT("TexCoord0_Ddy");
 		case EExternalInput::TexCoord1_Ddy: return TEXT("TexCoord1_Ddy");
 		case EExternalInput::TexCoord2_Ddy: return TEXT("TexCoord2_Ddy");
@@ -448,6 +458,15 @@ const TCHAR* LexToString(EExternalInput Input)
 		case EExternalInput::TexCoord5_Ddy: return TEXT("TexCoord5_Ddy");
 		case EExternalInput::TexCoord6_Ddy: return TEXT("TexCoord6_Ddy");
 		case EExternalInput::TexCoord7_Ddy: return TEXT("TexCoord7_Ddy");
+		case EExternalInput::WorldPosition_Absolute_Ddy: return TEXT("WorldPosition_Absolute_Ddy");
+		case EExternalInput::WorldPosition_AbsoluteNoOffsets_Ddy: return TEXT("WorldPosition_AbsoluteNoOffsets_Ddy");
+		case EExternalInput::WorldPosition_CameraRelative_Ddy: return TEXT("WorldPosition_CameraRelative_Ddy");
+		case EExternalInput::WorldPosition_CameraRelativeNoOffsets_Ddy: return TEXT("WorldPosition_CameraRelativeNoOffsets_Ddy");
+
+		case EExternalInput::ActorPosition_Absolute: return TEXT("ActorPosition_Absolute");
+		case EExternalInput::ActorPosition_CameraRelative: return TEXT("ActorPosition_CameraRelative");
+		case EExternalInput::ObjectPosition_Absolute: return TEXT("ObjectPosition_Absolute");
+		case EExternalInput::ObjectPosition_CameraRelative: return TEXT("ObjectPosition_CameraRelative");
 		case EExternalInput::ViewMaterialTextureMipBias: return TEXT("ViewMaterialTextureMipBias");
 		case EExternalInput::ViewMaterialTextureDerivativeMultiply: return TEXT("ViewMaterialTextureDerivativeMultiply");
 		default: UE_MIR_UNREACHABLE();
@@ -456,7 +475,7 @@ const TCHAR* LexToString(EExternalInput Input)
 
 const FType* GetExternalInputType(EExternalInput Id)
 {
-	if (Id >= EExternalInput::TexCoord0 && Id <= EExternalInput::TexCoord7_Ddy)
+	if (IsExternalInputTexCoordOrPartialDerivative(Id))
 	{
 		return FPrimitiveType::GetFloat2();
 	}
@@ -465,10 +484,52 @@ const FType* GetExternalInputType(EExternalInput Id)
 	{
 		case EExternalInput::ViewMaterialTextureMipBias: return FPrimitiveType::GetFloat();
 		case EExternalInput::ViewMaterialTextureDerivativeMultiply: return FPrimitiveType::GetFloat();
+
+		case EExternalInput::ActorPosition_Absolute:
+		case EExternalInput::ObjectPosition_Absolute:
+		case EExternalInput::WorldPosition_Absolute:
+		case EExternalInput::WorldPosition_AbsoluteNoOffsets:
+		case EExternalInput::WorldPosition_Absolute_Ddx:				// Technically these derivatives are not LWC, but the emitter needs them to be for evaluation of
+		case EExternalInput::WorldPosition_AbsoluteNoOffsets_Ddx:		// analytic derivative expressions.  So we cast them to LWC on load.  The shader compiler should
+		case EExternalInput::WorldPosition_Absolute_Ddy:				// be able to optimize the math (the tile offsets will all be zero, cancelling things out).
+		case EExternalInput::WorldPosition_AbsoluteNoOffsets_Ddy:
+			return FPrimitiveType::GetLWC3();
+
+		case EExternalInput::ActorPosition_CameraRelative:
+		case EExternalInput::ObjectPosition_CameraRelative:
+		case EExternalInput::WorldPosition_CameraRelative:
+		case EExternalInput::WorldPosition_CameraRelativeNoOffsets:
+		case EExternalInput::WorldPosition_CameraRelative_Ddx:
+		case EExternalInput::WorldPosition_CameraRelativeNoOffsets_Ddx:
+		case EExternalInput::WorldPosition_CameraRelative_Ddy:
+		case EExternalInput::WorldPosition_CameraRelativeNoOffsets_Ddy:
+			return FPrimitiveType::GetFloat3();
+
 		default: UE_MIR_UNREACHABLE();
 	}
 }
 
+// Validate derivative related assumptions in EExternalInput enum
+static_assert(EExternalInput::WithDerivatives_First == EExternalInput::TexCoord0);
+static_assert(EExternalInput::WithDerivatives_Last == EExternalInput::WithDerivatives_LastDdy);
+static_assert((int32)EExternalInput::WithDerivatives_LastDdx - (int32)EExternalInput::WithDerivatives_LastVal == ExternalInputWithDerivativesNum);
+static_assert((int32)EExternalInput::WithDerivatives_LastDdy - (int32)EExternalInput::WithDerivatives_LastDdx == ExternalInputWithDerivativesNum);
+
+bool IsExternalInputWithDerivatives(EExternalInput Id)
+{
+	return Id >= EExternalInput::WithDerivatives_First && Id <= EExternalInput::WithDerivatives_Last;
+}
+
+EExternalInput GetExternalInputDerivative(EExternalInput Id, EDerivativeAxis Axis)
+{
+	check(IsExternalInputWithDerivatives(Id));
+
+	int32 IntId = (int32(Id) - (int32)MIR::EExternalInput::WithDerivatives_First) % ExternalInputWithDerivativesNum + (int32)MIR::EExternalInput::WithDerivatives_First;
+
+	// Add one to axis enum (X==0, Y==1) to produce a value with range 1 to 2 for computing offset to the derivative variation of a value
+	return (EExternalInput)(IntId + ((int)Axis + 1) * ExternalInputWithDerivativesNum);
+}
+
 MIR::EExternalInput TexCoordIndexToExternalInput(int32 TexCoordIndex)
 {
 	check(TexCoordIndex < TexCoordMaxNum);
@@ -477,12 +538,8 @@ MIR::EExternalInput TexCoordIndexToExternalInput(int32 TexCoordIndex)
 
 int32 ExternalInputToTexCoordIndex(EExternalInput Id)
 {
-	int32 IntId = int32(Id) - 1;
-	if (IntId < 0 || IntId > TexCoordMaxNum * TexCoordGroups)
-	{
-		return UINT32_MAX;
-	}
-	return IntId % TexCoordMaxNum;
+	check(IsExternalInputTexCoordOrPartialDerivative(Id));
+	return (int32(Id) - (int32)MIR::EExternalInput::WithDerivatives_First) % ExternalInputWithDerivativesNum;
 }
 
 bool IsExternalInputTexCoord(EExternalInput Id)
@@ -502,7 +559,7 @@ bool IsExternalInputTexCoordDdy(EExternalInput Id)
 
 bool IsExternalInputTexCoordOrPartialDerivative(EExternalInput Id)
 {
-	return Id >= EExternalInput::TexCoord0 && Id <= EExternalInput::TexCoord7_Ddy;
+	return IsExternalInputTexCoord(Id) || IsExternalInputTexCoordDdx(Id) || IsExternalInputTexCoordDdy(Id);
 }
 
 FBlock* FBlock::FindCommonParentWith(MIR::FBlock* Other)
@@ -664,7 +721,7 @@ const TCHAR* LexToString(EOperator Op)
 		/* Binary operators */
 		case BO_Add: return TEXT("Add");
 		case BO_And: return TEXT("And");
-		case BO_Atan2: return TEXT("Atan2");
+		case BO_ATan2: return TEXT("ATan2");
 		case BO_BitShiftLeft: return TEXT("BitShiftLeft");
 		case BO_BitShiftRight: return TEXT("BitShiftRight");
 		case BO_BitwiseAnd: return TEXT("BitwiseAnd");
diff --git a/Engine/Source/Runtime/Engine/Private/Materials/MaterialIREmitter.cpp b/Engine/Source/Runtime/Engine/Private/Materials/MaterialIREmitter.cpp
--- a/Engine/Source/Runtime/Engine/Private/Materials/MaterialIREmitter.cpp
+++ b/Engine/Source/Runtime/Engine/Private/Materials/MaterialIREmitter.cpp
@@ -474,6 +474,7 @@ FValueRef FEmitter::ConstantZero(EScalarKind Kind)
 		case ScalarKind_Bool: return ConstantFalse();
 		case ScalarKind_Int: return ConstantInt(0);
 		case ScalarKind_Float: return ConstantFloat(0.0f);
+		case ScalarKind_LWC: return CastToScalarKind(ConstantFloat(0.0f), ScalarKind_LWC);
 		default: UE_MIR_UNREACHABLE();
 	}
 }
@@ -485,6 +486,7 @@ FValueRef FEmitter::ConstantOne(EScalarKind Kind)
 		case ScalarKind_Bool: return ConstantTrue();
 		case ScalarKind_Int: return ConstantInt(1);
 		case ScalarKind_Float: return ConstantFloat(1.0f);
+		case ScalarKind_LWC: return CastToScalarKind(ConstantFloat(1.0f), ScalarKind_LWC);
 		default: UE_MIR_UNREACHABLE();
 	}
 }
@@ -496,6 +498,7 @@ FValueRef FEmitter::ConstantScalar(EScalarKind Kind, TFloat FromFloat)
 		case ScalarKind_Bool: return ConstantBool(FromFloat != 0.0f);
 		case ScalarKind_Int: return ConstantInt(TInteger(FromFloat));
 		case ScalarKind_Float: return ConstantFloat(FromFloat);
+		case ScalarKind_LWC: return CastToScalarKind(ConstantFloat(FromFloat), ScalarKind_LWC);
 		default: UE_MIR_UNREACHABLE();
 	}
 }
@@ -996,7 +999,7 @@ static T FoldScalarOperator(FEmitter& Emitter, EOperator Operator, T A, T B, T C
 				{
 					return FGenericPlatformMath::TruncToDouble(A);
 				}
-			case BO_Atan2: return FGenericPlatformMath::Atan2(A, B);
+			case BO_ATan2: return FGenericPlatformMath::Atan2(A, B);
 			case BO_Fmod: return FGenericPlatformMath::Fmod(A, B);
 			case BO_Pow: return FGenericPlatformMath::Pow(A, B);
 			case TO_Lerp: return FMath::Lerp<T>(A, B, C);
@@ -1301,7 +1304,9 @@ enum EOperatorParameterTypeFilter
 	OPTF_CheckIsOneOrGreaterFloatConst = 1 << 7,				///< Check that if the argument is a constant float, it is 1 or greater (xFloat >= 1)
 	OPTF_CheckIsBetweenMinusOneAndPlusOneFloatConst = 1 << 8,	///< Check that if the argument is a constant float, it is between -1 and 1 (-1 < x < 1)
 	OPTF_CastToCommonType = 1 << 11,	 						///< Cast the argument to the common arguments type
+	OPTF_LWCInput = 1 << 12,									///< Operator can accept LWC input for this argument
 	OPTF_CastToCommonArithmeticType = OPTF_CheckIsArithmetic | OPTF_CastToCommonType,
+	OPTF_CastToCommonArithmeticTypeLWC = OPTF_CheckIsArithmetic | OPTF_CastToCommonType | OPTF_LWCInput,
 	OPTF_CastToCommonFloatType = OPTF_CastToAnyFloat | OPTF_CastToCommonType,
 };
 
@@ -1314,10 +1319,11 @@ inline EOperatorParameterTypeFilter operator|(EOperatorParameterTypeFilter A, EO
 enum EOperatorReturnType
 {
 	ORT_Unknown,							///< Unspecified
-	ORT_FirstArgumentType,					///< The same type as the first argument
+	ORT_FirstArgumentType,					///< The same type as the first argument, LWC input produces float result
 	ORT_BooleanWithFirstArgumentDimensions,	///< A boolean primitive type with the same dimensions (rows and columns) as the first argument type
-	ORT_FirstArgumentTypeToScalar,			///< A scalar primitive type with the same kind as the scalar type of the first argument
+	ORT_FirstArgumentTypeToScalarLWC,		///< A scalar primitive type with the same kind as the scalar type of the first argument, LWC results allowed
 	ORT_SecondArgumentType,					///< The same type as the second argument
+	ORT_FirstArgumentTypeLWC,				///< The same type as the first argument, LWC results allowed
 };
 
 /// The signature of an operator consisting of its parameter and return type information.
@@ -1333,86 +1339,91 @@ static const FOperatorSignature* GetOperatorSignature(EOperator Op)
 	static const FOperatorSignature* Signatures = [] ()
 	{
 		const FOperatorSignature UnaryFloat = { { OPTF_CheckIsArithmetic | OPTF_CastToAnyFloat }, ORT_FirstArgumentType };
+		const FOperatorSignature UnaryFloatLWC = { { OPTF_CheckIsArithmetic | OPTF_CastToAnyFloat | OPTF_LWCInput }, ORT_FirstArgumentTypeLWC };
+		const FOperatorSignature UnaryFloatLWCDemote = { { OPTF_CheckIsArithmetic | OPTF_CastToAnyFloat | OPTF_LWCInput }, ORT_FirstArgumentType };
 		const FOperatorSignature UnaryFloatToBoolean = { { OPTF_CheckIsArithmetic | OPTF_CastToAnyFloat }, ORT_BooleanWithFirstArgumentDimensions };
 		const FOperatorSignature BinaryArithmetic = { { OPTF_CastToCommonArithmeticType, OPTF_CastToCommonArithmeticType }, ORT_FirstArgumentType };
+		const FOperatorSignature BinaryArithmeticLWC = { { OPTF_CastToCommonArithmeticTypeLWC, OPTF_CastToCommonArithmeticTypeLWC }, ORT_FirstArgumentTypeLWC };
+		const FOperatorSignature BinaryArithmeticLWCDemote = { { OPTF_CastToCommonArithmeticTypeLWC, OPTF_CastToCommonArithmeticTypeLWC }, ORT_FirstArgumentType };
 		const FOperatorSignature BinaryInteger = { { OPTF_CheckIsInteger | OPTF_CastToCommonArithmeticType, OPTF_CheckIsInteger | OPTF_CastToCommonArithmeticType }, ORT_FirstArgumentType };
-		const FOperatorSignature BinaryFloat = { { OPTF_CastToCommonFloatType , OPTF_CastToCommonFloatType }, ORT_FirstArgumentType };
-		const FOperatorSignature BinaryArithmeticComparison = { { OPTF_CastToCommonArithmeticType, OPTF_CastToCommonArithmeticType }, ORT_BooleanWithFirstArgumentDimensions };
+		const FOperatorSignature BinaryFloat = { { OPTF_CastToCommonFloatType, OPTF_CastToCommonFloatType }, ORT_FirstArgumentType };
+		const FOperatorSignature BinaryArithmeticComparisonLWC = { { OPTF_CastToCommonArithmeticType | OPTF_LWCInput, OPTF_CastToCommonArithmeticType | OPTF_LWCInput }, ORT_BooleanWithFirstArgumentDimensions };
 		const FOperatorSignature BinaryLogical = { { OPTF_CheckIsBoolean | OPTF_CastToCommonType, OPTF_CheckIsBoolean | OPTF_CastToCommonType }, ORT_FirstArgumentType };
-		const FOperatorSignature TernaryArithmetic = { { OPTF_CastToCommonArithmeticType, OPTF_CastToCommonArithmeticType, OPTF_CastToCommonArithmeticType }, ORT_FirstArgumentType };
-		const FOperatorSignature TernaryFloat = { { OPTF_CastToCommonArithmeticType | OPTF_CastToAnyFloat, OPTF_CastToCommonArithmeticType, OPTF_CastToCommonArithmeticType }, ORT_FirstArgumentType };
+		const FOperatorSignature TernaryArithmeticLWC = { { OPTF_CastToCommonArithmeticTypeLWC, OPTF_CastToCommonArithmeticTypeLWC, OPTF_CastToCommonArithmeticTypeLWC }, ORT_FirstArgumentTypeLWC };
+		const FOperatorSignature TernaryFloatLWCDemote = { { OPTF_CastToCommonArithmeticTypeLWC | OPTF_CastToAnyFloat, OPTF_CastToCommonArithmeticTypeLWC, OPTF_CastToCommonArithmeticTypeLWC }, ORT_FirstArgumentType };
 
 		static FOperatorSignature S[MIR::OperatorCount];
 
 		/* unary operators */
 		S[UO_BitwiseNot] 				= { { OPTF_CheckIsInteger }, ORT_FirstArgumentType };
-		S[UO_Negate] 					= { { OPTF_CheckIsArithmetic }, ORT_FirstArgumentType };
+		S[UO_Negate] 					= { { OPTF_CheckIsArithmetic | OPTF_LWCInput }, ORT_FirstArgumentTypeLWC };
 		S[UO_Not] 						= { { OPTF_CheckIsBoolean }, ORT_FirstArgumentType };
 
-		S[UO_Abs] 						= UnaryFloat;
-		S[UO_ACos] 						= UnaryFloat;
+		S[UO_Abs] 						= UnaryFloatLWC;
+		S[UO_ACos] 						= UnaryFloatLWCDemote;
 		S[UO_ACosh] 					= { { OPTF_CheckIsArithmetic | OPTF_CastToAnyFloat | OPTF_CheckIsOneOrGreaterFloatConst }, ORT_FirstArgumentType };
-		S[UO_ASin] 						= UnaryFloat;
+		S[UO_ASin] 						= UnaryFloatLWCDemote;
 		S[UO_ASinh] 					= UnaryFloat;
-		S[UO_ATan] 						= UnaryFloat;
+		S[UO_ATan] 						= UnaryFloatLWCDemote;
 		S[UO_ATanh] 					= { { OPTF_CheckIsArithmetic | OPTF_CastToAnyFloat | OPTF_CheckIsBetweenMinusOneAndPlusOneFloatConst }, ORT_FirstArgumentType };
-		S[UO_Ceil] 						= UnaryFloat;
-		S[UO_Cos] 						= UnaryFloat;
+		S[UO_Ceil] 						= UnaryFloatLWC;
+		S[UO_Cos] 						= UnaryFloatLWCDemote;
 		S[UO_Exponential]				= UnaryFloat;
 		S[UO_Exponential2]				= UnaryFloat;
-		S[UO_Floor] 					= UnaryFloat;
-		S[UO_Frac]						= UnaryFloat;
+		S[UO_Floor] 					= UnaryFloatLWC;
+		S[UO_Frac]						= UnaryFloatLWCDemote;
 		S[UO_IsFinite]					= UnaryFloatToBoolean;
 		S[UO_IsInf]						= UnaryFloatToBoolean;
 		S[UO_IsNan]						= UnaryFloatToBoolean;
-		S[UO_Length]					= { { OPTF_CheckIsArithmetic | OPTF_CheckIsNotMatrix | OPTF_CastToAnyFloat }, ORT_FirstArgumentTypeToScalar };
+		S[UO_Length]					= { { OPTF_CheckIsArithmetic | OPTF_CheckIsNotMatrix | OPTF_CastToAnyFloat | OPTF_LWCInput }, ORT_FirstArgumentTypeToScalarLWC };
 		S[UO_Logarithm] 				= { { OPTF_CheckIsArithmetic | OPTF_CheckIsNonZeroFloatConst | OPTF_CheckIsNonNegativeFloatConst | OPTF_CastToAnyFloat }, ORT_FirstArgumentType };
 		S[UO_Logarithm10]				= { { OPTF_CheckIsArithmetic | OPTF_CheckIsNonZeroFloatConst | OPTF_CheckIsNonNegativeFloatConst | OPTF_CastToAnyFloat }, ORT_FirstArgumentType };
 		S[UO_Logarithm2] 				= { { OPTF_CheckIsArithmetic | OPTF_CheckIsNonZeroFloatConst | OPTF_CheckIsNonNegativeFloatConst | OPTF_CastToAnyFloat }, ORT_FirstArgumentType };
-		S[UO_Round] 					= UnaryFloat;
-		S[UO_Saturate]					= UnaryFloat;
-		S[UO_Sign]						= UnaryFloat;
-		S[UO_Sin] 						= UnaryFloat;
-		S[UO_Sqrt]						= { { OPTF_CheckIsArithmetic | OPTF_CheckIsNonNegativeFloatConst | OPTF_CastToAnyFloat }, ORT_FirstArgumentType };
-		S[UO_Tan] 						= UnaryFloat;
+		S[UO_Round] 					= UnaryFloatLWC;
+		S[UO_Saturate]					= UnaryFloatLWCDemote;
+		S[UO_Sign]						= UnaryFloatLWCDemote;
+		S[UO_Sin] 						= UnaryFloatLWCDemote;
+		S[UO_Sqrt]						= { { OPTF_CheckIsArithmetic | OPTF_CheckIsNonNegativeFloatConst | OPTF_CastToAnyFloat | OPTF_LWCInput }, ORT_FirstArgumentType };
+		S[UO_Tan] 						= UnaryFloatLWCDemote;
 		S[UO_Tanh] 						= UnaryFloat;
-		S[UO_Truncate]					= UnaryFloat;
+		S[UO_Truncate]					= UnaryFloatLWC;
 
 		/* binary operators */
-		S[BO_Equals]					= { { OPTF_CastToCommonType, OPTF_CastToCommonType }, ORT_BooleanWithFirstArgumentDimensions };
-		S[BO_GreaterThan]				= BinaryArithmeticComparison;
-		S[BO_GreaterThanOrEquals]		= BinaryArithmeticComparison;
-		S[BO_LessThan]					= BinaryArithmeticComparison;
-		S[BO_LessThanOrEquals]			= BinaryArithmeticComparison;
-		S[BO_NotEquals]					= { { OPTF_CastToCommonType, OPTF_CastToCommonType }, ORT_BooleanWithFirstArgumentDimensions };
+		S[BO_Equals]					= { { OPTF_CastToCommonType | OPTF_LWCInput, OPTF_CastToCommonType | OPTF_LWCInput }, ORT_BooleanWithFirstArgumentDimensions };
+		S[BO_GreaterThan]				= BinaryArithmeticComparisonLWC;
+		S[BO_GreaterThanOrEquals]		= BinaryArithmeticComparisonLWC;
+		S[BO_LessThan]					= BinaryArithmeticComparisonLWC;
+		S[BO_LessThanOrEquals]			= BinaryArithmeticComparisonLWC;
+		S[BO_NotEquals]					= { { OPTF_CastToCommonType | OPTF_LWCInput, OPTF_CastToCommonType | OPTF_LWCInput }, ORT_BooleanWithFirstArgumentDimensions };
 		
 		S[BO_And] 						= BinaryLogical;
 		S[BO_Or] 						= BinaryLogical;
-		S[BO_Add] 						= BinaryArithmetic;
-		S[BO_Subtract] 					= BinaryArithmetic;
-		S[BO_Multiply] 					= BinaryArithmetic;
-		S[BO_Divide] 					= BinaryArithmetic;
+		S[BO_Add] 						= BinaryArithmeticLWC;
+		S[BO_Subtract] 					= BinaryArithmeticLWC;
+		S[BO_Multiply] 					= BinaryArithmeticLWC;
+		S[BO_Divide] 					= BinaryArithmeticLWC;
 		S[BO_Modulo] 					= BinaryInteger;
 		S[BO_BitwiseAnd]				= BinaryInteger;
 		S[BO_BitwiseOr]					= BinaryInteger;
 		S[BO_BitShiftLeft]				= BinaryInteger;
 		S[BO_BitShiftRight]				= BinaryInteger;
 
-		S[BO_Atan2]						= BinaryFloat;
+		S[BO_ATan2]						= BinaryFloat;
 		S[BO_Cross] 					= { { OPTF_CheckIsArithmetic | OPTF_CheckIsVector3, OPTF_CastToFirstArgumentType }, ORT_FirstArgumentType };
-		S[BO_Distance]					= BinaryFloat;
-		S[BO_Dot] 						= { { OPTF_CheckIsArithmetic | OPTF_CheckIsNotMatrix, OPTF_CastToFirstArgumentType }, ORT_FirstArgumentTypeToScalar };
-		S[BO_Fmod] 						= BinaryFloat;
-		S[BO_Max] 						= BinaryArithmetic;
-		S[BO_Min] 						= BinaryArithmetic;
+		S[BO_Distance]					= { { OPTF_CastToCommonFloatType | OPTF_LWCInput, OPTF_CastToCommonFloatType | OPTF_LWCInput }, ORT_FirstArgumentTypeToScalarLWC };
+		S[BO_Dot] 						= { { OPTF_CheckIsArithmetic | OPTF_CheckIsNotMatrix | OPTF_LWCInput, OPTF_CastToFirstArgumentType | OPTF_LWCInput }, ORT_FirstArgumentTypeToScalarLWC };
+		S[BO_Fmod] 						= { { OPTF_CastToCommonFloatType | OPTF_LWCInput, OPTF_CastToCommonFloatType }, ORT_FirstArgumentType };		// First input can be LWC, but second and output are always demoted
+
+		S[BO_Max] 						= BinaryArithmeticLWC;
+		S[BO_Min] 						= BinaryArithmeticLWC;
 		S[BO_Pow] 						= BinaryFloat;
-		S[BO_Step] 						= BinaryArithmetic;
+		S[BO_Step] 						= BinaryArithmeticLWCDemote;
 
-		/* ternary operators */
-		S[TO_Clamp]						= TernaryArithmetic;
-		S[TO_Lerp] 						= TernaryFloat;
-		S[TO_Select]					= { { OPTF_CheckIsBoolean | OPTF_CheckIsNotMatrix, OPTF_CheckIsNotMatrix, OPTF_CheckIsNotMatrix}, ORT_SecondArgumentType }; // Note: this is a special operator, which is handled manually in the Validate function
-		S[TO_Smoothstep]				= TernaryFloat;
+		/* ternary operators -- note lerp doesn't support LWC for the third argument! */
+		S[TO_Clamp]						= TernaryArithmeticLWC;
+		S[TO_Lerp] 						= { { OPTF_CastToCommonArithmeticTypeLWC | OPTF_CastToAnyFloat, OPTF_CastToCommonArithmeticTypeLWC, OPTF_CastToCommonArithmeticType }, ORT_FirstArgumentTypeLWC };;
+		S[TO_Select]					= { { OPTF_CheckIsBoolean | OPTF_CheckIsNotMatrix, OPTF_CheckIsNotMatrix | OPTF_LWCInput, OPTF_CheckIsNotMatrix | OPTF_LWCInput}, ORT_SecondArgumentType }; // Note: this is a special operator, which is handled manually in the Validate function
+		S[TO_Smoothstep]				= TernaryFloatLWCDemote;
 		return S;
 	} ();
 	return &Signatures[Op];
@@ -1434,11 +1445,24 @@ static const FPrimitiveType* ValidateOperatorAndGetResultType(FEmitter& Emitter,
 	// Verify that C argument has been provided if operator is ternary.
 	check(!MIR::IsTernaryOperator(Op) || C);
 
+	// For Clamp, we do a special case and demote the first LWC argument if the second and third arguments (min / max) are non-LWC.
+	// We want to do this before fetching FirstArgumentPrimitiveType.
+	if (Op == TO_Clamp && A->Type->IsLWC() && !B->Type->IsLWC() && !C->Type->IsLWC())
+	{
+		A = Emitter.CastToFloat(A);
+	}
+
 	// Verify that the first argument type is primitive.
 	FValueRef Arguments[] = { A, B, C, nullptr };
 	const FPrimitiveType* FirstArgumentPrimitiveType = A->Type->AsPrimitive();
 	static const TCHAR* ArgumentsStr[] = { TEXT("first"), TEXT("second"), TEXT("third") };
 	const FOperatorSignature* Signature = GetOperatorSignature(Op);
+
+	// Handle automatic cast to float for operators that don't support LWC inputs
+	if (FirstArgumentPrimitiveType->ScalarKind == EScalarKind::ScalarKind_LWC && !(Signature->ParameterTypes[0] & OPTF_LWCInput))
+	{
+		FirstArgumentPrimitiveType = FirstArgumentPrimitiveType->WithScalarKind(EScalarKind::ScalarKind_Float);
+	}
 	const FType* ArgumentsCommonType = FirstArgumentPrimitiveType;
 
 	bool bValid = true;
@@ -1464,6 +1488,12 @@ static const FPrimitiveType* ValidateOperatorAndGetResultType(FEmitter& Emitter,
 			Arguments[i] = Emitter.Cast(Arguments[i], ArgumentPrimitiveType);
 			bValid &= Arguments[i];
 		}
+		else if (ArgumentPrimitiveType->ScalarKind == EScalarKind::ScalarKind_LWC && !(Filter & OPTF_LWCInput))
+		{
+		// Handle automatic cast to float for operators that don't support LWC inputs
+			ArgumentPrimitiveType = ArgumentPrimitiveType->WithScalarKind(EScalarKind::ScalarKind_Float);
+			Arguments[i] = Emitter.Cast(Arguments[i], ArgumentPrimitiveType);
+		}
 		else if (Filter & OPTF_CastToAnyFloat)
 		{
 			if (!ScalarKindIsAnyFloat(ArgumentPrimitiveType->ScalarKind))
@@ -1594,9 +1624,23 @@ static const FPrimitiveType* ValidateOperatorAndGetResultType(FEmitter& Emitter,
 			if (Filter & OPTF_CastToCommonType)
 			{
 				check(ArgumentsCommonType && ArgumentsCommonType->AsPrimitive());
-				Arguments[i] = Emitter.Cast(Arguments[i], ArgumentsCommonType);
+
+				// Lerp doesn't accept LWC for its third input, so we need to check if LWC is allowed per input when casting to the common type.
+				if (ArgumentsCommonType->IsLWC() && !(Filter & OPTF_LWCInput))
+				{
+					Arguments[i] = Emitter.Cast(Arguments[i], ArgumentsCommonType->AsPrimitive()->WithScalarKind(ScalarKind_Float));
+				}
+				else
+				{
+					Arguments[i] = Emitter.Cast(Arguments[i], ArgumentsCommonType);
+				}
 				bValid &= !Arguments[i]->IsPoison();
 			}
+			// Fmod doesn't accept LWC for its second input, so we need to potentially cast that to non-LWC.
+			else if (Arguments[i]->Type->IsLWC() && !(Filter & OPTF_LWCInput))
+			{
+				Arguments[i] = Emitter.CastToFloat(Arguments[i]);
+			}
 		}
 	}
 
@@ -1620,12 +1664,16 @@ static const FPrimitiveType* ValidateOperatorAndGetResultType(FEmitter& Emitter,
 			UE_MIR_UNREACHABLE(); // missing operator signature declaration
 
 		case ORT_FirstArgumentType:
+			// Some operators accept LWC as input, but always output non-LWC (examples:  modulus, sign, sine, cosine, saturate, frac).  Anything that may output LWC should use ORT_FirstArgumentTypeLWC.
+			return (FirstArgumentPrimitiveType->ScalarKind == EScalarKind::ScalarKind_LWC) ? FirstArgumentPrimitiveType->WithScalarKind(EScalarKind::ScalarKind_Float) : FirstArgumentPrimitiveType;
+
+		case ORT_FirstArgumentTypeLWC:
 			return FirstArgumentPrimitiveType;
 
 		case ORT_BooleanWithFirstArgumentDimensions:
 			return FPrimitiveType::Get(ScalarKind_Bool, FirstArgumentPrimitiveType->NumRows, FirstArgumentPrimitiveType->NumColumns);
 
-		case ORT_FirstArgumentTypeToScalar:
+ 		case ORT_FirstArgumentTypeToScalarLWC:
 			return FPrimitiveType::GetScalar(FirstArgumentPrimitiveType->ScalarKind);
 
 		case ORT_SecondArgumentType:
@@ -1859,6 +1907,21 @@ FValueRef FEmitter::Operator(EOperator Op, FValueRef A, FValueRef B, FValueRef C
 		return Poison();
 	}
 
+	// Subtract has a special case option to automatically truncate when subtracting two LWC inputs from each other, assuming this is a
+	// transition from LWC space to relative space, and no longer needs to be LWC.  We need to check that all arguments are LWC before
+	// the call to ValidateOperatorAndGetResultType, as that may cast the inputs, changing them.
+	bool bCastLWCOutputToFloat = false;
+	if (Op == BO_Subtract && B)
+	{
+		const FPrimitiveType* AType = A->Type->AsPrimitive();
+		const FPrimitiveType* BType = B->Type->AsPrimitive();
+	
+		if (AType && AType->ScalarKind == EScalarKind::ScalarKind_LWC && BType && BType->ScalarKind == EScalarKind::ScalarKind_LWC && UE::MaterialTranslatorUtils::GetLWCTruncateMode() == 2)
+		{
+			bCastLWCOutputToFloat = true;
+		}
+	}
+
 	// Validate the operation and retrieve the result type.
 	const FPrimitiveType* ResultType = ValidateOperatorAndGetResultType(*this, Op, A, B, C);
 	if (!ResultType)
@@ -1866,26 +1929,36 @@ FValueRef FEmitter::Operator(EOperator Op, FValueRef A, FValueRef B, FValueRef C
 		return Poison();
 	}
 
+	FValueRef Result;
+
 	// Try to apply some operator identity to simplify the operator.
 	if (FValueRef Simplified = TrySimplifyOperator(*this, Op, A, B, C))
 	{
-		return Simplified;
+		Result = Simplified;
 	}
-
 	// Try folding the operator first.
-	if (FValue* FoldedValue = TryFoldOperator(*this, Op, A, B, C, ResultType))
+	else if (FValue* FoldedValue = TryFoldOperator(*this, Op, A, B, C, ResultType))
 	{
-		return FoldedValue;
+		Result = FoldedValue;
 	}
+	else
+	{
+		// Otherwise, we must emit a new instruction that executes the operator.
+		FOperator Proto = MakePrototype<FOperator>(ResultType);
+		Proto.Op = Op;
+		Proto.AArg = A;
+		Proto.BArg = B;
+		Proto.CArg = C;
 
-	// Otherwise, we must emit a new instruction that executes the operator.
-	FOperator Proto = MakePrototype<FOperator>(ResultType);
-	Proto.Op = Op;
-	Proto.AArg = A;
-	Proto.BArg = B;
-	Proto.CArg = C;
+		Result = EmitPrototype(*this, Proto);
+	}
 
-	return EmitPrototype(*this, Proto);
+	// Optionally cast the output to float
+	if (bCastLWCOutputToFloat)
+	{
+		Result = CastToFloat(Result);
+	}
+	return Result;
 }
 
 FValueRef FEmitter::Branch(FValueRef Condition, FValueRef True, FValueRef False)
@@ -2109,8 +2182,9 @@ static FValue* CastValueToPrimitiveType(FEmitter& Emitter, FValueRef Value, cons
 		{
 			return Value;
 		}
-		// If value is a constant, we can cast the constant now.
-		else if (FConstant* ConstantInitializer = As<FConstant>(Value)) 
+		// If value is a constant, we can cast the constant now.  LWC is not a supported constant type at the moment, so in that case,
+		// we go through the dynamic cast code path below.
+		else if (FConstant* ConstantInitializer = TargetPrimitiveType->ScalarKind != EScalarKind::ScalarKind_LWC ? As<FConstant>(Value) : nullptr)
 		{
 			return CastConstant(Emitter, ConstantInitializer, ValuePrimitiveType->ScalarKind, TargetPrimitiveType->ScalarKind);
 		}
@@ -2521,16 +2595,10 @@ FValueRef FEmitter::PartialDerivative(FValueRef Value, EDerivativeAxis Axis)
 
 static FValue* DifferentiateExternalInput(FEmitter& Emitter, FExternalInput* ExternalInput, EDerivativeAxis Axis)
 {
-	// Texture coordinate external inputs have their own matching DDX/DDY inputs.
-	if (IsExternalInputTexCoord(ExternalInput->Id))
+	// External inputs with derivatives have their own matching DDX/DDY inputs.
+	if (IsExternalInputWithDerivatives(ExternalInput->Id))
 	{
-		int32 TexCoordIndex = ExternalInputToTexCoordIndex(ExternalInput->Id);
-
-		EExternalInput PartialDerivativeExternalInput = (Axis == EDerivativeAxis::X) 
-			? EExternalInput((int32)EExternalInput::TexCoord0_Ddx + TexCoordIndex)
-			: EExternalInput((int32)EExternalInput::TexCoord0_Ddy + TexCoordIndex);
-
-		return Emitter.ExternalInput(PartialDerivativeExternalInput);
+		return Emitter.ExternalInput(GetExternalInputDerivative(ExternalInput->Id, Axis));
 	}
 
 	// All other inputs are assumed constant.
@@ -2709,7 +2777,7 @@ static FValue* DifferentiateOperator(FEmitter& E, FOperator* Op, EDerivativeAxis
 		}
 
 		// d/dx atan2(f(x), g(x)) = g(x) / (f(x)^2 + g(x)^2) * f'(x)  -  f(x) / (f(x)^2 + g(x)^2) * g'(x)
-		case BO_Atan2:
+		case BO_ATan2:
 		{
 			FValue* Magnitude = E.Divide(One(), E.Add(E.Multiply(F, F), E.Multiply(G, G)));		// 1 / (f(x)^2 + g(x)^2)
 			return E.Subtract(E.Multiply(E.Multiply(G, Magnitude), dF), E.Multiply(E.Multiply(F, Magnitude), dG));
diff --git a/Engine/Source/Runtime/Engine/Private/Materials/MaterialIRToHLSLTranslator.cpp b/Engine/Source/Runtime/Engine/Private/Materials/MaterialIRToHLSLTranslator.cpp
--- a/Engine/Source/Runtime/Engine/Private/Materials/MaterialIRToHLSLTranslator.cpp
+++ b/Engine/Source/Runtime/Engine/Private/Materials/MaterialIRToHLSLTranslator.cpp
@@ -453,11 +453,16 @@ struct FTranslator : FMaterialIRToHLSLTranslation
 					Printer.Buffer.Appendf(TEXT("%.9g"), Constant->Float);
 				}
 				break;
+
+			case MIR::ScalarKind_LWC: UE_MIR_UNREACHABLE();
 		}
 	}
 
 	void LowerExternalInput(const MIR::FExternalInput* ExternalInput)
 	{
+		// @todo-laura.hermanns - Must be replaced as soon as this state is available in MIR
+		constexpr bool bCompilingPreviousFrame = false;
+
 		int32 ExternalInputIndex = (int32)ExternalInput->Id;
 		if (MIR::IsExternalInputTexCoord(ExternalInput->Id))
 		{
@@ -474,11 +479,63 @@ struct FTranslator : FMaterialIRToHLSLTranslation
 			int32 Index = ExternalInputIndex - (int32)MIR::EExternalInput::TexCoord0_Ddy;
 			Printer.Buffer.Appendf(TEXT("Parameters.TexCoords_DDY[%d]"), Index);
 		}
+		else if (ExternalInput->Id >= MIR::EExternalInput::WorldPosition_Absolute && ExternalInput->Id <= MIR::EExternalInput::WorldPosition_CameraRelativeNoOffsets)
+		{
+			// Various function permutations exist to fetch world position.  "Prev" permutations are only available in the vertex shader, while
+			// "MoMaterialOffsets" permutations are only available in the pixel shader, so we need to cobble together the permutation string
+			// factoring in those limitations.  Logic adapted from FHLSLMaterialTranslator::WorldPosition.
+			// 
+			// Format:  Get[Prev][Translated]WorldPosition[_NoMaterialOffsets](Parameters)
+			//
+			Printer.Buffer.Append(TEXT("Get"));
+
+			// If compiling for the previous frame in the vertex shader
+			if (CurrentStage == MIR::EStage::Stage_Vertex && bCompilingPreviousFrame)
+			{
+				Printer.Buffer.Append(TEXT("Prev"));
+			}
+
+			if (ExternalInput->Id == MIR::EExternalInput::WorldPosition_CameraRelative || ExternalInput->Id == MIR::EExternalInput::WorldPosition_CameraRelativeNoOffsets)
+			{
+				Printer.Buffer.Append(TEXT("Translated"));
+			}
+
+			Printer.Buffer.Append(TEXT("WorldPosition"));
+
+			if (CurrentStage == MIR::EStage::Stage_Pixel && (ExternalInput->Id == MIR::EExternalInput::WorldPosition_AbsoluteNoOffsets || ExternalInput->Id == MIR::EExternalInput::WorldPosition_CameraRelativeNoOffsets))
+			{
+				Printer.Buffer.Append(TEXT("_NoMaterialOffsets"));
+			}
+
+			Printer.Buffer.Append(TEXT("(Parameters)"));
+		}
+		else if (ExternalInput->Id >= MIR::EExternalInput::WorldPosition_Absolute_Ddx && ExternalInput->Id <= MIR::EExternalInput::WorldPosition_CameraRelativeNoOffsets_Ddx)
+		{
+			// Expression emitter assumes these to be LWC for analytic derivative evaluation, so we promote them on load
+			Printer.Buffer.Append(TEXT("WSPromote(Parameters.WorldPosition_DDX)"));
+		}
+		else if (ExternalInput->Id >= MIR::EExternalInput::WorldPosition_Absolute_Ddy && ExternalInput->Id <= MIR::EExternalInput::WorldPosition_CameraRelativeNoOffsets_Ddy)
+		{
+			// Expression emitter assumes these to be LWC for analytic derivative evaluation, so we promote them on load
+			Printer.Buffer.Append(TEXT("WSPromote(Parameters.WorldPosition_DDY)"));
+		}
+		else if (ExternalInput->Id == MIR::EExternalInput::ActorPosition_Absolute)
+		{
+			const bool bPrevious = (CurrentStage == MIR::EStage::Stage_Vertex && bCompilingPreviousFrame);
+			Printer.Buffer.Append(bPrevious ? TEXT("GetPreviousActorWorldPosition(Parameters)") : TEXT("GetActorWorldPosition(Parameters)"));
+		}
+		else if (ExternalInput->Id == MIR::EExternalInput::ActorPosition_CameraRelative)
+		{
+			const bool bPrevious = (CurrentStage == MIR::EStage::Stage_Vertex && bCompilingPreviousFrame);
+			Printer.Buffer.Append(bPrevious ? TEXT("GetPreviousActorTranslatedWorldPosition(Parameters)") : TEXT("GetActorTranslatedWorldPosition(Parameters)"));
+		}
 		else
 		{
 			const TCHAR* Code{};
 			switch (ExternalInput->Id)
 			{
+				case MIR::EExternalInput::ObjectPosition_Absolute: Code = TEXT("GetObjectWorldPosition(Parameters)"); break;
+				case MIR::EExternalInput::ObjectPosition_CameraRelative: Code = TEXT("GetObjectTranslatedWorldPosition(Parameters)"); break;
 				case MIR::EExternalInput::ViewMaterialTextureMipBias: Code = TEXT("View.MaterialTextureMipBias"); break;
 				case MIR::EExternalInput::ViewMaterialTextureDerivativeMultiply: Code = TEXT("View.MaterialTextureDerivativeMultiply"); break;
 				default: UE_MIR_UNREACHABLE();
@@ -807,20 +864,37 @@ struct FTranslator : FMaterialIRToHLSLTranslation
 		
 		// In order to generate smaller and tidier HLSL, first check whether all components
 		// of this dimensional are actually the same. If so, we can simply emit the 
-		// component and cast it to the type.
+		// component and cast it to the type.  LWC doesn't support casting, and always needs
+		// to call a function to convert types.
 		bool bSameComponents = true;
 		for (int32 i = 1; bSameComponents && i < Components.Num(); ++i)
 		{
 			bSameComponents &= (Components[i] == Components[0]);
 		}
 		
-		if (bSameComponents)
+		if (bSameComponents && ArithmeticType->ScalarKind != MIR::EScalarKind::ScalarKind_LWC)
 		{
 			Printer << TEXT("(") << LowerType(ArithmeticType) << TEXT(")") << LowerValue(Components[0]);
 		}
 		else
 		{
-			Printer << LowerType(ArithmeticType) << BeginArgs;
+			if (ArithmeticType->ScalarKind == MIR::EScalarKind::ScalarKind_LWC)
+			{
+				if (ArithmeticType->NumColumns == 1)
+				{
+					Printer << (ArithmeticType->NumRows == 1 ? TEXT("MakeWSScalar") : TEXT("MakeWSVector")) << BeginArgs;
+				}
+				else
+				{
+					// @todo-jason.hoerner - LWC Matrix support.  I don't think there are any nodes that can build LWC matrices from scratch, these
+					// generally can only come from external inputs or material parameter collections, so this isn't necessary at the moment.
+					UE_MIR_UNREACHABLE();
+				}
+			}
+			else
+			{
+				Printer << LowerType(ArithmeticType) << BeginArgs;
+			}
 			for (MIR::FValue* Component : Components)
 			{
 				Printer << ListSeparator << LowerValue(Component);
@@ -877,7 +951,9 @@ struct FTranslator : FMaterialIRToHLSLTranslation
 
 	void LowerOperator(const MIR::FOperator* Operator)
 	{
-		if (IsOperatorInfix(Operator->Op))
+		bool bIsLWC = Operator->AArg->Type->IsLWC() || (Operator->Op == MIR::TO_Select && Operator->BArg->Type->IsLWC());
+
+		if (IsOperatorInfix(Operator->Op) && !bIsLWC)
 		{
 			const TCHAR* OpString;
 			switch (Operator->Op)
@@ -901,57 +977,110 @@ struct FTranslator : FMaterialIRToHLSLTranslation
 		else
 		{
 			const TCHAR* OpString;
-			switch (Operator->Op)
+			if (bIsLWC)
 			{
-				case MIR::UO_Abs: OpString = TEXT("abs"); break;
-				case MIR::UO_ACos: OpString = TEXT("acos"); break;
-				case MIR::UO_ACosh: OpString = TEXT("acosh"); break;
-				case MIR::UO_ASin: OpString = TEXT("asin"); break;
-				case MIR::UO_ASinh: OpString = TEXT("asinh"); break;
-				case MIR::UO_ATan: OpString = TEXT("atan"); break;
-				case MIR::UO_ATanh: OpString = TEXT("atanh"); break;
-				case MIR::UO_Ceil: OpString = TEXT("ceil"); break;
-				case MIR::UO_Cos: OpString = TEXT("cos"); break;
-				case MIR::UO_Cosh: OpString = TEXT("cosh"); break;
-				case MIR::UO_Exponential: OpString = TEXT("exp"); break;
-				case MIR::UO_Exponential2: OpString = TEXT("exp2"); break;
-				case MIR::UO_Floor: OpString = TEXT("floor"); break;
-				case MIR::UO_Frac: OpString = TEXT("frac"); break;
-				case MIR::UO_IsFinite: OpString = TEXT("isfinite"); break;
-				case MIR::UO_IsInf: OpString = TEXT("isinf"); break;
-				case MIR::UO_IsNan: OpString = TEXT("isnan"); break;
-				case MIR::UO_Length: OpString = TEXT("length"); break;
-				case MIR::UO_Logarithm: OpString = TEXT("log"); break;
-				case MIR::UO_Logarithm10: OpString = TEXT("log10"); break;
-				case MIR::UO_Logarithm2: OpString = TEXT("log2"); break;
-				case MIR::UO_Round: OpString = TEXT("round"); break;
-				case MIR::UO_Saturate: OpString = TEXT("saturate"); break;
-				case MIR::UO_Sign: OpString = TEXT("sign"); break;
-				case MIR::UO_Sin: OpString = TEXT("sin"); break;
-				case MIR::UO_Sinh: OpString = TEXT("sinh"); break;
-				case MIR::UO_Sqrt: OpString = TEXT("sqrt"); break;
-				case MIR::UO_Tan: OpString = TEXT("tan"); break;
-				case MIR::UO_Tanh: OpString = TEXT("tanh"); break;
-				case MIR::UO_Truncate: OpString = TEXT("trunc"); break;
-
-				case MIR::BO_And: OpString = TEXT("and"); break;
-				case MIR::BO_Atan2: OpString = TEXT("atan2"); break;
-				case MIR::BO_Cross: OpString = TEXT("cross"); break;
-				case MIR::BO_Distance: OpString = TEXT("distance"); break;
-				case MIR::BO_Dot: OpString = TEXT("dot"); break;
-				case MIR::BO_Fmod: OpString = TEXT("fmod"); break;
-				case MIR::BO_Max: OpString = TEXT("max"); break;
-				case MIR::BO_Min: OpString = TEXT("min"); break;
-				case MIR::BO_Or: OpString = TEXT("or"); break;
-				case MIR::BO_Pow: OpString = TEXT("pow"); break;
-				case MIR::BO_Step: OpString = TEXT("step"); break;
-
-				case MIR::TO_Clamp: OpString = TEXT("clamp"); break;
-				case MIR::TO_Lerp: OpString = TEXT("lerp"); break;
-				case MIR::TO_Select: OpString = TEXT("select"); break;
-				case MIR::TO_Smoothstep: OpString = TEXT("smoothstep"); break;
+				// "Demotes LWC" indicates the given operator returns a non-LWC float, even if the input is LWC.  Besides that, comparison
+				// operators all return bool, instead of LWC, and a couple specific operators require specific inputs to always be non-LWC
+				// (second argument of Fmod, and third argument of Lerp).
+				switch (Operator->Op)
+				{
+					case MIR::UO_Abs: OpString = TEXT("WSAbs"); break;
+					case MIR::UO_ACos: OpString = TEXT("WSACos"); break;						// Demotes LWC
+					case MIR::UO_ASin: OpString = TEXT("WSASin"); break;						// Demotes LWC
+					case MIR::UO_ATan: OpString = TEXT("WSATan"); break;						// Demotes LWC
+					case MIR::UO_Ceil: OpString = TEXT("WSCeil"); break;
+					case MIR::UO_Cos: OpString = TEXT("WSCos"); break;							// Demotes LWC
+					case MIR::UO_Floor: OpString = TEXT("WSFloor"); break;
+					case MIR::UO_Frac: OpString = TEXT("WSFracDemote"); break;					// Demotes LWC
+					case MIR::UO_Length: OpString = TEXT("WSLength"); break;
+					case MIR::UO_Negate: OpString = TEXT("WSNegate"); break;
+					case MIR::UO_Round: OpString = TEXT("WSRound"); break;
+					case MIR::UO_Saturate: OpString = TEXT("WSSaturateDemote"); break;			// Demotes LWC
+					case MIR::UO_Sign: OpString = TEXT("WSSign"); break;						// Demotes LWC
+					case MIR::UO_Sin: OpString = TEXT("WSSin"); break;							// Demotes LWC
+					case MIR::UO_Sqrt: OpString = TEXT("WSSqrtDemote"); break;					// Demotes LWC
+					case MIR::UO_Tan: OpString = TEXT("WSTan"); break;							// Demotes LWC
+					case MIR::UO_Truncate: OpString = TEXT("WSTrunc"); break;
+
+					case MIR::BO_Add: OpString = TEXT("WSAdd"); break;
+					case MIR::BO_Distance: OpString = TEXT("WSDistance"); break;
+					case MIR::BO_Divide: OpString = TEXT("WSDivide"); break;
+					case MIR::BO_Dot: OpString = TEXT("WSDot"); break;
+					case MIR::BO_Equals: OpString = TEXT("WSEquals"); break;					// Bool output
+					case MIR::BO_Fmod: OpString = TEXT("WSFmodDemote"); break;					// Demotes LWC, second input must be float (not LWC)!
+					case MIR::BO_GreaterThan: OpString = TEXT("WSGreater"); break;				// Bool output
+					case MIR::BO_GreaterThanOrEquals: OpString = TEXT("WSGreaterEqual"); break;	// Bool output
+					case MIR::BO_LessThan: OpString = TEXT("WSLess"); break;					// Bool output
+					case MIR::BO_LessThanOrEquals: OpString = TEXT("WSLessEqual"); break;		// Bool output
+					case MIR::BO_Max: OpString = TEXT("WSMax"); break;
+					case MIR::BO_Min: OpString = TEXT("WSMin"); break;
+					case MIR::BO_Multiply: OpString = TEXT("WSMultiply"); break;
+					case MIR::BO_NotEquals: OpString = TEXT("WSNotEquals"); break;				// Bool output
+					case MIR::BO_Step: OpString = TEXT("WSStep"); break;						// Demotes LWC
+					case MIR::BO_Subtract: OpString = TEXT("WSSubtract"); break;
+
+					case MIR::TO_Clamp: OpString = TEXT("WSClamp"); break;
+					case MIR::TO_Lerp: OpString = TEXT("WSLerp"); break;						// Third input must be float (not LWC)!
+					case MIR::TO_Select: OpString = TEXT("WSSelect"); break;
+					case MIR::TO_Smoothstep: OpString = TEXT("WSSmoothStepDemote"); break;		// Demotes LWC
+
+					default: UE_MIR_UNREACHABLE();
+				}
+			}
+			else
+			{
+				switch (Operator->Op)
+				{
+					case MIR::UO_Abs: OpString = TEXT("abs"); break;
+					case MIR::UO_ACos: OpString = TEXT("acos"); break;
+					case MIR::UO_ACosh: OpString = TEXT("acosh"); break;
+					case MIR::UO_ASin: OpString = TEXT("asin"); break;
+					case MIR::UO_ASinh: OpString = TEXT("asinh"); break;
+					case MIR::UO_ATan: OpString = TEXT("atan"); break;
+					case MIR::UO_ATanh: OpString = TEXT("atanh"); break;
+					case MIR::UO_Ceil: OpString = TEXT("ceil"); break;
+					case MIR::UO_Cos: OpString = TEXT("cos"); break;
+					case MIR::UO_Cosh: OpString = TEXT("cosh"); break;
+					case MIR::UO_Exponential: OpString = TEXT("exp"); break;
+					case MIR::UO_Exponential2: OpString = TEXT("exp2"); break;
+					case MIR::UO_Floor: OpString = TEXT("floor"); break;
+					case MIR::UO_Frac: OpString = TEXT("frac"); break;
+					case MIR::UO_IsFinite: OpString = TEXT("isfinite"); break;
+					case MIR::UO_IsInf: OpString = TEXT("isinf"); break;
+					case MIR::UO_IsNan: OpString = TEXT("isnan"); break;
+					case MIR::UO_Length: OpString = TEXT("length"); break;
+					case MIR::UO_Logarithm: OpString = TEXT("log"); break;
+					case MIR::UO_Logarithm10: OpString = TEXT("log10"); break;
+					case MIR::UO_Logarithm2: OpString = TEXT("log2"); break;
+					case MIR::UO_Round: OpString = TEXT("round"); break;
+					case MIR::UO_Saturate: OpString = TEXT("saturate"); break;
+					case MIR::UO_Sign: OpString = TEXT("sign"); break;
+					case MIR::UO_Sin: OpString = TEXT("sin"); break;
+					case MIR::UO_Sinh: OpString = TEXT("sinh"); break;
+					case MIR::UO_Sqrt: OpString = TEXT("sqrt"); break;
+					case MIR::UO_Tan: OpString = TEXT("tan"); break;
+					case MIR::UO_Tanh: OpString = TEXT("tanh"); break;
+					case MIR::UO_Truncate: OpString = TEXT("trunc"); break;
+
+					case MIR::BO_And: OpString = TEXT("and"); break;
+					case MIR::BO_ATan2: OpString = TEXT("atan2"); break;
+					case MIR::BO_Cross: OpString = TEXT("cross"); break;
+					case MIR::BO_Distance: OpString = TEXT("distance"); break;
+					case MIR::BO_Dot: OpString = TEXT("dot"); break;
+					case MIR::BO_Fmod: OpString = TEXT("fmod"); break;
+					case MIR::BO_Max: OpString = TEXT("max"); break;
+					case MIR::BO_Min: OpString = TEXT("min"); break;
+					case MIR::BO_Or: OpString = TEXT("or"); break;
+					case MIR::BO_Pow: OpString = TEXT("pow"); break;
+					case MIR::BO_Step: OpString = TEXT("step"); break;
+
+					case MIR::TO_Clamp: OpString = TEXT("clamp"); break;
+					case MIR::TO_Lerp: OpString = TEXT("lerp"); break;
+					case MIR::TO_Select: OpString = TEXT("select"); break;
+					case MIR::TO_Smoothstep: OpString = TEXT("smoothstep"); break;
 			
-				default: UE_MIR_UNREACHABLE();
+					default: UE_MIR_UNREACHABLE();
+				}
 			}
 
 			// Unary
@@ -979,9 +1108,19 @@ struct FTranslator : FMaterialIRToHLSLTranslation
 	{
 		if (IsFoldable(Branch, CurrentStage))
 		{
-			Printer << LowerValue(Branch->ConditionArg)
-				<< TEXT(" ? ") << LowerValue(Branch->TrueArg)
-				<< TEXT(" : ") << LowerValue(Branch->FalseArg);
+			if (Branch->TrueArg->Type->IsLWC())
+			{
+				Printer << TEXT("WSSelect(") << LowerValue(Branch->ConditionArg)
+					<< TEXT(", ") << LowerValue(Branch->TrueArg)
+					<< TEXT(", ") << LowerValue(Branch->FalseArg)
+					<< TEXT(")");
+			}
+			else
+			{
+				Printer << LowerValue(Branch->ConditionArg)
+					<< TEXT(" ? ") << LowerValue(Branch->TrueArg)
+					<< TEXT(" : ") << LowerValue(Branch->FalseArg);
+			}
 		}
 		else
 		{
@@ -999,20 +1138,43 @@ struct FTranslator : FMaterialIRToHLSLTranslation
 			
 	void LowerSubscript(const MIR::FSubscript* Subscript)
 	{
-		LowerValue(Subscript->Arg);
+		const MIR::FPrimitiveType* ArgArithmeticType = Subscript->Arg->Type->AsVector();
 
-		if (const MIR::FPrimitiveType* ArgArithmeticType = Subscript->Arg->Type->AsVector())
+		if (ArgArithmeticType && ArgArithmeticType->ScalarKind == MIR::EScalarKind::ScalarKind_LWC)
 		{
-			const TCHAR* ComponentsStr[] = { TEXT(".x"), TEXT(".y"), TEXT(".z"), TEXT(".w") };
+			const TCHAR* LwcComponentsStr[] = { TEXT("WSGetX("), TEXT("WSGetY("), TEXT("WSGetZ("), TEXT("WSGetW(") };
 			check(Subscript->Index <= ArgArithmeticType->GetNumComponents());
 
-			Printer << ComponentsStr[Subscript->Index];
+			Printer << LwcComponentsStr[Subscript->Index] << LowerValue(Subscript->Arg) << TEXT(")");
+		}
+		else
+		{
+			LowerValue(Subscript->Arg);
+
+			if (ArgArithmeticType)
+			{
+				const TCHAR* ComponentsStr[] = { TEXT(".x"), TEXT(".y"), TEXT(".z"), TEXT(".w") };
+				check(Subscript->Index <= ArgArithmeticType->GetNumComponents());
+
+				Printer << ComponentsStr[Subscript->Index];
+			}
 		}
 	}
 
 	void LowerScalar(const MIR::FScalar* Scalar)
 	{
-		Printer << LowerType(Scalar->Type) << "(" << LowerValue(Scalar->Arg) << ")";
+		if (Scalar->Arg->Type->AsPrimitive()->ScalarKind == MIR::EScalarKind::ScalarKind_LWC)
+		{
+			Printer << LowerType(Scalar->Type) << "(WSDemote(" << LowerValue(Scalar->Arg) << "))";
+		}
+		else if (Scalar->Type->AsPrimitive()->ScalarKind == MIR::EScalarKind::ScalarKind_LWC)
+		{
+			Printer << "WSPromote(" << LowerValue(Scalar->Arg) << ")";
+		}
+		else
+		{
+			Printer << LowerType(Scalar->Type) << "(" << LowerValue(Scalar->Arg) << ")";
+		}
 	}
 
 	void LowerStandardTextureRead(const MIR::FTextureRead* TextureRead)
@@ -1294,7 +1456,15 @@ struct FTranslator : FMaterialIRToHLSLTranslation
 
 	void LowerHardwarePartialDerivative(const MIR::FHardwarePartialDerivative* StageSwitch)
 	{
-		Printer << (StageSwitch->Axis == MIR::EDerivativeAxis::X ? "ddx(" : "ddy(") << LowerValue(StageSwitch->Arg) << ")";
+		if (StageSwitch->Arg->Type->IsLWC())
+		{
+			// Expression emitter assumes these to be LWC for analytic derivative evaluation, so we promote them on load
+			Printer << (StageSwitch->Axis == MIR::EDerivativeAxis::X ? "WSPromote(WSDdxDemote(" : "WSPromote(WSDdyDemote(") << LowerValue(StageSwitch->Arg) << "))";
+		}
+		else
+		{
+			Printer << (StageSwitch->Axis == MIR::EDerivativeAxis::X ? "ddx(" : "ddy(") << LowerValue(StageSwitch->Arg) << ")";
+		}
 	}
 
 	/* Finalization */
@@ -1405,21 +1575,43 @@ struct FTranslator : FMaterialIRToHLSLTranslation
 	{
 		if (const MIR::FPrimitiveType* ArithmeticType = Type->AsPrimitive())
 		{
-			switch (ArithmeticType->ScalarKind)
+			if (ArithmeticType->ScalarKind == MIR::ScalarKind_LWC)
 			{
-				case MIR::ScalarKind_Bool:	Printer << TEXT("bool"); break;
-				case MIR::ScalarKind_Int: 	Printer << TEXT("int32"); break;
-				case MIR::ScalarKind_Float:	Printer << TEXT("MaterialFloat"); break;
+				if (ArithmeticType->NumRows == 1)
+				{
+					Printer << TEXT("FWSScalar");
+				}
+				else if (ArithmeticType->NumColumns == 1)
+				{
+					Printer << TEXT("FWSVector") << ArithmeticType->NumRows;
+				}
+				else if (ArithmeticType->bIsLWCInverseMatrix)
+				{
+					Printer << TEXT("FWSInverseMatrix");
+				}
+				else
+				{
+					Printer << TEXT("FWSMatrix");
+				}
 			}
-
-			if (ArithmeticType->NumRows > 1)
+			else
 			{
-				Printer << ArithmeticType->NumRows;
-			}
+				switch (ArithmeticType->ScalarKind)
+				{
+					case MIR::ScalarKind_Bool:	Printer << TEXT("bool"); break;
+					case MIR::ScalarKind_Int: 	Printer << TEXT("int32"); break;
+					case MIR::ScalarKind_Float:	Printer << TEXT("MaterialFloat"); break;
+				}
+
+				if (ArithmeticType->NumRows > 1)
+				{
+					Printer << ArithmeticType->NumRows;
+				}
 			
-			if (ArithmeticType->NumColumns > 1)
-			{
-				Printer << TEXT("x") << ArithmeticType->NumColumns;
+				if (ArithmeticType->NumColumns > 1)
+				{
+					Printer << TEXT("x") << ArithmeticType->NumColumns;
+				}
 			}
 		}
 		else if (const MIR::FObjectType* ObjectType = Type->AsObject())
@@ -1531,8 +1723,7 @@ struct FTranslator : FMaterialIRToHLSLTranslation
 			OutEnvironment.SetDefine(GetShadingModelParameterName(MSM_DefaultLit), true);
 		}
 
-		static IConsoleVariable* CVarLWCIsEnabled = IConsoleManager::Get().FindConsoleVariable(TEXT("r.MaterialEditor.LWCEnabled"));
-		OutEnvironment.SetDefine(TEXT("MATERIAL_LWC_ENABLED"), CVarLWCIsEnabled->GetInt());
+		OutEnvironment.SetDefine(TEXT("MATERIAL_LWC_ENABLED"), UE::MaterialTranslatorUtils::IsLWCEnabled() ? 1 : 0);
 		OutEnvironment.SetDefine(TEXT("WSVECTOR_IS_TILEOFFSET"), true);
 		OutEnvironment.SetDefine(TEXT("WSVECTOR_IS_DOUBLEFLOAT"), false);
 
diff --git a/Engine/Source/Runtime/Engine/Private/Materials/MaterialIRTypes.cpp b/Engine/Source/Runtime/Engine/Private/Materials/MaterialIRTypes.cpp
--- a/Engine/Source/Runtime/Engine/Private/Materials/MaterialIRTypes.cpp
+++ b/Engine/Source/Runtime/Engine/Private/Materials/MaterialIRTypes.cpp
@@ -44,7 +44,13 @@ const FType* FType::FromShaderType(const UE::Shader::FType& InShaderType)
 		case UE::Shader::EValueType::Bool2:
 		case UE::Shader::EValueType::Bool3:
 		case UE::Shader::EValueType::Bool4:
-			return FPrimitiveType::GetVector(ScalarKind_Bool, (int)InShaderType.ValueType - (int)UE::Shader::EValueType::Int1 + 1);
+			return FPrimitiveType::GetVector(ScalarKind_Bool, (int)InShaderType.ValueType - (int)UE::Shader::EValueType::Bool1 + 1);
+
+		case UE::Shader::EValueType::Double1:
+		case UE::Shader::EValueType::Double2:
+		case UE::Shader::EValueType::Double3:
+		case UE::Shader::EValueType::Double4:
+			return FPrimitiveType::GetVector(ScalarKind_LWC, (int)InShaderType.ValueType - (int)UE::Shader::EValueType::Double1 + 1);
 
 		default:
 			UE_MIR_UNREACHABLE();
@@ -85,6 +91,15 @@ const FType* FType::FromMaterialValueType(EMaterialValueType Type)
 		case MCT_Bool:
 			return FPrimitiveType::GetVector(ScalarKind_Bool, 1);
 
+		case MCT_LWCScalar:
+			return FPrimitiveType::GetVector(ScalarKind_LWC, 1);
+		case MCT_LWCVector2:
+			return FPrimitiveType::GetVector(ScalarKind_LWC, 2);
+		case MCT_LWCVector3:
+			return FPrimitiveType::GetVector(ScalarKind_LWC, 3);
+		case MCT_LWCVector4:
+			return FPrimitiveType::GetVector(ScalarKind_LWC, 4);
+
 		default:
 			UE_MIR_UNREACHABLE();
 	}
@@ -149,6 +164,10 @@ UE::Shader::EValueType FType::ToValueType() const
 				{
 					return EValueType::Float4x4;
 				}
+				else if (PrimitiveType->ScalarKind == ScalarKind_LWC)
+				{
+					return PrimitiveType->bIsLWCInverseMatrix ? EValueType::DoubleInverse4x4 : EValueType::Double4x4;
+				}
 				else
 				{
 					return EValueType::Numeric4x4;
@@ -165,6 +184,7 @@ UE::Shader::EValueType FType::ToValueType() const
 			case ScalarKind_Bool: 	return (EValueType)((int)EValueType::Bool1 + PrimitiveType->NumRows - 1);
 			case ScalarKind_Int: 	return (EValueType)((int)EValueType::Int1 + PrimitiveType->NumRows - 1);
 			case ScalarKind_Float: 	return (EValueType)((int)EValueType::Float1 + PrimitiveType->NumRows - 1);
+			case ScalarKind_LWC:	return (EValueType)((int)EValueType::Double1 + PrimitiveType->NumRows - 1);
 			default: UE_MIR_UNREACHABLE();
 		}
 	}
@@ -199,6 +219,12 @@ bool FType::IsInteger() const
 	return Ptr && Ptr->ScalarKind == ScalarKind_Int;
 }
 
+bool FType::IsLWC() const
+{
+	const FPrimitiveType* Ptr = AsPrimitive();
+	return Ptr && Ptr->ScalarKind == ScalarKind_LWC;
+}
+
 bool FType::IsBoolScalar() const
 {
 	return this == FPrimitiveType::GetBool();
@@ -264,7 +290,7 @@ bool ScalarKindIsArithmetic(EScalarKind Kind)
 
 bool ScalarKindIsAnyFloat(EScalarKind Kind)
 {
-	return Kind == ScalarKind_Float;
+	return Kind == ScalarKind_Float || Kind == ScalarKind_LWC;
 }
 
 const TCHAR* ScalarKindToString(EScalarKind Kind)
@@ -274,6 +300,7 @@ const TCHAR* ScalarKindToString(EScalarKind Kind)
 		case ScalarKind_Bool: return TEXT("bool");
 		case ScalarKind_Int: return TEXT("int");
 		case ScalarKind_Float: return TEXT("MaterialFloat");
+		case ScalarKind_LWC: return TEXT("FLWCScalar");
 		default: UE_MIR_UNREACHABLE();
 	}
 }
@@ -308,6 +335,11 @@ const FPrimitiveType* FPrimitiveType::GetFloat4()
 	return GetVector(ScalarKind_Float, 4);
 }
 
+const FPrimitiveType* FPrimitiveType::GetLWC3()
+{
+	return GetVector(ScalarKind_LWC, 3);
+}
+
 const FPrimitiveType* FPrimitiveType::GetScalar(EScalarKind InScalarKind)
 {
 	return Get(InScalarKind, 1, 1);
@@ -319,9 +351,9 @@ const FPrimitiveType* FPrimitiveType::GetVector(EScalarKind InScalarKind, int Nu
 	return Get(InScalarKind, NumRows, 1);
 }
 
-const FPrimitiveType* FPrimitiveType::Get(EScalarKind InScalarKind, int NumRows, int NumColumns)
+const FPrimitiveType* FPrimitiveType::Get(EScalarKind InScalarKind, int NumRows, int NumColumns, bool bIsLWCInverseMatrix)
 {
-	check(InScalarKind >= 0 && InScalarKind <= ScalarKind_Float);
+	check(InScalarKind >= 0 && InScalarKind <= ScalarKind_LWC);
 	
 	static const FStringView Invalid = TEXT("invalid");
 
@@ -359,9 +391,9 @@ const FPrimitiveType* FPrimitiveType::Get(EScalarKind InScalarKind, int NumRows,
 		{ { TypeKind_Primitive }, { TEXT("int4x3") }, 	ScalarKind_Int, 4, 3 },
 		{ { TypeKind_Primitive }, { TEXT("int4x4") }, 	ScalarKind_Int, 4, 4 },
 		{ { TypeKind_Primitive }, { TEXT("float") }, 		ScalarKind_Float, 1, 1 },
-		{ { TypeKind_Primitive }, Invalid, 				ScalarKind_Float, 1, 2 },
-		{ { TypeKind_Primitive }, Invalid, 				ScalarKind_Float, 1, 3 },
-		{ { TypeKind_Primitive }, Invalid, 				ScalarKind_Float, 1, 4 },
+		{ { TypeKind_Primitive }, Invalid, 					ScalarKind_Float, 1, 2 },
+		{ { TypeKind_Primitive }, Invalid, 					ScalarKind_Float, 1, 3 },
+		{ { TypeKind_Primitive }, Invalid, 					ScalarKind_Float, 1, 4 },
 		{ { TypeKind_Primitive }, { TEXT("float2") },   	ScalarKind_Float, 2, 1 },
 		{ { TypeKind_Primitive }, { TEXT("float2x2") }, 	ScalarKind_Float, 2, 2 },
 		{ { TypeKind_Primitive }, { TEXT("float2x3") }, 	ScalarKind_Float, 2, 3 },
@@ -374,8 +406,33 @@ const FPrimitiveType* FPrimitiveType::Get(EScalarKind InScalarKind, int NumRows,
 		{ { TypeKind_Primitive }, { TEXT("float4x2") }, 	ScalarKind_Float, 4, 2 },
 		{ { TypeKind_Primitive }, { TEXT("float4x3") }, 	ScalarKind_Float, 4, 3 },
 		{ { TypeKind_Primitive }, { TEXT("float4x4") }, 	ScalarKind_Float, 4, 4 },
+		{ { TypeKind_Primitive }, { TEXT("FWSScalar") },		ScalarKind_LWC, 1, 1 },
+		{ { TypeKind_Primitive }, Invalid, 						ScalarKind_LWC, 1, 2 },
+		{ { TypeKind_Primitive }, Invalid, 						ScalarKind_LWC, 1, 3 },
+		{ { TypeKind_Primitive }, Invalid, 						ScalarKind_LWC, 1, 4 },
+		{ { TypeKind_Primitive }, { TEXT("FWSVector2") },		ScalarKind_LWC, 2, 1 },
+		{ { TypeKind_Primitive }, Invalid,						ScalarKind_LWC, 2, 2 },
+		{ { TypeKind_Primitive }, Invalid,						ScalarKind_LWC, 2, 3 },
+		{ { TypeKind_Primitive }, Invalid,						ScalarKind_LWC, 2, 4 },
+		{ { TypeKind_Primitive }, { TEXT("FWSVector3") },		ScalarKind_LWC, 3, 1 },
+		{ { TypeKind_Primitive }, Invalid,						ScalarKind_LWC, 3, 2 },
+		{ { TypeKind_Primitive }, Invalid,						ScalarKind_LWC, 3, 3 },
+		{ { TypeKind_Primitive }, Invalid,						ScalarKind_LWC, 3, 4 },
+		{ { TypeKind_Primitive }, { TEXT("FWSVector4") },		ScalarKind_LWC, 4, 1 },
+		{ { TypeKind_Primitive }, Invalid,						ScalarKind_LWC, 4, 2 },
+		{ { TypeKind_Primitive }, Invalid,						ScalarKind_LWC, 4, 3 },
+		{ { TypeKind_Primitive }, { TEXT("FWSMatrix") }, 		ScalarKind_LWC, 4, 4 },
 	};
 
+	static const FPrimitiveType LWCInverseMatrixType = { { TypeKind_Primitive }, { TEXT("FWSInverseMatrix") }, ScalarKind_LWC, 4, 4, true };
+
+	if (bIsLWCInverseMatrix)
+	{
+		// LWC inverse matrix flag only supported for ScalarKind_LWC and 4x4 dimensions.
+		check(InScalarKind == ScalarKind_LWC && NumRows == 4 && NumColumns == 4);
+		return &LWCInverseMatrixType;
+	}
+
 	int Index = InScalarKind * 4 * 4 + (NumRows - 1) * 4 + (NumColumns - 1);
 	check(Index < UE_ARRAY_COUNT(Types));
 	return &Types[Index];
diff --git a/Engine/Source/Runtime/Engine/Private/Materials/MaterialShader.cpp b/Engine/Source/Runtime/Engine/Private/Materials/MaterialShader.cpp
--- a/Engine/Source/Runtime/Engine/Private/Materials/MaterialShader.cpp
+++ b/Engine/Source/Runtime/Engine/Private/Materials/MaterialShader.cpp
@@ -602,6 +602,20 @@ void GetTextureAddressForSamplerSource(const UTexture* InTexture, ESamplerSource
 	}
 }
 
+bool IsLWCEnabled()
+{
+	static IConsoleVariable* CVarLWCIsEnabled = IConsoleManager::Get().FindConsoleVariable(TEXT("r.MaterialEditor.LWCEnabled"));
+	const bool bLWCIsEnabled = CVarLWCIsEnabled ? CVarLWCIsEnabled->GetBool() : true;
+	return bLWCIsEnabled;
+}
+
+int32 GetLWCTruncateMode()
+{
+	static const auto CVar = IConsoleManager::Get().FindTConsoleVariableDataInt(TEXT("r.MaterialEditor.LWCTruncateMode"));
+	const int32 LWCTruncateValue = CVar ? CVar->GetValueOnAnyThread() : 0;
+	return LWCTruncateValue;
+}
+
 } // namespace UE::MaterialTranslatorUtils
 
 #endif // WITH_EDITOR
diff --git a/Engine/Source/Runtime/Engine/Private/Materials/MaterialSharedPrivate.h b/Engine/Source/Runtime/Engine/Private/Materials/MaterialSharedPrivate.h
--- a/Engine/Source/Runtime/Engine/Private/Materials/MaterialSharedPrivate.h
+++ b/Engine/Source/Runtime/Engine/Private/Materials/MaterialSharedPrivate.h
@@ -55,6 +55,10 @@ namespace UE::MaterialTranslatorUtils
 {
 	const TCHAR* GetVTAddressMode(TextureAddress Address);
 	void GetTextureAddressForSamplerSource(const UTexture* InTexture, ESamplerSourceMode InSamplerSource, TextureAddress& OutStaticAddressX, TextureAddress& OutStaticAddressY, TextureAddress& OutStaticAddressZ);
+
+	bool IsLWCEnabled();
+	int32 GetLWCTruncateMode();
+
 } // /namespace UE::MaterialTranslatorUtils
 
 ///////////////////////////////////////////////////////
diff --git a/Engine/Source/Runtime/Engine/Public/Materials/MaterialExpressionActorPositionWS.h b/Engine/Source/Runtime/Engine/Public/Materials/MaterialExpressionActorPositionWS.h
--- a/Engine/Source/Runtime/Engine/Public/Materials/MaterialExpressionActorPositionWS.h
+++ b/Engine/Source/Runtime/Engine/Public/Materials/MaterialExpressionActorPositionWS.h
@@ -20,6 +20,7 @@ public:
 
 	//~ Begin UMaterialExpression Interface
 #if WITH_EDITOR
+	virtual void Build(MIR::FEmitter& Em) override;
 	virtual int32 Compile(class FMaterialCompiler* Compiler, int32 OutputIndex) override;
 	virtual void GetCaption(TArray<FString>& OutCaptions) const override;
 #endif
diff --git a/Engine/Source/Runtime/Engine/Public/Materials/MaterialExpressionMax.h b/Engine/Source/Runtime/Engine/Public/Materials/MaterialExpressionMax.h
--- a/Engine/Source/Runtime/Engine/Public/Materials/MaterialExpressionMax.h
+++ b/Engine/Source/Runtime/Engine/Public/Materials/MaterialExpressionMax.h
@@ -32,6 +32,7 @@ public:
 
 	//~ Begin UMaterialExpression Interface
 #if WITH_EDITOR
+	virtual void Build(MIR::FEmitter& Em) override;
 	virtual int32 Compile(class FMaterialCompiler* Compiler, int32 OutputIndex) override;
 	virtual void GetCaption(TArray<FString>& OutCaptions) const override;
 #endif
diff --git a/Engine/Source/Runtime/Engine/Public/Materials/MaterialExpressionObjectPositionWS.h b/Engine/Source/Runtime/Engine/Public/Materials/MaterialExpressionObjectPositionWS.h
--- a/Engine/Source/Runtime/Engine/Public/Materials/MaterialExpressionObjectPositionWS.h
+++ b/Engine/Source/Runtime/Engine/Public/Materials/MaterialExpressionObjectPositionWS.h
@@ -20,6 +20,7 @@ public:
 
 	//~ Begin UMaterialExpression Interface
 #if WITH_EDITOR
+	virtual void Build(MIR::FEmitter& Em) override;
 	virtual int32 Compile(class FMaterialCompiler* Compiler, int32 OutputIndex) override;
 	virtual void GetCaption(TArray<FString>& OutCaptions) const override;
     virtual void GetExpressionToolTip(TArray<FString>& OutToolTip) override;
diff --git a/Engine/Source/Runtime/Engine/Public/Materials/MaterialExpressionTruncateLWC.h b/Engine/Source/Runtime/Engine/Public/Materials/MaterialExpressionTruncateLWC.h
--- a/Engine/Source/Runtime/Engine/Public/Materials/MaterialExpressionTruncateLWC.h
+++ b/Engine/Source/Runtime/Engine/Public/Materials/MaterialExpressionTruncateLWC.h
@@ -21,6 +21,7 @@ public:
 
 	//~ Begin UMaterialExpression Interface
 #if WITH_EDITOR
+	virtual void Build(MIR::FEmitter& Em) override;
 	virtual int32 Compile(class FMaterialCompiler* Compiler, int32 OutputIndex) override;
 	virtual void GetCaption(TArray<FString>& OutCaptions) const override { OutCaptions.Add(TEXT("TruncateLWC")); }
 	virtual FText GetKeywords() const override { return FText::FromString(TEXT("truncate lwc")); }
diff --git a/Engine/Source/Runtime/Engine/Public/Materials/MaterialExpressionWorldPosition.h b/Engine/Source/Runtime/Engine/Public/Materials/MaterialExpressionWorldPosition.h
--- a/Engine/Source/Runtime/Engine/Public/Materials/MaterialExpressionWorldPosition.h
+++ b/Engine/Source/Runtime/Engine/Public/Materials/MaterialExpressionWorldPosition.h
@@ -37,6 +37,7 @@ class UMaterialExpressionWorldPosition : public UMaterialExpression
 	
 	//~ Begin UMaterialExpression Interface
 #if WITH_EDITOR
+	virtual void Build(MIR::FEmitter& Emitter) override;
 	virtual int32 Compile(class FMaterialCompiler* Compiler, int32 OutputIndex) override;
 	virtual void GetCaption(TArray<FString>& OutCaptions) const override;
 	virtual FText GetKeywords() const override {return FText::FromString(TEXT("position"));}
diff --git a/Engine/Source/Runtime/Engine/Public/Materials/MaterialIR.h b/Engine/Source/Runtime/Engine/Public/Materials/MaterialIR.h
--- a/Engine/Source/Runtime/Engine/Public/Materials/MaterialIR.h
+++ b/Engine/Source/Runtime/Engine/Public/Materials/MaterialIR.h
@@ -325,6 +325,17 @@ struct FConstant : TValue<VK_Constant>
 	}
 };
 
+// Specifies the axis for computing screen-space derivatives.
+//
+// This enum is used to indicate whether a partial derivative should be taken
+// along the X or Y screen-space direction, corresponding to HLSL's ddx and ddy
+// functions.
+enum class EDerivativeAxis
+{
+	X, // Corresponding to d/dx
+	Y, // Corresponding to d/dy
+};
+
 // Enumeration of all supported material external inputs.
 enum class EExternalInput
 {
@@ -338,6 +349,10 @@ enum class EExternalInput
 	TexCoord5,
 	TexCoord6,
 	TexCoord7,
+	WorldPosition_Absolute,					// (LWC)	WPT_Default						Absolute World Position, including Material Shader Offsets
+	WorldPosition_AbsoluteNoOffsets,		// (LWC)	WPT_ExcludeAllShaderOffsets		Absolute World Position, excluding Material Shader Offsets
+	WorldPosition_CameraRelative,			// (float)	WPT_CameraRelative				Camera Relative World Position, including Material Shader Offsets
+	WorldPosition_CameraRelativeNoOffsets,	// (float)	WPT_CameraRelativeNoOffsets		Camera Relative World Position, excluding Material Shader Offsets
 
 	TexCoord0_Ddx,
 	TexCoord1_Ddx,
@@ -347,6 +362,10 @@ enum class EExternalInput
 	TexCoord5_Ddx,
 	TexCoord6_Ddx,
 	TexCoord7_Ddx,
+	WorldPosition_Absolute_Ddx,				// LWC (cast from float on load)
+	WorldPosition_AbsoluteNoOffsets_Ddx,	// LWC (cast from float on load)
+	WorldPosition_CameraRelative_Ddx,
+	WorldPosition_CameraRelativeNoOffsets_Ddx,
 
 	TexCoord0_Ddy,
 	TexCoord1_Ddy,
@@ -356,25 +375,50 @@ enum class EExternalInput
 	TexCoord5_Ddy,
 	TexCoord6_Ddy,
 	TexCoord7_Ddy,
-
+	WorldPosition_Absolute_Ddy,				// LWC (cast from float on load)
+	WorldPosition_AbsoluteNoOffsets_Ddy,	// LWC (cast from float on load)
+	WorldPosition_CameraRelative_Ddy,
+	WorldPosition_CameraRelativeNoOffsets_Ddy,
+
+	// Ranges of External inputs with derivatives.  There must be an equal number of value, ddx, and ddy variations.  TexCoord0 must remain first,
+	// as code that translates texture coordinate indices to external input ID and back assumes this.
+	WithDerivatives_First = TexCoord0,
+	WithDerivatives_LastVal = WorldPosition_CameraRelativeNoOffsets,
+	WithDerivatives_LastDdx = WorldPosition_CameraRelativeNoOffsets_Ddx,
+	WithDerivatives_LastDdy = WorldPosition_CameraRelativeNoOffsets_Ddy,
+	WithDerivatives_Last = WithDerivatives_LastDdy,
+
+	ActorPosition_Absolute,					// LWC
+	ActorPosition_CameraRelative,
+	ObjectPosition_Absolute,				// LWC
+	ObjectPosition_CameraRelative,
 	ViewMaterialTextureMipBias,
 	ViewMaterialTextureDerivativeMultiply,
 
 	Count,
 };
 
-// The number of texcoord external input groups (uv, ddx, ddy).
-static constexpr int32 TexCoordGroups = 3;
+// The number of external input derivative groups (value, ddx, ddy).
+static constexpr int32 ExternalInputWithDerivativesGroups = 3;
+
+// Number of external inputs with derivatives
+static constexpr int32 ExternalInputWithDerivativesNum = (int32)EExternalInput::WithDerivatives_LastVal - (int32)EExternalInput::WithDerivatives_First + 1;
 
 // Maximum number of supported texture coordinates.
 static constexpr int32 TexCoordMaxNum = 8;
 
-// Returns the string representation of given given external input.
+// Returns the string representation of given external input.
 const TCHAR* LexToString(EExternalInput Input);
 
 // Returns the given external input type.
 const FType* GetExternalInputType(EExternalInput Id);
 
+// Returns whether the given external input needs to generate non-trivial derivative expressions.
+bool IsExternalInputWithDerivatives(EExternalInput Id);
+
+// Translate an input ID to a derivative variation
+EExternalInput GetExternalInputDerivative(EExternalInput Id, EDerivativeAxis Axis);
+
 // Converts the given texture coordinate index to its corresponding external input mapping.
 MIR::EExternalInput TexCoordIndexToExternalInput(int32 TexCoordIndex);
 
@@ -635,7 +679,7 @@ enum EOperator
 	BO_BitShiftRight,
 	
 	// Binary intrinsics
-	BO_Atan2,
+	BO_ATan2,
 	BO_Cross,
 	BO_Distance,
 	BO_Dot,
@@ -848,17 +892,6 @@ struct FStageSwitch : TInstruction<VK_StageSwitch, 1>
 	void SetArgs(FValue* PixelStageArg, FValue* OtherStagesArg);
 };
 
-// Specifies the axis for computing screen-space derivatives.
-//
-// This enum is used to indicate whether a partial derivative should be taken
-// along the X or Y screen-space direction, corresponding to HLSL's ddx and ddy
-// functions.
-enum class EDerivativeAxis
-{
-	X, // Corresponding to d/dx
-	Y, // Corresponding to d/dy
-};
-
 // Instruction that maps to hardware ddx()/ddy().
 //
 // Note: This is only available in stages that support hardware derivatives (e.g. pixel shader)
diff --git a/Engine/Source/Runtime/Engine/Public/Materials/MaterialIREmitter.h b/Engine/Source/Runtime/Engine/Public/Materials/MaterialIREmitter.h
--- a/Engine/Source/Runtime/Engine/Public/Materials/MaterialIREmitter.h
+++ b/Engine/Source/Runtime/Engine/Public/Materials/MaterialIREmitter.h
@@ -524,7 +524,7 @@ public:
 	FValueRef Min(FValueRef A, FValueRef B) { return Operator(BO_Min, A, B); }
 	FValueRef Dot(FValueRef A, FValueRef B) { return Operator(BO_Dot, A, B); }
 	FValueRef Cross(FValueRef A, FValueRef B) { return Operator(BO_Cross, A, B); }
-	FValueRef Atan2(FValueRef A, FValueRef B) { return Operator(BO_Atan2, A, B); }
+	FValueRef Atan2(FValueRef A, FValueRef B) { return Operator(BO_ATan2, A, B); }
 	FValueRef Pow(FValueRef A, FValueRef B) { return Operator(BO_Pow, A, B); }
 
 	/*-------------------------- Ternary operators shortcuts ---------------------------*/
diff --git a/Engine/Source/Runtime/Engine/Public/Materials/MaterialIRTypes.h b/Engine/Source/Runtime/Engine/Public/Materials/MaterialIRTypes.h
--- a/Engine/Source/Runtime/Engine/Public/Materials/MaterialIRTypes.h
+++ b/Engine/Source/Runtime/Engine/Public/Materials/MaterialIRTypes.h
@@ -57,6 +57,9 @@ struct FType
 	/// Returns whether this is an integer primitive type (of any dimension).
 	bool IsInteger() const;
 
+	/// Returns whether this is an LWC primitive type (of any dimension).
+	bool IsLWC() const;
+
 	/// Returns whether this type is a `bool` scalar.
 	bool IsBoolScalar() const;
 
@@ -98,6 +101,7 @@ enum EScalarKind
 	ScalarKind_Bool,
 	ScalarKind_Int,
 	ScalarKind_Float,
+	ScalarKind_LWC,
 };
 
 /// Returns whether the specified scalar kind supports arithmetic operators (plus, minus, etc).
@@ -125,6 +129,9 @@ struct FPrimitiveType : FType
 	/// Number of columns (rows > 1 and columns == 1, this is a vector).
 	int NumColumns;
 
+	/// LWC inverse matrix type -- only supported for ScalarKind_LWC and 4x4 dimensions
+	bool bIsLWCInverseMatrix;
+
 	// Returns the boolean scalar type.
 	static const FPrimitiveType* GetBool();
 	
@@ -143,14 +150,17 @@ struct FPrimitiveType : FType
 	// Returns the floating point 4D vector type.
 	static const FPrimitiveType* GetFloat4();
 
+	// Returns the LWC 3D vector type
+	static const FPrimitiveType* GetLWC3();
+
 	// Returns the scalar type with given kind.
 	static const FPrimitiveType* GetScalar(EScalarKind InScalarKind);
 
 	// Returns the column vector type with given kind and number of rows.
 	static const FPrimitiveType* GetVector(EScalarKind InScalarKind, int NumRows);
 	
-	// Returns the primitive type with given kind and number of columns and rows.
-	static const FPrimitiveType* Get(EScalarKind InScalarKind, int NumRows, int NumColumns);
+	// Returns the primitive type with given kind and number of columns and rows.  LWC inverse matrix flag only supported for ScalarKind_LWC and 4x4 dimensions.
+	static const FPrimitiveType* Get(EScalarKind InScalarKind, int NumRows, int NumColumns, bool bIsLWCInverseMatrix = false);
 
 	// Returns the number of components in this primitive type.
 	int  GetNumComponents() const { return NumRows * NumColumns; }
