New material translator:
- Introduced support for user defined Material Aggregates (UMaterialAggregate), which work similarly to C/C++ structs.
- Introduced new expression "Aggregate" for aggregate value creation, access and manipulation. It embeds all the semantics of Make/Break/Set/GetMaterialAttributes into a single expression. It is only supported in the new translator.
- Introduced new expression "Blend" to generalize "BlendMaterialAttributes". It accepts any arithmetic value, including material aggregates (potentially nesting other material aggregates).
- Fixed the node preview scalar to grayscale bug I introduce in my previous change.
- Implemented support for Make/Break/Set/Get/Blend MaterialAttributes expressions in the new translator.
- Some bug fixes and and general code improvements.

#rnx

diff --git a/Engine/Config/BaseMaterialExpressions.ini b/Engine/Config/BaseMaterialExpressions.ini
--- a/Engine/Config/BaseMaterialExpressions.ini
+++ b/Engine/Config/BaseMaterialExpressions.ini
@@ -1053,3 +1053,9 @@ bShaderInputData=true
 
 [/Script/Engine.MaterialExpressionEnumerationParameter]
 +MenuCategories=NSLOCTEXT("MaterialExpression", "Parameters", "Parameters")
+
+[/Script/Engine.MaterialExpressionAggregate]
++MenuCategories=NSLOCTEXT("MaterialExpression", "MaterialAggregates", "Material Aggregates")
+
+[/Script/Engine.MaterialExpressionBlend]
++MenuCategories=NSLOCTEXT("MaterialExpression", "Math", "Math")
diff --git a/Engine/Source/Runtime/Engine/Private/Materials/MaterialAggregate.cpp b/Engine/Source/Runtime/Engine/Private/Materials/MaterialAggregate.cpp
new file mode 100644
--- /dev/null
+++ b/Engine/Source/Runtime/Engine/Private/Materials/MaterialAggregate.cpp
@@ -0,0 +1,255 @@
+// Copyright Epic Games, Inc. All Rights Reserved.
+
+#include "Materials/MaterialAggregate.h"
+#include "Framework/Notifications/NotificationManager.h"
+#include "Widgets/Notifications/SNotificationList.h"
+#include "Materials/MaterialAttributeDefinitionMap.h"
+#include "SceneTypes.h"
+#include "UObject/Package.h"
+
+#if WITH_EDITOR
+
+EMaterialAggregateAttributeType MaterialValueTypeToMaterialAggregateAttributeType(EMaterialValueType Type)
+{
+	switch (Type)
+	{
+		case MCT_Float1: return EMaterialAggregateAttributeType::Float1;
+		case MCT_Float2: return EMaterialAggregateAttributeType::Float2;
+		case MCT_Float3: return EMaterialAggregateAttributeType::Float3;
+		case MCT_Float4: return EMaterialAggregateAttributeType::Float4;
+		case MCT_Float: return EMaterialAggregateAttributeType::Float4;
+		case MCT_Bool: return EMaterialAggregateAttributeType::Bool1;
+		case MCT_ShadingModel: return EMaterialAggregateAttributeType::ShadingModel;
+		case MCT_UInt1: return EMaterialAggregateAttributeType::UInt1;
+		case MCT_UInt2: return EMaterialAggregateAttributeType::UInt2;
+		case MCT_UInt3: return EMaterialAggregateAttributeType::UInt3;
+		case MCT_UInt4: return EMaterialAggregateAttributeType::UInt4;
+		case MCT_UInt: return EMaterialAggregateAttributeType::UInt4;
+		default:
+			checkf(false, TEXT("This material value type '%d' cannot be expressed to a material aggregate attribute type."), (int32)Type);
+			return {};
+	}
+}
+
+EMaterialValueType FMaterialAggregateAttribute::ToMaterialValueType() const
+{
+	switch (Type)
+	{
+		case EMaterialAggregateAttributeType::Bool1: return MCT_Bool;
+		case EMaterialAggregateAttributeType::Bool2: return MCT_Unknown; // MCT does't have this yet
+		case EMaterialAggregateAttributeType::Bool3: return MCT_Unknown; // MCT does't have this yet
+		case EMaterialAggregateAttributeType::Bool4: return MCT_Unknown; // MCT does't have this yet
+		case EMaterialAggregateAttributeType::UInt1: return MCT_UInt1;
+		case EMaterialAggregateAttributeType::UInt2: return MCT_UInt2;
+		case EMaterialAggregateAttributeType::UInt3: return MCT_UInt3;
+		case EMaterialAggregateAttributeType::UInt4: return MCT_UInt4;
+		case EMaterialAggregateAttributeType::Float1: return MCT_Float1;
+		case EMaterialAggregateAttributeType::Float2: return MCT_Float2;
+		case EMaterialAggregateAttributeType::Float3: return MCT_Float3;
+		case EMaterialAggregateAttributeType::Float4: return MCT_Float4;
+		case EMaterialAggregateAttributeType::ShadingModel: return MCT_ShadingModel;
+		case EMaterialAggregateAttributeType::Aggregate:
+			return Aggregate.Get() == UMaterialAggregate::GetMaterialAttributes() ? MCT_MaterialAttributes : MCT_Unknown;
+		default:
+			checkNoEntry();
+			return {};
+	}
+}
+
+static bool ContainsCyclicReference(const UMaterialAggregate* Aggregate, TArray<const UMaterialAggregate*>& Stack)
+{
+	if (!Aggregate)
+	{
+		return false;
+	}
+
+	// Check if Current is already in the stack (i.e., a cycle)
+	if (Stack.Contains(Aggregate))
+	{
+		return true;
+	}
+
+	Stack.Push(Aggregate);
+
+	// Look for cycles in any attribute of this aggregate
+	for (const FMaterialAggregateAttribute& Attr : Aggregate->Attributes)
+	{
+		if (Attr.Type == EMaterialAggregateAttributeType::Aggregate && Attr.Aggregate)
+		{
+			if (ContainsCyclicReference(Attr.Aggregate, Stack))
+			{
+				return true;
+			}
+		}
+	}
+
+	Stack.Pop();
+	return false;
+}
+
+// Makes sure that all attributes with an aggregate type don't reference back to this creating a cyclic dependency.
+static void ClearCyclicDependencies(UMaterialAggregate* Aggregate, bool bDisplayNotifications)
+{
+	TArray<const UMaterialAggregate*> Stack;
+	for (FMaterialAggregateAttribute& Attr : Aggregate->Attributes)
+	{
+		Stack.Empty(Stack.Num());
+		if (Attr.Type == EMaterialAggregateAttributeType::Aggregate && ContainsCyclicReference(Aggregate, Stack))
+		{
+			if (bDisplayNotifications)
+			{
+				// Display a notification message.
+				FString Message = FString::Printf(TEXT("Setting '%s' to reference material aggregate '%s' would introduce a cyclic dependency."), *Attr.Name.ToString(), *Aggregate->GetName());
+
+				FNotificationInfo Info(FText::FromString(Message));
+				Info.ExpireDuration = 2.5f;
+				Info.bFireAndForget = true;
+				FSlateNotificationManager::Get().AddNotification(Info)->SetCompletionState(SNotificationItem::CS_Fail);
+			}
+
+			// Reset the attribute to null.
+			Attr.Aggregate = nullptr;
+		}
+	}
+}
+
+static FMaterialAggregateAttribute MaterialPropertyToAttribute(EMaterialProperty Property)
+{ 
+	FMaterialAggregateAttribute Attribute;
+	Attribute.Name = *FMaterialAttributeDefinitionMap::GetAttributeName(Property);
+	Attribute.Type = MaterialValueTypeToMaterialAggregateAttributeType(FMaterialAttributeDefinitionMap::GetValueType(Property));
+	return Attribute;
+}
+
+// This array declares the material properties that are referenced by each attribute in the MaterialAttributes aggregate.
+static constexpr EMaterialProperty MaterialAttributesProperties[] = {
+	MP_BaseColor,
+	MP_Metallic,
+	MP_Specular,
+	MP_Roughness,
+	MP_Anisotropy,
+	MP_EmissiveColor,
+	MP_Opacity,
+	MP_OpacityMask,
+	MP_Normal,
+	MP_Tangent,
+	MP_WorldPositionOffset,
+	MP_SubsurfaceColor,
+	MP_CustomData0,
+	MP_CustomData1,
+	MP_AmbientOcclusion,
+	MP_Refraction,
+	MP_CustomizedUVs0,
+	MP_CustomizedUVs1,
+	MP_CustomizedUVs2,
+	MP_CustomizedUVs3,
+	MP_CustomizedUVs4,
+	MP_CustomizedUVs5,
+	MP_CustomizedUVs6,
+	MP_CustomizedUVs7,
+	MP_PixelDepthOffset,
+	MP_ShadingModel,
+	MP_Displacement,
+};
+
+const UMaterialAggregate* UMaterialAggregate::GetMaterialAttributes()
+{
+	static TObjectPtr<UMaterialAggregate> Instance = []
+	{
+		TObjectPtr<UMaterialAggregate> MA = NewObject<UMaterialAggregate>(GetTransientPackage());
+		MA->Rename(TEXT("MaterialAttributes"));
+		MA->AddToRoot();
+		for (EMaterialProperty Property : MaterialAttributesProperties)
+		{
+			MA->Attributes.Add(MaterialPropertyToAttribute(Property));
+		}
+		return MA;
+	}();
+	return Instance.Get();
+}
+
+int32 UMaterialAggregate::MaterialPropertyToAttributeIndex(EMaterialProperty Property)
+{
+	switch (Property)
+	{
+		case MP_BaseColor: return 0;
+		case MP_Metallic: return 1;
+		case MP_Specular: return 2;
+		case MP_Roughness: return 3;
+		case MP_Anisotropy: return 4;
+		case MP_EmissiveColor: return 5;
+		case MP_Opacity: return 6;
+		case MP_OpacityMask: return 7;
+		case MP_Normal: return 8;
+		case MP_Tangent: return 9;
+		case MP_WorldPositionOffset: return 10;
+		case MP_SubsurfaceColor: return 11;
+		case MP_CustomData0: return 12;
+		case MP_CustomData1: return 13;
+		case MP_AmbientOcclusion: return 14;
+		case MP_Refraction: return 15;
+		case MP_CustomizedUVs0: return 16;
+		case MP_CustomizedUVs1: return 17;
+		case MP_CustomizedUVs2: return 18;
+		case MP_CustomizedUVs3: return 19;
+		case MP_CustomizedUVs4: return 20;
+		case MP_CustomizedUVs5: return 21;
+		case MP_CustomizedUVs6: return 22;
+		case MP_CustomizedUVs7: return 23;
+		case MP_PixelDepthOffset: return 24;
+		case MP_ShadingModel: return 25;
+		case MP_Displacement: return 26;
+		default: return INDEX_NONE;
+	}
+}
+
+EMaterialProperty UMaterialAggregate::AttributeIndexToMaterialProperty(int32 MaterialAttributeIndex)
+{
+	return MaterialAttributeIndex < UE_ARRAY_COUNT(MaterialAttributesProperties)
+		? MaterialAttributesProperties[MaterialAttributeIndex]
+		: MP_MAX;
+}
+
+int32 UMaterialAggregate::FindAttributeIndexByName(FName Name) const
+{
+	return Attributes.IndexOfByPredicate([Name](const FMaterialAggregateAttribute& Attribute)
+	{
+		return Attribute.Name == Name;
+	});
+}
+
+void UMaterialAggregate::PostLoad()
+{
+	ClearCyclicDependencies(this, false);
+
+	Super::PostLoad();
+}
+
+void UMaterialAggregate::PostEditChangeProperty(FPropertyChangedEvent& PropertyChangedEvent)
+{
+	Super::PostEditChangeProperty(PropertyChangedEvent);
+	
+	if (!(PropertyChangedEvent.ChangeType & EPropertyChangeType::ValueSet))
+	{
+		return;
+	}
+
+	const FName PropertyName = PropertyChangedEvent.Property ? PropertyChangedEvent.Property->GetFName() : NAME_None;
+	
+	if (PropertyName == GET_MEMBER_NAME_CHECKED(FMaterialAggregateAttribute, Type))
+	{
+		for (FMaterialAggregateAttribute& Attr : Attributes)
+		{
+			if (Attr.Type != EMaterialAggregateAttributeType::Aggregate)
+			{
+				Attr.Aggregate = nullptr;
+			}
+		}
+	}
+	else if (PropertyName == GET_MEMBER_NAME_CHECKED(FMaterialAggregateAttribute, Aggregate))
+	{
+		ClearCyclicDependencies(this, true);
+	}
+}
+
+#endif // WITH_EDITOR
diff --git a/Engine/Source/Runtime/Engine/Private/Materials/MaterialExpressions.cpp b/Engine/Source/Runtime/Engine/Private/Materials/MaterialExpressions.cpp
--- a/Engine/Source/Runtime/Engine/Private/Materials/MaterialExpressions.cpp
+++ b/Engine/Source/Runtime/Engine/Private/Materials/MaterialExpressions.cpp
@@ -60,6 +60,7 @@
 #include "Materials/MaterialExpressionBentNormalCustomOutput.h"
 #include "Materials/MaterialExpressionBindlessSwitch.h"
 #include "Materials/MaterialExpressionBlackBody.h"
+#include "Materials/MaterialExpressionBlend.h"
 #include "Materials/MaterialExpressionBlendMaterialAttributes.h"
 #include "Materials/MaterialExpressionBreakMaterialAttributes.h"
 #include "Materials/MaterialExpressionBumpOffset.h"
@@ -303,6 +304,7 @@
 #include "Materials/MaterialExternalCodeRegistry.h"
 #include "Materials/MaterialExpressionRecordTextureStreamingInfo.h"
 #include "Materials/MaterialExpressionOperator.h"
+#include "Materials/MaterialExpressionAggregate.h"
 #include "EditorSupportDelegates.h"
 #if WITH_EDITOR
 #include "MaterialGraph/MaterialGraphNode_Comment.h"
@@ -743,8 +745,6 @@ UMaterialExpression::UMaterialExpression(const FObjectInitializer& ObjectInitial
 #if WITH_EDITORONLY_DATA
 	Outputs.Add(FExpressionOutput(TEXT("")));
 
-	NumExecutionInputs = 0;
-
 	bShowInputs = true;
 	bShowOutputs = true;
 	bCollapsed = true;
@@ -1690,6 +1690,7 @@ bool UMaterialExpression::CanEditChange(const FProperty* InProperty) const
 			if (!EditConditionPropertyName.IsEmpty())
 			{
 				FBoolProperty* EditConditionProperty = FindFProperty<FBoolProperty>(GetClass(), *EditConditionPropertyName);
+				if (EditConditionProperty)
 				{
 					bIsEditable = *EditConditionProperty->ContainerPtrToValuePtr<bool>(this);
 				}
@@ -1705,7 +1706,6 @@ TArray<FExpressionOutput>& UMaterialExpression::GetOutputs()
 	return Outputs;
 }
 
-
 TArrayView<FExpressionInput*> UMaterialExpression::GetInputsView()
 {
 	return CachedInputs;
@@ -2315,6 +2315,12 @@ bool UMaterialExpression::ContainsInputLoopInternal(const FContainsInputLoopInte
 	return false;
 }
 
+int32 UMaterialExpression::Compile(class FMaterialCompiler* Compiler, int32 OutputIndex)
+{
+	Compiler->Error(TEXT("Unsupported material expression."));
+	return INDEX_NONE;
+}
+
 bool UMaterialExpression::IsUsingNewHLSLGenerator() const
 {
 	if (Material)
@@ -24217,4 +24223,244 @@ TArray<FString> UMaterialExpressionEnumerationParameter::GetEntryOptions() const
 
 #endif // WITH_EDITOR
 
+///////////////////////////////////////////////////////////////////////////////
+// UMaterialExpressionAggregate
+///////////////////////////////////////////////////////////////////////////////
+
+UMaterialExpressionAggregate::UMaterialExpressionAggregate(const FObjectInitializer& ObjectInitializer)
+: Super(ObjectInitializer)
+{
+	Outputs.Empty();
+	Outputs.Emplace();
+}
+
+#if WITH_EDITOR
+
+const UMaterialAggregate* UMaterialExpressionAggregate::GetAggregate() const
+{
+	switch (Kind)
+	{
+		case EMaterialExpressionMakeAggregateKind::MaterialAttributes: return UMaterialAggregate::GetMaterialAttributes();
+		case EMaterialExpressionMakeAggregateKind::UserDefined: return UserAggregate;
+		default: checkNoEntry(); return nullptr;
+	}
+}
+
+const FMaterialAggregateAttribute* UMaterialExpressionAggregate::GetAggregateAttribute(int32 Index) const
+{
+	check(Index >= 0);
+	if (const UMaterialAggregate* Aggregate = GetAggregate())
+	{
+		if (Index < Aggregate->Attributes.Num())
+		{
+			return &Aggregate->Attributes[Index];
+		}
+	}
+	return nullptr;
+}
+
+TArray<FString> UMaterialExpressionAggregate::GetPossibleAttributeNames() const
+{
+	TArray<FString> Values;
+	if (const UMaterialAggregate* Aggregate = GetAggregate())
+	{
+		for (int32 i = 0; i < Aggregate->Attributes.Num(); ++i)
+		{
+			if (!AttributeNames.Contains(Aggregate->Attributes[i].Name))
+			{
+				Values.Add(Aggregate->Attributes[i].Name.ToString());
+			}
+		}
+	}
+	return Values;
+}
+
+void UMaterialExpressionAggregate::GetCaption(TArray<FString>& OutCaptions) const
+{
+	const UMaterialAggregate* Aggregate = GetAggregate();
+	if (Aggregate)
+	{
+		OutCaptions.Add(*Aggregate->GetName());
+	}
+	else
+	{
+		OutCaptions.Add(TEXT("Aggregate")); 
+	}
+}
+
+FExpressionInput* UMaterialExpressionAggregate::GetInput(int32 InputIndex)
+{
+	check(InputIndex >= 0);
+	if (InputIndex == 0)
+	{
+		return &PrototypeInput;
+	}
+
+	InputIndex -= 1;
+	return InputIndex < Entries.Num() ? &Entries[InputIndex].Input : nullptr;
+}
+
+FName UMaterialExpressionAggregate::GetInputName(int32 InputIndex) const
+{
+	check(InputIndex >= 0);
+	if (InputIndex == 0)
+	{
+		if (const UMaterialAggregate* Aggregate = GetAggregate())
+		{
+			return *GetAggregate()->GetName();
+		}
+		return NAME_None;
+	}
+
+	InputIndex -= 1;
+	return InputIndex < AttributeNames.Num() ? AttributeNames[InputIndex] : NAME_None;
+}
+
+static EMaterialValueType GetPinValueType(const UMaterialExpressionAggregate* Expr, int32 PinIndex)
+{
+	check(PinIndex >= 0);
+	if (PinIndex == 0)
+	{
+		return Expr->GetAggregate() == UMaterialAggregate::GetMaterialAttributes() ? MCT_MaterialAttributes : MCT_Unknown;
+	}
+	if (const FMaterialAggregateAttribute* Attribute = Expr->GetAggregateAttribute(PinIndex - 1))
+	{
+		return Attribute->ToMaterialValueType();
+	}
+	return MCT_Unknown;
+}
+
+EMaterialValueType UMaterialExpressionAggregate::GetInputValueType(int32 InputIndex)
+{
+	return GetPinValueType(this, InputIndex);
+}
+
+EMaterialValueType UMaterialExpressionAggregate::GetOutputValueType(int32 OutputIndex)
+{
+	return GetPinValueType(this, OutputIndex);
+}
+
+void UMaterialExpressionAggregate::PreEditChange(FProperty* PropertyAboutToChange)
+{
+	PreEditData.Aggregate = GetAggregate();
+	PreEditData.AttributeNames = AttributeNames;
+	PreEditData.Entries = Entries;
+}
+
+static FName PickAvailableAttribute(const UMaterialAggregate* Aggregate, const TArray<FName>& AttributeNames)
+{
+	if (Aggregate)
+	{
+		for (int i = 0; i < Aggregate->Attributes.Num(); ++i)
+		{
+			if (!AttributeNames.Contains(Aggregate->Attributes[i].Name))
+			{
+				return Aggregate->Attributes[i].Name;
+			}
+		}
+	}
+	return {};
+}
+void UMaterialExpressionAggregate::PostEditChangeProperty(FPropertyChangedEvent& PropertyChangedEvent)
+{
+    // Early-out if there’s no graph node or no property change to handle
+    if (!GraphNode || !PropertyChangedEvent.MemberProperty)
+    {
+        Super::PostEditChangeProperty(PropertyChangedEvent);
+        return;
+    }
+
+    const UMaterialAggregate* Aggregate = GetAggregate();
+
+    // If the underlying aggregate asset changed, reset all attribute selections
+    if (PreEditData.Aggregate != Aggregate)
+    {
+        AttributeNames.Empty();
+    }
+    
+    // Handle newly added attributes
+    if (AttributeNames.Num() > PreEditData.AttributeNames.Num())
+    {
+        for (int32 i = 0; i < AttributeNames.Num(); )
+        {
+            if (AttributeNames[i] == NAME_None)
+            {
+                AttributeNames[i] = PickAvailableAttribute(Aggregate, AttributeNames);
+                if (AttributeNames[i] == NAME_None)
+                {
+                    AttributeNames.RemoveAt(i);
+                    continue;
+                }
+                Entries.EmplaceAt(i);
+                Outputs.Emplace();
+            }
+            ++i;
+        }
+    }
+    // Handle removed attributes
+    else if (AttributeNames.Num() < PreEditData.AttributeNames.Num())
+    {
+        int32 CurrIndex = 0, PrevIndex = 0;
+
+        // Remove entries/pins for any attributes dropped from the list
+        for (; CurrIndex < AttributeNames.Num(); ++PrevIndex)
+        {
+            if (PrevIndex >= PreEditData.AttributeNames.Num() ||
+                AttributeNames[CurrIndex] != PreEditData.AttributeNames[PrevIndex])
+            {
+                Entries.RemoveAt(CurrIndex);
+
+                int PinIndex = CurrIndex + 1;
+                Outputs.RemoveAt(PinIndex);
+                GraphNode->RemovePinAt(PinIndex, EGPD_Input);
+                GraphNode->RemovePinAt(PinIndex, EGPD_Output);
+
+                continue;
+            }
+            ++CurrIndex;
+        }
+
+        // Clean up any leftover old entries beyond the new list size
+        for (; PrevIndex < PreEditData.AttributeNames.Num(); ++PrevIndex)
+        {
+            Entries.Pop();
+            Outputs.Pop();
+            GraphNode->RemovePinAt(Outputs.Num(), EGPD_Input);
+            GraphNode->RemovePinAt(Outputs.Num(), EGPD_Output);
+        }
+    }
+
+    // Update each entry’s index to match the aggregate’s attribute ordering
+    for (int i = 0; i < AttributeNames.Num(); ++i)
+    {
+        Entries[i].AttributeIndex = Aggregate->FindAttributeIndexByName(AttributeNames[i]);
+    }
+
+    // Rebuild the node to reflect any pin/entry changes
+    GraphNode->ReconstructNode();
+
+    Super::PostEditChangeProperty(PropertyChangedEvent);
+}
+
+///////////////////////////////////////////////////////////////////////////////
+// UMaterialExpressionBlend
+///////////////////////////////////////////////////////////////////////////////
+
+void UMaterialExpressionBlend::GetCaption(TArray<FString>& OutCaptions) const
+{
+	OutCaptions.Add(TEXT("Blend")); 
+}
+
+EMaterialValueType UMaterialExpressionBlend::GetInputValueType(int32 InputIndex)
+{
+	return InputIndex == 2 ? MCT_Float1 : MCT_Unknown;
+}
+
+EMaterialValueType UMaterialExpressionBlend::GetOutputValueType(int32 OutputIndex)
+{
+	return MCT_Unknown;
+}
+
+#endif // WITH_EDITOR
+
 #undef LOCTEXT_NAMESPACE
diff --git a/Engine/Source/Runtime/Engine/Private/Materials/MaterialExpressionsIR.cpp b/Engine/Source/Runtime/Engine/Private/Materials/MaterialExpressionsIR.cpp
--- a/Engine/Source/Runtime/Engine/Private/Materials/MaterialExpressionsIR.cpp
+++ b/Engine/Source/Runtime/Engine/Private/Materials/MaterialExpressionsIR.cpp
@@ -28,6 +28,7 @@
 #include "Materials/MaterialExpressionAtmosphericLightVector.h"
 #include "Materials/MaterialExpressionBentNormalCustomOutput.h"
 #include "Materials/MaterialExpressionBlackBody.h"
+#include "Materials/MaterialExpressionBlend.h"
 #include "Materials/MaterialExpressionBlendMaterialAttributes.h"
 #include "Materials/MaterialExpressionBreakMaterialAttributes.h"
 #include "Materials/MaterialExpressionBumpOffset.h"
@@ -240,6 +241,7 @@
 #include "Materials/MaterialFunctionInterface.h"
 #include "Materials/MaterialParameterCollection.h"
 #include "Materials/MaterialExpressionOperator.h"
+#include "Materials/MaterialExpressionAggregate.h"
 #include "Materials/MaterialSharedPrivate.h"
 #include "MaterialShared.h"
 #include "Misc/MemStackUtility.h"
@@ -401,15 +403,15 @@ void UMaterialExpressionStaticBoolParameter::Build(MIR::FEmitter& Em)
 void UMaterialExpressionStaticComponentMaskParameter::Build(MIR::FEmitter& Em)
 {
 	FValueRef Value = BuildMaterialExpressionParameter(Em, this);
-	const MIR::FDimensional* Dimensional = Value->As<MIR::FDimensional>();
+	const MIR::FComposite* Composite = Value->As<MIR::FComposite>();
 
 	// BuildMaterialExpressionParameter should return this as a constant 4 dimensional bool.
-	check(Dimensional);
-	check(Dimensional->AreComponentsConstant());
+	check(Composite);
+	check(Composite->AreComponentsConstant());
 
 	MIR::FSwizzleMask Mask;
 
-	TConstArrayView<MIR::FValue*> Components = Dimensional->GetComponents();
+	TConstArrayView<MIR::FValue*> Components = Composite->GetComponents();
 	check(Components.Num() == 4);
 
 	if (Components[0]->IsTrue())
@@ -874,7 +876,7 @@ void UMaterialExpressionIf::Build(MIR::FEmitter& Em)
 	// Equals value is optional -- if present, generate an additional conditional.
 	if (AEqualsBValue.IsValid())
 	{
-		FValueRef AEqualsBConditionValue = EmitAlmostEquals(Em, AValue, BValue, EqualsThreshold);
+	FValueRef AEqualsBConditionValue = EmitAlmostEquals(Em, AValue, BValue, EqualsThreshold);
 
 		if (AEqualsBConditionValue->Type->IsBoolScalar())
 		{
@@ -1329,7 +1331,7 @@ void UMaterialExpressionConvert::Build(MIR::FEmitter& Em)
 				const MIR::FPrimitiveType* InputPrimitiveType = InputValues[InputIndex]->Type->AsPrimitive();
 				if (!InputPrimitiveType)
 				{
-					Em.Errorf(TEXT("Input `%d` of type `%s` is not primitive."), InputComponentIndex, InputValues[InputIndex]->Type->GetSpelling().GetData());
+					Em.Errorf(TEXT("Input `%d` of type `%s` is not primitive."), InputComponentIndex, *InputValues[InputIndex]->Type->GetSpelling());
 					continue;
 				}
 
@@ -1391,7 +1393,7 @@ static FValueRef BuildViewProperty(MIR::FEmitter& Em, EMaterialExposedViewProper
 	const MIR::FPrimitiveType* PrimitiveType = HLSLCodeType->AsPrimitive();
 	if (PrimitiveType && PrimitiveType->IsLWC() && !UE::MaterialTranslatorUtils::IsLWCEnabled())
 	{
-		Result = Em.CastToFloat(Result);
+		Result = Em.CastToFloatKind(Result);
 	}
 
 	// Fall back to compute the property's inverse from PropertyCode, if no custom inverse
@@ -1485,12 +1487,12 @@ void UMaterialExpressionOperator::Build(MIR::FEmitter& Em)
 
 void UMaterialExpressionFloatToUInt::Build(MIR::FEmitter& Em)
 {
-	Em.Output(0, Em.CastToInt(Em.Input(&Input)));
+	Em.Output(0, Em.CastToIntKind(Em.Input(&Input)));
 }
 
 void UMaterialExpressionUIntToFloat::Build(MIR::FEmitter& Em)
 {
-	Em.Output(0, Em.CastToFloat(Em.Input(&Input)));
+	Em.Output(0, Em.CastToFloatKind(Em.Input(&Input)));
 }
 
 void UMaterialExpressionTruncateLWC::Build(MIR::FEmitter& Em)
@@ -1503,7 +1505,7 @@ void UMaterialExpressionTruncateLWC::Build(MIR::FEmitter& Em)
 		const MIR::FPrimitiveType* Primitive = InputValue->Type->AsPrimitive();
 		if (Primitive && Primitive->ScalarKind == MIR::EScalarKind::ScalarKind_LWC)
 		{
-			Em.Output(0, Em.CastToFloat(Em.Input(&Input)));
+			Em.Output(0, Em.CastToFloatKind(Em.Input(&Input)));
 			return;
 		}
 	}
@@ -1520,7 +1522,7 @@ void UMaterialExpressionActorPositionWS::Build(MIR::FEmitter& Em)
 	else if (!UE::MaterialTranslatorUtils::IsLWCEnabled())
 	{
 		// CastToNonLWCIfDisabled
-		Em.Output(0, Em.CastToFloat(Em.ExternalInput(MIR::EExternalInput::ActorPosition_Absolute)));
+		Em.Output(0, Em.CastToFloatKind(Em.ExternalInput(MIR::EExternalInput::ActorPosition_Absolute)));
 	}
 	else
 	{
@@ -1537,7 +1539,7 @@ void UMaterialExpressionObjectPositionWS::Build(MIR::FEmitter& Em)
 	else if (!UE::MaterialTranslatorUtils::IsLWCEnabled())
 	{
 		// CastToNonLWCIfDisabled
-		Em.Output(0, Em.CastToFloat(Em.ExternalInput(MIR::EExternalInput::ObjectPosition_Absolute)));
+		Em.Output(0, Em.CastToFloatKind(Em.ExternalInput(MIR::EExternalInput::ObjectPosition_Absolute)));
 	}
 	else
 	{
@@ -1559,7 +1561,7 @@ void UMaterialExpressionWorldPosition::Build(MIR::FEmitter& Em)
 	// CastToNonLWCIfDisabled
 	if (!UE::MaterialTranslatorUtils::IsLWCEnabled())
 	{
-		WorldPosition = Em.CastToFloat(WorldPosition);
+		WorldPosition = Em.CastToFloatKind(WorldPosition);
 	}
 
 	Em.Output(0, WorldPosition);
@@ -1567,4 +1569,346 @@ void UMaterialExpressionWorldPosition::Build(MIR::FEmitter& Em)
 	Em.Output(2, Em.Subscript(WorldPosition, 2));
 }
 
+
+void UMaterialExpressionMakeMaterialAttributes::Build(MIR::FEmitter& Em)
+{
+	MIR::TTemporaryArray<MIR::FAttributeAssignment> Assignments{ MP_MAX };
+	int32 NumAssignments = 0;
+
+	auto PushAttributeAssignment = [&](EMaterialProperty Property, FExpressionInput* Input)
+	{
+		if (FValueRef Value = Em.TryInput(Input))
+		{
+			Assignments[NumAssignments++] = { *FMaterialAttributeDefinitionMap::GetAttributeName(Property), Value };
+		}
+	};
+
+	PushAttributeAssignment(MP_BaseColor, &BaseColor);
+	PushAttributeAssignment(MP_Metallic, &Metallic);
+	PushAttributeAssignment(MP_Specular, &Specular);
+	PushAttributeAssignment(MP_Roughness, &Roughness);
+	PushAttributeAssignment(MP_Anisotropy, &Anisotropy);
+	PushAttributeAssignment(MP_EmissiveColor, &EmissiveColor);
+	PushAttributeAssignment(MP_Opacity, &Opacity);
+	PushAttributeAssignment(MP_OpacityMask, &OpacityMask);
+	PushAttributeAssignment(MP_Normal, &Normal);
+	PushAttributeAssignment(MP_Tangent, &Tangent);
+	PushAttributeAssignment(MP_WorldPositionOffset, &WorldPositionOffset);
+	PushAttributeAssignment(MP_SubsurfaceColor, &SubsurfaceColor);
+	PushAttributeAssignment(MP_CustomData0, &ClearCoat);
+	PushAttributeAssignment(MP_CustomData1, &ClearCoatRoughness);
+	PushAttributeAssignment(MP_AmbientOcclusion, &AmbientOcclusion);
+	PushAttributeAssignment(MP_Refraction, &Refraction);
+	PushAttributeAssignment(MP_PixelDepthOffset, &PixelDepthOffset);
+	PushAttributeAssignment(MP_ShadingModel, &ShadingModel);
+	PushAttributeAssignment(MP_Displacement, &Displacement);
+
+	for (int32 i = 0; i < 8; ++i)
+	{
+		PushAttributeAssignment(EMaterialProperty(MP_CustomizedUVs0 + i), &CustomizedUVs[i]);
+	}
+
+	Em.Output(0, Em.Aggregate(UMaterialAggregate::GetMaterialAttributes(), {}, Assignments.Left(NumAssignments)));
+}
+
+void UMaterialExpressionBreakMaterialAttributes::Build(MIR::FEmitter& Em)
+{
+	FValueRef Prototype = Em.CheckIsAggregate(Em.Input(&MaterialAttributes), UMaterialAggregate::GetMaterialAttributes());
+	UE_MIR_CHECKPOINT(Em);
+
+	static const EMaterialProperty Properties[] = {
+		MP_BaseColor,
+		MP_Metallic,
+		MP_Specular,
+		MP_Roughness,
+		MP_Anisotropy,
+		MP_EmissiveColor,
+		MP_Opacity,
+		MP_OpacityMask,
+		MP_Normal,
+		MP_Tangent,
+		MP_WorldPositionOffset,
+		MP_SubsurfaceColor,
+		MP_CustomData0, // ClearColor
+		MP_CustomData1, // ClearColorRoughness
+		MP_AmbientOcclusion,
+		MP_Refraction,
+		MP_PixelDepthOffset,
+		MP_ShadingModel,
+		MP_Displacement
+	};
+
+	for (int32 Index = 0; Index < UE_ARRAY_COUNT(Properties); ++Index)
+	{
+		Em.Output(Index, Em.Subscript(Prototype, UMaterialAggregate::MaterialPropertyToAttributeIndex(Properties[Index])));
+	}
+}
+
+// Verifies that the attribute ids in the MaterialAttributes expression are valid (e.g. no duplicates, proper mapping).
+static void CheckMaterialAttributesExpression(MIR::FEmitter& Em, TConstArrayView<FGuid> AttributeIds)
+{
+	for (int32 i = 0; i < AttributeIds.Num(); ++i)
+	{
+		for (int j = i + 1; j < AttributeIds.Num(); ++j)
+		{
+			if (AttributeIds[i] == AttributeIds[j])
+			{
+				Em.Error(TEXT("Duplicate attribute types."));
+				return;
+			}
+		}
+
+		if (FMaterialAttributeDefinitionMap::GetProperty(AttributeIds[i]) == MP_MAX)
+		{
+			Em.Error(TEXT("Property type doesn't exist, needs re-mapping?"));
+			return;
+		}
+	}
+}
+
+void UMaterialExpressionGetMaterialAttributes::Build(MIR::FEmitter& Em)
+{
+	CheckMaterialAttributesExpression(Em, AttributeGetTypes);
+
+	FValueRef Prototype = Em.CheckIsAggregate(Em.Input(&MaterialAttributes), UMaterialAggregate::GetMaterialAttributes());
+	
+	UE_MIR_CHECKPOINT(Em);
+	
+	Em.Output(0, Prototype);
+
+	const UMaterialAggregate* MaterialAttributesAggregate = UMaterialAggregate::GetMaterialAttributes();
+	for (int32 i = 0; i < AttributeGetTypes.Num(); ++i)
+	{
+		EMaterialProperty Property = FMaterialAttributeDefinitionMap::GetProperty(AttributeGetTypes[i]);
+		check(Property != MP_MAX);
+
+		int32 AttributeIndex = MaterialAttributesAggregate->FindAttributeIndexByName(*FMaterialAttributeDefinitionMap::GetAttributeName(Property));
+
+		Em.Output(i + 1, Em.Subscript(Prototype, AttributeIndex));
+	}
+}
+
+void UMaterialExpressionSetMaterialAttributes::Build(MIR::FEmitter& Em)
+{
+	CheckMaterialAttributesExpression(Em, AttributeSetTypes);
+	
+	FValueRef Prototype = Em.CheckIsAggregate(Em.Input(&Inputs[0]), UMaterialAggregate::GetMaterialAttributes());
+	
+	UE_MIR_CHECKPOINT(Em);
+
+	const UMaterialAggregate* MaterialAttributesAggregate = UMaterialAggregate::GetMaterialAttributes();
+	MIR::TTemporaryArray<MIR::FAttributeAssignment> Assignments{ AttributeSetTypes.Num() };
+	int32 NumAssignments = 0;
+
+	for (int32 i = 0; i < AttributeSetTypes.Num(); ++i)
+	{
+		EMaterialProperty Property = FMaterialAttributeDefinitionMap::GetProperty(AttributeSetTypes[i]);
+		check(Property != MP_MAX);
+
+		if (FValueRef Value = Em.TryInput(&Inputs[i + 1]))
+		{
+			Assignments[NumAssignments++] = { *FMaterialAttributeDefinitionMap::GetAttributeName(Property), Value };
+		}
+	}
+
+	Prototype = Em.Aggregate(UMaterialAggregate::GetMaterialAttributes(), Prototype, Assignments.Left(NumAssignments));
+
+	Em.Output(0, Prototype);
+}
+
+// Utility to input a MaterialAttributes value, or return the default instance (with each attribute set to zero). 
+static FValueRef InputDefaultMaterialAttributes(MIR::FEmitter& Em, FExpressionInput* Input)
+{
+	const UMaterialAggregate* MaterialAttributes = UMaterialAggregate::GetMaterialAttributes();
+	FValueRef Value = Em.CheckIsAggregate(Em.TryInput(Input), MaterialAttributes);
+	return Value ? Value : Em.Aggregate(MaterialAttributes);
+}
+
+// Converts old EMaterialAttributeBlend::Type to EMaterialExpressionBlendMode.
+static EMaterialExpressionBlendMode ConvertMaterialAttributeBlend(EMaterialAttributeBlend::Type InBlend)
+{
+	switch (InBlend)
+	{
+		case EMaterialAttributeBlend::Blend: return EMaterialExpressionBlendMode::Blend;
+		case EMaterialAttributeBlend::UseA: return EMaterialExpressionBlendMode::UseA;
+		case EMaterialAttributeBlend::UseB: return EMaterialExpressionBlendMode::UseB;
+		default: UE_MIR_UNREACHABLE();
+	}
+}
+
+// Forward declaration.
+static FValueRef Blend(MIR::FEmitter& Em, EMaterialExpressionBlendMode PixelAttributeBlendMode, EMaterialExpressionBlendMode VertexAttributeBlendMode, FValueRef A, FValueRef B, FValueRef Alpha);
+
+// Blends two argument aggregate values based on [0-1] alpha value. See Blend() below for more info.
+static FValueRef BlendAggregate(MIR::FEmitter& Em,
+						   EMaterialExpressionBlendMode PixelAttributesBlendMode,
+						   EMaterialExpressionBlendMode VertexAttributesBlendMode,
+						   FValueRef A,
+						   FValueRef B,
+						   FValueRef Alpha)
+{
+	const UMaterialAggregate* MaterialAggregate = A->Type->AsAggregate()->Aggregate;
+	MIR::TTemporaryArray<MIR::FValueRef> AttributeValues{ MaterialAggregate->Attributes.Num() };
+
+	for (int32 i = 0; i < MaterialAggregate->Attributes.Num(); ++i)
+	{
+		EMaterialExpressionBlendMode BlendMode = PixelAttributesBlendMode;
+		if (MaterialAggregate == UMaterialAggregate::GetMaterialAttributes())
+		{
+			EMaterialProperty Property = UMaterialAggregate::AttributeIndexToMaterialProperty(i);
+			BlendMode = (Property == MP_WorldPositionOffset) ? VertexAttributesBlendMode : PixelAttributesBlendMode;
+		}
+
+		if (BlendMode == EMaterialExpressionBlendMode::UseA)
+		{
+			AttributeValues[i] = Em.Subscript(A, i);
+		}
+		if (BlendMode == EMaterialExpressionBlendMode::UseB)
+		{
+			AttributeValues[i] = Em.Subscript(B, i);
+		}
+		else if (BlendMode == EMaterialExpressionBlendMode::Blend)
+		{
+			AttributeValues[i] = Blend(Em, PixelAttributesBlendMode, VertexAttributesBlendMode, Em.Subscript(A, i), Em.Subscript(B, i), Alpha);
+		}
+	}
+
+	return Em.Aggregate(MaterialAggregate, {}, AttributeValues);
+}
+
+// Blends two argument values based on [0-1] alpha value. If argument values are or contain MaterialAttributes aggregates,
+// PixelAttributeBlendMode and VertexAttributeBlendMode instruct on how to blend the attributes depending on whether
+// they're evaluated in pixel or vertex shaders.
+// Note: VertexAttributeBlendMode are only used when blending MaterialAttributes. Otherwise, PixelAttributeBlendMode is used.
+static FValueRef Blend(MIR::FEmitter& Em, EMaterialExpressionBlendMode PixelAttributeBlendMode, EMaterialExpressionBlendMode VertexAttributeBlendMode, FValueRef A, FValueRef B, FValueRef Alpha)
+{
+	// Find the common type between arguments
+	if (const MIR::FType* CommonType = Em.GetCommonType(A->Type, B->Type))
+	{
+		// And cast both arguments to the common type
+		A = Em.Cast(A, CommonType);
+		B = Em.Cast(B, CommonType);
+	}
+	else
+	{
+		return Em.Poison();
+	}
+
+	if (A->Type->IsAnyFloat())
+	{
+		// Blend floating point values using linear interpolation
+		return Em.Lerp(A, B, Alpha);
+	}
+	else if (A->Type->IsInteger())
+	{
+		// "Blend" integer values by selecting A or B based on whether alpha is less than 0.5.
+		return Em.Select(Em.LessThan(Alpha, Em.ConstantFloat(0.5f)), A, B);
+	}
+	else if (const MIR::FAggregateType* AggregateType = A->Type->AsAggregate())
+	{
+		// Arguments are aggregates, so recursively blend each attribute pair.
+		return BlendAggregate(Em, PixelAttributeBlendMode, VertexAttributeBlendMode, A, B, Alpha);
+	}
+	else
+	{
+		Em.Errorf(TEXT("Cannot blend values of type '%s'."), *A->Type->GetSpelling());
+		return Em.Poison();
+	}
+}
+
+void UMaterialExpressionBlendMaterialAttributes::Build(MIR::FEmitter& Em)
+{
+	const UMaterialAggregate* MaterialAttributes = UMaterialAggregate::GetMaterialAttributes();
+	
+	FValueRef AValue = InputDefaultMaterialAttributes(Em, &A);
+	FValueRef BValue = InputDefaultMaterialAttributes(Em, &B);
+	FValueRef AlphaValue = Em.CastToFloat(Em.InputDefaultFloat(&Alpha, 0.0f), 1);
+
+	UE_MIR_CHECKPOINT(Em);
+
+	FValueRef Result = BlendAggregate(
+		Em, 
+		ConvertMaterialAttributeBlend(PixelAttributeBlendType),
+		ConvertMaterialAttributeBlend(VertexAttributeBlendType),
+		AValue,
+		BValue,
+		AlphaValue);
+
+	Em.Output(0, Result);
+}
+
+void UMaterialExpressionAggregate::Build(MIR::FEmitter& Em)
+{
+	// Get the material aggregate definition.
+	const UMaterialAggregate* Aggregate = GetAggregate();
+	if (!Aggregate)
+	{
+		Em.Error(TEXT("Unspecified material aggregate."));
+		return;
+	}
+
+	// Read the aggregate prototype value, if present, and make sure it is of the right type.
+	FValueRef Prototype = Em.CheckIsAggregate(Em.TryInput(&PrototypeInput), Aggregate);
+	UE_MIR_CHECKPOINT(Em);
+
+	// Collect the attribute assignments from the input pins.
+	MIR::TTemporaryArray<MIR::FAttributeAssignment> Assignments{ Entries.Num() };
+	int32 NumAssignments = 0;
+
+	for (const FMaterialExpressionAggregateEntry& Entry : Entries)
+	{
+		// If value is present, push this attribute assignment.
+		if (FValueRef AttributeValue = Em.TryInput(&Entry.Input))
+		{
+			Assignments[NumAssignments++] = { Aggregate->Attributes[Entry.AttributeIndex].Name, AttributeValue };
+		}
+	}
+	
+	// Make hte aggregate value using the optional prototype and assignments.
+	Prototype = Em.Aggregate(Aggregate, Prototype, Assignments.Left(NumAssignments));
+
+	// Output the aggregate value
+	Em.Output(0, Prototype);
+
+	// And output each individual aggregate attribute through the invidual output pins
+	for (int i = 0; i < Entries.Num(); ++i)
+	{
+		Em.Output(i + 1, Em.Subscript(Prototype, Entries[i].AttributeIndex));
+	}
+}
+
+void UMaterialExpressionBlend::Build(MIR::FEmitter& Em)
+{
+	// Try reading the input values (could be null).
+	FValueRef AValue = Em.TryInput(&A);
+	FValueRef BValue = Em.TryInput(&B);
+
+	if (!AValue && !BValue)
+	{
+		Em.Error(TEXT("No input value provided."));
+		return;
+	}
+
+	// Create a default value from the other input's type if any input is missing.
+	if (!AValue)
+	{
+		AValue = Em.ConstantDefault(BValue->Type);
+	}
+	else if (!BValue)
+	{
+		BValue = Em.ConstantDefault(AValue->Type);
+	}
+
+	// Read the alpha value (defaulting it to 0.0f)
+	FValueRef AlphaValue = Em.CastToFloat(Em.InputDefaultFloat(&Alpha, 0.0f), 1);
+
+	// Make sure all previous operations went well.
+	UE_MIR_CHECKPOINT(Em);
+
+	// Blend the input values.
+	FValueRef Result = Blend(Em, PixelAttributesBlendMode, VertexAttributesBlendMode, AValue, BValue, AlphaValue);
+	
+	Em.Output(0, Result);
+}
+
 #endif // WITH_EDITOR
diff --git a/Engine/Source/Runtime/Engine/Private/Materials/MaterialIR.cpp b/Engine/Source/Runtime/Engine/Private/Materials/MaterialIR.cpp
--- a/Engine/Source/Runtime/Engine/Private/Materials/MaterialIR.cpp
+++ b/Engine/Source/Runtime/Engine/Private/Materials/MaterialIR.cpp
@@ -2,6 +2,7 @@
 
 #include "Materials/MaterialIR.h"
 #include "Materials/MaterialIRTypes.h"
+#include "Materials/MaterialAggregate.h"
 #include "Engine/Texture.h"
 
 #if WITH_EDITOR
@@ -42,7 +43,7 @@ const TCHAR* LexToString(EValueKind Kind)
 		case VK_ExternalInput: return TEXT("ExternalInput");
 		case VK_TextureObject: return TEXT("TextureObject");
 		case VK_UniformParameter: return TEXT("UniformParameter");
-		case VK_Dimensional: return TEXT("Dimensional");
+		case VK_Composite: return TEXT("Composite");
 		case VK_SetMaterialOutput: return TEXT("SetMaterialOutput");
 		case VK_Operator: return TEXT("Operator");
 		case VK_Branch: return TEXT("Branch");
@@ -101,7 +102,7 @@ uint32 FValue::GetSizeInBytes() const
 		case VK_ExternalInput: return sizeof(FExternalInput);
 		case VK_TextureObject: return sizeof(FTextureObject);
 		case VK_UniformParameter: return sizeof(FUniformParameter);
-		case VK_Dimensional: return sizeof(FDimensional) + sizeof(FValue*) * static_cast<const FDimensional*>(this)->GetComponents().Num();
+		case VK_Composite: return sizeof(FComposite) + sizeof(FValue*) * static_cast<const FComposite*>(this)->GetComponents().Num();
 		case VK_SetMaterialOutput: return sizeof(FSetMaterialOutput);
 		case VK_Operator: return sizeof(FOperator);
 		case VK_Branch: return sizeof(FBranch);
@@ -159,9 +160,9 @@ TConstArrayView<FValue*> FValue::GetUses() const
 
 	switch (Kind)
 	{
-		case VK_Dimensional:
+		case VK_Composite:
 		{
-			auto This = static_cast<const FDimensional*>(this);
+			auto This = static_cast<const FComposite*>(this);
 			return This->GetComponents();
 		}
 
@@ -268,7 +269,7 @@ bool FValue::IsFalse() const
 
 bool FValue::AreAllTrue() const
 {
-	if (const MIR::FDimensional* Dimensional = As<MIR::FDimensional>())
+	if (const MIR::FComposite* Dimensional = As<MIR::FComposite>())
 	{
 		for (const MIR::FValue* Component : Dimensional->GetComponents())
 		{
@@ -287,7 +288,7 @@ bool FValue::AreAllTrue() const
 
 bool FValue::AreAllFalse() const
 {
-	if (const MIR::FDimensional* Dimensional = As<MIR::FDimensional>())
+	if (const MIR::FComposite* Dimensional = As<MIR::FComposite>())
 	{
 		for (const MIR::FValue* Component : Dimensional->GetComponents())
 		{
@@ -306,7 +307,7 @@ bool FValue::AreAllFalse() const
 
 bool FValue::AreAllExactlyZero() const
 {
-	if (const MIR::FDimensional* Dimensional = As<MIR::FDimensional>())
+	if (const MIR::FComposite* Dimensional = As<MIR::FComposite>())
 	{
 		for (const MIR::FValue* Component : Dimensional->GetComponents())
 		{
@@ -326,7 +327,7 @@ bool FValue::AreAllExactlyZero() const
 
 bool FValue::AreAllNearlyZero() const
 {
-	if (const MIR::FDimensional* Dimensional = As<MIR::FDimensional>())
+	if (const MIR::FComposite* Dimensional = As<MIR::FComposite>())
 	{
 		for (const MIR::FValue* Component : Dimensional->GetComponents())
 		{
@@ -346,7 +347,7 @@ bool FValue::AreAllNearlyZero() const
 
 bool FValue::AreAllExactlyOne() const
 {
-	if (const MIR::FDimensional* Dimensional = As<MIR::FDimensional>())
+	if (const MIR::FComposite* Dimensional = As<MIR::FComposite>())
 	{
 		for (const MIR::FValue* Component : Dimensional->GetComponents())
 		{
@@ -366,7 +367,7 @@ bool FValue::AreAllExactlyOne() const
 
 bool FValue::AreAllNearlyOne() const
 {
-	if (const MIR::FDimensional* Dimensional = As<MIR::FDimensional>())
+	if (const MIR::FComposite* Dimensional = As<MIR::FComposite>())
 	{
 		for (const MIR::FValue* Component : Dimensional->GetComponents())
 		{
@@ -591,23 +592,34 @@ FBlock* FBlock::FindCommonParentWith(MIR::FBlock* Other)
 	return A;
 }
 
-TConstArrayView<FValue*> FDimensional::GetComponents() const
+TConstArrayView<FValue*> FComposite::GetComponents() const
 {
-	const FPrimitiveType* PrimitiveType = Type->AsPrimitive();
-	check(PrimitiveType);
+	int32 NumComponents;
+	if (const FPrimitiveType* PrimitiveType = Type->AsPrimitive())
+	{
+		NumComponents = PrimitiveType->NumComponents();
+	}
+	else if (const FAggregateType* AggregateType = Type->AsAggregate())
+	{
+		NumComponents = AggregateType->Aggregate->Attributes.Num();
+	}
+	else
+	{
+		UE_MIR_UNREACHABLE();
+	}
 
-	auto Ptr = static_cast<const TDimensional<1>*>(this)->Components;
-	return { Ptr, PrimitiveType->NumRows };
+	auto Ptr = static_cast<const TComposite<1>*>(this)->Components;
+	return { Ptr, NumComponents };
 }
 
-TArrayView<FValue*> FDimensional::GetMutableComponents()
+TArrayView<FValue*> FComposite::GetMutableComponents()
 {
 	// const_cast is okay here as it's used only to get the array of components
-	TConstArrayView<FValue*> Components = static_cast<const FDimensional*>(this)->GetComponents();
+	TConstArrayView<FValue*> Components = static_cast<const FComposite*>(this)->GetComponents();
 	return { const_cast<FValue**>(Components.GetData()), Components.Num() };
 }
 
-bool FDimensional::AreComponentsConstant() const
+bool FComposite::AreComponentsConstant() const
 {
 	for (FValue const* Component : GetComponents())
 	{
diff --git a/Engine/Source/Runtime/Engine/Private/Materials/MaterialIRDebug.cpp b/Engine/Source/Runtime/Engine/Private/Materials/MaterialIRDebug.cpp
--- a/Engine/Source/Runtime/Engine/Private/Materials/MaterialIRDebug.cpp
+++ b/Engine/Source/Runtime/Engine/Private/Materials/MaterialIRDebug.cpp
@@ -61,7 +61,7 @@ struct FDebugDumpIRUseGraphState
 		Out.Appendf(TEXT("\"%p\" [label=< <b>%s</b>  (%s) <br/> "),
 						Value,
 						LexToString(Value->Kind),
-						Value->Type ? Value->Type->GetSpelling().GetData() : TEXT("???"));
+						Value->Type ? *Value->Type->GetSpelling() : TEXT("???"));
 
 		DumpValueInfo(Value);
 
@@ -152,9 +152,9 @@ struct FDebugDumpIRUseGraphState
 
 	void DumpUseInfo(const FValue* Value, const FValue* Use, int32 UseIndex)
 	{
-		if (const FDimensional* Dimensional = Value->As<FDimensional>())
+		if (const FComposite* Composite = Value->As<FComposite>())
 		{
-			if (Dimensional->Type->AsPrimitive()->IsVector())
+			if (Composite->Type->AsPrimitive()->IsVector())
 			{
 				check(UseIndex < 4);
 				Out.AppendChar(TEXT("xyzw")[UseIndex]);
@@ -266,7 +266,7 @@ struct FDebugDumpIRState
 				}
 
 				// First the type...
-				Out.Appendf(TEXT("%s "), Use->Type->GetSpelling().GetData());
+				Out.Appendf(TEXT("%s "), *Use->Type->GetSpelling());
 
 				// If this use is in a block different from current's, dump the block in "{}" first.
 				const FBlock* UseBlock = Instr->GetDesiredBlockForUse(CurrentStage, UseIndex);
diff --git a/Engine/Source/Runtime/Engine/Private/Materials/MaterialIREmitter.cpp b/Engine/Source/Runtime/Engine/Private/Materials/MaterialIREmitter.cpp
--- a/Engine/Source/Runtime/Engine/Private/Materials/MaterialIREmitter.cpp
+++ b/Engine/Source/Runtime/Engine/Private/Materials/MaterialIREmitter.cpp
@@ -5,6 +5,7 @@
 #include "Materials/MaterialIRModuleBuilder.h"
 #include "Materials/MaterialIRInternal.h"
 #include "Materials/MaterialExternalCodeRegistry.h"
+#include "Materials/MaterialAggregate.h"
 #include "Shader/ShaderTypes.h"
 #include "MaterialShared.h"
 #include "MaterialSharedPrivate.h"
@@ -80,10 +81,11 @@ struct FEmitter::FPrivate
 		return Value ? *Value : nullptr;
 	}
 
-	// Allocates a block of memory using the emitter's allocator.
+	// Allocates a block of memory.
 	static uint8* Allocate(FEmitter& Emitter, int32 Size, int32 Alignment)
 	{
-		uint8* Bytes = (uint8*) FMemory::Malloc(Size, Alignment); // Emitter.Module->Allocator.PushBytes(Size, Alignment);
+		// TODO: use a linear allocator inside the module instead.
+		uint8* Bytes = (uint8*) FMemory::Malloc(Size, Alignment); 
 		FMemory::Memzero(Bytes, Size);
 		return Bytes;
 	}
@@ -96,25 +98,22 @@ struct FEmitter::FPrivate
 	}
 };
 
-// Creates a new `FDimensional` value of specified `Type` and returns it.
-static FDimensional* NewDimensionalValue(FEmitter& Emitter, const FPrimitiveType* Type)
+// Creates a new composite value of specified type and returns it.
+static FComposite* NewCompositeValue(FEmitter& Emitter, const FType* Type, int NumComponents)
 {
-	check(!Type->IsScalar());
+	int32 SizeInBytes = sizeof(FComposite) + sizeof(FValue*) * NumComponents;
 
-	int32 Dimensions = Type->NumRows * Type->NumColumns;
-	int32 SizeInBytes = sizeof(FDimensional) + sizeof(FValue*) * Dimensions;
+	uint8* Bytes = FEmitter::FPrivate::Allocate(Emitter, SizeInBytes, alignof(FComposite));
 
-	uint8* Bytes = FEmitter::FPrivate::Allocate(Emitter, SizeInBytes, alignof(FDimensional));
-
-	FDimensional* Value = new (Bytes) FDimensional{};
-	Value->Kind = VK_Dimensional;
+	FComposite* Value = new (Bytes) FComposite{};
+	Value->Kind = VK_Composite;
 	Value->Type = Type;
 
 	return Value;
 }
 
-// Emits specified newly created `Value`. If the exact value already exists,
-// specified one is *destroyed* and existing one is returned instead.
+// It adds the given value to the module and returns it. If the exact value already existed,
+// it destroys the the given value and returns the existing one instead.
 static FValueRef EmitNew(FEmitter& Emitter, FValue* Value)
 {
 	if (FValue* Existing = FEmitter::FPrivate::FindValue(Emitter, Value))
@@ -219,6 +218,11 @@ bool FValueRef::IsValid() const
 	return Value && !Value->IsPoison();
 }
 
+bool FValueRef::IsPoison() const
+{
+	return Value && Value->IsPoison();
+}
+
 FValueRef FValueRef::To(FValue* InValue) const
 {
 	return { InValue, Input };
@@ -262,10 +266,61 @@ const FType* FEmitter::GetCommonType(const FType* A, const FType* B)
 		return CommonType;
 	}
 
-	Errorf(TEXT("No common type between '%s' and '%s'."), A->GetSpelling().GetData(), B->GetSpelling().GetData());
+	Errorf(TEXT("No common type between '%s' and '%s'."), *A->GetSpelling(), *B->GetSpelling());
 	return nullptr;
 }
 
+const FAggregateType* FEmitter::GetMaterialAggregateType(const UMaterialAggregate* Aggregate)
+{
+	// First, see if the aggregate type hasn't been mapped already.
+	FAggregateType* Type{};
+	if (Internal::Find(AggregateTypeMap, Aggregate, Type))
+	{
+		return Type;
+	}
+
+	// If note, create a MIR type container for this aggregate.
+	Type = new (FPrivate::Allocate(*this, sizeof(FAggregateType), alignof(FAggregateType))) FAggregateType{};
+	Type->Kind = TypeKind_Aggregate;
+	Type->Aggregate = Aggregate;
+
+	// Add the new aggregate type to the map.
+	AggregateTypeMap.Add(Aggregate, Type);
+
+	return Type;
+}
+
+const FType* FEmitter::GetMaterialAggregateAttributeType(const UMaterialAggregate* Aggregate, int32 AttributeIndex)
+{
+	check(AttributeIndex >= 0);
+
+	if (AttributeIndex >= Aggregate->Attributes.Num())
+	{
+		Errorf(TEXT("Invalid attribute index %d for material aggregate '%s'. Index is out of range (Num = %d)."),
+			   AttributeIndex, *Aggregate->GetName(), Aggregate->Attributes.Num());
+		return nullptr;
+	}
+
+	switch (Aggregate->Attributes[AttributeIndex].Type)
+	{
+		case EMaterialAggregateAttributeType::Bool1: return FPrimitiveType::GetVector(ScalarKind_Bool, 1);
+		case EMaterialAggregateAttributeType::Bool2: return FPrimitiveType::GetVector(ScalarKind_Bool, 2);
+		case EMaterialAggregateAttributeType::Bool3: return FPrimitiveType::GetVector(ScalarKind_Bool, 3);
+		case EMaterialAggregateAttributeType::Bool4: return FPrimitiveType::GetVector(ScalarKind_Bool, 4);
+		case EMaterialAggregateAttributeType::ShadingModel:
+		case EMaterialAggregateAttributeType::UInt1: return FPrimitiveType::GetVector(ScalarKind_Int, 1);
+		case EMaterialAggregateAttributeType::UInt2: return FPrimitiveType::GetVector(ScalarKind_Int, 2);
+		case EMaterialAggregateAttributeType::UInt3: return FPrimitiveType::GetVector(ScalarKind_Int, 3);
+		case EMaterialAggregateAttributeType::UInt4: return FPrimitiveType::GetVector(ScalarKind_Int, 4);
+		case EMaterialAggregateAttributeType::Float1: return FPrimitiveType::GetVector(ScalarKind_Float, 1);
+		case EMaterialAggregateAttributeType::Float2: return FPrimitiveType::GetVector(ScalarKind_Float, 2);
+		case EMaterialAggregateAttributeType::Float3: return FPrimitiveType::GetVector(ScalarKind_Float, 3);
+		case EMaterialAggregateAttributeType::Float4: return FPrimitiveType::GetVector(ScalarKind_Float, 4);
+		case EMaterialAggregateAttributeType::Aggregate: return GetMaterialAggregateType(Aggregate->Attributes[AttributeIndex].Aggregate.Get());
+		default: UE_MIR_UNREACHABLE();
+	}
+}
+
 /*-------------------------------- Input management --------------------------------*/
 
 FValueRef FEmitter::TryInput(const FExpressionInput* InInput)
@@ -342,7 +397,7 @@ FValueRef FEmitter::CheckTypeIsKind(FValueRef Value, ETypeKind Kind)
 {
 	if (Value.IsValid() && Value->Type->Kind != Kind)
 	{
-		Errorf(Value, TEXT("Expected a '%s' value, got a '%s' instead."), TypeKindToString(Kind), Value->Type->GetSpelling().GetData());
+		Errorf(Value, TEXT("Expected a '%s' value, got a '%s' instead."), TypeKindToString(Kind), *Value->Type->GetSpelling());
 		return Value.ToPoison();
 	}
 	return Value;
@@ -352,7 +407,7 @@ FValueRef FEmitter::CheckIsPrimitive(FValueRef Value)
 {
 	if (Value.IsValid() && !Value->Type->AsPrimitive())
 	{
-		Errorf(Value, TEXT("Expected a primitive value, got a '%s' instead."), Value->Type->GetSpelling().GetData());
+		Errorf(Value, TEXT("Expected a primitive value, got a '%s' instead."), *Value->Type->GetSpelling());
 		return Value.ToPoison();
 	}
 	return Value;
@@ -362,7 +417,7 @@ FValueRef FEmitter::CheckIsArithmetic(FValueRef Value)
 {
 	if (Value.IsValid() && !Value->Type->AsArithmetic())
 	{
-		Errorf(Value, TEXT("Expected an arithmetic value, got a '%s' instead."), Value->Type->GetSpelling().GetData());
+		Errorf(Value, TEXT("Expected an arithmetic value, got a '%s' instead."), *Value->Type->GetSpelling());
 		return Value.ToPoison();
 	}
 	return Value;
@@ -372,7 +427,7 @@ FValueRef FEmitter::CheckIsInteger(FValueRef Value)
 {
 	if (Value.IsValid() && !Value->Type->IsInteger())
 	{
-		Errorf(Value, TEXT("Expected an integer value, got a '%s' instead."), Value->Type->GetSpelling().GetData());
+		Errorf(Value, TEXT("Expected an integer value, got a '%s' instead."), *Value->Type->GetSpelling());
 		return Value.ToPoison();
 	}
 	return Value;
@@ -382,7 +437,7 @@ FValueRef FEmitter::CheckIsScalar(FValueRef Value)
 {
 	if (Value.IsValid() && (!Value->Type->AsPrimitive() || !Value->Type->AsPrimitive()->IsScalar()))
 	{
-		Errorf(Value, TEXT("Expected a scalar value, got a '%s' instead."), Value->Type->GetSpelling().GetData());
+		Errorf(Value, TEXT("Expected a scalar value, got a '%s' instead."), *Value->Type->GetSpelling());
 		return Value.ToPoison();
 	}
 	return Value;
@@ -392,7 +447,7 @@ FValueRef FEmitter::CheckIsScalarOrVector(FValueRef Value)
 {
 	if (Value.IsValid() && (!Value->Type->AsPrimitive() || Value->Type->AsPrimitive()->IsMatrix()))
 	{
-		Errorf(Value, TEXT("Expected a scalar or vector value, got a '%s' instead."), Value->Type->GetSpelling().GetData());
+		Errorf(Value, TEXT("Expected a scalar or vector value, got a '%s' instead."), *Value->Type->GetSpelling());
 		return Value.ToPoison();
 	}
 	return Value;
@@ -402,12 +457,30 @@ FValueRef FEmitter::CheckIsTexture(FValueRef Value)
 {
 	if (Value.IsValid() && !Value->Type->IsTexture())
 	{
-		Errorf(Value, TEXT("Expected a texture value, got a '%s' instead."), Value->Type->GetSpelling().GetData());
+		Errorf(Value, TEXT("Expected a texture value, got a '%s' instead."), *Value->Type->GetSpelling());
 		return Value.ToPoison();
 	}
 	return Value;
 }
 
+FValueRef FEmitter::CheckIsAggregate(FValueRef Value, const UMaterialAggregate* Aggregate)
+{
+	if (Value.IsValid())
+	{
+		if (!Value->Type->AsAggregate())
+		{
+			Errorf(Value, TEXT("Expected an aggregate value, got a '%s' instead."), *Value->Type->GetSpelling());
+			return Value.ToPoison();
+		}
+		if (Aggregate && Value->Type->AsAggregate()->Aggregate != Aggregate)
+		{
+			Errorf(Value, TEXT("Expected a value of aggregate type '%s', got a value of aggregate type '%s' instead."), *Aggregate->GetName(), *Value->Type->AsAggregate()->Aggregate->GetName());
+			return Value.ToPoison();
+		}
+	}
+	return Value;
+}
+
 bool FEmitter::ToConstantBool(FValueRef Value)
 {
 	if (!Value.IsValid())
@@ -422,7 +495,7 @@ bool FEmitter::ToConstantBool(FValueRef Value)
     }
 	if (Constant->Type != FPrimitiveType::GetBool())
     {
-        Errorf(Value, TEXT("Expected a constant bool value, got a '%s' instead."), Constant->Type->GetSpelling().GetData());
+        Errorf(Value, TEXT("Expected a constant bool value, got a '%s' instead."), *Constant->Type->GetSpelling());
         return false;
     }
 	return Constant->Boolean;
@@ -467,6 +540,36 @@ FValueRef FEmitter::ConstantFromShaderValue(const UE::Shader::FValue& InValue)
 	UE_MIR_UNREACHABLE();
 }
 
+FValueRef FEmitter::ConstantDefault(const FType* Type)
+{
+	if (const FPrimitiveType* PrimitiveType = Type->AsPrimitive())
+	{
+		FValueRef Zero = ConstantZero(PrimitiveType->ScalarKind);
+		if (PrimitiveType->IsScalar())
+		{
+			return Zero;
+		}
+		else
+		{
+			FComposite* Composite = NewCompositeValue(*this, Type, PrimitiveType->NumComponents());
+			for (FValue*& Component : Composite->GetMutableComponents())
+			{
+				Component = Zero;
+			}
+			return EmitNew(*this, Composite);
+		}
+	}
+	else if (const FAggregateType* AggregateType = Type->AsAggregate())
+	{
+		return Aggregate(AggregateType->Aggregate);
+	}
+	else
+	{
+		Errorf(TEXT("Type '%s' has no default. Expected primitive or aggregate type."), *Type->GetSpelling());
+		return Poison();
+	}
+}
+
 FValueRef FEmitter::ConstantZero(EScalarKind Kind)
 {
 	switch (Kind)
@@ -808,18 +911,20 @@ FValueRef FEmitter::Vector2(FValueRef InX, FValueRef InY)
 	check(InX->Type->AsScalar());
 	check(InX->Type == InY->Type);
 
-	TDimensional<2> Vector = MakePrototype<TDimensional<2>>(FPrimitiveType::GetVector(InX->Type->AsPrimitive()->ScalarKind, 2));
+	TComposite<2> Vector = MakePrototype<TComposite<2>>(FPrimitiveType::GetVector(InX->Type->AsPrimitive()->ScalarKind, 2));
 	TArrayView<FValue*> Components = Vector.GetMutableComponents();
 	Components[0] = InX;
 	Components[1] = InY;
 
 	FValueRef Value = EmitPrototype(*this, Vector);
+	
+	// If all components source inputs match, refer to the same input in the vector value too (for debugging purposes).
 	if (InX.Input == InY.Input)
 	{
 		Value.Input = InX.Input;
 	}
 
-	return EmitPrototype(*this, Vector);
+	return Value;
 }
 
 FValueRef FEmitter::Vector3(FValueRef InX, FValueRef InY, FValueRef InZ)
@@ -833,13 +938,15 @@ FValueRef FEmitter::Vector3(FValueRef InX, FValueRef InY, FValueRef InZ)
 	check(InX->Type == InY->Type);
 	check(InY->Type == InZ->Type);
 
-	TDimensional<3> Vector = MakePrototype<TDimensional<3>>(FPrimitiveType::GetVector(InX->Type->AsPrimitive()->ScalarKind, 3));
+	TComposite<3> Vector = MakePrototype<TComposite<3>>(FPrimitiveType::GetVector(InX->Type->AsPrimitive()->ScalarKind, 3));
 	TArrayView<FValue*> Components = Vector.GetMutableComponents();
 	Components[0] = InX;
 	Components[1] = InY;
 	Components[2] = InZ;
 
 	FValueRef Value = EmitPrototype(*this, Vector);
+
+	// If all components source inputs match, refer to the same input in the vector value too (for debugging purposes).
 	if (InX.Input == InY.Input && InX.Input == InZ.Input)
 	{
 		Value.Input = InX.Input;
@@ -860,7 +967,7 @@ FValueRef FEmitter::Vector4(FValueRef InX, FValueRef InY, FValueRef InZ, FValueR
 	check(InY->Type == InZ->Type);
 	check(InZ->Type == InW->Type);
 
-	TDimensional<4> Vector = MakePrototype<TDimensional<4>>(FPrimitiveType::GetVector(InX->Type->AsPrimitive()->ScalarKind, 4));
+	TComposite<4> Vector = MakePrototype<TComposite<4>>(FPrimitiveType::GetVector(InX->Type->AsPrimitive()->ScalarKind, 4));
 	TArrayView<FValue*> Components = Vector.GetMutableComponents();
 	Components[0] = InX;
 	Components[1] = InY;
@@ -868,12 +975,108 @@ FValueRef FEmitter::Vector4(FValueRef InX, FValueRef InY, FValueRef InZ, FValueR
 	Components[3] = InW;
 
 	FValueRef Value = EmitPrototype(*this, Vector);
+
+	// If all components source inputs match, refer to the same input in the vector value too (for debugging purposes).
 	if (InX.Input == InY.Input && InX.Input == InZ.Input && InX.Input == InW.Input)
 	{
 		Value.Input = InX.Input;
 	}
 
-	return EmitPrototype(*this, Vector);
+	return Value;
+}
+
+FValueRef FEmitter::Aggregate(const UMaterialAggregate* InAggregate)
+{
+	return Aggregate(InAggregate, {}, TConstArrayView<FValueRef>{});
+}
+
+FValueRef FEmitter::Aggregate(const UMaterialAggregate* InAggregate, FValueRef InPrototype, TConstArrayView<FValueRef> AttributeValues)
+{
+	// Check that the specified prototype (if any) aggregate matches the one provided.
+	InPrototype = CheckIsAggregate(InPrototype, InAggregate);
+
+	// If a prototype is provided and there are no attribute assignments, this is a no-op, return the prototype.
+	if (InPrototype.IsValid() && AttributeValues.IsEmpty())
+	{
+		return InPrototype;
+	}
+
+	// Create the new composite to store the aggregate attribute values.
+	FComposite* AggregateValue = NewCompositeValue(*this, GetMaterialAggregateType(InAggregate), InAggregate->Attributes.Num());
+	
+	// Assign all components of the new composite value.
+	TArrayView<FValue*> Components = AggregateValue->GetMutableComponents();
+	for (int32 i = 0; i < Components.Num(); ++i)
+	{
+		// Get the the ith aggregate aggregate MIR type.
+		const FType* AttributeType = GetMaterialAggregateAttributeType(InAggregate, i);
+
+		if (i < AttributeValues.Num())
+		{
+			// Set the this aggregate component to the specified value cast to the attribute type, if present...
+			Components[i] = Cast(AttributeValues[i], AttributeType);
+		}
+		else if (InPrototype)
+		{
+			// ... otherwise use the component value as in the prototype, if defined...
+			Components[i] = Subscript(InPrototype, i);
+		}
+		else
+		{
+			// ...otherwise construct the default value.
+			Components[i] = ConstantDefault(AttributeType);
+		}
+	}
+
+	return EmitNew(*this, AggregateValue);
+}
+
+FValueRef FEmitter::Aggregate(const UMaterialAggregate* InAggregate, FValueRef InPrototype, TConstArrayView<FAttributeAssignment> AttributeAssignments)
+{
+	// Check that the specified prototype (if any) aggregate matches the one provided.
+	InPrototype = CheckIsAggregate(InPrototype, InAggregate);
+
+	// If a prototype is provided and there are no attribute assignments, this is a no-op, return the prototype.
+	if (InPrototype.IsValid() && AttributeAssignments.IsEmpty())
+	{
+		return InPrototype;
+	}
+
+	// Create the new composite to store the aggregate attribute values.
+	FComposite* AggregateValue = NewCompositeValue(*this, GetMaterialAggregateType(InAggregate), InAggregate->Attributes.Num());
+	
+	// Assign all components of the new composite value.
+	TArrayView<FValue*> Components = AggregateValue->GetMutableComponents();
+	for (int32 i = 0; i < Components.Num(); ++i)
+	{
+		// Search the assignments for this component's name.
+		const FAttributeAssignment* Assignment = AttributeAssignments.FindByPredicate(
+			[AttributeName = InAggregate->Attributes[i].Name](const FAttributeAssignment& Item)
+			{
+				return Item.Name == AttributeName;
+			});
+
+		// Get the the ith aggregate aggregate MIR type.
+		const FType* AttributeType = GetMaterialAggregateAttributeType(InAggregate, i);
+
+		if (Assignment && Assignment->Value)
+		{
+			// Set the this aggregate component to the value specified in the assignments, if present...
+			Components[i] = Cast(Assignment->Value, AttributeType);
+		}
+		else if (!Assignment && InPrototype)
+		{
+			// ... otherwise use the component value as in the prototype, if defined...
+			Components[i] = Subscript(InPrototype, i);
+		}
+		else
+		{
+			// ...otherwise construct the default value.
+			Components[i] = ConstantDefault(AttributeType);
+		}
+	}
+
+	return EmitNew(*this, AggregateValue);
 }
 
 /*--------------------------------- Operator emission ---------------------------------*/
@@ -1449,7 +1652,7 @@ static const FPrimitiveType* ValidateOperatorAndGetResultType(FEmitter& Emitter,
 	// We want to do this before fetching FirstArgumentPrimitiveType.
 	if (Op == TO_Clamp && A->Type->IsLWC() && !B->Type->IsLWC() && !C->Type->IsLWC())
 	{
-		A = Emitter.CastToFloat(A);
+		A = Emitter.CastToFloatKind(A);
 	}
 
 	// Verify that the first argument type is primitive.
@@ -1530,7 +1733,7 @@ static const FPrimitiveType* ValidateOperatorAndGetResultType(FEmitter& Emitter,
 
 		if (Filter & OPTF_CheckIsVector3)
 		{
-			if (!ArgumentPrimitiveType->IsVector() || ArgumentPrimitiveType->GetNumComponents() != 3)
+			if (!ArgumentPrimitiveType->IsVector() || ArgumentPrimitiveType->NumComponents() != 3)
 			{
 				Emitter.Errorf(Arguments[i], TEXT("Expected a 3D vector."));
 				bValid = false;
@@ -1595,7 +1798,7 @@ static const FPrimitiveType* ValidateOperatorAndGetResultType(FEmitter& Emitter,
 		const FPrimitiveType* CPrimitiveType = static_cast<const FPrimitiveType*>(Arguments[2]->Type);
 
 		// Compute the maximum number of vector components between all arguments. We know they're scalar or vectors, as it was checked before.
-		int32 MaxNumComponents = FMath::Max3(FirstArgumentPrimitiveType->GetNumComponents(), BPrimitiveType->GetNumComponents(), CPrimitiveType->GetNumComponents());
+		int32 MaxNumComponents = FMath::Max3(FirstArgumentPrimitiveType->NumComponents(), BPrimitiveType->NumComponents(), CPrimitiveType->NumComponents());
 
 		// Cast the first argument (the boolean condition) to a bool vector of the maximum number of components.
 		Arguments[0] = Emitter.Cast(Arguments[0], FPrimitiveType::Get(ScalarKind_Bool, MaxNumComponents, 1));
@@ -1639,7 +1842,7 @@ static const FPrimitiveType* ValidateOperatorAndGetResultType(FEmitter& Emitter,
 			// Fmod doesn't accept LWC for its second input, so we need to potentially cast that to non-LWC.
 			else if (Arguments[i]->Type->IsLWC() && !(Filter & OPTF_LWCInput))
 			{
-				Arguments[i] = Emitter.CastToFloat(Arguments[i]);
+				Arguments[i] = Emitter.CastToFloatKind(Arguments[i]);
 			}
 		}
 	}
@@ -1691,13 +1894,13 @@ static bool IsComponentwiseOperator(EOperator Op)
 }
 
 /// Tries to fold the operator by applying the operator componentwise on arguments components.
-/// If a value is returned, it will be a dimensional with some component folded to a constant. If some argument
-/// isn't a dimensional, or all arguments components are non-constant, the folding will not be carried out.
+/// If a value is returned, it will be a composite with some component folded to a constant. If some argument
+/// isn't a composite, or all arguments components are non-constant, the folding will not be carried out.
 /// If no folding is carried out, this function simply returns nullptr.
 static FValue* TryFoldComponentwiseOperator(FEmitter& Emitter, EOperator Op, FValue* A, FValue* B, FValue* C, const FPrimitiveType* ResultType)
 {
-	// Check that at least one component of the resulting dimensional value would folded.
-	// If all components of resulting dimensional value are not folded, then instead of emitting
+	// Check that at least one component of the resulting composite value would folded.
+	// If all components of resulting composite value are not folded, then instead of emitting
 	// an individual operator instruction for each component, simply emit a single binary operator
 	// instruction applied between lhs and rhs as a whole. (v1 + v2 rather than float2(v1.x + v2.x, v1.y + v2.y)
 	bool bSomeResultComponentWasFolded = false;
@@ -1706,10 +1909,9 @@ static FValue* TryFoldComponentwiseOperator(FEmitter& Emitter, EOperator Op, FVa
 	bool bResultIsIdenticalToC = true;
 
 	// Allocate the temporary array to store the folded component results
-	FMemMark Mark(FMemStack::Get());
-	TArrayView<FValue*> TempResultComponents = MakeTemporaryArray<FValue*>(Mark, ResultType->GetNumComponents());
+	TTemporaryArray<FValue*> TempResultComponents{ ResultType->NumComponents() };
 	
-	for (int32 i = 0; i < ResultType->GetNumComponents(); ++i)
+	for (int32 i = 0; i < ResultType->NumComponents(); ++i)
 	{
 		// Extract the arguments individual components
 		FValue* AComponent = Emitter.Subscript(A, i);
@@ -1751,24 +1953,24 @@ static FValue* TryFoldComponentwiseOperator(FEmitter& Emitter, EOperator Op, FVa
 	// so that we retain as much compile-time information as possible.
 	if (bSomeResultComponentWasFolded)
 	{
-		// If result type is scalar, simply return the single folded result (instead of creating a dimensional value)
+		// If result type is scalar, simply return the single folded result (instead of creating a composite value)
 		if (ResultType->IsScalar())
 		{
 			check(TempResultComponents[0]);
 			return TempResultComponents[0];
 		}
 
-		// Make the new dimensional value
-		FDimensional* Result = NewDimensionalValue(Emitter, ResultType);
+		// Make the new composite value
+		FComposite* Result = NewCompositeValue(Emitter, ResultType, ResultType->NumComponents());
 
-		// Fetch the components array from the result dimensional
+		// Fetch the components array from the result composite
 		TArrayView<FValue*> ResultComponents = Result->GetMutableComponents();
 
 		// Also cache the type of a single component
 		const FPrimitiveType* ComponentType = ResultType->ToScalar();
 
 		// Create the operator instruction for each component pair
-		for (int32 i = 0; i < ResultType->GetNumComponents(); ++i)
+		for (int32 i = 0; i < ResultType->NumComponents(); ++i)
 		{
 			// Reuse cached result if possible
 			ResultComponents[i] = TempResultComponents[i];
@@ -1791,18 +1993,18 @@ static FValue* TryFoldComponentwiseOperator(FEmitter& Emitter, EOperator Op, FVa
 	return {};
 }
 
-/// If V is a dimensional and all its components are constants, it unpacks the components into OutComponents and returns true.
+/// If V is a composite and all its components are constants, it unpacks the components into OutComponents and returns true.
 /// If this is not possible for any reason, it returns false.
 static bool TryUnpackConstantScalarOrVector(FValue* V, TArrayView<FConstant*> OutComponents, int32& OutNumComponents)
 {
 	// V not specified? Or not a scalar/vector?
-	FDimensional* Dimensional = As<FDimensional>(V);
-	if (!Dimensional || V->Type->AsPrimitive()->IsMatrix())
+	FComposite* Composite = As<FComposite>(V);
+	if (!Composite || V->Type->AsPrimitive()->IsMatrix())
 	{
 		return false;
 	}
 
-	TConstArrayView<FValue*> Components = Dimensional->GetComponents();
+	TConstArrayView<FValue*> Components = Composite->GetComponents();
 	for (int32 i = 0; i < Components.Num(); ++i)
 	{
 		OutComponents[i] = As<FConstant>(Components[i]);
@@ -1956,7 +2158,7 @@ FValueRef FEmitter::Operator(EOperator Op, FValueRef A, FValueRef B, FValueRef C
 	// Optionally cast the output to float
 	if (bCastLWCOutputToFloat)
 	{
-		Result = CastToFloat(Result);
+		Result = CastToFloatKind(Result);
 	}
 	return Result;
 }
@@ -2011,11 +2213,19 @@ FValueRef FEmitter::Subscript(FValueRef Value, int32 Index)
 	{
 		return Value;
 	}
-
+	
+	// Subscripting a composite by index is always possible and simply yields the i-th component.
+	if (FComposite* Composite = As<FComposite>(Value))
+	{
+		check(Index < Composite->GetComponents().Num());
+		return Value.To(Composite->GetComponents()[Index]);
+	}
+	
+	// Other operations supported if value is primitive. Check it first.
 	const FPrimitiveType* PrimitiveType = Value->Type->AsPrimitive();
 	if (!PrimitiveType)
 	{
-		Errorf(Value, TEXT("Value of type '%s' cannot be subscripted."), Value->Type->GetSpelling().GetData());
+		Errorf(Value, TEXT("Value of type '%s' cannot be subscripted."), *Value->Type->GetSpelling());
 		return Value.ToPoison();
 	}
 
@@ -2025,19 +2235,12 @@ FValueRef FEmitter::Subscript(FValueRef Value, int32 Index)
 		return Value;
 	}
 
-	if (Index >= PrimitiveType->GetNumComponents())
+	if (Index >= PrimitiveType->NumComponents())
 	{
-		Errorf(Value, TEXT("Value of type '%s' has fewer dimensions than subscript index `%d`."), Value->Type->GetSpelling().GetData(), Index);
+		Errorf(Value, TEXT("Value of type '%s' has fewer dimensions than subscript index `%d`."), *Value->Type->GetSpelling(), Index);
 		return Value.ToPoison();
 	}
 
-	// Getting first component and Value is already a scalar, just return itself.
-	if (FDimensional* DimensionalValue = As<FDimensional>(Value))
-	{
-		check(Index < DimensionalValue->GetComponents().Num());
-		return Value.To(DimensionalValue->GetComponents()[Index]);
-	}
-	
 	// Avoid subscripting a subscript (e.g. no value.xy.x)
 	if (FSubscript* Subscript = As<FSubscript>(Value))
 	{
@@ -2070,7 +2273,7 @@ FValueRef FEmitter::Swizzle(FValueRef Value, FSwizzleMask Mask)
 	const FPrimitiveType* PrimitiveType = Value->Type->AsVector();
 	if (!PrimitiveType || PrimitiveType->IsMatrix())
 	{
-		Errorf(Value, TEXT("Cannot swizzle a '%s' value."), Value->Type->GetSpelling().GetData());
+		Errorf(Value, TEXT("Cannot swizzle a '%s' value."), *Value->Type->GetSpelling());
 		return Value.ToPoison();
 	}
 
@@ -2086,7 +2289,7 @@ FValueRef FEmitter::Swizzle(FValueRef Value, FSwizzleMask Mask)
 
 	// If the requested number of components is the same as Value and the order in which the components
 	// are specified in the mask is sequential (e.g. x, y, z) then this is a no op, simply return Value as is.
-	if (Mask.NumComponents == PrimitiveType->GetNumComponents())
+	if (Mask.NumComponents == PrimitiveType->NumComponents())
 	{
 		bool InOrder = true;
 		for (int32 i = 0; i < Mask.NumComponents; ++i)
@@ -2112,7 +2315,7 @@ FValueRef FEmitter::Swizzle(FValueRef Value, FSwizzleMask Mask)
 
 	// Make the result vector type.
 	const FPrimitiveType* ResultType = FPrimitiveType::GetVector(PrimitiveType->ScalarKind, Mask.NumComponents);
-	FDimensional* Result = NewDimensionalValue(*this, ResultType);
+	FComposite* Result = NewCompositeValue(*this, ResultType, ResultType->NumComponents());
 	TArrayView<FValue*> ResultComponents = Result->GetMutableComponents();
 
 	for (int32 i = 0; i < Mask.NumComponents; ++i)
@@ -2165,7 +2368,7 @@ static FValue* CastValueToPrimitiveType(FEmitter& Emitter, FValueRef Value, cons
 	const FPrimitiveType* ValuePrimitiveType = Value->Type->AsPrimitive();
 	if (!ValuePrimitiveType)
 	{
-		Emitter.Errorf(Value, TEXT("Cannot construct a '%s' from non primitive type '%s'."), TargetPrimitiveType->Spelling.GetData(), Value->Type->GetSpelling().GetData());
+		Emitter.Errorf(Value, TEXT("Cannot construct a '%s' from non primitive type '%s'."), TargetPrimitiveType->Spelling.GetData(), *Value->Type->GetSpelling());
 		return FPoison::Get();
 	}
 
@@ -2201,10 +2404,10 @@ static FValue* CastValueToPrimitiveType(FEmitter& Emitter, FValueRef Value, cons
 	// Note: we know target isn't scalar as it's been handled above.
 	if (ValuePrimitiveType->IsScalar())
 	{
-		// Create the result dimensional value.
-		FDimensional* Result = NewDimensionalValue(Emitter, TargetPrimitiveType);
+		// Create the result composite value.
+		FComposite* Result = NewCompositeValue(Emitter, TargetPrimitiveType, TargetPrimitiveType->NumComponents());
 
-		// Create a dimensional and initialize each of its components to the conversion
+		// Create a composite and initialize each of its components to the conversion
 		// of initializer value to the single component type.
 		FValue* Component = Emitter.Cast(Value, TargetPrimitiveType->ToScalar());
 
@@ -2212,7 +2415,7 @@ static FValue* CastValueToPrimitiveType(FEmitter& Emitter, FValueRef Value, cons
 		TArrayView<FValue*> ResultComponents = Result->GetMutableComponents();
 
 		// Initialize all result components to the same scalar.
-		for (int32 i = 0; i < TargetPrimitiveType->GetNumComponents(); ++i)
+		for (int32 i = 0; i < TargetPrimitiveType->NumComponents(); ++i)
 		{
 			ResultComponents[i] = Component;
 		}
@@ -2227,17 +2430,17 @@ static FValue* CastValueToPrimitiveType(FEmitter& Emitter, FValueRef Value, cons
 	{
 		// #todo-massimo.tristano Use swizzle when scalartypes are the same, and target num components is less than initializer's.
 
-		int32 TargetNumComponents = TargetPrimitiveType->GetNumComponents();
-		int32 InitializerNumComponents = ValuePrimitiveType->GetNumComponents();
+		int32 TargetNumComponents = TargetPrimitiveType->NumComponents();
+		int32 InitializerNumComponents = ValuePrimitiveType->NumComponents();
 
-		// Create the result dimensional value.
-		FDimensional* Result = NewDimensionalValue(Emitter, TargetPrimitiveType);
+		// Create the result composite value.
+		FComposite* Result = NewCompositeValue(Emitter, TargetPrimitiveType, TargetPrimitiveType->NumComponents());
 		TArrayView<FValue*> ResultComponents = Result->GetMutableComponents();
 
 		// Determine the result component type (scalar).
 		const FPrimitiveType* ResultComponentType = TargetPrimitiveType->ToScalar();
 
-		// For iterating over the components of the result dimensional value.
+		// For iterating over the components of the result composite value.
 		int32 Index = 0;
 		
 		// Convert components from the initializer vector.
@@ -2247,7 +2450,7 @@ static FValue* CastValueToPrimitiveType(FEmitter& Emitter, FValueRef Value, cons
 			ResultComponents[Index] = Emitter.Cast(Emitter.Subscript(Value, Index), ResultComponentType);
 		}
 
-		// Initialize remaining result dimensional components to zero.
+		// Initialize remaining result composite components to zero.
 		for (; Index < TargetNumComponents; ++Index)
 		{
 			ResultComponents[Index] = Emitter.ConstantZero(ResultComponentType->ScalarKind);
@@ -2262,15 +2465,15 @@ static FValue* CastValueToPrimitiveType(FEmitter& Emitter, FValueRef Value, cons
 	{
 		check(TargetPrimitiveType->IsMatrix());
 
-		// Create the result dimensional value.
-		FDimensional* Result = NewDimensionalValue(Emitter, TargetPrimitiveType);
+		// Create the result composite value.
+		FComposite* Result = NewCompositeValue(Emitter, TargetPrimitiveType, TargetPrimitiveType->NumComponents());
 		TArrayView<FValue*> ResultComponents = Result->GetMutableComponents();
 			
 		// Determine the result component type (scalar).
 		const FPrimitiveType* ResultComponentType = TargetPrimitiveType->ToScalar();
 
 		// Convert components from the initializer vector.
-		for (int32 Index = 0, Num = TargetPrimitiveType->GetNumComponents(); Index < Num; ++Index)
+		for (int32 Index = 0, Num = TargetPrimitiveType->NumComponents(); Index < Num; ++Index)
 		{
 			ResultComponents[Index] = Emitter.Cast(Emitter.Subscript(Value, Index), ResultComponentType);
 		}
@@ -2301,9 +2504,10 @@ FValueRef FEmitter::Cast(FValueRef Value, const FType* TargetType)
 	{
 		Result = CastValueToPrimitiveType(*this, Value, PrimitiveType);
 	}
-	else // No other legal conversions applicable. Report error if we haven't converted the value.
+	else
 	{
-		Errorf(Value, TEXT("Cannot construct a '%s' from a '%s'."), TargetType->GetSpelling().GetData(), Value->Type->GetSpelling().GetData());
+		// No other legal conversions applicable. Report error if we haven't converted the value.
+		Errorf(Value, TEXT("Cannot construct a '%s' from a '%s'."), *TargetType->GetSpelling(), *Value->Type->GetSpelling());
 		return Poison();
 	}
 	
@@ -2326,17 +2530,17 @@ FValueRef FEmitter::CastToScalarKind(FValueRef Value, EScalarKind ToScalarKind)
 		: Value;
 }
 
-FValueRef FEmitter::CastToBool(FValueRef Value)
+FValueRef FEmitter::CastToBoolKind(FValueRef Value)
 {
 	return CastToScalarKind(Value, ScalarKind_Bool);
 }
 
-FValueRef FEmitter::CastToInt(FValueRef Value)
+FValueRef FEmitter::CastToIntKind(FValueRef Value)
 {
 	return CastToScalarKind(Value, ScalarKind_Bool);
 }
 
-FValueRef FEmitter::CastToFloat(FValueRef Value)
+FValueRef FEmitter::CastToFloatKind(FValueRef Value)
 {
 	return CastToScalarKind(Value, ScalarKind_Float);
 }
@@ -2570,7 +2774,7 @@ FValueRef FEmitter::PartialDerivative(FValueRef Value, EDerivativeAxis Axis)
 	const MIR::FPrimitiveType* ValuePrimitiveType = Value->Type->AsPrimitive();
 	if (!ValuePrimitiveType || !ScalarKindIsAnyFloat(ValuePrimitiveType->ScalarKind))
 	{
-		Errorf(Value, TEXT("Trying to differentiate a value of type `%s` is invalid. Expected a float type."), Value->Type->GetSpelling().GetData());
+		Errorf(Value, TEXT("Trying to differentiate a value of type `%s` is invalid. Expected a float type."), *Value->Type->GetSpelling());
 		return Poison();
 	}
 
@@ -2771,8 +2975,8 @@ static FValue* DifferentiateOperator(FEmitter& E, FOperator* Op, EDerivativeAxis
 		// d/dx pow(f(x), g(x)) = f(x)^g(x) * (g'(x) * ln(f(x)) + g(x) * f'(x) / f(x))
 		case BO_Pow:
 		{
-			FValueRef Term1 = E.Multiply(dG, E.Logarithm(F)); // g'(x) * ln(f(x))
-			FValueRef Term2 = E.Divide(E.Multiply(G, dF), F); // g(x) * f'(x) / f(x)
+			FValue* Term1 = E.Multiply(dG, E.Logarithm(F)); // g'(x) * ln(f(x))
+			FValue* Term2 = E.Divide(E.Multiply(G, dF), F); // g(x) * f'(x) / f(x)
 			return E.Multiply(Op, E.Add(Term1, Term2));
 		}
 
@@ -2854,7 +3058,7 @@ FValueRef FEmitter::AnalyticalPartialDerivative(FValueRef Value, EDerivativeAxis
 	const MIR::FPrimitiveType* ValuePrimitiveType = Value->Type->AsPrimitive();
 	if (!ValuePrimitiveType || !ScalarKindIsAnyFloat(ValuePrimitiveType->ScalarKind))
 	{
-		Errorf(Value, TEXT("Trying to differentiate a value of type `%s` is invalid. Expected a float type."), Value->Type->GetSpelling().GetData());
+		Errorf(Value, TEXT("Trying to differentiate a value of type `%s` is invalid. Expected a float type."), *Value->Type->GetSpelling());
 		return Poison();
 	}
 
@@ -2863,10 +3067,10 @@ FValueRef FEmitter::AnalyticalPartialDerivative(FValueRef Value, EDerivativeAxis
 		case MIR::VK_ExternalInput:
 			return DifferentiateExternalInput(*this, Value->As<FExternalInput>(), Axis);
 
-		case MIR::VK_Dimensional:
+		case MIR::VK_Composite:
 		{
-			TConstArrayView<FValue*> ValueComponents = Value->As<FDimensional>()->GetComponents();
-			FDimensional* Derivative = NewDimensionalValue(*this, ValuePrimitiveType);
+			TConstArrayView<FValue*> ValueComponents = Value->As<FComposite>()->GetComponents();
+			FComposite* Derivative = NewCompositeValue(*this, ValuePrimitiveType, ValuePrimitiveType->NumComponents());
 			TArrayView<FValue*> DerivativeComponents = Derivative->GetMutableComponents();
 			for (int32 i = 0; i < ValueComponents.Num(); ++i)
 			{
diff --git a/Engine/Source/Runtime/Engine/Private/Materials/MaterialIRInternal.h b/Engine/Source/Runtime/Engine/Private/Materials/MaterialIRInternal.h
--- a/Engine/Source/Runtime/Engine/Private/Materials/MaterialIRInternal.h
+++ b/Engine/Source/Runtime/Engine/Private/Materials/MaterialIRInternal.h
@@ -6,13 +6,6 @@
 
 #if WITH_EDITOR
 
-template <typename T>
-TArrayView<T> MakeTemporaryArray(FMemMark&, int Count)
-{
-	auto Ptr = (T*)FMemStack::Get().Alloc(sizeof(T) * Count, alignof(T));
-	return { Ptr, Count };
-}
-
 namespace MIR::Internal {
 
 //
diff --git a/Engine/Source/Runtime/Engine/Private/Materials/MaterialIRModule.cpp b/Engine/Source/Runtime/Engine/Private/Materials/MaterialIRModule.cpp
--- a/Engine/Source/Runtime/Engine/Private/Materials/MaterialIRModule.cpp
+++ b/Engine/Source/Runtime/Engine/Private/Materials/MaterialIRModule.cpp
@@ -36,9 +36,14 @@ void FMaterialIRModule::Empty()
 	{
 		FMemory::Free(Value);
 	}
-
 	Values.Empty();
 
+	for (MIR::FType* Type : Types)
+	{
+		FMemory::Free(Type);
+	}
+	Types.Empty();
+
 	// Reset module statistics.
 	for (int i = 0; i < MIR::NumStages; ++i)
 	{
@@ -47,8 +52,6 @@ void FMaterialIRModule::Empty()
 
 	Statistics.NumVertexTexCoords = 0;
 	Statistics.NumPixelTexCoords = 0;
-
-	Allocator.Flush();
 }
 
 const TCHAR* FMaterialIRModule::PushUserString(FString InString)
diff --git a/Engine/Source/Runtime/Engine/Private/Materials/MaterialIRModuleBuilder.cpp b/Engine/Source/Runtime/Engine/Private/Materials/MaterialIRModuleBuilder.cpp
--- a/Engine/Source/Runtime/Engine/Private/Materials/MaterialIRModuleBuilder.cpp
+++ b/Engine/Source/Runtime/Engine/Private/Materials/MaterialIRModuleBuilder.cpp
@@ -266,8 +266,7 @@ struct FMaterialIRModuleBuilderImpl
 
 	void PushFunctionCall(UMaterialExpressionMaterialFunctionCall* Call)
 	{
-		FMemMark Mark(FMemStack::Get());
-		TArrayView<MIR::FValue*> CallInputValues = MakeTemporaryArray<MIR::FValue*>(Mark, Call->FunctionInputs.Num());
+		MIR::TTemporaryArray<MIR::FValue*> CallInputValues{ Call->FunctionInputs.Num() };
 
 		// Make sure each function input is connected and has a value. If so, cache the values flowing into this
 		// funcion call inside the auxiliary value array.
@@ -370,9 +369,28 @@ struct FMaterialIRModuleBuilderImpl
 				}
 				else if (MIR::FValue* Value = MIR::Internal::FetchValueFromExpressionInput(this, Input.Input))
 				{
-					// If a value is flowing in through the connection, cast it to this material attribute type and assign it.
-					const MIR::FType* OutputArgType = MIR::FType::FromShaderType(Input.Type);
-					Output->Arg = Emitter->Cast(Value, OutputArgType);
+					// If this property is the emissive color and we're previewing the material, apply gamma correction to the previewed value.
+					if (Output->Property == MP_EmissiveColor && PreviewInput.IsConnected())
+					{
+						MIR::FValue* Zero = Emitter->ConstantZero(MIR::EScalarKind::ScalarKind_Float);
+
+						// Get preview expression back into gamma corrected space, as DrawTile does not do this adjustment.
+						Value = Emitter->Pow(Emitter->Max(Value, Zero), Emitter->ConstantFloat(1.0f / 2.2f));
+
+						// Preview should display scalars as red, so if this is a scalar, create a vector padded with zeroes
+						if (Value->IsScalar())
+						{
+							Value = Emitter->Vector3(Value, Zero, Zero);
+						}
+
+						Output->Arg = Value;
+					}
+					else
+					{
+						// If a value is flowing in through the connection, cast it to this material attribute type and assign it.
+						const MIR::FType* OutputArgType = MIR::FType::FromShaderType(Input.Type);
+						Output->Arg = Emitter->Cast(Value, OutputArgType);
+					}
 				}
 				else
 				{
@@ -383,21 +401,6 @@ struct FMaterialIRModuleBuilderImpl
 				// The value being set to this material output is now valid.
 				check(Output->Arg);
 
-				// If this property is the emissive color and we're previewing the material, apply gamma correction to the previewed value.
-				if (Output->Property == MP_EmissiveColor && PreviewInput.IsConnected())
-				{
-					MIR::FValue* Zero = Emitter->ConstantZero(MIR::EScalarKind::ScalarKind_Float);
-
-					// Get preview expression back into gamma corrected space, as DrawTile does not do this adjustment.
-					Output->Arg = Emitter->Pow(Emitter->Max(Output->Arg, Zero), Emitter->ConstantFloat(1.0f / 2.2f));
-
-					// Preview should display scalars as red, so if this is a scalar, create a vector padded with zeroes
-					if (Output->Arg->IsScalar())
-					{
-						Output->Arg = Emitter->Vector3(Output->Arg, Zero, Zero);
-					}
-				}
-
 				// Connect output properties with substrate parameter if a FrontMaterial property is provided
 				if (Output->Property == MP_FrontMaterial)
 				{
@@ -418,6 +421,12 @@ struct FMaterialIRModuleBuilderImpl
 					}
 				}
 
+				//
+				if (!Module->IsValid())
+				{
+					return;
+				}
+
 				// Finally, push this connection insight
 				PushConnectionInsight(Builder->Material, (int)Output->Property, Input.Input->Expression, Input.Input->OutputIndex, Output->Arg->Type);
 			}
@@ -708,6 +717,12 @@ bool FMaterialIRModuleBuilder::Build(FMaterialIRModule* TargetModule)
 	}
 
 	Impl.Step_FlowValuesIntoMaterialOutputs();
+
+	if (!TargetModule->IsValid())
+	{
+		return false;
+	}
+
 	Impl.Step_AnalyzeIRGraph();
 	Impl.Step_ConsolidateEnvironmentDefines();
 	Impl.Step_AnalyzeBuiltinDefines();
diff --git a/Engine/Source/Runtime/Engine/Private/Materials/MaterialIRToHLSLTranslator.cpp b/Engine/Source/Runtime/Engine/Private/Materials/MaterialIRToHLSLTranslator.cpp
--- a/Engine/Source/Runtime/Engine/Private/Materials/MaterialIRToHLSLTranslator.cpp
+++ b/Engine/Source/Runtime/Engine/Private/Materials/MaterialIRToHLSLTranslator.cpp
@@ -406,7 +406,7 @@ struct FTranslator : FMaterialIRToHLSLTranslation
 	{
 		switch (Instr->Kind)
 		{
-			case MIR::VK_Dimensional: LowerDimensional(static_cast<const MIR::FDimensional*>(Instr)); break;
+			case MIR::VK_Composite: LowerComposite(static_cast<const MIR::FComposite*>(Instr)); break;
 			case MIR::VK_SetMaterialOutput: LowerSetMaterialOutput(static_cast<const MIR::FSetMaterialOutput*>(Instr)); break;
 			case MIR::VK_Operator: LowerOperator(static_cast<const MIR::FOperator*>(Instr)); break;
 			case MIR::VK_Branch: LowerBranch(static_cast<const MIR::FBranch*>(Instr)); break;
@@ -773,36 +773,36 @@ struct FTranslator : FMaterialIRToHLSLTranslation
 
 				default:
 				{
-					if (Argument.MaterialProperty != MP_MAX && SubstrateParameter->PropertyArgs[Argument.MaterialProperty])
-					{
-						LowerValue(SubstrateParameter->PropertyArgs[Argument.MaterialProperty]);
-					}
-					else if (Argument.DefaultInlineHLSL)
-					{
-						Printer << Argument.DefaultInlineHLSL;
-					}
-					else if (const MIR::FPrimitiveType* PrimitiveType = Argument.ParameterType->AsPrimitive())
-					{
-						if (const MIR::FPrimitiveType* ScalarType = Argument.ParameterType->AsScalar())
-						{
-							Printer << TEXT("0.0f");
-						}
-						else
-						{
-							LowerType(Argument.ParameterType);
-							Printer << BeginArgs;
-							for (int32 ComponentIndex = 0; ComponentIndex < PrimitiveType->NumRows; ++ComponentIndex)
-							{
-								Printer << ListSeparator << TEXT("0.0f");
-							}
-							Printer << EndArgs;
-						}
-					}
-					else
+			if (Argument.MaterialProperty != MP_MAX && SubstrateParameter->PropertyArgs[Argument.MaterialProperty])
+			{
+				LowerValue(SubstrateParameter->PropertyArgs[Argument.MaterialProperty]);
+			}
+			else if (Argument.DefaultInlineHLSL)
+			{
+				Printer << Argument.DefaultInlineHLSL;
+			}
+			else if (const MIR::FPrimitiveType* PrimitiveType = Argument.ParameterType->AsPrimitive())
+			{
+				if (const MIR::FPrimitiveType* ScalarType = Argument.ParameterType->AsScalar())
+				{
+					Printer << TEXT("0.0f");
+				}
+				else
+				{
+					LowerType(Argument.ParameterType);
+					Printer << BeginArgs;
+					for (int32 ComponentIndex = 0; ComponentIndex < PrimitiveType->NumRows; ++ComponentIndex)
 					{
-						UE_MIR_UNREACHABLE();
+						Printer << ListSeparator << TEXT("0.0f");
 					}
+					Printer << EndArgs;
 				}
+			}
+			else
+			{
+				UE_MIR_UNREACHABLE();
+			}
+		}
 				break;
 			}
 		}
@@ -838,12 +838,12 @@ struct FTranslator : FMaterialIRToHLSLTranslation
 
 		Printer << TEXT("Material.PreshaderBuffer[") << BufferSlotIndex << TEXT("]");
 
-		if (PrimitiveType->GetNumComponents() < 4)
+		if (PrimitiveType->NumComponents() < 4)
 		{
 			Printer << TEXT(".");
 
 			static const TCHAR* Components = TEXT("xyzw");
-			for (int32 i = 0; i < PrimitiveType->GetNumComponents(); ++i)
+			for (int32 i = 0; i < PrimitiveType->NumComponents(); ++i)
 			{
 				check(BufferSlotOffset + i < sizeof(Components));
 				Printer.Buffer.AppendChar(Components[BufferSlotOffset + i]);
@@ -856,7 +856,7 @@ struct FTranslator : FMaterialIRToHLSLTranslation
 		}
 	}
 
-	void LowerDimensional(const MIR::FDimensional* Dimensional) 
+	void LowerComposite(const MIR::FComposite* Dimensional) 
 	{
 		TArrayView<MIR::FValue* const> Components = Dimensional->GetComponents();
 		const MIR::FPrimitiveType* ArithmeticType = Dimensional->Type->AsPrimitive();
@@ -1143,7 +1143,7 @@ struct FTranslator : FMaterialIRToHLSLTranslation
 		if (ArgArithmeticType && ArgArithmeticType->ScalarKind == MIR::EScalarKind::ScalarKind_LWC)
 		{
 			const TCHAR* LwcComponentsStr[] = { TEXT("WSGetX("), TEXT("WSGetY("), TEXT("WSGetZ("), TEXT("WSGetW(") };
-			check(Subscript->Index <= ArgArithmeticType->GetNumComponents());
+			check(Subscript->Index <= ArgArithmeticType->NumComponents());
 
 			Printer << LwcComponentsStr[Subscript->Index] << LowerValue(Subscript->Arg) << TEXT(")");
 		}
@@ -1154,7 +1154,7 @@ struct FTranslator : FMaterialIRToHLSLTranslation
 			if (ArgArithmeticType)
 			{
 				const TCHAR* ComponentsStr[] = { TEXT(".x"), TEXT(".y"), TEXT(".z"), TEXT(".w") };
-				check(Subscript->Index <= ArgArithmeticType->GetNumComponents());
+				check(Subscript->Index <= ArgArithmeticType->NumComponents());
 
 				Printer << ComponentsStr[Subscript->Index];
 			}
diff --git a/Engine/Source/Runtime/Engine/Private/Materials/MaterialIRTypes.cpp b/Engine/Source/Runtime/Engine/Private/Materials/MaterialIRTypes.cpp
--- a/Engine/Source/Runtime/Engine/Private/Materials/MaterialIRTypes.cpp
+++ b/Engine/Source/Runtime/Engine/Private/Materials/MaterialIRTypes.cpp
@@ -1,6 +1,7 @@
 // Copyright Epic Games, Inc. All Rights Reserved.
 
 #include "Materials/MaterialIRTypes.h"
+#include "Materials/MaterialAggregate.h"
 
 #if WITH_EDITOR
 
@@ -128,7 +129,7 @@ const FType* FType::GetVoid()
 	return &Type;
 }
 
-FStringView FType::GetSpelling() const
+FString FType::GetSpelling() const
 {
 	if (IsPoison())
 	{
@@ -136,7 +137,7 @@ FStringView FType::GetSpelling() const
 	}
 	else if (auto PrimitiveType = AsPrimitive())
 	{
-		return PrimitiveType->Spelling;
+		return PrimitiveType->Spelling.GetData();
 	}
 	else if (auto ObjectType = AsObject())
 	{
@@ -147,6 +148,10 @@ FStringView FType::GetSpelling() const
 			case ObjectKind_VTPageTableResult: return TEXT("VTPageTableResult");
 		}
 	}
+	else if (auto AggregateType = AsAggregate())
+	{
+		return AggregateType->Aggregate->GetName();
+	}
 	UE_MIR_UNREACHABLE();
 }
 
@@ -192,6 +197,10 @@ UE::Shader::EValueType FType::ToValueType() const
 	{
 		return EValueType::Object;
 	}
+	else if (const FAggregateType* AggregateType = AsAggregate())
+	{
+		return EValueType::Struct;
+	}
 	
 	UE_MIR_UNREACHABLE();
 }
@@ -219,6 +228,12 @@ bool FType::IsInteger() const
 	return Ptr && Ptr->ScalarKind == ScalarKind_Int;
 }
 
+bool FType::IsAnyFloat() const
+{
+	const FPrimitiveType* Ptr = AsPrimitive();
+	return Ptr && ScalarKindIsAnyFloat(Ptr->ScalarKind);
+}
+
 bool FType::IsLWC() const
 {
 	const FPrimitiveType* Ptr = AsPrimitive();
@@ -268,19 +283,24 @@ const FObjectType* FType::AsObject() const
 const FPrimitiveType* FType::AsScalar() const
 {
 	const FPrimitiveType* Type = AsPrimitive();
-	return Type->IsScalar() ? Type : nullptr;
+	return Type && Type->IsScalar() ? Type : nullptr;
 }
 
 const FPrimitiveType* FType::AsVector() const
 {
 	const FPrimitiveType* Type = AsPrimitive();
-	return Type->IsVector() ? Type : nullptr;
+	return Type && Type->IsVector() ? Type : nullptr;
 }
 
 const FPrimitiveType* FType::AsMatrix() const
 {
 	const FPrimitiveType* Type = AsPrimitive();
-	return Type->IsMatrix() ? Type : nullptr;
+	return Type && Type->IsMatrix() ? Type : nullptr;
+}
+
+const FAggregateType* FType::AsAggregate() const
+{
+	return Kind == TypeKind_Aggregate ? static_cast<const FAggregateType*>(this) : nullptr; 
 }
 
 bool ScalarKindIsArithmetic(EScalarKind Kind)
@@ -391,9 +411,9 @@ const FPrimitiveType* FPrimitiveType::Get(EScalarKind InScalarKind, int NumRows,
 		{ { TypeKind_Primitive }, { TEXT("int4x3") }, 	ScalarKind_Int, 4, 3 },
 		{ { TypeKind_Primitive }, { TEXT("int4x4") }, 	ScalarKind_Int, 4, 4 },
 		{ { TypeKind_Primitive }, { TEXT("float") }, 		ScalarKind_Float, 1, 1 },
-		{ { TypeKind_Primitive }, Invalid, 					ScalarKind_Float, 1, 2 },
-		{ { TypeKind_Primitive }, Invalid, 					ScalarKind_Float, 1, 3 },
-		{ { TypeKind_Primitive }, Invalid, 					ScalarKind_Float, 1, 4 },
+		{ { TypeKind_Primitive }, Invalid, 				ScalarKind_Float, 1, 2 },
+		{ { TypeKind_Primitive }, Invalid, 				ScalarKind_Float, 1, 3 },
+		{ { TypeKind_Primitive }, Invalid, 				ScalarKind_Float, 1, 4 },
 		{ { TypeKind_Primitive }, { TEXT("float2") },   	ScalarKind_Float, 2, 1 },
 		{ { TypeKind_Primitive }, { TEXT("float2x2") }, 	ScalarKind_Float, 2, 2 },
 		{ { TypeKind_Primitive }, { TEXT("float2x3") }, 	ScalarKind_Float, 2, 3 },
diff --git a/Engine/Source/Runtime/Engine/Private/Materials/MaterialIRValueAnalyzer.cpp b/Engine/Source/Runtime/Engine/Private/Materials/MaterialIRValueAnalyzer.cpp
--- a/Engine/Source/Runtime/Engine/Private/Materials/MaterialIRValueAnalyzer.cpp
+++ b/Engine/Source/Runtime/Engine/Private/Materials/MaterialIRValueAnalyzer.cpp
@@ -217,7 +217,7 @@ static void AnalyzePrimitiveUniformParameter(FMaterialIRValueAnalyzer& Analyzer,
 	// in the appropriate freelist.
 
 	uint32 GlobalComponentOffset = UINT32_MAX;
-	uint32 NumComponents = ParameterType->GetNumComponents();
+	uint32 NumComponents = ParameterType->NumComponents();
 	check(NumComponents > 0);
 	check(NumComponents <= 4); // Only vectors supported for now
 
diff --git a/Engine/Source/Runtime/Engine/Public/Materials/MaterialAggregate.h b/Engine/Source/Runtime/Engine/Public/Materials/MaterialAggregate.h
new file mode 100644
--- /dev/null
+++ b/Engine/Source/Runtime/Engine/Public/Materials/MaterialAggregate.h
@@ -0,0 +1,100 @@
+// Copyright Epic Games, Inc. All Rights Reserved.
+
+#pragma once
+
+#include "CoreMinimal.h"
+#include "MaterialValueType.h"
+#include "Engine/DataAsset.h"
+#include "SceneTypes.h"
+#include "MaterialAggregate.generated.h"
+
+class UMaterialAggregate;
+
+// Specifies the type of an individual material aggregate attribute.
+UENUM()
+enum class EMaterialAggregateAttributeType
+{
+	Bool1,
+	Bool2,
+	Bool3,
+	Bool4,
+	
+	UInt1,
+	UInt2,
+	UInt3,
+	UInt4,
+
+	Float1,
+	Float2,
+	Float3,
+	Float4,
+
+	ShadingModel,
+
+	Aggregate,
+};
+
+#if WITH_EDITOR
+
+// Converts a EMaterialValueType to the equivalent EMaterialAggregateAttributeTypeKind, if possible.
+EMaterialAggregateAttributeType MaterialValueTypeToMaterialAggregateAttributeType(EMaterialValueType Type);
+
+#endif
+
+// Represents a single attribute within a material aggregate.
+USTRUCT(MinimalAPI)
+struct FMaterialAggregateAttribute
+{
+	GENERATED_BODY()
+
+	// The name of the attribute.
+	UPROPERTY(EditAnywhere, Category=MaterialAggregateAttribute)
+	FName Name{};
+
+	// The type of the attribute.
+	UPROPERTY(EditAnywhere, Category=MaterialAggregateAttribute)
+	EMaterialAggregateAttributeType Type = EMaterialAggregateAttributeType::Float4;
+
+	// The nested aggregate reference, used only if TypeKind is Aggregate.
+	UPROPERTY(EditAnywhere, Category=MaterialAggregateAttribute, meta=(EditCondition="TypeKind == EMaterialAggregateAttributeTypeKind::Aggregate", EditConditionHides))
+	TObjectPtr<UMaterialAggregate> Aggregate{};
+
+#if WITH_EDITOR
+	// Converts this attribute's type to a corresponding material value type.
+	EMaterialValueType ToMaterialValueType() const;
+#endif
+};
+
+// It defines a collection of arithmetic material values to be bundled together.
+// A material aggregate works similarly to a struct in C/C++. Each attribute has a name and specifies a type, either a
+// primitive one like float4 or another aggregate (for nested structures).
+UCLASS(MinimalAPI)
+class UMaterialAggregate : public UDataAsset
+{
+	GENERATED_BODY()
+
+public:
+	// List of material aggregate attributes.
+	UPROPERTY(EditAnywhere, Category=MaterialAggregate)
+	TArray<FMaterialAggregateAttribute> Attributes;
+
+public:
+#if WITH_EDITOR
+
+	// Returns a reference to the global material attributes aggregate used by the engine.
+	static const UMaterialAggregate* GetMaterialAttributes();
+
+	// Maps a material property to the corresponding attribute index in the aggregate.
+	static int32 MaterialPropertyToAttributeIndex(EMaterialProperty Property);
+
+	// Maps an attribute index back to the corresponding material property.
+	static EMaterialProperty AttributeIndexToMaterialProperty(int32 MaterialAttributeIndex);
+
+	// Finds the index of an attribute in the aggregate by name.
+	int32 FindAttributeIndexByName(FName Name) const;
+
+	virtual void PostLoad() override;
+	virtual void PostEditChangeProperty(FPropertyChangedEvent& PropertyChangedEvent) override;
+
+#endif // #if WITH_EDITOR
+};
diff --git a/Engine/Source/Runtime/Engine/Public/Materials/MaterialExpression.h b/Engine/Source/Runtime/Engine/Public/Materials/MaterialExpression.h
--- a/Engine/Source/Runtime/Engine/Public/Materials/MaterialExpression.h
+++ b/Engine/Source/Runtime/Engine/Public/Materials/MaterialExpression.h
@@ -189,9 +189,6 @@ class UMaterialExpression : public UObject
 	UPROPERTY(EditAnywhere, Category=MaterialExpression, meta=(MultiLine=true, DisplayAfter = "SortPriority"))
 	FString Desc;
 
-	/** Number of expressions connected to this expression's execution input */
-	int32 NumExecutionInputs;
-
 	/** Set to true by RecursiveUpdateRealtimePreview() if the expression's preview needs to be updated in realtime in the material editor. */
 	UPROPERTY()
 	uint32 bRealtimePreview:1;
@@ -268,7 +265,7 @@ class UMaterialExpression : public UObject
 	 * @param	Compiler - UMaterial compiler that knows how to handle this expression.
 	 * @return	Index to the new FMaterialCompiler::CodeChunk entry for this expression
 	 */	
-	virtual int32 Compile(class FMaterialCompiler* Compiler, int32 OutputIndex) { return INDEX_NONE; }
+	virtual int32 Compile(class FMaterialCompiler* Compiler, int32 OutputIndex);
 	virtual int32 CompilePreview(class FMaterialCompiler* Compiler, int32 OutputIndex) { return Compile(Compiler, OutputIndex); }
 
 	ENGINE_API bool IsUsingNewHLSLGenerator() const;
diff --git a/Engine/Source/Runtime/Engine/Public/Materials/MaterialExpressionAggregate.h b/Engine/Source/Runtime/Engine/Public/Materials/MaterialExpressionAggregate.h
new file mode 100644
--- /dev/null
+++ b/Engine/Source/Runtime/Engine/Public/Materials/MaterialExpressionAggregate.h
@@ -0,0 +1,89 @@
+// Copyright Epic Games, Inc. All Rights Reserved.
+
+#pragma once
+
+#include "CoreMinimal.h"
+#include "UObject/ObjectMacros.h"
+#include "MaterialExpressionIO.h"
+#include "Materials/MaterialExpression.h"
+#include "Materials/MaterialAggregate.h"
+#include "MaterialExpressionAggregate.generated.h"
+
+UENUM()
+enum class EMaterialExpressionMakeAggregateKind
+{
+	MaterialAttributes,
+	UserDefined
+};
+
+USTRUCT(MinimalAPI)
+struct FMaterialExpressionAggregateEntry
+{
+	GENERATED_BODY()
+
+	UPROPERTY()
+	int32 AttributeIndex;
+
+	UPROPERTY()
+	FExpressionInput Input;
+};
+
+UCLASS(MinimalAPI)
+class UMaterialExpressionAggregate : public UMaterialExpression
+{
+	GENERATED_BODY()
+
+public:
+	UPROPERTY(EditAnywhere, Category=MaterialAggregate)
+	EMaterialExpressionMakeAggregateKind Kind = EMaterialExpressionMakeAggregateKind::MaterialAttributes;
+
+	UPROPERTY(EditAnywhere, Category=MaterialAggregate, Meta=(EditCondition="Kind == EMaterialExpressionMakeAggregateKind::UserDefined"))
+	TObjectPtr<UMaterialAggregate> UserAggregate;
+
+	UPROPERTY(EditAnywhere, Category=MaterialAggregate, Meta=(DisplayName="Attributes", GetOptions="GetPossibleAttributeNames"))
+	TArray<FName> AttributeNames;
+	
+	UPROPERTY()
+	FExpressionInput PrototypeInput;
+
+	UPROPERTY()
+	TArray<FMaterialExpressionAggregateEntry> Entries;
+
+	UMaterialExpressionAggregate(const FObjectInitializer& ObjectInitializer);
+
+#if WITH_EDITOR
+	// Returns the aggregate currently bound to this expression.
+	const UMaterialAggregate* GetAggregate() const;
+
+	// Returns the attribute with given index of the aggregate currently bound to this expression.
+	const FMaterialAggregateAttribute* GetAggregateAttribute(int32 Index) const;
+	
+	// Utility function to list the attribute names in the attribute selection combobox.
+	UFUNCTION()
+	TArray<FString> GetPossibleAttributeNames() const;
+
+	//~ Begin UMaterialExpression Interface
+	virtual void 				GetCaption(TArray<FString>& OutCaptions) const override;
+	virtual FExpressionInput* 	GetInput(int32 InputIndex) override;
+	virtual FName 				GetInputName(int32 InputIndex) const override;
+	virtual EMaterialValueType 	GetInputValueType(int32 InputIndex) override;
+	virtual EMaterialValueType 	GetOutputValueType(int32 OutputIndex) override;
+	virtual void 				PreEditChange(FProperty* PropertyAboutToChange) override;
+	virtual void 				PostEditChangeProperty(FPropertyChangedEvent& PropertyChangedEvent) override;
+	virtual void 				Build(MIR::FEmitter& Emitter) override;
+	//~ End UMaterialExpression Interface
+#endif // WITH_EDITOR
+
+private:
+#if WITH_EDITOR
+	// Used to store the state of the expression before a change was made.
+	struct FPrevEditData
+	{
+		const UMaterialAggregate* Aggregate{};
+		TArray<FName> AttributeNames{};
+		TArray<FMaterialExpressionAggregateEntry> Entries{};
+	};
+
+	FPrevEditData PreEditData;
+#endif // WITH_EDITOR
+};
diff --git a/Engine/Source/Runtime/Engine/Public/Materials/MaterialExpressionBlend.h b/Engine/Source/Runtime/Engine/Public/Materials/MaterialExpressionBlend.h
new file mode 100644
--- /dev/null
+++ b/Engine/Source/Runtime/Engine/Public/Materials/MaterialExpressionBlend.h
@@ -0,0 +1,48 @@
+// Copyright Epic Games, Inc. All Rights Reserved.
+
+#pragma once
+
+#include "CoreMinimal.h"
+#include "UObject/ObjectMacros.h"
+#include "MaterialExpressionIO.h"
+#include "Materials/MaterialExpression.h"
+#include "MaterialExpressionBlend.generated.h"
+
+UENUM()
+enum class EMaterialExpressionBlendMode : uint8
+{
+	Blend,
+	UseA,
+	UseB,
+};
+
+UCLASS(MinimalAPI, CollapseCategories, HideCategories = Object)
+class UMaterialExpressionBlend : public UMaterialExpression
+{
+	GENERATED_BODY()
+
+public:
+	UPROPERTY()
+	FExpressionInput A;
+
+	UPROPERTY()
+	FExpressionInput B;
+
+	UPROPERTY()
+	FExpressionInput Alpha;
+
+	UPROPERTY(EditAnywhere, Category=MaterialExpressionBlend)
+	EMaterialExpressionBlendMode PixelAttributesBlendMode;
+	
+	UPROPERTY(EditAnywhere, Category=MaterialExpressionBlend)
+	EMaterialExpressionBlendMode VertexAttributesBlendMode;
+
+#if WITH_EDITOR
+	//~ Begin UMaterialExpression Interface
+	virtual void 				GetCaption(TArray<FString>& OutCaptions) const override;
+	virtual EMaterialValueType 	GetInputValueType(int32 InputIndex) override;
+	virtual EMaterialValueType 	GetOutputValueType(int32 OutputIndex) override;
+	virtual void 				Build(MIR::FEmitter& Emitter) override;
+	//~ End UMaterialExpression Interface
+#endif // WITH_EDITOR
+};
diff --git a/Engine/Source/Runtime/Engine/Public/Materials/MaterialExpressionBlendMaterialAttributes.h b/Engine/Source/Runtime/Engine/Public/Materials/MaterialExpressionBlendMaterialAttributes.h
--- a/Engine/Source/Runtime/Engine/Public/Materials/MaterialExpressionBlendMaterialAttributes.h
+++ b/Engine/Source/Runtime/Engine/Public/Materials/MaterialExpressionBlendMaterialAttributes.h
@@ -54,6 +54,7 @@ class UMaterialExpressionBlendMaterialAttributes : public UMaterialExpression
 
 	//~ Begin UMaterialExpression Interface
 #if WITH_EDITOR
+	virtual void Build(MIR::FEmitter& Emitter) override;
 	virtual FExpressionInput* GetInput(int32 InputIndex) override;
 	virtual int32 Compile(class FMaterialCompiler* Compiler, int32 OutputIndex) override;
 	virtual void GetCaption(TArray<FString>& OutCaptions) const override;
diff --git a/Engine/Source/Runtime/Engine/Public/Materials/MaterialExpressionBreakMaterialAttributes.h b/Engine/Source/Runtime/Engine/Public/Materials/MaterialExpressionBreakMaterialAttributes.h
--- a/Engine/Source/Runtime/Engine/Public/Materials/MaterialExpressionBreakMaterialAttributes.h
+++ b/Engine/Source/Runtime/Engine/Public/Materials/MaterialExpressionBreakMaterialAttributes.h
@@ -24,6 +24,7 @@ class UMaterialExpressionBreakMaterialAttributes : public UMaterialExpression
  
 	//~ Begin UMaterialExpression Interface
 #if WITH_EDITOR
+	virtual void Build(MIR::FEmitter& Emitter) override;
 	virtual FExpressionInput* GetInput(int32 InputIndex) override { return InputIndex == 0 ? &MaterialAttributes : nullptr; }
 	virtual int32 Compile(class FMaterialCompiler* Compiler, int32 OutputIndex) override;
 	virtual void GetCaption(TArray<FString>& OutCaptions) const override;
diff --git a/Engine/Source/Runtime/Engine/Public/Materials/MaterialExpressionGetMaterialAttributes.h b/Engine/Source/Runtime/Engine/Public/Materials/MaterialExpressionGetMaterialAttributes.h
--- a/Engine/Source/Runtime/Engine/Public/Materials/MaterialExpressionGetMaterialAttributes.h
+++ b/Engine/Source/Runtime/Engine/Public/Materials/MaterialExpressionGetMaterialAttributes.h
@@ -37,6 +37,7 @@ class UMaterialExpressionGetMaterialAttributes : public UMaterialExpression
 	//~ Begin UMaterialExpression Interface
 #if WITH_EDITOR
 	ENGINE_API int32 CreateOrGetOutputAttribute(EMaterialProperty Attribute);
+	virtual void Build(MIR::FEmitter& Emitter) override;
 	virtual int32 Compile(class FMaterialCompiler* Compiler, int32 OutputIndex) override;
 	virtual void GetCaption(TArray<FString>& OutCaptions) const override;
 	virtual FName GetInputName(int32 InputIndex) const override;
diff --git a/Engine/Source/Runtime/Engine/Public/Materials/MaterialExpressionMakeMaterialAttributes.h b/Engine/Source/Runtime/Engine/Public/Materials/MaterialExpressionMakeMaterialAttributes.h
--- a/Engine/Source/Runtime/Engine/Public/Materials/MaterialExpressionMakeMaterialAttributes.h
+++ b/Engine/Source/Runtime/Engine/Public/Materials/MaterialExpressionMakeMaterialAttributes.h
@@ -85,6 +85,7 @@ public:
 
 	//~ Begin UMaterialExpression Interface
 #if WITH_EDITOR
+	virtual void Build(MIR::FEmitter& Emitter) override;
 	virtual int32 Compile(class FMaterialCompiler* Compiler, int32 OutputIndex) override;
 	virtual void GetCaption(TArray<FString>& OutCaptions) const override;
 	virtual bool IsResultMaterialAttributes(int32 OutputIndex) override {return true;}
diff --git a/Engine/Source/Runtime/Engine/Public/Materials/MaterialExpressionSetMaterialAttributes.h b/Engine/Source/Runtime/Engine/Public/Materials/MaterialExpressionSetMaterialAttributes.h
--- a/Engine/Source/Runtime/Engine/Public/Materials/MaterialExpressionSetMaterialAttributes.h
+++ b/Engine/Source/Runtime/Engine/Public/Materials/MaterialExpressionSetMaterialAttributes.h
@@ -34,6 +34,7 @@ class UMaterialExpressionSetMaterialAttributes : public UMaterialExpression
 
 	//~ Begin UMaterialExpression Interface
 #if WITH_EDITOR
+	virtual void Build(MIR::FEmitter& Emitter) override;
 	virtual int32 Compile(class FMaterialCompiler* Compiler, int32 OutputIndex) override;
 	virtual void GetCaption(TArray<FString>& OutCaptions) const override;
 	virtual TArrayView<FExpressionInput*> GetInputsView()override;
diff --git a/Engine/Source/Runtime/Engine/Public/Materials/MaterialIR.h b/Engine/Source/Runtime/Engine/Public/Materials/MaterialIR.h
--- a/Engine/Source/Runtime/Engine/Public/Materials/MaterialIR.h
+++ b/Engine/Source/Runtime/Engine/Public/Materials/MaterialIR.h
@@ -96,7 +96,7 @@ enum EValueKind
 	VK_InstructionBegin,
 
 	VK_SetMaterialOutput,
-	VK_Dimensional,
+	VK_Composite,
 	VK_Operator,
 	VK_Branch,
 	VK_Subscript,
@@ -540,7 +540,7 @@ struct TInstruction : FInstruction
 // An aggregate of other values.
 //
 // A dimensional is a fixed array of other values. This value is used to model vectors and matrices.
-struct FDimensional : TInstruction<VK_Dimensional, 0>
+struct FComposite : TInstruction<VK_Composite, 0>
 {
 	// Returns the constant array of component values.
 	TConstArrayView<FValue*> GetComponents() const;
@@ -552,10 +552,10 @@ struct FDimensional : TInstruction<VK_Dimensional, 0>
 	bool AreComponentsConstant() const;
 };
 
-template <int TDimension>
-struct TDimensional : FDimensional
+template <int TSize>
+struct TComposite : FComposite
 {
-	FValue* Components[TDimension];
+	FValue* Components[TSize];
 };
 
 // Instruction that sets material attribute (i.e., "BaseColor") to its value.
diff --git a/Engine/Source/Runtime/Engine/Public/Materials/MaterialIRCommon.h b/Engine/Source/Runtime/Engine/Public/Materials/MaterialIRCommon.h
--- a/Engine/Source/Runtime/Engine/Public/Materials/MaterialIRCommon.h
+++ b/Engine/Source/Runtime/Engine/Public/Materials/MaterialIRCommon.h
@@ -5,6 +5,7 @@
 #include "RHIShaderPlatform.h"
 #include "Containers/Array.h"
 #include "Math/MathFwd.h"
+#include "Misc/MemStack.h"
 
 #if WITH_EDITOR
 
@@ -14,6 +15,7 @@ class FMaterial;
 class FMaterialIRModule;
 class ITargetPlatform;
 class UMaterial;
+class UMaterialAggregate;
 class UMaterialExpression;
 class UMaterialExpressionMaterialFunctionCall;
 class UTexture;
@@ -41,6 +43,7 @@ namespace MIR
 struct FType;
 struct FPrimitiveType;
 struct FObjectType;
+struct FAggregateType;
 
 /* IR */
 struct FValue;
@@ -57,6 +60,31 @@ namespace Internal
 	uint32 HashBytes(const char* Ptr, uint32 Size);
 }
 
+/* Helper structures */
+
+// Use this to efficiently allocate a temporary array using MemStack, instead
+// of using TArray and going through the global allocator. Only declare local variable
+// of this struct (i.e. do not "new TTemporaryArray")
+// Remark: The allocated memory lifespan is the same as the TTemporaryArray local variable
+// (it is deallocated when TTemporaryArray goes out of scope).
+template <typename T>
+struct TTemporaryArray : public TArrayView<T>
+{
+	FMemMark MemMark;
+
+	TTemporaryArray(int32 Num)
+	: MemMark{ FMemStack::Get() }
+	{
+		*static_cast<TArrayView<T>*>(this) = { (T*)FMemStack::Get().Alloc(sizeof(T) * Num, alignof(T)), Num };
+	}
+
+	operator TConstArrayView<T>() const
+	{
+		return { TArrayView<T>::GetData(), TArrayView<T>::Num() };
+	}
+};
+
+
 }
 
 #define UE_MIR_UNREACHABLE() { check(!"Unreachable"); UE_ASSUME(false); }
diff --git a/Engine/Source/Runtime/Engine/Public/Materials/MaterialIREmitter.h b/Engine/Source/Runtime/Engine/Public/Materials/MaterialIREmitter.h
--- a/Engine/Source/Runtime/Engine/Public/Materials/MaterialIREmitter.h
+++ b/Engine/Source/Runtime/Engine/Public/Materials/MaterialIREmitter.h
@@ -93,6 +93,9 @@ struct FValueRef
 	// Checks if the reference holds a valid non-null value.
 	bool IsValid() const;
 
+	// Checks if the reference holds a poison value.
+	bool IsPoison() const;
+
 	// Creates a new reference with the same input but a different value.
 	FValueRef To(FValue* Value) const;
 
@@ -106,6 +109,16 @@ struct FValueRef
 	const FExpressionInput* Input{};
 };
 
+//
+struct FAttributeAssignment
+{
+	//
+	FName Name;
+
+	//
+	FValueRef Value;
+};
+
 // The FEmitter is a helper object responsible for emitting MIR values and instructions.
 // It is primarily used within UMaterialExpression::Build() implementations to lower the
 // high-level semantics of a material expression into corresponding lower-level MIR nodes.
@@ -172,6 +185,12 @@ public:
 	// See TryGetCommonType()
 	const FType* GetCommonType(const FType* A, const FType* B);
 
+	// Returns the type associated to the given material aggregate. 
+	const FAggregateType* GetMaterialAggregateType(const UMaterialAggregate* Aggregate);
+
+	// Returns the type of the i-th attribute in the given material aggregate.
+	const FType* GetMaterialAggregateAttributeType(const UMaterialAggregate* Aggregate, int32 AttributeIndex);
+
 	/*-------------------------------- Input management --------------------------------*/
 
 	// Retrieves the value from the given input, if present. A null reference otherwise.
@@ -194,8 +213,6 @@ public:
 	// If the input has no value flowing in (e.g., because the input is disconnected),
 	// these DefaultXXX() functions will emit a constant value, bind it to the input, and
 	// return it.
-	// Note that these functions may return poison if the value flowing into this input
-	// was poison in	FValueRef DefaultBool(FValueRef Value, bool Default);
 	FValueRef InputDefaultBool(const FExpressionInput* Input, bool Default);
 	FValueRef InputDefaultInt(const FExpressionInput* Input, TInteger Default);
 	FValueRef InputDefaultInt2(const FExpressionInput* Input, UE::Math::TIntVector2<TInteger> Default);
@@ -229,6 +246,9 @@ public:
 	// If the value is missing or cannot be cast to a texture, it reports an error and returns poison.
 	FValueRef CheckIsTexture(FValueRef Value);
 
+	// Validates that the value is of an aggregate type. If not, reports error and returns poison.
+	FValueRef CheckIsAggregate(FValueRef Value, const UMaterialAggregate* Aggregate = nullptr);
+
 	// Casts the value flowing into this input to a constant boolean and returns it.
 	// If the value is missing, or is not a constant boolean it reports an error and returns false.
 	bool ToConstantBool(FValueRef Value);
@@ -246,6 +266,9 @@ public:
 	// Converts the given UE::Shader::FValue to a MIR constant value and returns it.
 	FValueRef ConstantFromShaderValue(const UE::Shader::FValue& InValue);
 
+	// Emits a default value of given type (zero initialized).
+	FValueRef ConstantDefault(const FType* Type);
+
 	// Emits a scalar constant value of given kind with value 0 and returns it.
 	FValueRef ConstantZero(EScalarKind Kind);
 
@@ -329,9 +352,24 @@ public:
 	// Emits a 4D column vector dimensional value and returns it.
 	FValueRef Vector4(FValueRef InX, FValueRef InY, FValueRef InZ, FValueRef InW);
 	
+	// Constructs an aggregate value using the provided aggregate definition.
+	// The resulting value has default-initialized attributes.
+	FValueRef Aggregate(const UMaterialAggregate* InAggregate);
+
+	// Constructs an aggregate value using the given type and prototype, and initializes it with
+	// the provided list of attribute values. Each value corresponds to an attribute in declaration order.
+	// AttributeValues can be smaller than the number attributes in the aggregate. For excess attributes,
+	// the value in the prototype (if provided) or default initialized otherwise.
+	FValueRef Aggregate(const UMaterialAggregate* InAggregate, FValueRef InPrototype, TConstArrayView<FValueRef> AttributeValues);
+
+	// Constructs an aggregate value using the given type and prototype, initializing its attributes using the provided assignments.
+	// Each assignment specifies the target attribute and its initial value.
+	// Unassigned attributes use the value in the prototype (if provided) or default initialized otherwise.
+	FValueRef Aggregate(const UMaterialAggregate* InAggregate, FValueRef InPrototype, TConstArrayView<FAttributeAssignment> AttributeAssignments);
+
 	// Emits a mathematical operation instruction with given operator and arguments.
-	// Note: this function will try to simplify the operation at translation time if possible.
-	// The returned value is therefore not guaranteed to be an FOperator instruction instance.
+	// Note: This function will try to simplify the operation at translation time if possible.
+	//       The returned value is therefore not guaranteed to be an FOperator instruction instance.
 	FValueRef Operator(EOperator Operator, FValueRef A, FValueRef B = {}, FValueRef C = {});
 	
 	// Emits a branch instruction.
@@ -351,9 +389,9 @@ public:
 	// Casts the given primitive value to a given scalar kind, maintaining the same
 	// number of rows and columns.
 	FValueRef CastToScalarKind(FValueRef Value, EScalarKind ToScalarKind);
-	FValueRef CastToBool(FValueRef Value);
-	FValueRef CastToInt(FValueRef Value);
-	FValueRef CastToFloat(FValueRef Value);
+	FValueRef CastToBoolKind(FValueRef Value);
+	FValueRef CastToIntKind(FValueRef Value);
+	FValueRef CastToFloatKind(FValueRef Value);
 
 	// These functions cast the given value to the a primitive scalar or column vector type of given rows.
 	// If the value is missing ore could not be cast to the relative primitive type, they repors an error and return poison.
@@ -580,6 +618,9 @@ private:
 	// calculations.
 	TSet<FValue*, FValueKeyFuncs> ValueSet{};
 
+	// Maps aggregate types to their MIR type instance. Used to avoid duplicates.
+	TMap<const UMaterialAggregate*, FAggregateType*> AggregateTypeMap{};
+
 	friend FMaterialIRModuleBuilder;
 	friend FMaterialIRModuleBuilderImpl;
 };
diff --git a/Engine/Source/Runtime/Engine/Public/Materials/MaterialIRModule.h b/Engine/Source/Runtime/Engine/Public/Materials/MaterialIRModule.h
--- a/Engine/Source/Runtime/Engine/Public/Materials/MaterialIRModule.h
+++ b/Engine/Source/Runtime/Engine/Public/Materials/MaterialIRModule.h
@@ -95,8 +95,8 @@ private:
 	// Compilation output data.
 	FMaterialCompilationOutput CompilationOutput;
 
-	// Memory allocator used to allocate IR data (values, payloads, etc).
-	FMemStackBase Allocator{};
+	// List of all the dynamic types contained in this module.
+	TArray<MIR::FType*> Types;
 
 	// List of all the IR values contained in this module.
 	TArray<MIR::FValue*> Values;
diff --git a/Engine/Source/Runtime/Engine/Public/Materials/MaterialIRTypes.h b/Engine/Source/Runtime/Engine/Public/Materials/MaterialIRTypes.h
--- a/Engine/Source/Runtime/Engine/Public/Materials/MaterialIRTypes.h
+++ b/Engine/Source/Runtime/Engine/Public/Materials/MaterialIRTypes.h
@@ -17,6 +17,7 @@ enum ETypeKind
 	TypeKind_Void,
 	TypeKind_Primitive,
 	TypeKind_Object,
+	TypeKind_Aggregate,
 };
 
 enum EObjectKind 
@@ -57,6 +58,9 @@ struct FType
 	/// Returns whether this is an integer primitive type (of any dimension).
 	bool IsInteger() const;
 
+	/// Returns whether this is a floating point primitive type (of any dimension).
+	bool IsAnyFloat() const;
+
 	/// Returns whether this is an LWC primitive type (of any dimension).
 	bool IsLWC() const;
 
@@ -69,23 +73,26 @@ struct FType
 	/// Returns whether this type is a texture.
 	bool IsTexture() const;
 
-	/// Returns this type upcast this type to PrimitiveType if it is one, otherwise nullptr.
+	/// Returns this type upcast to PrimitiveType if it is one, otherwise nullptr.
 	const FPrimitiveType* AsPrimitive() const;
 	
-	/// Returns this type upcast this type to PrimitiveType if it is one, otherwise nullptr.
+	/// Returns this type upcast to PrimitiveType if it is one, otherwise nullptr.
 	const FPrimitiveType* AsArithmetic() const;
 
-	/// Returns this type upcast this type to PrimitiveType if it's a scalar, otherwise nullptr.
+	/// Returns this type upcast to PrimitiveType if it's a scalar, otherwise nullptr.
 	const FPrimitiveType* AsScalar() const;
 
-	/// Returns this type upcast this type to PrimitiveType if it's a vector, otherwise nullptr.
+	/// Returns this type upcast to PrimitiveType if it's a vector, otherwise nullptr.
 	const FPrimitiveType* AsVector() const;
 
-	/// Returns this type upcast this type to PrimitiveType if it's a matrix, otherwise nullptr.
+	/// Returns this type upcast to PrimitiveType if it's a matrix, otherwise nullptr.
 	const FPrimitiveType* AsMatrix() const;
 
+	/// Returns this type upcast to AggregateType, otherwise nullptr.
+	const FAggregateType* AsAggregate() const;
+
 	/// Returns the this type name spelling (e.g. float4x4).
-	FStringView GetSpelling() const;
+	FString GetSpelling() const;
 	
 	/// Returns this type cast as a FObjectPtr if it is one, otherwise nullptr.
 	const FObjectType* AsObject() const;
@@ -163,10 +170,10 @@ struct FPrimitiveType : FType
 	static const FPrimitiveType* Get(EScalarKind InScalarKind, int NumRows, int NumColumns, bool bIsLWCInverseMatrix = false);
 
 	// Returns the number of components in this primitive type.
-	int  GetNumComponents() const { return NumRows * NumColumns; }
+	int  NumComponents() const { return NumRows * NumColumns; }
 
 	// Whether this primitive type is scalar.
-	bool IsScalar() const { return GetNumComponents() == 1; }
+	bool IsScalar() const { return NumComponents() == 1; }
 	
 	// Whether this primitive type is a column vector.
 	bool IsVector() const { return NumRows > 1 && NumColumns == 1; }
@@ -187,6 +194,7 @@ struct FPrimitiveType : FType
 	const FPrimitiveType* ToVector(int NumRows) const;
 };
 
+// It represents an opaque object-like type such as a texture.
 struct FObjectType : FType
 {
 	EObjectKind ObjectKind;
@@ -196,6 +204,12 @@ struct FObjectType : FType
 	static const FType* GetVTPageTableResult();
 };
 
+// A value of this type is a material aggregate.
+struct FAggregateType : FType
+{
+	const UMaterialAggregate* Aggregate;
+};
+
 } // namespace MIR
 
 #endif // #if WITH_EDITOR
