Build fix for and inaccuracy fix to previous Material Aggregates CL I submitted.

#rnx

diff --git a/Engine/Source/Runtime/Engine/Private/Materials/MaterialExpressions.cpp b/Engine/Source/Runtime/Engine/Private/Materials/MaterialExpressions.cpp
--- a/Engine/Source/Runtime/Engine/Private/Materials/MaterialExpressions.cpp
+++ b/Engine/Source/Runtime/Engine/Private/Materials/MaterialExpressions.cpp
@@ -2315,12 +2315,6 @@ bool UMaterialExpression::ContainsInputLoopInternal(const FContainsInputLoopInte
 	return false;
 }
 
-int32 UMaterialExpression::Compile(class FMaterialCompiler* Compiler, int32 OutputIndex)
-{
-	Compiler->Error(TEXT("Unsupported material expression."));
-	return INDEX_NONE;
-}
-
 bool UMaterialExpression::IsUsingNewHLSLGenerator() const
 {
 	if (Material)
diff --git a/Engine/Source/Runtime/Engine/Public/Materials/MaterialAggregate.h b/Engine/Source/Runtime/Engine/Public/Materials/MaterialAggregate.h
--- a/Engine/Source/Runtime/Engine/Public/Materials/MaterialAggregate.h
+++ b/Engine/Source/Runtime/Engine/Public/Materials/MaterialAggregate.h
@@ -56,7 +56,7 @@ struct FMaterialAggregateAttribute
 	EMaterialAggregateAttributeType Type = EMaterialAggregateAttributeType::Float4;
 
 	// The nested aggregate reference, used only if TypeKind is Aggregate.
-	UPROPERTY(EditAnywhere, Category=MaterialAggregateAttribute, meta=(EditCondition="TypeKind == EMaterialAggregateAttributeTypeKind::Aggregate", EditConditionHides))
+	UPROPERTY(EditAnywhere, Category=MaterialAggregateAttribute, meta=(EditCondition="Type == EMaterialAggregateAttributeType::Aggregate", EditConditionHides))
 	TObjectPtr<UMaterialAggregate> Aggregate{};
 
 #if WITH_EDITOR
diff --git a/Engine/Source/Runtime/Engine/Public/Materials/MaterialExpression.h b/Engine/Source/Runtime/Engine/Public/Materials/MaterialExpression.h
--- a/Engine/Source/Runtime/Engine/Public/Materials/MaterialExpression.h
+++ b/Engine/Source/Runtime/Engine/Public/Materials/MaterialExpression.h
@@ -265,7 +265,7 @@ class UMaterialExpression : public UObject
 	 * @param	Compiler - UMaterial compiler that knows how to handle this expression.
 	 * @return	Index to the new FMaterialCompiler::CodeChunk entry for this expression
 	 */	
-	virtual int32 Compile(class FMaterialCompiler* Compiler, int32 OutputIndex);
+	virtual int32 Compile(class FMaterialCompiler* Compiler, int32 OutputIndex) { return INDEX_NONE; }
 	virtual int32 CompilePreview(class FMaterialCompiler* Compiler, int32 OutputIndex) { return Compile(Compiler, OutputIndex); }
 
 	ENGINE_API bool IsUsingNewHLSLGenerator() const;
