[VisualLogger] log entry initialization can now also assign a valid location for components. With this change, double-clicking on such an entry will bring the camera to that location instead of the world's origin.

diff --git a/Engine/Source/Runtime/Engine/Private/VisualLogger/VisualLogger.cpp b/Engine/Source/Runtime/Engine/Private/VisualLogger/VisualLogger.cpp
--- a/Engine/Source/Runtime/Engine/Private/VisualLogger/VisualLogger.cpp
+++ b/Engine/Source/Runtime/Engine/Private/VisualLogger/VisualLogger.cpp
@@ -579,6 +579,19 @@ FVisualLogEntry* FVisualLogger::GetEntryToWriteInternal(const UObject* Object, c
 			CurrentEntry->Location = ObjectAsActor->GetActorLocation();
 			CurrentEntry->bIsLocationValid = true;
 		}
+		else if (const UActorComponent* ObjectAsActorComponent = Cast<UActorComponent>(LogOwner))
+		{
+			if (const USceneComponent* ObjectAsSceneComponent = Cast<USceneComponent>(ObjectAsActorComponent))
+			{
+				CurrentEntry->Location = ObjectAsSceneComponent->GetComponentLocation();
+				CurrentEntry->bIsLocationValid = true;
+			}
+			else if (const AActor* ComponentOwner = ObjectAsActorComponent->GetOwner())
+			{
+				CurrentEntry->Location = ComponentOwner->GetActorLocation();
+				CurrentEntry->bIsLocationValid = true;
+			}
+		}
 
 		UE::TReadScopeLock RedirectScopeLock(RedirectRWLock);
 		FOwnerToChildrenRedirectionMap& RedirectionMap = GetRedirectionMap(LogOwner);
