[NewMaterialTranslator] - Fix bug with some use cases of Pow in the new translator.

- The implementation of Power::Build uses the Select operator which emits the HLSL select.  We hand implement our own select functions in Common.ush so in some scenarios it was incorrectly picking up up select_internal(bool, int, int) instead of the float version, which would cause a cast of the pow result and it would truncate to black.
- This was because of %g turns 0.00000000 into 0
- Change the %g to using %#.8gf.  The `#` forces the written output to contain a decimal point even if no more digits follow.
- This explicitly adds the `f` at the end of the constant to avoid further ambiguity, and to also match the emitted code of the old translator.

#rnx

diff --git a/Engine/Source/Runtime/Engine/Private/Materials/MaterialIRToHLSLTranslator.cpp b/Engine/Source/Runtime/Engine/Private/Materials/MaterialIRToHLSLTranslator.cpp
--- a/Engine/Source/Runtime/Engine/Private/Materials/MaterialIRToHLSLTranslator.cpp
+++ b/Engine/Source/Runtime/Engine/Private/Materials/MaterialIRToHLSLTranslator.cpp
@@ -450,7 +450,7 @@ struct FTranslator : FMaterialIRToHLSLTranslation
 				}
 				else
 				{
-					Printer.Buffer.Appendf(TEXT("%.9g"), Constant->Float);
+					Printer.Buffer.Appendf(TEXT("%#.8gf"), Constant->Float);
 				}
 				break;
 
