[WorldPartition] Make sure the full update logic (instead of optimized) is used to update streaming sources right after the World Partition was in a state in which it couldn't stream

diff --git a/Engine/Source/Runtime/Engine/Private/WorldPartition/WorldPartitionStreamingPolicy.cpp b/Engine/Source/Runtime/Engine/Private/WorldPartition/WorldPartitionStreamingPolicy.cpp
--- a/Engine/Source/Runtime/Engine/Private/WorldPartition/WorldPartitionStreamingPolicy.cpp
+++ b/Engine/Source/Runtime/Engine/Private/WorldPartition/WorldPartitionStreamingPolicy.cpp
@@ -110,6 +110,7 @@ void UWorldPartitionStreamingPolicy::UpdateStreamingSources(bool bCanOptimizeUpd
 	if (!WorldPartition->CanStream())
 	{
 		CurrentState.StreamingSources.Reset();
+		UpdateStreamingSourcesHash = 0;
 		return;
 	}
 
