New Translator static analysis fix in static_assert subtracting the same value with itself.

#rnx

diff --git a/Engine/Source/Runtime/Engine/Private/Materials/MaterialExpressionsIR.cpp b/Engine/Source/Runtime/Engine/Private/Materials/MaterialExpressionsIR.cpp
--- a/Engine/Source/Runtime/Engine/Private/Materials/MaterialExpressionsIR.cpp
+++ b/Engine/Source/Runtime/Engine/Private/Materials/MaterialExpressionsIR.cpp
@@ -1551,7 +1551,6 @@ void UMaterialExpressionWorldPosition::Build(MIR::FEmitter& Em)
 {
 	// Make sure EWorldPositionIncludedOffsets and corresponding elements of EExternalInput stay in sync, so this enum addition is valid
 	static_assert(WPT_MAX == 4 && WPT_Default == 0);
-	static_assert((int32)MIR::EExternalInput::WorldPosition_Absolute - (int32)MIR::EExternalInput::WorldPosition_Absolute == WPT_Default);
 	static_assert((int32)MIR::EExternalInput::WorldPosition_AbsoluteNoOffsets - (int32)MIR::EExternalInput::WorldPosition_Absolute == WPT_ExcludeAllShaderOffsets);
 	static_assert((int32)MIR::EExternalInput::WorldPosition_CameraRelative - (int32)MIR::EExternalInput::WorldPosition_Absolute == WPT_CameraRelative);
 	static_assert((int32)MIR::EExternalInput::WorldPosition_CameraRelativeNoOffsets - (int32)MIR::EExternalInput::WorldPosition_Absolute == WPT_CameraRelativeNoOffsets);
