New Material Translator:  Natvis support for FType and FValue.  Makes debugging FValue trees infinitely easier.

#rnx

diff --git a/Engine/Source/Runtime/Engine/Private/Materials/MaterialIR.cpp b/Engine/Source/Runtime/Engine/Private/Materials/MaterialIR.cpp
--- a/Engine/Source/Runtime/Engine/Private/Materials/MaterialIR.cpp
+++ b/Engine/Source/Runtime/Engine/Private/Materials/MaterialIR.cpp
@@ -592,6 +592,9 @@ FBlock* FBlock::FindCommonParentWith(MIR::FBlock* Other)
 	return A;
 }
 
+// This global is needed to allow the Visual Studio debugger (and subsequently Natvis) to have access to the "TComposite<1>" type identifer, as used in the function below.
+TComposite<1> GCompositeNatvisPrototype;
+
 TConstArrayView<FValue*> FComposite::GetComponents() const
 {
 	int32 NumComponents;
