Removed mip change callbacks in SkinWeightProfileManager due to stability and performance

#rnx

diff --git a/Engine/Source/Runtime/Engine/Private/Animation/SkinWeightProfileManager.cpp b/Engine/Source/Runtime/Engine/Private/Animation/SkinWeightProfileManager.cpp
--- a/Engine/Source/Runtime/Engine/Private/Animation/SkinWeightProfileManager.cpp
+++ b/Engine/Source/Runtime/Engine/Private/Animation/SkinWeightProfileManager.cpp
@@ -185,6 +185,8 @@ void FSkinWeightProfileManager::RequestSkinWeightProfileStack(
 		int32& NumRequestsForMesh = PendingMeshes.FindOrAdd(Mesh, 0);
 		++NumRequestsForMesh;
 
+		WaitingForStreaming = false;
+
 		TickFunction.SetTickFunctionEnable(true);
 	}
 }
@@ -197,6 +199,7 @@ void FSkinWeightProfileManager::CancelSkinWeightProfileRequest(UObject* Requeste
 	{
 		if (Request.IdentifyingObject == Requester || Request.WeakSkeletalMesh == Requester)
 		{
+			WaitingForStreaming = false;
 			CanceledRequest.Add(Request);
 			UE_LOG(LogSkinWeightProfileManager, Display, TEXT("Cancel [%s]"), *Requester->GetName());			
 		}
@@ -217,6 +220,13 @@ void FSkinWeightProfileManager::DoTick(float DeltaTime, ENamedThreads::Type Curr
 
 	if (PendingSetProfileRequests.Num())
 	{
+		if (ShouldSkipTick())
+		{
+			return;
+		}
+
+		WaitingForStreaming = false;
+
 		if (AsyncTask.IsValid())
 		{
 			FTaskGraphInterface::Get().WaitUntilTaskCompletes(AsyncTask, ENamedThreads::GameThread);
@@ -229,6 +239,8 @@ void FSkinWeightProfileManager::DoTick(float DeltaTime, ENamedThreads::Type Curr
 			QUICK_SCOPE_CYCLE_COUNTER(STAT_SkinWeightProfileManager_TickMesh);
 
 			FSetProfileRequest& Request = PendingSetProfileRequests[RequestIndex];
+
+			Request.ExpectedResidentLODs = INDEX_NONE;
 			
 			// Remove if canceled or objects have gone stale 
 			bool bRemove = Request.IdentifyingObject.IsStale() || Request.WeakSkeletalMesh.IsStale() ||
@@ -256,7 +268,7 @@ void FSkinWeightProfileManager::DoTick(float DeltaTime, ENamedThreads::Type Curr
 						if (ShouldIgnoreLOD(SkinWeightProfilesData, Request.WeakSkeletalMesh))
 						{
 							// This LOD isn't streamed in and has no requests for that. There may be several such LODs for the current model.
-							// We will create a callback to be notified when any of these LODs will be streamed in,
+							// We will create a record in StreamingBlockedMeshes to track when any of these LODs will be streamed in,
 							// and for that we need to know the index of the least detailed LOD among them.
 							LODIndexToWaitStreaming = FMath::Max(LODIndexToWaitStreaming, LODRenderDataArrayIndex);
 						}
@@ -271,10 +283,10 @@ void FSkinWeightProfileManager::DoTick(float DeltaTime, ENamedThreads::Type Curr
 					}
 				}
 
-				if ((LODIndexToWaitStreaming != INDEX_NONE) &&
-					(LODIndexToWaitStreaming != Request.StreamingCallbackLODIndex))
+				if (LODIndexToWaitStreaming != INDEX_NONE)
 				{
-					CreateLODStreamingCallback(Request, LODIndexToWaitStreaming);
+					const FStreamableRenderResourceState& SRRState = SkeletalMesh->GetStreamableResourceState();
+					Request.ExpectedResidentLODs = SRRState.MaxNumLODs - (LODIndexToWaitStreaming - SRRState.AssetLODBias);
 				}
 
 				if (bAllBuffersReady)
@@ -310,9 +322,8 @@ void FSkinWeightProfileManager::DoTick(float DeltaTime, ENamedThreads::Type Curr
 		}
 		else
 		{
-			// This is called when we still have pending requests, but all their unprocessed LODs have not been streamed in.
-			// TickFunction will be re-enabled once the streaming is complete or a new request is received.
-			TickFunction.SetTickFunctionEnable(false);
+			// This is set when we still have pending requests, but all their unprocessed LODs have not been streamed in.
+			WaitingForStreaming = true;
 		}
 	}
 	else
@@ -321,34 +332,6 @@ void FSkinWeightProfileManager::DoTick(float DeltaTime, ENamedThreads::Type Curr
 	}
 }
 
-void FSkinWeightProfileManager::OnRequestedMeshesStreamingUpdated(const UPrimitiveComponent* PrimitiveComponent, const UStreamableRenderAsset* StreamableRenderAsset, const FSkinWeightProfileStack& ProfileStack)
-{
-	if (!PrimitiveComponent || !StreamableRenderAsset)
-	{
-		return;
-	}
-
-	const USkeletalMesh* SkinnedMeshComponent = Cast<USkeletalMesh>(StreamableRenderAsset);
-	if (!SkinnedMeshComponent)
-	{
-		return;
-	}
-
-	// Find the request for which streaming was updated to reset its callback index
-	for (FSetProfileRequest& Request : PendingSetProfileRequests)
-	{
-		if ((Request.ProfileStack == ProfileStack) &&
-			(Request.IdentifyingObject.Get() == PrimitiveComponent) &&
-			(Request.WeakSkeletalMesh.Get() == SkinnedMeshComponent))
-		{
-			Request.StreamingCallbackLODIndex = INDEX_NONE;
-			break;
-		}
-	}
-
-	TickFunction.SetTickFunctionEnable(true);
-}
-
 void FSkinWeightProfileManager::Tick(float DeltaTime)
 {
 	DoTick(DeltaTime, ENamedThreads::GameThread, nullptr);
@@ -419,33 +402,29 @@ bool FSkinWeightProfileManager::ShouldIgnoreLOD(const FSkinWeightProfilesData& S
 	return true;
 }
 
-void FSkinWeightProfileManager::CreateLODStreamingCallback(FSetProfileRequest& Request, int32 LODIndexToWaitStreaming)
+// Used to return early from DoTick if the entire request queue consists of LOD meshes that have not yet been streamed in.
+bool FSkinWeightProfileManager::ShouldSkipTick() const
 {
-	UObject* IdentifyingObject = Request.IdentifyingObject.Get();
-	if (!IdentifyingObject)
+	if (!WaitingForStreaming || PendingSetProfileRequests.IsEmpty() || !GSkipLODWithoutDataOrStreamingRequest)
 	{
-		return;
+		return false;
 	}
 
-	USkinnedMeshComponent* Component = Cast<USkinnedMeshComponent>(IdentifyingObject);
-	if (!Component)
+	for (const FSetProfileRequest& Request : PendingSetProfileRequests)
 	{
-		return;
-	}
+		if (Request.ExpectedResidentLODs == INDEX_NONE || Request.IdentifyingObject.IsStale())
+		{
+			return false; // Continue ticking because there are stale requests to clean up
+		}
 
-	constexpr float CallbackTimeoutSecs = 60.0f;
-	Component->RegisterLODStreamingCallback(
-		[ProfileStack=Request.ProfileStack](UPrimitiveComponent* PrimitiveComponent, UStreamableRenderAsset* StreamableRenderAsset, ELODStreamingCallbackResult LODStreamingCallbackResult)
+		const USkeletalMesh* Mesh = Request.WeakSkeletalMesh.Get();
+		if (!Mesh || (Mesh->GetStreamableResourceState().NumResidentLODs >= Request.ExpectedResidentLODs))
 		{
-			FSkinWeightProfileManager* Manager = FSkinWeightProfileManager::Get(PrimitiveComponent->GetWorld());
-			if (Manager)
-			{
-				Manager->OnRequestedMeshesStreamingUpdated(PrimitiveComponent, StreamableRenderAsset, ProfileStack);
-			}
-		},
-		LODIndexToWaitStreaming, CallbackTimeoutSecs, true);
+			return false; // Continue ticking to clean up or because something was streamed in
+		}
+	}
 
-	Request.StreamingCallbackLODIndex = LODIndexToWaitStreaming;
+	return true; // All current requests are still waiting for streaming, we can end current tick early
 }
 
 bool FSkinWeightProfileManager::IsTickableWhenPaused() const
diff --git a/Engine/Source/Runtime/Engine/Public/Animation/SkinWeightProfileManager.h b/Engine/Source/Runtime/Engine/Public/Animation/SkinWeightProfileManager.h
--- a/Engine/Source/Runtime/Engine/Public/Animation/SkinWeightProfileManager.h
+++ b/Engine/Source/Runtime/Engine/Public/Animation/SkinWeightProfileManager.h
@@ -35,8 +35,8 @@ struct FSetProfileRequest
 	/** Weak skeletal mesh for which the skin weight profile is to be loaded */
 	TWeakObjectPtr<USkeletalMesh> WeakSkeletalMesh;
 
-	/** An LOD index for which we have an active streaming callback, INDEX_NONE if none exists */
-	int32 StreamingCallbackLODIndex = INDEX_NONE;
+	/** Tells how much LODs need to be streamed in so that we can continue processing the request. If INDEX_NONE then no waiting is required */
+	int32 ExpectedResidentLODs = INDEX_NONE;
 
 	friend bool operator==(const FSetProfileRequest& A, const FSetProfileRequest& B)
 	{
@@ -118,18 +118,17 @@ public:
 	void CancelSkinWeightProfileRequest(UObject* Requester);
 	
 	void DoTick(float DeltaTime, ENamedThreads::Type CurrentThread, const FGraphEventRef& MyCompletionGraphEvent);
-
-	void OnRequestedMeshesStreamingUpdated(const UPrimitiveComponent* PrimitiveComponent, const UStreamableRenderAsset* StreamableRenderAsset, const FSkinWeightProfileStack& ProfileStack);
 protected:	
 	void CleanupRequest(const FSetProfileRequest& Request);
 	bool ShouldIgnoreLOD(const FSkinWeightProfilesData& SkinWeightProfilesData, const TWeakObjectPtr<USkeletalMesh>& WeakSkeletalMesh) const;
-	void CreateLODStreamingCallback(FSetProfileRequest& Request, int32 LODIndexToWaitStreaming);
+	bool ShouldSkipTick() const;
 protected:
 	TArray<FSetProfileRequest, TInlineAllocator<4>> CanceledRequest;
 	TArray<FSetProfileRequest> PendingSetProfileRequests;
 	TMap<TWeakObjectPtr<USkeletalMesh>, int32> PendingMeshes;
 	FSkinWeightProfileManagerTickFunction TickFunction;
 	int32 LastGamethreadProfileIndex;
+	bool WaitingForStreaming = false;
 
 	TWeakObjectPtr<UWorld> WeakWorld;
 
