[Engine]
* Added dllstorage on FODSCManager::IsODSCEnabled() to fix modular UnrealGame linker error

diff --git a/Engine/Source/Runtime/Engine/Public/ODSC/ODSCManager.h b/Engine/Source/Runtime/Engine/Public/ODSC/ODSCManager.h
--- a/Engine/Source/Runtime/Engine/Public/ODSC/ODSCManager.h
+++ b/Engine/Source/Runtime/Engine/Public/ODSC/ODSCManager.h
@@ -111,7 +111,7 @@ public:
 	static void UnregisterMaterialInstance(const UMaterialInstance* MI);
 
 	/** Check options to see if ODSC could be active in this session */
-	static bool IsODSCEnabled();
+	static ENGINE_API bool IsODSCEnabled();
 
 	/** Check if it is currently running */
 	static inline bool IsODSCActive();
