Add lock to prevent race condition in allocation/free of virtual texture upload tiles.
#rnx

diff --git a/Engine/Source/Runtime/Engine/Private/VT/VirtualTextureUploadCache.cpp b/Engine/Source/Runtime/Engine/Private/VT/VirtualTextureUploadCache.cpp
--- a/Engine/Source/Runtime/Engine/Private/VT/VirtualTextureUploadCache.cpp
+++ b/Engine/Source/Runtime/Engine/Private/VT/VirtualTextureUploadCache.cpp
@@ -41,6 +41,8 @@ static TAutoConsoleVariable<int32> CVarMaxUploadRequests(
 
 uint32 FVTUploadTileAllocator::Allocate(FRHICommandListBase& RHICmdList, EVTUploadType InUploadType, EPixelFormat InFormat, uint32 InTileSize)
 {
+	UE::TScopeLock Lock(Mutex);
+
 	// Find matching FormatBuffer.
 	const FPixelFormatInfo& FormatInfo = GPixelFormats[InFormat];
 	const uint32 TileWidthInBlocks = FMath::DivideAndRoundUp(InTileSize, (uint32)FormatInfo.BlockSizeX);
@@ -112,6 +114,8 @@ uint32 FVTUploadTileAllocator::Allocate(FRHICommandListBase& RHICmdList, EVTUplo
 
 void FVTUploadTileAllocator::Free(FRHICommandListBase& RHICmdList, uint32 InHandle)
 {
+	UE::TScopeLock Lock(Mutex);
+
 	FHandle Handle;
 	Handle.PackedValue = InHandle;
 
diff --git a/Engine/Source/Runtime/Engine/Private/VT/VirtualTextureUploadCache.h b/Engine/Source/Runtime/Engine/Private/VT/VirtualTextureUploadCache.h
--- a/Engine/Source/Runtime/Engine/Private/VT/VirtualTextureUploadCache.h
+++ b/Engine/Source/Runtime/Engine/Private/VT/VirtualTextureUploadCache.h
@@ -78,6 +78,9 @@ public:
 	uint32 TotalAllocatedBytes() const { return NumAllocatedBytes; }
 
 private:
+	/** Mutex for Allocate/Free. */
+	UE::FMutex Mutex;
+
 	/** Handle for an allocated tile used by the allocation system. */
 	union FHandle
 	{
