Make r.Mobile.AllowDitheredLODTransition per ShaderPlatform and not per project

diff --git a/Engine/Source/Runtime/Engine/Private/Components/StaticMeshComponent.cpp b/Engine/Source/Runtime/Engine/Private/Components/StaticMeshComponent.cpp
--- a/Engine/Source/Runtime/Engine/Private/Components/StaticMeshComponent.cpp
+++ b/Engine/Source/Runtime/Engine/Private/Components/StaticMeshComponent.cpp
@@ -2057,7 +2057,7 @@ bool UStaticMeshComponent::SupportsDefaultCollision()
 
 bool UStaticMeshComponent::SupportsDitheredLODTransitions(ERHIFeatureLevel::Type FeatureLevel)
 {
-	if (!AllowDitheredLODTransition(FeatureLevel))
+	if (!AllowDitheredLODTransition(GShaderPlatformForFeatureLevel[FeatureLevel]))
 	{
 		return false;
 	}
diff --git a/Engine/Source/Runtime/Engine/Private/Materials/MaterialShared.cpp b/Engine/Source/Runtime/Engine/Private/Materials/MaterialShared.cpp
--- a/Engine/Source/Runtime/Engine/Private/Materials/MaterialShared.cpp
+++ b/Engine/Source/Runtime/Engine/Private/Materials/MaterialShared.cpp
@@ -205,13 +205,20 @@ static FMaterialDumpDebugInfoExecHelper GMaterialDumpDebugInfoExecHelper;
 // defined in the same module (Material.cpp)
 bool PoolSpecialMaterialsCompileJobs();
 
+// UE_DEPRECATED(5.7, "Please use AllowDitheredLODTransition with a EShaderPlatform argument")
 bool AllowDitheredLODTransition(ERHIFeatureLevel::Type FeatureLevel)
+{
+	EShaderPlatform ShaderPlatform = GShaderPlatformForFeatureLevel[FeatureLevel];
+	return AllowDitheredLODTransition(ShaderPlatform);
+}
+
+bool AllowDitheredLODTransition(EShaderPlatform ShaderPlatform)
 {
 	// On mobile support for 'Dithered LOD Transition' has to be explicitly enabled in projects settings
-	if (FeatureLevel <= ERHIFeatureLevel::ES3_1)
+	if (IsMobilePlatform(ShaderPlatform))
 	{
-		static TConsoleVariableData<int32>* CVar = IConsoleManager::Get().FindTConsoleVariableDataInt(TEXT("r.Mobile.AllowDitheredLODTransition"));
-		return (CVar && CVar->GetValueOnAnyThread() != 0) ? true : false;
+		static FShaderPlatformCachedIniValue<int32> CVarAllowDitheredLODTransition(TEXT("r.Mobile.AllowDitheredLODTransition"));
+		return (CVarAllowDitheredLODTransition.Get(ShaderPlatform) != 0) ? true : false;
 	}
 	return true;
 }
@@ -1933,7 +1940,7 @@ bool FMaterialResource::IsThinSurface() const
 
 bool FMaterialResource::IsDitheredLODTransition() const 
 {
-	if (!AllowDitheredLODTransition(GetFeatureLevel()))
+	if (!AllowDitheredLODTransition(GShaderPlatformForFeatureLevel[GetFeatureLevel()]))
 	{
 		return false;
 	}
diff --git a/Engine/Source/Runtime/Engine/Public/MaterialShared.h b/Engine/Source/Runtime/Engine/Public/MaterialShared.h
--- a/Engine/Source/Runtime/Engine/Public/MaterialShared.h
+++ b/Engine/Source/Runtime/Engine/Public/MaterialShared.h
@@ -137,9 +137,12 @@ inline uint32 GetUseSubsurfaceProfileShadingModelMask()
 	return (1 << MSM_SubsurfaceProfile) | (1 << MSM_Eye);
 }
 
-/** Whether to allow dithered LOD transitions for a specific feature level. */
+UE_DEPRECATED(5.7, "Please use AllowDitheredLODTransition with a EShaderPlatform argument")
 ENGINE_API bool AllowDitheredLODTransition(ERHIFeatureLevel::Type FeatureLevel);
 
+/** Whether to allow dithered LOD transitions for a specific shader platform. */
+ENGINE_API bool AllowDitheredLODTransition(EShaderPlatform ShaderPlatform);
+
 /** Blend modes utility functions */
 ENGINE_API bool IsOpaqueBlendMode(EBlendMode BlendMode);
 ENGINE_API bool IsOpaqueBlendMode(const FMaterial& In);
diff --git a/Engine/Source/Runtime/RenderCore/Private/Shader.cpp b/Engine/Source/Runtime/RenderCore/Private/Shader.cpp
--- a/Engine/Source/Runtime/RenderCore/Private/Shader.cpp
+++ b/Engine/Source/Runtime/RenderCore/Private/Shader.cpp
@@ -1974,8 +1974,8 @@ void ShaderMapAppendKey(EShaderPlatform Platform, FShaderKeyGenerator& KeyGen)
 		}
 
 		{
-			static IConsoleVariable* CVar = IConsoleManager::Get().FindConsoleVariable(TEXT("r.Mobile.AllowDitheredLODTransition"));
-			if (CVar && CVar->GetInt() != 0)
+			static FShaderPlatformCachedIniValue<int32> CVarAllowDitheredLODTransition(TEXT("r.Mobile.AllowDitheredLODTransition"));
+			if (CVarAllowDitheredLODTransition.Get(Platform) != 0)
 			{
 				KeyGen.AppendSeparator();
 				KeyGen.Append(TEXT("DLODT"));
