[MIR] Fix VT support for gradient sampling and generalize emitting floating-point values.

#rnx

diff --git a/Engine/Source/Runtime/Engine/Private/Materials/MaterialIRToHLSLTranslator.cpp b/Engine/Source/Runtime/Engine/Private/Materials/MaterialIRToHLSLTranslator.cpp
--- a/Engine/Source/Runtime/Engine/Private/Materials/MaterialIRToHLSLTranslator.cpp
+++ b/Engine/Source/Runtime/Engine/Private/Materials/MaterialIRToHLSLTranslator.cpp
@@ -64,6 +64,29 @@ struct FHLSLPrinter
 		return *this;
 	}
 
+	FHLSLPrinter& operator<<(uint32 Value)
+	{
+		Buffer.Appendf(TEXT("%u"), Value);
+		return *this;
+	}
+
+	FHLSLPrinter& operator<<(float Value)
+	{
+		if (FGenericPlatformMath::IsNaN(Value))
+		{
+			Buffer.Append(TEXT("(0.0f / 0.0f)"));
+		}
+		else if (!FGenericPlatformMath::IsFinite(Value))
+		{
+			Buffer.Append(TEXT("INFINITE_FLOAT"));
+		}
+		else
+		{
+			Buffer.Appendf(TEXT("%#.8gf"), Value);
+		}
+		return *this;
+	}
+
     FHLSLPrinter& operator<<(ENoOp)
 	{
 		return *this;
@@ -440,18 +463,7 @@ struct FTranslator : FMaterialIRToHLSLTranslation
 				break;
 			
 			case MIR::ScalarKind_Float:
-				if (FGenericPlatformMath::IsNaN(Constant->Float))
-				{
-					Printer.Buffer.Append(TEXT("(0.0f / 0.0f)"));
-				}
-				else if (!FGenericPlatformMath::IsFinite(Constant->Float))
-				{
-					Printer.Buffer.Append(TEXT("INFINITE_FLOAT"));
-				}
-				else
-				{
-					Printer.Buffer.Appendf(TEXT("%#.8gf"), Constant->Float);
-				}
+				Printer << Constant->Float;
 				break;
 
 			case MIR::ScalarKind_LWC: UE_MIR_UNREACHABLE();
@@ -1325,6 +1337,7 @@ struct FTranslator : FMaterialIRToHLSLTranslation
 		}
 		else
 		{
+			constexpr float MipBias = 0.0f;
 			Printer
 				<< VTPageTableReadFunctionName
 				<< BeginArgs
@@ -1333,22 +1346,24 @@ struct FTranslator : FMaterialIRToHLSLTranslation
 				<< ListSeparator << LowerValue(VTPageTableRead->TexCoord)
 				<< ListSeparator << GetVTAddressMode(VTPageTableRead->AddressU)
 				<< ListSeparator << GetVTAddressMode(VTPageTableRead->AddressV)
+				<< ListSeparator << MipBias
 				;
 		}
 
+		Printer << ListSeparator << TEXT("Parameters.SvPosition.xy");
+
 		// VT feedback
 		const bool bGenerateFeedback = VTPageTableRead->VTFeedbackRequestIndex != INDEX_NONE;
 
 		if (bGenerateFeedback)
 		{
 			Printer
-				<< ListSeparator << 0
-				<< ListSeparator << TEXT("Parameters.SvPosition.xy")
 				<< ListSeparator << VTPageTableRead->VTFeedbackRequestIndex << TEXT("u + LIGHTMAP_VT_ENABLED")
 				<< ListSeparator << TEXT("Parameters.VirtualTextureFeedback")
-				<< EndArgs
 				;
 		}
+
+		Printer << EndArgs;
 	}
 
 	void LowerSamplerType(EMaterialSamplerType SamplerType, bool& bOutSamplerNeedsBrackets)
