Make debug logging for the new layer tool so it doesn't spam the log.

#rnx

diff --git a/Engine/Source/Runtime/Engine/Private/Materials/MaterialExpressions.cpp b/Engine/Source/Runtime/Engine/Private/Materials/MaterialExpressions.cpp
--- a/Engine/Source/Runtime/Engine/Private/Materials/MaterialExpressions.cpp
+++ b/Engine/Source/Runtime/Engine/Private/Materials/MaterialExpressions.cpp
@@ -16163,7 +16163,7 @@ namespace MLFT {
 		LogMessage.ParseIntoArray(Lines, TEXT("\n"));
 		for (const FString& Line : Lines)
 		{
-			UE_LOG(LogMaterial, Log, TEXT("%s"), *Line);
+			UE_LOG(LogMaterial, VeryVerbose, TEXT("%s"), *Line);
 		}
 	}
 
