Splines - FSpline Thread Safety - Make OnSplineImplementationChanged thread safe so that it is safe to write to it on FSpline construction/destruction

diff --git a/Engine/Source/Runtime/Engine/Private/Curves/Spline.cpp b/Engine/Source/Runtime/Engine/Private/Curves/Spline.cpp
--- a/Engine/Source/Runtime/Engine/Private/Curves/Spline.cpp
+++ b/Engine/Source/Runtime/Engine/Private/Curves/Spline.cpp
@@ -20,9 +20,9 @@ namespace UE::Spline
 
 int32 GImplementation = 0;
 #if WITH_EDITOR
-static FSimpleMulticastDelegate& GetOnSplineImplementationChanged()
+static FTSSimpleMulticastDelegate& GetOnSplineImplementationChanged()
 {
-	static FSimpleMulticastDelegate OnSplineImplementationChanged;
+	static FTSSimpleMulticastDelegate OnSplineImplementationChanged;
 	return OnSplineImplementationChanged;
 };
 static void SplineImplementationSink(IConsoleVariable*)
