Material Shading Path Node can be forced to follow a specific path using the console variable r.Material.ShadingPathNodeOverride (disabled by default).

#rnx

diff --git a/Engine/Source/Runtime/Core/Private/HAL/ConsoleManager.cpp b/Engine/Source/Runtime/Core/Private/HAL/ConsoleManager.cpp
--- a/Engine/Source/Runtime/Core/Private/HAL/ConsoleManager.cpp
+++ b/Engine/Source/Runtime/Core/Private/HAL/ConsoleManager.cpp
@@ -3938,6 +3938,13 @@ static TAutoConsoleVariable<int32> CVarMaterialQualityLevel(
 	TEXT("0 corresponds to low quality materials, as defined by quality switches in materials, 1 corresponds to high, 2 for medium, and 3 for Epic."),
 	ECVF_Scalability | ECVF_RenderThreadSafe);
 
+static TAutoConsoleVariable<int32> CVarMaterialShadingPathNodeOverride(
+	TEXT("r.Material.ShadingPathNodeOverride"),
+	-1,
+	TEXT("Material Shading Path Node Override. -1 = no override (default value), 0 = deferred rendering, 1 = forward, 2 = mobile)"),
+	ECVF_ReadOnly | ECVF_Preview | ECVF_RenderThreadSafe
+);
+
 static TAutoConsoleVariable<int32> CVarUseDXT5NormalMaps(
 	TEXT("Compat.UseDXT5NormalMaps"),
 	0,
diff --git a/Engine/Source/Runtime/Engine/Private/Materials/MaterialExpressions.cpp b/Engine/Source/Runtime/Engine/Private/Materials/MaterialExpressions.cpp
--- a/Engine/Source/Runtime/Engine/Private/Materials/MaterialExpressions.cpp
+++ b/Engine/Source/Runtime/Engine/Private/Materials/MaterialExpressions.cpp
@@ -9855,13 +9855,23 @@ FExpressionInput* UMaterialExpressionShadingPathSwitch::GetEffectiveInput(class
 	const EShaderPlatform ShaderPlatform = Compiler->GetShaderPlatform();
 	ERHIShadingPath::Type ShadingPathToCompile = ERHIShadingPath::Deferred;
 
-	if (IsForwardShadingEnabled(ShaderPlatform))
+	static_assert(int(ERHIShadingPath::Deferred) == 0);
+
+	int32 ShadingPathNodeOverride = GetMaterialShadingPathNodeOverride(ShaderPlatform);
+	if (ShadingPathNodeOverride != -1)
 	{
-		ShadingPathToCompile = ERHIShadingPath::Forward;
+		ShadingPathToCompile = ERHIShadingPath::Type(ShadingPathNodeOverride);
 	}
-	else if (Compiler->GetFeatureLevel() < ERHIFeatureLevel::SM5)
+	else
 	{
-		ShadingPathToCompile = ERHIShadingPath::Mobile;
+		if (IsForwardShadingEnabled(ShaderPlatform))
+		{
+			ShadingPathToCompile = ERHIShadingPath::Forward;
+		}
+		else if (Compiler->GetFeatureLevel() < ERHIFeatureLevel::SM5)
+		{
+			ShadingPathToCompile = ERHIShadingPath::Mobile;
+		}
 	}
 
 	check(ShadingPathToCompile < UE_ARRAY_COUNT(Inputs));
diff --git a/Engine/Source/Runtime/RenderCore/Private/RenderUtils.cpp b/Engine/Source/Runtime/RenderCore/Private/RenderUtils.cpp
--- a/Engine/Source/Runtime/RenderCore/Private/RenderUtils.cpp
+++ b/Engine/Source/Runtime/RenderCore/Private/RenderUtils.cpp
@@ -1571,6 +1571,12 @@ bool IsMobileDistanceFieldAOEnabled(const FStaticShaderPlatform Platform)
 	return MobileDistanceFieldAO.Get(Platform) > 0;
 }
 
+int32 GetMaterialShadingPathNodeOverride(const FStaticShaderPlatform Platform)
+{
+	static FShaderPlatformCachedIniValue<int32> MaterialShadingPathNodeOverride(TEXT("r.Material.ShadingPathNodeOverride"));
+	return MaterialShadingPathNodeOverride.Get(Platform);
+}
+
 bool UseGPUScene(const FStaticShaderPlatform Platform, const FStaticFeatureLevel FeatureLevel)
 {
 	if (FeatureLevel == ERHIFeatureLevel::ES3_1)
diff --git a/Engine/Source/Runtime/RenderCore/Private/Shader.cpp b/Engine/Source/Runtime/RenderCore/Private/Shader.cpp
--- a/Engine/Source/Runtime/RenderCore/Private/Shader.cpp
+++ b/Engine/Source/Runtime/RenderCore/Private/Shader.cpp
@@ -2573,6 +2573,26 @@ void ShaderMapAppendKey(EShaderPlatform Platform, FShaderKeyGenerator& KeyGen)
 		}
 	}
 
+	{
+		int32 MaterialShadingPathOverride = GetMaterialShadingPathNodeOverride(Platform);
+		if (MaterialShadingPathOverride != -1)
+		{
+			KeyGen.AppendSeparator();
+			switch (ERHIShadingPath::Type(MaterialShadingPathOverride))
+			{
+				case ERHIShadingPath::Deferred:
+					KeyGen.Append(TEXT("MATSPD"));
+					break;
+				case ERHIShadingPath::Forward:
+					KeyGen.Append(TEXT("MATSPF"));
+					break;
+				case ERHIShadingPath::Mobile:
+					KeyGen.Append(TEXT("MATSPM"));
+					break;
+			}
+		}
+	}
+
 	{
 		if (MaskedInEarlyPass(Platform))
 		{
diff --git a/Engine/Source/Runtime/RenderCore/Public/RenderUtils.h b/Engine/Source/Runtime/RenderCore/Public/RenderUtils.h
--- a/Engine/Source/Runtime/RenderCore/Public/RenderUtils.h
+++ b/Engine/Source/Runtime/RenderCore/Public/RenderUtils.h
@@ -364,6 +364,9 @@ RENDERCORE_API bool IsWaterSeparateMainDirLightEnabled(const FStaticShaderPlatfo
 /** Returns if Mobile Distance Field AO Enabled. */
 RENDERCORE_API bool IsMobileDistanceFieldAOEnabled(const FStaticShaderPlatform Platform);
 
+/** Returns shading path node material override. Returns -1 if no override. */
+RENDERCORE_API int32 GetMaterialShadingPathNodeOverride(const FStaticShaderPlatform Platform);
+
 RENDERCORE_API bool UseGPUScene(const FStaticShaderPlatform Platform, const FStaticFeatureLevel FeatureLevel);
 
 RENDERCORE_API bool UseGPUScene(const FStaticShaderPlatform Platform);
