[NewMaterialTranslator] - Adding Reciprocal and RSqrt to the material IR.

- This is in prepration for implementing `normalize` (which will use the derivative of rsqrt to caluate its derivative)
- These don't have material expressions associated with them so can only be used in C++ when targetting the IR, or with the new `UMaterialExpressionOperator` type.
- Both these operators validate to see if constants args are non-zero to avoid divide by zero (OPTF_CheckIsNonZeroFloatConst)
- Implementation of the derivatives are based on what is in HLSLMaterialDerivativeAutogen.cpp.

Derivative of rcp:
```
Ret.Value = rcp(A.Value)
dFdA = -Ret.Value * Ret.Value
```

Derivative of Rsqrt:
```
dFdA = -0.5f * rsqrt(A.Value) * rcp(A.Value)
```

#rnx

diff --git a/Engine/Source/Runtime/Engine/Private/Materials/MaterialIR.cpp b/Engine/Source/Runtime/Engine/Private/Materials/MaterialIR.cpp
--- a/Engine/Source/Runtime/Engine/Private/Materials/MaterialIR.cpp
+++ b/Engine/Source/Runtime/Engine/Private/Materials/MaterialIR.cpp
@@ -723,7 +723,9 @@ const TCHAR* LexToString(EOperator Op)
 		case UO_Logarithm2: return TEXT("Logarithm2");
 		case UO_Negate: return TEXT("Negate");
 		case UO_Not: return TEXT("Not");
+		case UO_Reciprocal: return TEXT("Reciprocal");
 		case UO_Round: return TEXT("Round");
+		case UO_Rsqrt: return TEXT("Rsqrt");
 		case UO_Saturate: return TEXT("Saturate");
 		case UO_Sign: return TEXT("Sign");
 		case UO_Sin: return TEXT("Sin");
diff --git a/Engine/Source/Runtime/Engine/Private/Materials/MaterialIREmitter.cpp b/Engine/Source/Runtime/Engine/Private/Materials/MaterialIREmitter.cpp
--- a/Engine/Source/Runtime/Engine/Private/Materials/MaterialIREmitter.cpp
+++ b/Engine/Source/Runtime/Engine/Private/Materials/MaterialIREmitter.cpp
@@ -1179,6 +1179,7 @@ static T FoldScalarOperator(FEmitter& Emitter, EOperator Operator, T A, T B, T C
 					: T(10);
 				return FGenericPlatformMath::LogX(Base, A);
 			}
+			case UO_Reciprocal: return 1 / A;
 			case UO_Round:
 				if constexpr (std::is_same<T, float>::value)
 				{
@@ -1188,6 +1189,7 @@ static T FoldScalarOperator(FEmitter& Emitter, EOperator Operator, T A, T B, T C
 				{
 					return FGenericPlatformMath::RoundToDouble(A);
 				}
+			case UO_Rsqrt: return FGenericPlatformMath::InvSqrt(A);
 			case UO_Sin: return FGenericPlatformMath::Sin(A);
 			case UO_Sinh: return FGenericPlatformMath::Sinh(A);
 			case UO_Sqrt: return FGenericPlatformMath::Sqrt(A);
@@ -1582,7 +1584,9 @@ static const FOperatorSignature* GetOperatorSignature(EOperator Op)
 		S[UO_Logarithm] 				= { { OPTF_CheckIsArithmetic | OPTF_CheckIsNonZeroFloatConst | OPTF_CheckIsNonNegativeFloatConst | OPTF_CastToAnyFloat }, ORT_FirstArgumentType };
 		S[UO_Logarithm10]				= { { OPTF_CheckIsArithmetic | OPTF_CheckIsNonZeroFloatConst | OPTF_CheckIsNonNegativeFloatConst | OPTF_CastToAnyFloat }, ORT_FirstArgumentType };
 		S[UO_Logarithm2] 				= { { OPTF_CheckIsArithmetic | OPTF_CheckIsNonZeroFloatConst | OPTF_CheckIsNonNegativeFloatConst | OPTF_CastToAnyFloat }, ORT_FirstArgumentType };
+		S[UO_Reciprocal]				= { { OPTF_CheckIsArithmetic | OPTF_CheckIsNonZeroFloatConst | OPTF_CastToAnyFloat }, ORT_FirstArgumentType };
 		S[UO_Round] 					= UnaryFloatLWC;
+		S[UO_Rsqrt]						= { { OPTF_CheckIsArithmetic | OPTF_CheckIsNonZeroFloatConst | OPTF_CastToAnyFloat }, ORT_FirstArgumentType };
 		S[UO_Saturate]					= UnaryFloatLWCDemote;
 		S[UO_Sign]						= UnaryFloatLWCDemote;
 		S[UO_Sin] 						= UnaryFloatLWCDemote;
@@ -2921,6 +2925,17 @@ static FValue* DifferentiateOperator(FEmitter& E, FOperator* Op, EDerivativeAxis
 		case UO_Sqrt:
 			return E.Divide(dF, E.Multiply(Constant(2), E.Sqrt(F)));
 
+		// d/dx rcp(f(x)) = -1 / (f(x)^2) * f'(x) 
+		case UO_Reciprocal:
+		{
+			FValueRef rcp = E.Reciprocal(F);
+			return E.Multiply(E.Multiply(E.Negate(rcp), rcp), dF);
+		}
+
+		// d/dx rsqrt(f(x)) = -1 / (2 * sqrt(f(x)) * f(x)) * f'(x)
+		case UO_Rsqrt:
+			return E.Multiply(E.Multiply(E.ConstantFloat(-0.5f), E.Multiply(E.Rsqrt(F), E.Reciprocal(F))), dF);
+
 		// d/dx tan(f(x)) = 1 / cos^2(f(x)) * f'(x)
 		case UO_Tan:
 		{
diff --git a/Engine/Source/Runtime/Engine/Private/Materials/MaterialIRToHLSLTranslator.cpp b/Engine/Source/Runtime/Engine/Private/Materials/MaterialIRToHLSLTranslator.cpp
--- a/Engine/Source/Runtime/Engine/Private/Materials/MaterialIRToHLSLTranslator.cpp
+++ b/Engine/Source/Runtime/Engine/Private/Materials/MaterialIRToHLSLTranslator.cpp
@@ -1064,7 +1064,9 @@ struct FTranslator : FMaterialIRToHLSLTranslation
 					case MIR::UO_Logarithm: OpString = TEXT("log"); break;
 					case MIR::UO_Logarithm10: OpString = TEXT("log10"); break;
 					case MIR::UO_Logarithm2: OpString = TEXT("log2"); break;
+					case MIR::UO_Reciprocal: OpString = TEXT("rcp"); break;
 					case MIR::UO_Round: OpString = TEXT("round"); break;
+					case MIR::UO_Rsqrt: OpString = TEXT("rsqrt"); break;
 					case MIR::UO_Saturate: OpString = TEXT("saturate"); break;
 					case MIR::UO_Sign: OpString = TEXT("sign"); break;
 					case MIR::UO_Sin: OpString = TEXT("sin"); break;
diff --git a/Engine/Source/Runtime/Engine/Public/Materials/MaterialExpressionOperator.h b/Engine/Source/Runtime/Engine/Public/Materials/MaterialExpressionOperator.h
--- a/Engine/Source/Runtime/Engine/Public/Materials/MaterialExpressionOperator.h
+++ b/Engine/Source/Runtime/Engine/Public/Materials/MaterialExpressionOperator.h
@@ -37,7 +37,9 @@ enum class EMaterialExpressionOperatorKind
 	Logarithm,
 	Logarithm10,
 	Logarithm2,
+	Reciprocal,
 	Round,
+	Rsqrt,
 	Saturate,
 	Sign,
 	Sin,
diff --git a/Engine/Source/Runtime/Engine/Public/Materials/MaterialIR.h b/Engine/Source/Runtime/Engine/Public/Materials/MaterialIR.h
--- a/Engine/Source/Runtime/Engine/Public/Materials/MaterialIR.h
+++ b/Engine/Source/Runtime/Engine/Public/Materials/MaterialIR.h
@@ -643,7 +643,9 @@ enum EOperator
 	UO_Logarithm,
 	UO_Logarithm10,
 	UO_Logarithm2,
+	UO_Reciprocal,
 	UO_Round,
+	UO_Rsqrt,
 	UO_Saturate,
 	UO_Sign,
 	UO_Sin,
diff --git a/Engine/Source/Runtime/Engine/Public/Materials/MaterialIREmitter.h b/Engine/Source/Runtime/Engine/Public/Materials/MaterialIREmitter.h
--- a/Engine/Source/Runtime/Engine/Public/Materials/MaterialIREmitter.h
+++ b/Engine/Source/Runtime/Engine/Public/Materials/MaterialIREmitter.h
@@ -531,7 +531,9 @@ public:
 	FValueRef Logarithm(FValueRef A) { return Operator(UO_Logarithm, A); }
 	FValueRef Logarithm2(FValueRef A) { return Operator(UO_Logarithm2, A); }
 	FValueRef Logarithm10(FValueRef A) { return Operator(UO_Logarithm10, A); }
+	FValueRef Reciprocal(FValueRef A) { return Operator(UO_Reciprocal, A); }
 	FValueRef Round(FValueRef A) { return Operator(UO_Round, A); }
+	FValueRef Rsqrt(FValueRef A) { return Operator(UO_Rsqrt, A); }
 	FValueRef Saturate(FValueRef A) { return Operator(UO_Saturate, A); }
 	FValueRef Sin(FValueRef A) { return Operator(UO_Sin, A); }
 	FValueRef Sinh(FValueRef A) { return Operator(UO_Sinh, A); }
