Remove unnecessary copying of audio data.

#tests RR, PIE

diff --git a/Engine/Source/Programs/Shared/EpicGames.UHT/Exporters/CodeGen/UhtHeaderCodeGeneratorCppFile.cs b/Engine/Source/Programs/Shared/EpicGames.UHT/Exporters/CodeGen/UhtHeaderCodeGeneratorCppFile.cs
--- a/Engine/Source/Programs/Shared/EpicGames.UHT/Exporters/CodeGen/UhtHeaderCodeGeneratorCppFile.cs
+++ b/Engine/Source/Programs/Shared/EpicGames.UHT/Exporters/CodeGen/UhtHeaderCodeGeneratorCppFile.cs
@@ -1717,7 +1717,7 @@ namespace EpicGames.UHT.Exporters.CodeGen
 			{
 				if (parameter is UhtProperty property)
 				{
-					builder.Append("\t\tParms.").Append(property.SourceName).Append('=').Append(property.SourceName).Append(";\r\n");
+					builder.Append("\t\tParms.").Append(property.SourceName).Append("=std::move(").Append(property.SourceName).Append(");\r\n");
 				}
 			}
 			builder
@@ -1732,7 +1732,7 @@ namespace EpicGames.UHT.Exporters.CodeGen
 				{
 					if (property.PropertyFlags.HasExactFlags(EPropertyFlags.OutParm | EPropertyFlags.ConstParm, EPropertyFlags.OutParm))
 					{
-						builder.Append("\t\t").Append(property.SourceName).Append("=Parms.").Append(property.SourceName).Append(";\r\n");
+						builder.Append("\t\t").Append(property.SourceName).Append("=std::move(Parms.").Append(property.SourceName).Append(");\r\n");
 					}
 				}
 			}
diff --git a/Engine/Source/Runtime/Engine/Private/Audio/ActorSoundParameterInterface.cpp b/Engine/Source/Runtime/Engine/Private/Audio/ActorSoundParameterInterface.cpp
--- a/Engine/Source/Runtime/Engine/Private/Audio/ActorSoundParameterInterface.cpp
+++ b/Engine/Source/Runtime/Engine/Private/Audio/ActorSoundParameterInterface.cpp
@@ -23,36 +23,32 @@ void UActorSoundParameterInterface::Fill(const AActor* OwningActor, TArray<FAudi
 {
 	QUICK_SCOPE_CYCLE_COUNTER(UActorSoundParameterInterface_Fill);
 
-	TArray<const AActor*> Actors;
-	TArray<const UActorComponent*> Components;
-
-	if (ActorSoundParameterInterfaceConsoleVariables::bGatherImplementers)
-	{
-		// This is prohibitively expensive, as it goes through our owning actor and all attached actors, looking for 
-		// components that implement the IActorSoundParameterInterface
-		GetImplementers(OwningActor, Actors, Components);
-	}
-	else
+	if (!ActorSoundParameterInterfaceConsoleVariables::bGatherImplementers)
 	{
 		// Prior to the GetImplementers change, we only considered the owning actor, and no attached actors or components.
 		if (OwningActor && OwningActor->Implements<UActorSoundParameterInterface>())
 		{
-			Actors.Add(OwningActor);
+			IActorSoundParameterInterface::Execute_GetActorSoundParams(OwningActor, OutParams);
 		}
+		return;
 	}
-
-	TArray<FAudioParameter> TempParams;
-
-	for (const AActor* Actor : Actors)
+	else
 	{
-		IActorSoundParameterInterface::Execute_GetActorSoundParams(Actor, TempParams);
-		OutParams.Append(MoveTemp(TempParams));
-	}
+		TArray<const AActor*> Actors;
+		TArray<const UActorComponent*> Components;
+		// This is prohibitively expensive, as it goes through our owning actor and all attached actors, looking for 
+		// components that implement the IActorSoundParameterInterface
+		GetImplementers(OwningActor, Actors, Components);
 
-	for (const UActorComponent* Component : Components)
-	{
-		IActorSoundParameterInterface::Execute_GetActorSoundParams(Component, TempParams);
-		OutParams.Append(MoveTemp(TempParams));
+		for (const AActor* Actor : Actors)
+		{
+			IActorSoundParameterInterface::Execute_GetActorSoundParams(Actor, OutParams);
+		}
+
+		for (const UActorComponent* Component : Components)
+		{
+			IActorSoundParameterInterface::Execute_GetActorSoundParams(Component, OutParams);
+		}
 	}
 }
 
diff --git a/Engine/Source/Runtime/Engine/Private/AudioDevice.cpp b/Engine/Source/Runtime/Engine/Private/AudioDevice.cpp
--- a/Engine/Source/Runtime/Engine/Private/AudioDevice.cpp
+++ b/Engine/Source/Runtime/Engine/Private/AudioDevice.cpp
@@ -5248,18 +5248,48 @@ void FAudioDevice::AddNewActiveSound(const FActiveSound& ActiveSound, const TArr
 	{
 		Params.Append(*InDefaultParams);
 	}
+	//Previously the copy was done in AddNewActiveSoundInternal when creating the Lamba function. 
+	TSharedRef<FActiveSound> NewUniqueActiveSound = MakeShared<FActiveSound>(ActiveSound);
 
-	AddNewActiveSound(ActiveSound, MoveTemp(Params));
+	AddNewActiveSound(NewUniqueActiveSound, MoveTemp(Params));
 }
 
-void FAudioDevice::AddNewActiveSound(const FActiveSound& NewActiveSound, TArray<FAudioParameter>&& InDefaultParams)
+void FAudioDevice::AddNewActiveSound(TSharedRef<FActiveSound> ActiveSound, TArray<FAudioParameter>&& InDefaultParams)
 {
-	if (USoundBase* Sound = NewActiveSound.GetSound())
+	if (USoundBase* Sound = ActiveSound.Get().GetSound())
 	{
 		Sound->InitResources();
 	}
+	AddNewActiveSoundInternal(ActiveSound, MoveTemp(InDefaultParams));
+
+}
+
+void FAudioDevice::AddNewActiveSound(const FActiveSound& NewActiveSound, TArray<FAudioParameter>&& InDefaultParams)
+{
+	//Previously the copy was done in AddNewActiveSoundInternal when creating the Lamba function. 
+	TSharedRef<FActiveSound> NewUniqueActiveSound = MakeShared<FActiveSound>(NewActiveSound);
+
+	AddNewActiveSound(NewUniqueActiveSound, MoveTemp(InDefaultParams));
+}
+
+void FAudioDevice::AddNewActiveSoundInternal(TSharedRef<FActiveSound> NewActiveSound, TArray<FAudioParameter>&& InDefaultParams, FAudioVirtualLoop* InVirtualLoopToRetrigger)
+{
+	if (!IsInAudioThread())
+	{
+		DECLARE_CYCLE_STAT(TEXT("FAudioThreadTask.AddNewActiveSound"), STAT_AudioAddNewActiveSound, STATGROUP_AudioThreadCommands);
+
+		checkf(nullptr == InVirtualLoopToRetrigger, TEXT("Virtual loops can only be retrigger from the audio thread"));
+
+		FAudioThread::RunCommandOnAudioThread([AudioDevice = this, NewActiveSound = MoveTemp(NewActiveSound), DefaultParams = MoveTemp(InDefaultParams)]() mutable
+			{
+				AudioDevice->AddNewActiveSoundInternal(NewActiveSound, MoveTemp(DefaultParams));
+			}, GET_STATID(STAT_AudioAddNewActiveSound));
+
+		return;
+	}
+
+	AddNewActiveSoundInternal(NewActiveSound.Get(), MoveTemp(InDefaultParams), InVirtualLoopToRetrigger);
 
-	AddNewActiveSoundInternal(NewActiveSound, MoveTemp(InDefaultParams));
 }
 
 void FAudioDevice::AddNewActiveSoundInternal(const FActiveSound& InNewActiveSound, TArray<FAudioParameter>&& InDefaultParams, FAudioVirtualLoop* InVirtualLoopToRetrigger)
@@ -6613,7 +6643,7 @@ UAudioComponent* FAudioDevice::CreateComponent(USoundBase* Sound, const FCreateC
 	return AudioComponent;
 }
 
-void FAudioDevice::PlaySoundAtLocation(USoundBase* Sound, UWorld* World, float VolumeMultiplier, float PitchMultiplier, float StartTime, const FVector& Location, const FRotator& Rotation, USoundAttenuation* AttenuationSettings, USoundConcurrency* Concurrency, const TArray<FAudioParameter>* Params, const AActor* OwningActor)
+void FAudioDevice::PlaySoundAtLocation(USoundBase* Sound, UWorld* World, float VolumeMultiplier, float PitchMultiplier, float StartTime, const FVector& Location, const FRotator& Rotation, USoundAttenuation* AttenuationSettings, USoundConcurrency* Concurrency, TArray<FAudioParameter>&& Params, const AActor* OwningActor)
 {
 	check(IsInGameThread());
 
@@ -6637,8 +6667,8 @@ void FAudioDevice::PlaySoundAtLocation(USoundBase* Sound, UWorld* World, float V
 	if (Sound->IsLooping() || Sound->IsPlayWhenSilent() || SoundIsAudible(Sound, World, Location, AttenuationSettingsToApply, MaxDistance, FocusFactor))
 	{
 		const bool bIsInGameWorld = World->IsGameWorld();
-
-		FActiveSound NewActiveSound;
+		TSharedRef<FActiveSound> NewSharedActiveSound = MakeShared<FActiveSound>();
+		FActiveSound& NewActiveSound = NewSharedActiveSound.Get();
 		NewActiveSound.SetWorld(World);
 		NewActiveSound.SetSound(Sound);
 		NewActiveSound.SetVolume(VolumeMultiplier);
@@ -6674,8 +6704,7 @@ void FAudioDevice::PlaySoundAtLocation(USoundBase* Sound, UWorld* World, float V
 		{
 			NotifySubsystemsSetOwner(NewActiveSound, OwningActor);
 		}
-
-		AddNewActiveSound(NewActiveSound, Params);
+		AddNewActiveSound(MoveTemp(NewSharedActiveSound), MoveTemp(Params));
 	}
 	else
 	{
@@ -6684,6 +6713,20 @@ void FAudioDevice::PlaySoundAtLocation(USoundBase* Sound, UWorld* World, float V
 	}
 }
 
+
+void FAudioDevice::PlaySoundAtLocation(USoundBase* Sound, UWorld* World, float VolumeMultiplier, float PitchMultiplier, float StartTime, const FVector& Location, const FRotator& Rotation, USoundAttenuation* AttenuationSettings, USoundConcurrency* Concurrency, const TArray<FAudioParameter>* Params, const AActor* OwningActor)
+{
+	check(IsInGameThread());
+
+	TArray<FAudioParameter> ParamsCopy;
+	if (Params)
+	{
+		ParamsCopy.Append(*Params);
+	}
+
+	PlaySoundAtLocation(Sound, World, VolumeMultiplier, PitchMultiplier, StartTime, Location, Rotation, AttenuationSettings, Concurrency, MoveTemp(ParamsCopy), OwningActor);
+}
+
 void FAudioDevice::Flush(UWorld* WorldToFlush, bool bClearActivatedReverb)
 {
 	if (!IsInAudioThread())
diff --git a/Engine/Source/Runtime/Engine/Private/Components/AudioComponent.cpp b/Engine/Source/Runtime/Engine/Private/Components/AudioComponent.cpp
--- a/Engine/Source/Runtime/Engine/Private/Components/AudioComponent.cpp
+++ b/Engine/Source/Runtime/Engine/Private/Components/AudioComponent.cpp
@@ -717,7 +717,8 @@ void UAudioComponent::PlayInternal(const PlayInternalRequestData& InPlayRequestD
 
 	AudioDevice->GetMaxDistanceAndFocusFactor(SoundToPlay, World, Location, AttenuationSettingsToApply, MaxDistance, FocusFactor);
 
-	FActiveSound NewActiveSound;
+	TSharedRef<FActiveSound> NewSharedActiveSound = MakeShared<FActiveSound>();
+	FActiveSound& NewActiveSound = NewSharedActiveSound.Get();
 	NewActiveSound.SetAudioComponent(*this);
 	NewActiveSound.SetWorld(World);
 	NewActiveSound.SetSound(SoundToPlay);
@@ -867,7 +868,7 @@ void UAudioComponent::PlayInternal(const PlayInternalRequestData& InPlayRequestD
 	TArray<FAudioParameter> InstanceParamsCopy = InstanceParameters;
 	FAudioParameter::Merge(MoveTemp(InstanceParamsCopy), SoundParams);
 
-	AudioDevice->AddNewActiveSound(NewActiveSound, MoveTemp(SoundParams));
+	AudioDevice->AddNewActiveSound(NewSharedActiveSound, MoveTemp(SoundParams));
 
 	LastSoundPlayOrder = NewActiveSound.GetPlayOrder();
 
diff --git a/Engine/Source/Runtime/Engine/Private/GameplayStatics.cpp b/Engine/Source/Runtime/Engine/Private/GameplayStatics.cpp
--- a/Engine/Source/Runtime/Engine/Private/GameplayStatics.cpp
+++ b/Engine/Source/Runtime/Engine/Private/GameplayStatics.cpp
@@ -1718,7 +1718,7 @@ void UGameplayStatics::PlaySoundAtLocation(const UObject* WorldContextObject, cl
 		const AActor* ActiveSoundOwner = OwningActor ? OwningActor : GameplayStatics::GetActorOwnerFromWorldContextObject(WorldContextObject);
 		UActorSoundParameterInterface::Fill(ActiveSoundOwner, Params);
 
-		AudioDevice->PlaySoundAtLocation(Sound, ThisWorld, VolumeMultiplier, PitchMultiplier, StartTime, Location, Rotation, AttenuationSettings, ConcurrencySettings, &Params, ActiveSoundOwner);
+		AudioDevice->PlaySoundAtLocation(Sound, ThisWorld, VolumeMultiplier, PitchMultiplier, StartTime, Location, Rotation, AttenuationSettings, ConcurrencySettings, MoveTemp(Params), ActiveSoundOwner);
 	}
 }
 
diff --git a/Engine/Source/Runtime/Engine/Public/AudioDevice.h b/Engine/Source/Runtime/Engine/Public/AudioDevice.h
--- a/Engine/Source/Runtime/Engine/Public/AudioDevice.h
+++ b/Engine/Source/Runtime/Engine/Public/AudioDevice.h
@@ -769,15 +769,24 @@ public:
 	 * @param	Params				An optional list of audio component params to immediately apply to a sound.
 	 */
 	ENGINE_API void PlaySoundAtLocation(USoundBase* Sound, UWorld* World, float VolumeMultiplier, float PitchMultiplier, float StartTime, const FVector& Location, const FRotator& Rotation, USoundAttenuation* AttenuationSettings = nullptr, USoundConcurrency* ConcurrencySettings = nullptr, const TArray<FAudioParameter>* Params = nullptr, const AActor* OwningActor = nullptr);
+	/** As above, but supporting move semantics to reduce copying of Params */
+	ENGINE_API void PlaySoundAtLocation(USoundBase* Sound, UWorld* World, float VolumeMultiplier, float PitchMultiplier, float StartTime, const FVector& Location, const FRotator& Rotation, USoundAttenuation* AttenuationSettings, USoundConcurrency* Concurrency, TArray<FAudioParameter>&& Params, const AActor* OwningActor = nullptr);
+
 
 	/**
 	 * Adds an active sound to the audio device
 	 */
 	ENGINE_API void AddNewActiveSound(const FActiveSound& ActiveSound, const TArray<FAudioParameter>* InDefaultParams = nullptr);
 
+	/**
+	 * Adds an active sound to the audio device
+	*/
+	ENGINE_API void AddNewActiveSound(TSharedRef<FActiveSound> ActiveSound, TArray<FAudioParameter>&& InDefaultParams);
+
 	/**
 	 * Adds an active sound to the audio device
 	 */
+	UE_DEPRECATED(5.7, "Use AddNewActiveSound which accepts a TSharedRef<>")
 	ENGINE_API void AddNewActiveSound(const FActiveSound& ActiveSound, TArray<FAudioParameter>&& InDefaultParams);
 
 	/**
@@ -1542,6 +1551,13 @@ private:
 	 * Adds an active sound to the audio device. Can be a new active sound or one provided by the re-triggering
 	 * loop system.
 	 */
+	void AddNewActiveSoundInternal(TSharedRef<FActiveSound> InActiveSound, TArray<FAudioParameter>&& InDefaultParams, FAudioVirtualLoop* InVirtualLoopToRetrigger = nullptr);
+
+
+	/**
+	 * Adds an active sound to the audio device. Can be a new active sound or one provided by the re-triggering
+	 * loop system.
+	*/
 	void AddNewActiveSoundInternal(const FActiveSound& InActiveSound, TArray<FAudioParameter>&& InDefaultParams, FAudioVirtualLoop* InVirtualLoopToRetrigger = nullptr);
 
 	/**
