Fix bad default local merge
#rnx

diff --git a/Engine/Source/Runtime/Engine/Classes/Sound/SoundConcurrency.h b/Engine/Source/Runtime/Engine/Classes/Sound/SoundConcurrency.h
--- a/Engine/Source/Runtime/Engine/Classes/Sound/SoundConcurrency.h
+++ b/Engine/Source/Runtime/Engine/Classes/Sound/SoundConcurrency.h
@@ -78,7 +78,7 @@ struct FSoundConcurrencySettings
 
 	/** The max number of allowable concurrent active voices for voices playing in this concurrency group. */
 	UPROPERTY(EditAnywhere, BlueprintReadWrite, Category = Concurrency, meta = (UIMin = "1", ClampMin = "1"))
-	int32 MaxCount = 0;
+	int32 MaxCount = 16;
 
 	/* Whether or not to limit the concurrency to per sound owner (i.e. the actor that plays the sound). If the sound doesn't have an owner, it falls back to global concurrency. */
 	UPROPERTY(EditAnywhere, BlueprintReadWrite, Category = Concurrency)
