Slate: Fix GameViewport Lifecycle

#rnx

diff --git a/Engine/Source/Runtime/Engine/Private/Slate/SGameLayerManager.cpp b/Engine/Source/Runtime/Engine/Private/Slate/SGameLayerManager.cpp
--- a/Engine/Source/Runtime/Engine/Private/Slate/SGameLayerManager.cpp
+++ b/Engine/Source/Runtime/Engine/Private/Slate/SGameLayerManager.cpp
@@ -361,6 +361,16 @@ void SGameLayerManager::ClearWidgets()
 {
 	PlayerCanvas->ClearChildren();
 
+	for (const TWeakPtr<SWidget>& GameLayerWeak : GameLayers)
+	{
+		if (const TSharedPtr<SWidget>& GameLayer = GameLayerWeak.Pin())
+		{
+			WindowOverlay->RemoveSlot(GameLayer.ToSharedRef());
+		}
+	}
+
+	GameLayers.Empty();
+
 	// Potential for removed layers to impact the map, so need to
 	// remove & delete as separate steps
 	while (PlayerLayers.Num())
@@ -381,6 +391,8 @@ void SGameLayerManager::ClearWidgets()
 
 void SGameLayerManager::AddGameLayer(TSharedRef<SWidget> ViewportContent, int32 ZOrder)
 {
+	GameLayers.Add(ViewportContent);
+
 	ensure(!WindowOverlay->HasSlotWithZOrder(ZOrder));
 
 	WindowOverlay->AddSlot(ZOrder)
diff --git a/Engine/Source/Runtime/Engine/Public/Slate/SGameLayerManager.h b/Engine/Source/Runtime/Engine/Public/Slate/SGameLayerManager.h
--- a/Engine/Source/Runtime/Engine/Public/Slate/SGameLayerManager.h
+++ b/Engine/Source/Runtime/Engine/Public/Slate/SGameLayerManager.h
@@ -208,6 +208,7 @@ private:
 	FGeometry CachedGeometry;
 
 	TMap<FObjectKey, TSharedPtr<FPlayerLayer>> PlayerLayers;
+	TArray<TWeakPtr<SWidget>> GameLayers;
 
 	TAttribute<FSceneViewport*> SceneViewport;
 	TSharedPtr<SVerticalBox> WidgetHost;
