Add per-platform MaxCount, not utilized until future CL (breaking up earlier MaxCount change to determine offending TestCook code serialization crash)

diff --git a/Engine/Source/Runtime/Engine/Classes/Sound/SoundConcurrency.h b/Engine/Source/Runtime/Engine/Classes/Sound/SoundConcurrency.h
--- a/Engine/Source/Runtime/Engine/Classes/Sound/SoundConcurrency.h
+++ b/Engine/Source/Runtime/Engine/Classes/Sound/SoundConcurrency.h
@@ -80,6 +80,9 @@ struct FSoundConcurrencySettings
 	UPROPERTY(EditAnywhere, BlueprintReadWrite, Category = Concurrency, meta = (UIMin = "1", ClampMin = "1"))
 	int32 MaxCount = 16;
 
+	UPROPERTY()
+	FPerPlatformInt PlatformMaxCount;
+
 	/* Whether or not to limit the concurrency to per sound owner (i.e. the actor that plays the sound). If the sound doesn't have an owner, it falls back to global concurrency. */
 	UPROPERTY(EditAnywhere, BlueprintReadWrite, Category = Concurrency)
 	uint8 bLimitToOwner : 1 = 0;
diff --git a/Engine/Source/Runtime/Engine/Private/SoundConcurrency.cpp b/Engine/Source/Runtime/Engine/Private/SoundConcurrency.cpp
--- a/Engine/Source/Runtime/Engine/Private/SoundConcurrency.cpp
+++ b/Engine/Source/Runtime/Engine/Private/SoundConcurrency.cpp
@@ -58,6 +58,7 @@ USoundConcurrency::USoundConcurrency(const FObjectInitializer& ObjectInitializer
 
 FSoundConcurrencySettings::FSoundConcurrencySettings()
 {
+	PlatformMaxCount = FPerPlatformInt(16);
 }
 
 float FSoundConcurrencySettings::GetVolumeScale() const
