Testing addition of no-op PostSerialize to SoundConcurrency with PerPlatform property against AutoTest
#rnx

diff --git a/Engine/Source/Runtime/Engine/Classes/Sound/SoundConcurrency.h b/Engine/Source/Runtime/Engine/Classes/Sound/SoundConcurrency.h
--- a/Engine/Source/Runtime/Engine/Classes/Sound/SoundConcurrency.h
+++ b/Engine/Source/Runtime/Engine/Classes/Sound/SoundConcurrency.h
@@ -149,8 +149,21 @@ public:
 	  * prior to being removed from a concurrency group.
 	  */
 	ENGINE_API bool IsEvictionSupported() const;
+
+	ENGINE_API void PostSerialize(const FArchive& Ar);
 };
 
+template<>
+struct TStructOpsTypeTraits<FSoundConcurrencySettings>
+	: public TStructOpsTypeTraitsBase2<FSoundConcurrencySettings>
+{
+	enum
+	{
+		WithPostSerialize = true
+	};
+};
+
+
 UCLASS(BlueprintType, hidecategories=Object, editinlinenew, MinimalAPI)
 class USoundConcurrency : public UObject
 {
diff --git a/Engine/Source/Runtime/Engine/Private/SoundConcurrency.cpp b/Engine/Source/Runtime/Engine/Private/SoundConcurrency.cpp
--- a/Engine/Source/Runtime/Engine/Private/SoundConcurrency.cpp
+++ b/Engine/Source/Runtime/Engine/Private/SoundConcurrency.cpp
@@ -90,6 +90,9 @@ bool FSoundConcurrencySettings::IsEvictionSupported() const
 	}
 }
 
+void FSoundConcurrencySettings::PostSerialize(const FArchive& Ar)
+{
+}
 
 FConcurrencyHandle::FConcurrencyHandle(const FSoundConcurrencySettings& InSettings)
 	: Settings(InSettings)
