New Material Translator:  Tidier vector type casting in HLSL translator, and more robust casting in general.
* LowerComposite checks for a matching chain of FScalar casts from a matching set of FSubscript's (consecutive indices pointing to the same underlying FValue).  If this pattern is found, a single chain of vector casts is generated instead of subcripts and scalar casts per component, merged back into a vector.  In a test LWC expression, this reduced the number of characters in the expression from 472 to 142.
* Both Scalar and newly added Vector cast logic upgraded to support more strict GLSL and HLSL 2021 rules for bool and float to integer casting.  There is no implicit bool casting in HLSL, so casting to bool requires a conditional, and from a bool requires a ternary conditional or select.  Ternary conditional only works on scalars, and select requires all arguments to have the same number of components ("select" had been overloaded to solve this, but it's cleaner going forward to not count on this).  There is also no automatic promotion from integer to float -- it requires an explicit cast.  Debug code was added to test every permutation of scalar and vector casts, since most of the casts can't be generated with existing expressions.

#rnx

diff --git a/Engine/Source/Runtime/Engine/Private/Materials/MaterialIREmitter.cpp b/Engine/Source/Runtime/Engine/Private/Materials/MaterialIREmitter.cpp
--- a/Engine/Source/Runtime/Engine/Private/Materials/MaterialIREmitter.cpp
+++ b/Engine/Source/Runtime/Engine/Private/Materials/MaterialIREmitter.cpp
@@ -2541,7 +2541,7 @@ FValueRef FEmitter::CastToBoolKind(FValueRef Value)
 
 FValueRef FEmitter::CastToIntKind(FValueRef Value)
 {
-	return CastToScalarKind(Value, ScalarKind_Bool);
+	return CastToScalarKind(Value, ScalarKind_Int);
 }
 
 FValueRef FEmitter::CastToFloatKind(FValueRef Value)
diff --git a/Engine/Source/Runtime/Engine/Private/Materials/MaterialIRToHLSLTranslator.cpp b/Engine/Source/Runtime/Engine/Private/Materials/MaterialIRToHLSLTranslator.cpp
--- a/Engine/Source/Runtime/Engine/Private/Materials/MaterialIRToHLSLTranslator.cpp
+++ b/Engine/Source/Runtime/Engine/Private/Materials/MaterialIRToHLSLTranslator.cpp
@@ -868,6 +868,105 @@ struct FTranslator : FMaterialIRToHLSLTranslation
 		}
 	}
 
+	bool HasMatchingScalarComponentCastChain(MIR::FValue* FirstComponent, MIR::FValue* CurrentComponent, int32 Index, int32 NumComponents)
+	{
+		if (FirstComponent->Kind == MIR::VK_Subscript && CurrentComponent->Kind == MIR::VK_Subscript)
+		{
+			if (((MIR::FSubscript*)FirstComponent)->Arg == ((MIR::FSubscript*)CurrentComponent)->Arg)
+			{
+				const MIR::FPrimitiveType* PrimitiveType = ((MIR::FSubscript*)FirstComponent)->Arg->Type->AsVector();
+				return PrimitiveType && PrimitiveType->NumRows == NumComponents && ((MIR::FSubscript*)FirstComponent)->Index == 0 && ((MIR::FSubscript*)CurrentComponent)->Index == Index;
+			}
+		}
+		else if (FirstComponent->Kind == MIR::VK_Scalar && CurrentComponent->Kind == MIR::VK_Scalar)
+		{
+			if (((MIR::FScalar*)FirstComponent)->Type == ((MIR::FScalar*)CurrentComponent)->Type)
+			{
+				return HasMatchingScalarComponentCastChain(((MIR::FScalar*)FirstComponent)->Arg, ((MIR::FScalar*)CurrentComponent)->Arg, Index, NumComponents);
+			}
+		}
+		return false;
+	}
+
+	ENoOp LowerVectorCastChain(MIR::FValue* FirstComponent, int32 NumComponents)
+	{
+		if (FirstComponent->Kind == MIR::VK_Scalar)
+		{
+			MIR::FValue* Arg = ((MIR::FScalar*)FirstComponent)->Arg;
+
+			const MIR::EScalarKind ScalarKind = ((MIR::FPrimitiveType*)FirstComponent->Type)->ScalarKind;
+			const MIR::EScalarKind ArgScalarKind = ((MIR::FPrimitiveType*)Arg->Type)->ScalarKind;
+
+			if (ArgScalarKind == MIR::ScalarKind_LWC)
+			{
+				// Cast from LWC
+				if (ScalarKind == MIR::ScalarKind_Bool)
+				{
+					// Cast to bool requires a comparison with zero, outside the WSDemote
+					Printer << "(WSDemote(" << LowerVectorCastChain(Arg, NumComponents) << ") != 0)";
+				}
+				else if (ScalarKind != MIR::ScalarKind_Float)
+				{
+					// Cast to non-float (integer) requires a cast to the type, outside the WSDemote
+					const MIR::FPrimitiveType* VectorType = MIR::FPrimitiveType::GetVector(ScalarKind, NumComponents);
+					Printer << LowerType(VectorType) << "(WSDemote(" << LowerVectorCastChain(Arg, NumComponents) << "))";
+				}
+				else
+				{
+					// Cast to float
+					Printer << "WSDemote(" << LowerVectorCastChain(Arg, NumComponents) << ")";
+				}
+			}
+			else if (ScalarKind == MIR::ScalarKind_LWC)
+			{
+				// Cast to LWC
+				if (ArgScalarKind == MIR::ScalarKind_Bool)
+				{
+					// Cast from bool requires a select between 1.0f and 0.0f, inside the WSPromote
+					Printer << "WSPromote(select(" << LowerVectorCastChain(Arg, NumComponents) << ", (float" << NumComponents << ")1.0f, (float" << NumComponents << ")0.0f))";
+				}
+				else if (ArgScalarKind != MIR::ScalarKind_Float)
+				{
+					// Cast from non-float (integer) requires a cast to float, inside the WSPromote
+					Printer << "WSPromote(float" << NumComponents << "(" << LowerVectorCastChain(Arg, NumComponents) << "))";
+				}
+				else
+				{
+					// Cast from float
+					Printer << "WSPromote(" << LowerVectorCastChain(Arg, NumComponents) << ")";
+				}
+			}
+			else
+			{
+				// Cast between intrinsic types
+				if (ScalarKind == MIR::ScalarKind_Bool)
+				{
+					// Cast to bool requires a comparison with zero
+					Printer << "(" << LowerVectorCastChain(Arg, NumComponents) << " != 0)";
+				}
+				else if (ArgScalarKind == MIR::ScalarKind_Bool)
+				{
+					// Cast from bool requires a select between 1 and 0
+					const MIR::FPrimitiveType* VectorType = MIR::FPrimitiveType::GetVector(ScalarKind, NumComponents);
+					Printer << "select(" << LowerVectorCastChain(Arg, NumComponents) << ", (" << LowerType(VectorType) << ")1, (" << LowerType(VectorType) << ")0)";
+				}
+				else
+				{
+					// Cast between arithmetic types
+					const MIR::FPrimitiveType* VectorType = MIR::FPrimitiveType::GetVector(ScalarKind, NumComponents);
+					Printer << LowerType(VectorType) << "(" << LowerVectorCastChain(Arg, NumComponents) << ")";
+				}
+			}
+		}
+		else
+		{
+			// Finally reached the inner subscript, print its vector argument
+			check(FirstComponent->Kind == MIR::VK_Subscript);
+			Printer << LowerValue(((MIR::FSubscript*)FirstComponent)->Arg);
+		}
+		return NoOp;
+	}
+
 	void LowerComposite(const MIR::FComposite* Dimensional) 
 	{
 		TArrayView<MIR::FValue* const> Components = Dimensional->GetComponents();
@@ -879,18 +978,31 @@ struct FTranslator : FMaterialIRToHLSLTranslation
 		// component and cast it to the type.  LWC doesn't support casting, and always needs
 		// to call a function to convert types.
 		bool bSameComponents = true;
+
+		// We can also generate tidier HLSL for cases where casts are done for whole vectors.  This is a frequent case for LWC, where casts happen
+		// in both directions (LWC to float and back), due to operations automatically casting their inputs or outputs.  For example, BO_Sub
+		// automatically downcasts LWC to float, but if you feed that into a BO_Add that has LWC for its other input, it will immediately cast it
+		// back to LWC -- two consecutive casts.  These then get expanded into individual scalar casts, which we would like to do as whole vector
+		// casts for readability.  We only need to consider this special case if the first component is a Scalar, and the type is a Vector.
+		bool bWholeVectorCast = Components[0]->Kind == MIR::VK_Scalar && ArithmeticType->IsVector();
+
 		for (int32 i = 1; bSameComponents && i < Components.Num(); ++i)
 		{
 			bSameComponents &= (Components[i] == Components[0]);
+			bWholeVectorCast &= HasMatchingScalarComponentCastChain(Components[0], Components[i], i, Components.Num());
 		}
 		
-		if (bSameComponents && ArithmeticType->ScalarKind != MIR::EScalarKind::ScalarKind_LWC)
+		if (bSameComponents && ArithmeticType->ScalarKind != MIR::ScalarKind_LWC)
 		{
 			Printer << TEXT("(") << LowerType(ArithmeticType) << TEXT(")") << LowerValue(Components[0]);
 		}
+		else if (bWholeVectorCast)
+		{
+			LowerVectorCastChain(Components[0], Components.Num());
+		}
 		else
 		{
-			if (ArithmeticType->ScalarKind == MIR::EScalarKind::ScalarKind_LWC)
+			if (ArithmeticType->ScalarKind == MIR::ScalarKind_LWC)
 			{
 				if (ArithmeticType->NumColumns == 1)
 				{
@@ -1177,17 +1289,65 @@ struct FTranslator : FMaterialIRToHLSLTranslation
 
 	void LowerScalar(const MIR::FScalar* Scalar)
 	{
-		if (Scalar->Arg->Type->AsPrimitive()->ScalarKind == MIR::EScalarKind::ScalarKind_LWC)
+		const MIR::EScalarKind ScalarKind = Scalar->Type->AsPrimitive()->ScalarKind;
+		const MIR::EScalarKind ArgScalarKind = Scalar->Arg->Type->AsPrimitive()->ScalarKind;
+
+		if (ArgScalarKind == MIR::ScalarKind_LWC)
 		{
-			Printer << LowerType(Scalar->Type) << "(WSDemote(" << LowerValue(Scalar->Arg) << "))";
+			// Cast from LWC
+			if (ScalarKind == MIR::ScalarKind_Bool)
+			{
+				// Cast to bool requires a comparison with zero, outside the WSDemote
+				Printer << "(WSDemote(" << LowerValue(Scalar->Arg) << ") != 0)";
+			}
+			else if (ScalarKind != MIR::ScalarKind_Float)
+			{
+				// Cast to non-float (integer) requires a cast to the type, outside the WSDemote
+				Printer << LowerType(Scalar->Type) << "(WSDemote(" << LowerValue(Scalar->Arg) << "))";
+			}
+			else
+			{
+				// Cast to float
+				Printer << "WSDemote(" << LowerValue(Scalar->Arg) << ")";
+			}
 		}
-		else if (Scalar->Type->AsPrimitive()->ScalarKind == MIR::EScalarKind::ScalarKind_LWC)
+		else if (ScalarKind == MIR::ScalarKind_LWC)
 		{
-			Printer << "WSPromote(" << LowerValue(Scalar->Arg) << ")";
+			// Cast to LWC
+			if (ArgScalarKind == MIR::ScalarKind_Bool)
+			{
+				// Cast from bool requires a select between 1.0f and 0.0f, inside the WSPromote
+				Printer << "WSPromote(" << LowerValue(Scalar->Arg) << " ? 1.0f : 0.0f)";
+			}
+			else if (ArgScalarKind != MIR::ScalarKind_Float)
+			{
+				// Cast from non-float (integer) requires a cast to float, inside the WSPromote
+				Printer << "WSPromote(float(" << LowerValue(Scalar->Arg) << "))";
+			}
+			else
+			{
+				// Cast from float
+				Printer << "WSPromote(" << LowerValue(Scalar->Arg) << ")";
+			}
 		}
 		else
 		{
-			Printer << LowerType(Scalar->Type) << "(" << LowerValue(Scalar->Arg) << ")";
+			// Cast between intrinsic types
+			if (ScalarKind == MIR::ScalarKind_Bool)
+			{
+				// Cast to bool requires a comparison with zero
+				Printer << "(" << LowerValue(Scalar->Arg) << " != 0)";
+			}
+			else if (ArgScalarKind == MIR::ScalarKind_Bool)
+			{
+				// Cast from bool requires a select between 1 and 0
+				Printer << LowerType(Scalar->Type) << "(" << LowerValue(Scalar->Arg) << " ? 1 : 0)";
+			}
+			else
+			{
+				// Cast between arithmetic types
+				Printer << LowerType(Scalar->Type) << "(" << LowerValue(Scalar->Arg) << ")";
+			}
 		}
 	}
 
@@ -1616,7 +1776,7 @@ struct FTranslator : FMaterialIRToHLSLTranslation
 				switch (ArithmeticType->ScalarKind)
 				{
 					case MIR::ScalarKind_Bool:	Printer << TEXT("bool"); break;
-					case MIR::ScalarKind_Int: 	Printer << TEXT("int32"); break;
+					case MIR::ScalarKind_Int: 	Printer << TEXT("int"); break;
 					case MIR::ScalarKind_Float:	Printer << TEXT("MaterialFloat"); break;
 				}
 
