- More incremental submission of concurrency update
- Readd no-op PostSerialize to SoundConcurrency with PerPlatform MaxCount
- Readd Getter/Setter and runtime code redirect accordingly
#rnx

diff --git a/Engine/Source/Runtime/Engine/Classes/Sound/SoundConcurrency.h b/Engine/Source/Runtime/Engine/Classes/Sound/SoundConcurrency.h
--- a/Engine/Source/Runtime/Engine/Classes/Sound/SoundConcurrency.h
+++ b/Engine/Source/Runtime/Engine/Classes/Sound/SoundConcurrency.h
@@ -136,6 +136,16 @@ public:
 	UPROPERTY(EditAnywhere, BlueprintReadWrite, Category = "Voice Stealing", meta = (UIMin = "0.0", ClampMin = "0.0", UIMax = "1.0", ClampMax="1000000.0"))
 	float VoiceStealReleaseTime = 0.0f;
 
+	/**
+	  * Returns the MaxCount as a per-platform integer value.
+	  */
+	ENGINE_API int32 GetMaxCount() const;
+
+	/**
+	  * Sets the MaxCount as a per-platform integer value.
+	  * (Optional) PlatformName - If set, applies the given max count for the platform of the given name (editor only). Otherwise, applies value to default.
+	  */
+	ENGINE_API bool SetMaxCount(int32 InMaxCount, FName PlatformName = { });
 
 	FSoundConcurrencySettings();
 
@@ -149,8 +159,21 @@ public:
 	  * prior to being removed from a concurrency group.
 	  */
 	ENGINE_API bool IsEvictionSupported() const;
+
+	ENGINE_API void PostSerialize(const FArchive& Ar);
 };
 
+template<>
+struct TStructOpsTypeTraits<FSoundConcurrencySettings>
+	: public TStructOpsTypeTraitsBase2<FSoundConcurrencySettings>
+{
+	enum
+	{
+		WithPostSerialize = true
+	};
+};
+
+
 UCLASS(BlueprintType, hidecategories=Object, editinlinenew, MinimalAPI)
 class USoundConcurrency : public UObject
 {
diff --git a/Engine/Source/Runtime/Engine/Private/SoundBase.cpp b/Engine/Source/Runtime/Engine/Private/SoundBase.cpp
--- a/Engine/Source/Runtime/Engine/Private/SoundBase.cpp
+++ b/Engine/Source/Runtime/Engine/Private/SoundBase.cpp
@@ -244,7 +244,7 @@ void USoundBase::PostLoad()
 	{
 		bOverrideConcurrency = true;
 		ConcurrencyOverrides.bLimitToOwner = false;
-		ConcurrencyOverrides.MaxCount = FMath::Max(MaxConcurrentPlayCount_DEPRECATED, 1);
+		ConcurrencyOverrides.SetMaxCount(FMath::Max(MaxConcurrentPlayCount_DEPRECATED, 1));
 		ConcurrencyOverrides.ResolutionRule = MaxConcurrentResolutionRule_DEPRECATED;
 	}
 
diff --git a/Engine/Source/Runtime/Engine/Private/SoundConcurrency.cpp b/Engine/Source/Runtime/Engine/Private/SoundConcurrency.cpp
--- a/Engine/Source/Runtime/Engine/Private/SoundConcurrency.cpp
+++ b/Engine/Source/Runtime/Engine/Private/SoundConcurrency.cpp
@@ -57,10 +57,51 @@ USoundConcurrency::USoundConcurrency(const FObjectInitializer& ObjectInitializer
 }
 
 FSoundConcurrencySettings::FSoundConcurrencySettings()
+PRAGMA_DISABLE_DEPRECATION_WARNINGS
+	: MaxCount(0)
+PRAGMA_ENABLE_DEPRECATION_WARNINGS
 {
 	PlatformMaxCount = FPerPlatformInt(16);
 }
 
+
+int32 FSoundConcurrencySettings::GetMaxCount() const
+{
+PRAGMA_DISABLE_DEPRECATION_WARNINGS
+	if (MaxCount > 0)
+	{
+		return MaxCount;
+	}
+PRAGMA_ENABLE_DEPRECATION_WARNINGS
+
+	return PlatformMaxCount.GetValue();
+}
+
+bool FSoundConcurrencySettings::SetMaxCount(int32 InMaxCount, FName PlatformName)
+{
+	if (InMaxCount > 0)
+	{
+PRAGMA_DISABLE_DEPRECATION_WARNINGS
+		MaxCount = 0;
+PRAGMA_ENABLE_DEPRECATION_WARNINGS
+
+		if (PlatformName.IsNone())
+		{
+			PlatformMaxCount.Default = InMaxCount;
+			return true;
+		}
+#if WITH_EDITORONLY_DATA
+		else
+		{
+			PlatformMaxCount.PerPlatform.FindOrAdd(PlatformName) = InMaxCount;
+			return true;
+		}
+#endif // WITH_EDITORONLY_DATA
+	}
+
+	return false;
+}
+
 float FSoundConcurrencySettings::GetVolumeScale() const
 {
 	return FMath::Clamp(VolumeScale, 0.0f, 1.0f);;
@@ -90,6 +131,9 @@ bool FSoundConcurrencySettings::IsEvictionSupported() const
 	}
 }
 
+void FSoundConcurrencySettings::PostSerialize(const FArchive& Ar)
+{
+}
 
 FConcurrencyHandle::FConcurrencyHandle(const FSoundConcurrencySettings& InSettings)
 	: Settings(InSettings)
@@ -215,7 +259,7 @@ bool FConcurrencyGroup::IsEmpty() const
 
 bool FConcurrencyGroup::IsFull() const
 {
-	return Settings.MaxCount <= ActiveSounds.Num();
+	return Settings.GetMaxCount() <= ActiveSounds.Num();
 }
 
 void FConcurrencyGroup::AddActiveSound(FActiveSound& ActiveSound)
@@ -359,7 +403,7 @@ void FConcurrencyGroup::CullSoundsDueToMaxConcurrency()
 {
 	// Nothing to do if our active sound count is less than or equal to our max active sounds
 	// or if eviction is supported
-	if (ActiveSounds.Num() <= Settings.MaxCount)
+	if (ActiveSounds.Num() <= Settings.GetMaxCount())
 	{
 		return;
 	}
@@ -425,7 +469,8 @@ void FConcurrencyGroup::CullSoundsDueToMaxConcurrency()
 
 	// Need to make a new list when stopping the sounds since the process of stopping an active sound
 	// will remove the sound from this concurrency group's ActiveSounds array.
-	const int32 NumSoundsToStop = ActiveSounds.Num() - Settings.MaxCount;
+	const int32 NumSoundsToStop = ActiveSounds.Num() - Settings.GetMaxCount();
+
 	check(NumSoundsToStop > 0);
 	for (int32 i = 0; i < ActiveSounds.Num(); ++i)
 	{
@@ -564,7 +609,7 @@ FActiveSound* FSoundConcurrencyManager::CreateNewActiveSound(const FActiveSound&
 #if !UE_BUILD_SHIPPING
 	for (const FConcurrencyHandle& ConcurrencyHandle : ConcurrencyHandles)
 	{
-		check(ConcurrencyHandle.Settings.MaxCount > 0);
+		check(ConcurrencyHandle.Settings.GetMaxCount() > 0);
 	}
 #endif
 
