Enact PerPlatform Sound Concurrency MaxCount.
- Expose field in editor
- Provide PostSerialize upgrade logic (testing as editor-only logic to fix AutoTest)

[RN] MaxCount in SoundConcurrency can now be set per-platform

diff --git a/Engine/Source/Runtime/Engine/Classes/Sound/SoundConcurrency.h b/Engine/Source/Runtime/Engine/Classes/Sound/SoundConcurrency.h
--- a/Engine/Source/Runtime/Engine/Classes/Sound/SoundConcurrency.h
+++ b/Engine/Source/Runtime/Engine/Classes/Sound/SoundConcurrency.h
@@ -77,10 +77,10 @@ struct FSoundConcurrencySettings
 	GENERATED_BODY()
 
 	/** The max number of allowable concurrent active voices for voices playing in this concurrency group. */
-	UPROPERTY(EditAnywhere, BlueprintReadWrite, Category = Concurrency, meta = (UIMin = "1", ClampMin = "1"))
-	int32 MaxCount = 16;
+	UPROPERTY(BlueprintReadWrite, Category = Concurrency, meta = (UIMin = "1", ClampMin = "1"))
+	int32 MaxCount = 0;
 
-	UPROPERTY()
+	UPROPERTY(EditAnywhere, Category = Concurrency, meta = (DisplayName = "Max Count", UIMin = "1", ClampMin = "1"))
 	FPerPlatformInt PlatformMaxCount;
 
 	/* Whether or not to limit the concurrency to per sound owner (i.e. the actor that plays the sound). If the sound doesn't have an owner, it falls back to global concurrency. */
@@ -160,9 +160,12 @@ public:
 	  */
 	ENGINE_API bool IsEvictionSupported() const;
 
+#if WITH_EDITORONLY_DATA
 	ENGINE_API void PostSerialize(const FArchive& Ar);
+#endif // WITH_EDITORONLY_DATA
 };
 
+#if WITH_EDITORONLY_DATA
 template<>
 struct TStructOpsTypeTraits<FSoundConcurrencySettings>
 	: public TStructOpsTypeTraitsBase2<FSoundConcurrencySettings>
@@ -172,6 +175,7 @@ struct TStructOpsTypeTraits<FSoundConcurrencySettings>
 		WithPostSerialize = true
 	};
 };
+#endif // WITH_EDITORONLY_DATA
 
 
 UCLASS(BlueprintType, hidecategories=Object, editinlinenew, MinimalAPI)
diff --git a/Engine/Source/Runtime/Engine/Private/SoundConcurrency.cpp b/Engine/Source/Runtime/Engine/Private/SoundConcurrency.cpp
--- a/Engine/Source/Runtime/Engine/Private/SoundConcurrency.cpp
+++ b/Engine/Source/Runtime/Engine/Private/SoundConcurrency.cpp
@@ -131,9 +131,21 @@ bool FSoundConcurrencySettings::IsEvictionSupported() const
 	}
 }
 
+#if WITH_EDITORONLY_DATA
 void FSoundConcurrencySettings::PostSerialize(const FArchive& Ar)
 {
+	if (Ar.IsLoading())
+	{
+PRAGMA_DISABLE_DEPRECATION_WARNINGS
+		if (MaxCount > 0)
+		{
+			PlatformMaxCount.Default = MaxCount;
+			MaxCount = 0;
+		}
+PRAGMA_ENABLE_DEPRECATION_WARNINGS
+	}
 }
+#endif // WITH_EDITORONLY_DATA
 
 FConcurrencyHandle::FConcurrencyHandle(const FSoundConcurrencySettings& InSettings)
 	: Settings(InSettings)
