New Material Translator:  Analytic derivative support for External Code Declarations, affects user facing expressions PixelDepth, ScreenPosition, VertexColor, and VertexNormal.

#rnx

diff --git a/Engine/Source/Runtime/Engine/Private/Materials/MaterialIREmitter.cpp b/Engine/Source/Runtime/Engine/Private/Materials/MaterialIREmitter.cpp
--- a/Engine/Source/Runtime/Engine/Private/Materials/MaterialIREmitter.cpp
+++ b/Engine/Source/Runtime/Engine/Private/Materials/MaterialIREmitter.cpp
@@ -3117,11 +3117,21 @@ FValueRef FEmitter::AnalyticalPartialDerivative(FValueRef Value, EDerivativeAxis
 			return Cast(AnalyticalPartialDerivative(AsCast->Arg, Axis), AsCast->Type);
 		}
 
+		case MIR::VK_InlineHLSL:
+		{
+			FInlineHLSL* AsInlineHLSL = Value->As<FInlineHLSL>();
+			if (!EnumHasAnyFlags(AsInlineHLSL->Flags, MIR::EValueFlags::HasDynamicHLSLCode) && AsInlineHLSL->ExternalCodeDeclaration->Derivative == EDerivativeStatus::Valid)
+			{
+				TArray<FValueRef> Arguments(TArrayView<FValue*>(AsInlineHLSL->Arguments, AsInlineHLSL->NumArguments));
+				return InlineHLSL(AsInlineHLSL->ExternalCodeDeclaration, Arguments, Axis == EDerivativeAxis::X ? MIR::EValueFlags::DerivativeDDX : MIR::EValueFlags::DerivativeDDY);
+			}
+			return Cast(ConstantZero(ValuePrimitiveType->ScalarKind), ValuePrimitiveType);
+		}
+
 		// These values are uniform (constant), thus their value is always zero.
 		case MIR::VK_Constant:
 		case MIR::VK_UniformParameter:
 		case MIR::VK_TextureRead:
-		case MIR::VK_InlineHLSL:
 			return Cast(ConstantZero(ValuePrimitiveType->ScalarKind), ValuePrimitiveType);
 
 		default:
