MaterialParameterCollection Extension

diff --git a/Engine/Source/Editor/DetailCustomizations/Private/DetailCustomizations.cpp b/Engine/Source/Editor/DetailCustomizations/Private/DetailCustomizations.cpp
--- a/Engine/Source/Editor/DetailCustomizations/Private/DetailCustomizations.cpp
+++ b/Engine/Source/Editor/DetailCustomizations/Private/DetailCustomizations.cpp
@@ -90,6 +90,7 @@
 #include "MaterialExpressionLandscapeGrassCustomization.h"
 #include "MaterialExpressionTextureBaseDetails.h"
 #include "MaterialInstanceDynamicDetails.h"
+#include "MaterialParameterCollectionDetails.h"
 #include "MaterialProxySettingsCustomizations.h"
 #include "MaterialShadingModelCustomization.h"
 #include "Math/Matrix.h"
@@ -486,6 +487,7 @@ void FDetailCustomizationsModule::RegisterObjectCustomizations()
 	RegisterCustomClassLayout("MaterialExpressionSetMaterialAttributes", FOnGetDetailCustomizationInstance::CreateStatic(&FMaterialAttributePropertyDetails::MakeInstance));
 	RegisterCustomClassLayout("MaterialExpressionTextureBase", FOnGetDetailCustomizationInstance::CreateStatic(&FMaterialExpressionTextureBaseDetails::MakeInstance));
 	RegisterCustomClassLayout("MaterialInstanceDynamic", FOnGetDetailCustomizationInstance::CreateStatic(&FMaterialInstanceDynamicDetails::MakeInstance));
+	RegisterCustomClassLayout("MaterialParameterCollection", FOnGetDetailCustomizationInstance::CreateStatic(&FMaterialParameterCollectionDetails::MakeInstance));
 	RegisterCustomClassLayout("SkeletalMeshLODSettings", FOnGetDetailCustomizationInstance::CreateStatic(&FSkeletalMeshLODSettingsDetails::MakeInstance));
 	RegisterCustomClassLayout("SkeletalMesh", FOnGetDetailCustomizationInstance::CreateStatic(&FSkeletalMeshDetails::MakeInstance));
 	RegisterCustomClassLayout("Skeleton", FOnGetDetailCustomizationInstance::CreateStatic(&FSkeletonDetails::MakeInstance));
diff --git a/Engine/Source/Editor/DetailCustomizations/Private/MaterialParameterCollectionDetails.cpp b/Engine/Source/Editor/DetailCustomizations/Private/MaterialParameterCollectionDetails.cpp
new file mode 100644
--- /dev/null
+++ b/Engine/Source/Editor/DetailCustomizations/Private/MaterialParameterCollectionDetails.cpp
@@ -0,0 +1,277 @@
+// Copyright Epic Games, Inc. All Rights Reserved.
+
+#include "MaterialParameterCollectionDetails.h"
+
+#include "DetailLayoutBuilder.h"
+#include "Editor.h"
+#include "IDetailGroup.h"
+#include "IPropertyUtilities.h"
+#include "Materials/MaterialParameterCollection.h"
+
+#define LOCTEXT_NAMESPACE "MaterialParameterCollectionDetails"
+
+TSharedRef<IDetailCustomization> FMaterialParameterCollectionDetails::MakeInstance()
+{
+	return MakeShareable(new FMaterialParameterCollectionDetails);
+}
+
+FMaterialParameterCollectionDetails::~FMaterialParameterCollectionDetails()
+{
+	FEditorDelegates::PostUndoRedo.RemoveAll(DetailLayout);
+}
+
+class FMaterialParameterCollectionDetails::FBaseOverridesSelectorConst
+{
+public:
+	FBaseOverridesSelectorConst(const UMaterialParameterCollection* OverrideCollection) :
+		Collection(const_cast<UMaterialParameterCollection*>(OverrideCollection))
+	{
+	}
+
+	template<class FCollectionParameterType, class... FArgs>
+	decltype(auto) Contains(const FCollectionParameterType& CollectionParameter, FArgs&&... Args) const
+	{
+		return SelectBaseOverridesMap(CollectionParameter).Contains(CollectionParameter.Id, Forward<FArgs>(Args)...);
+	}
+
+	template<class FCollectionParameterType, class... FArgs>
+	decltype(auto) Find(const FCollectionParameterType& CollectionParameter, FArgs&&... Args) const
+	{
+		return SelectBaseOverridesMap(CollectionParameter).Find(CollectionParameter.Id, Forward<FArgs>(Args)...);
+	}
+
+protected:
+	TMap<FGuid, float>& SelectBaseOverridesMap(const FCollectionScalarParameter&) const
+	{
+		return Collection->ScalarParameterBaseOverrides;
+	}
+
+	TMap<FGuid, FLinearColor>& SelectBaseOverridesMap(const FCollectionVectorParameter&) const
+	{
+		return Collection->VectorParameterBaseOverrides;
+	}
+
+private:
+	UMaterialParameterCollection* Collection;
+};
+
+class FMaterialParameterCollectionDetails::FBaseOverridesSelector : public FBaseOverridesSelectorConst
+{
+public:
+	FBaseOverridesSelector(UMaterialParameterCollection* OverrideCollection) :
+		FBaseOverridesSelectorConst(OverrideCollection)
+	{
+	}
+
+	template<class FCollectionParameterType, class... FArgs>
+	decltype(auto) Add(const FCollectionParameterType& CollectionParameter, FArgs&&... Args) const
+	{
+		return SelectBaseOverridesMap(CollectionParameter).Add(CollectionParameter.Id, Forward<FArgs>(Args)...);
+	}
+
+	template<class FCollectionParameterType, class... FArgs>
+	decltype(auto) FindOrAdd(const FCollectionParameterType& CollectionParameter, FArgs&&... Args) const
+	{
+		return SelectBaseOverridesMap(CollectionParameter).FindOrAdd(CollectionParameter.Id, Forward<FArgs>(Args)...);
+	}
+
+	template<class FCollectionParameterType, class... FArgs>
+	decltype(auto) Remove(const FCollectionParameterType& CollectionParameter, FArgs&&... Args) const
+	{
+		return SelectBaseOverridesMap(CollectionParameter).Remove(CollectionParameter.Id, Forward<FArgs>(Args)...);
+	}
+};
+
+FMaterialParameterCollectionDetails::FBaseOverridesSelector FMaterialParameterCollectionDetails::GetBaseOverridesMap(UMaterialParameterCollection* OverrideCollection)
+{
+	return FBaseOverridesSelector(OverrideCollection);
+}
+
+FMaterialParameterCollectionDetails::FBaseOverridesSelectorConst FMaterialParameterCollectionDetails::GetBaseOverridesMap(const UMaterialParameterCollection* OverrideCollection)
+{
+	return FBaseOverridesSelectorConst(OverrideCollection);
+}
+
+template<class FCollectionParameterType>
+auto FMaterialParameterCollectionDetails::GetParameterValue(const UMaterialParameterCollection* OverrideCollection, const FCollectionParameterType& CollectionParameter, const UMaterialParameterCollection* BaseCollection)
+{
+	for (; OverrideCollection != BaseCollection; OverrideCollection = OverrideCollection->Base)
+	{
+		if (auto OverrideValue = GetBaseOverridesMap(OverrideCollection).Find(CollectionParameter))
+		{
+			return *OverrideValue;
+		}
+	}
+
+	return CollectionParameter.DefaultValue;
+}
+
+template<class FCollectionParameterType>
+void FMaterialParameterCollectionDetails::AddParameter(IDetailGroup& DetailGroup, TSharedPtr<IPropertyHandle> PropertyHandle, UMaterialParameterCollection* Collection, const FCollectionParameterType& CollectionParameter, const UMaterialParameterCollection* BaseCollection)
+{
+	using FValueType = decltype(CollectionParameter.DefaultValue);
+
+	// Set the property display name to the parameter's name.
+	PropertyHandle->SetPropertyDisplayName(FText::FromName(CollectionParameter.ParameterName));
+
+	// Initialize the property value with the current parameter value from the collection instance.
+	void* ValueData = nullptr;
+	if (PropertyHandle->GetValueData(ValueData) == FPropertyAccess::Success && ValueData)
+	{
+		FValueType& Value = *static_cast<FValueType*>(ValueData);
+		Value = GetParameterValue(Collection, CollectionParameter, BaseCollection);
+	}
+
+	// Create a delegate to invoke when the property value changes.
+	TDelegate<void(const FPropertyChangedEvent&)> OnPropertyValueChangedDelegate = TDelegate<void(const FPropertyChangedEvent&)>::CreateWeakLambda(Collection, [Collection = Collection, CollectionParameter, PropertyHandle](const FPropertyChangedEvent& PropertyChangedEvent)
+		{
+			// Determine if this is the final setting of the property value.
+			if (PropertyChangedEvent.ChangeType != EPropertyChangeType::ValueSet)
+			{
+				return;
+			}
+
+			// Update the collection's parameter value with the new property value.
+			void* ValueData = nullptr;
+			if (PropertyHandle->GetValueData(ValueData) == FPropertyAccess::Success && ValueData)
+			{
+				FValueType& Value = *static_cast<FValueType*>(ValueData);
+				FValueType& ParameterValue = GetBaseOverridesMap(Collection).FindOrAdd(CollectionParameter);
+				ParameterValue = Value;
+			}
+		});
+
+	// Invoke the delegate when the property value or one of its child property values changes.
+	PropertyHandle->SetOnPropertyValueChangedWithData(OnPropertyValueChangedDelegate);
+	PropertyHandle->SetOnChildPropertyValueChangedWithData(OnPropertyValueChangedDelegate);
+
+	// If the collection instance includes the parameter in its parameter map, then the parameter is overridden.
+	auto IsOverridden = [Collection = Collection, CollectionParameter]()
+		{
+			return GetBaseOverridesMap(Collection).Contains(CollectionParameter);
+		};
+
+	auto OnOverride = [Collection = Collection, BaseCollection, CollectionParameter, PropertyHandle](bool bOverride)
+		{
+			void* ValueData = nullptr;
+			if (PropertyHandle->GetValueData(ValueData) == FPropertyAccess::Success && ValueData)
+			{
+				FValueType& Value = *static_cast<FValueType*>(ValueData);
+
+				// Remove the parameter from the collection instance.
+				auto&& BaseOverridesMap = GetBaseOverridesMap(Collection);
+				BaseOverridesMap.Remove(CollectionParameter);
+
+				// Update the property value with the parameter value from the collection instance's parent.
+				Value = GetParameterValue(Collection, CollectionParameter, BaseCollection);
+
+				if (bOverride)
+				{
+					// Add the parameter to the collection instance, with the value specified by its parent.
+					BaseOverridesMap.Add(CollectionParameter, Value);
+				}
+			}
+		};
+
+	// Set the edit condition to use the override lambdas.
+	DetailGroup.AddPropertyRow(PropertyHandle.ToSharedRef()).
+		EditCondition(
+			TAttribute<bool>::Create(TAttribute<bool>::FGetter::CreateWeakLambda(Collection, IsOverridden)),
+			FOnBooleanValueChanged::CreateWeakLambda(Collection, OnOverride));
+}
+
+template<class FCollectionParameterType>
+void FMaterialParameterCollectionDetails::AddParameters(IDetailCategoryBuilder& DetailCategory, FName PropertyName, UMaterialParameterCollection* Collection, const TArray<FCollectionParameterType>& CollectionParameters, UMaterialParameterCollection* BaseCollection)
+{
+	IDetailGroup& DetailGroup = DetailCategory.AddGroup(PropertyName, FText::FromName(PropertyName), false, false);
+
+	if (CollectionParameters.IsEmpty())
+	{
+		return;
+	}
+
+	if (UScriptStruct* CollectionParameterStruct = FCollectionParameterType::StaticStruct())
+	{
+		for (auto& CollectionParameter : CollectionParameters)
+		{
+			TSharedPtr<FStructOnScope> StructOnScope = MakeShared<FStructOnScope>(CollectionParameterStruct);
+
+			if (IDetailPropertyRow* DetailPropertyRow = DetailCategory.AddExternalStructure(StructOnScope))
+			{
+				DetailPropertyRow->Visibility(EVisibility::Collapsed);
+
+				if (TSharedPtr<IPropertyHandle> ParentHandle = DetailPropertyRow->GetPropertyHandle())
+				{
+					if (TSharedPtr<IPropertyHandle> PropertyHandle = ParentHandle->GetChildHandle("DefaultValue"))
+					{
+						AddParameter(DetailGroup, PropertyHandle, Collection, CollectionParameter, BaseCollection);
+					}
+				}
+			}
+		}
+	}
+}
+
+void FMaterialParameterCollectionDetails::CustomizeCollectionDetails(UMaterialParameterCollection* Collection)
+{
+	FName ScalarParametersName = "ScalarParameters";
+	FName VectorParametersName = "VectorParameters";
+
+	TSharedRef<IPropertyHandle> ScalarParametersHandle = DetailLayout->GetProperty(ScalarParametersName, UMaterialParameterCollection::StaticClass());
+	TSharedRef<IPropertyHandle> VectorParametersHandle = DetailLayout->GetProperty(VectorParametersName, UMaterialParameterCollection::StaticClass());
+	TSharedRef<IPropertyHandle> BaseHandle = DetailLayout->GetProperty("Base", UMaterialParameterCollection::StaticClass());
+
+	DetailLayout->HideProperty(ScalarParametersHandle);
+	DetailLayout->HideProperty(VectorParametersHandle);
+	DetailLayout->HideProperty(BaseHandle);
+
+	IDetailCategoryBuilder& MaterialCategory = DetailLayout->EditCategory("Material");
+	MaterialCategory.AddProperty(ScalarParametersHandle);
+	MaterialCategory.AddProperty(VectorParametersHandle);
+	MaterialCategory.AddProperty(BaseHandle);
+
+	// Create a delegate to invoke when the Base changes.
+	TDelegate<void(const FPropertyChangedEvent&)> OnPropertyValueChangedDelegate = TDelegate<void(const FPropertyChangedEvent&)>::CreateWeakLambda(Collection,
+		[Collection, DetailLayout = DetailLayout, BaseHandle, BaseCollectionWeak = TWeakObjectPtr<UMaterialParameterCollection>(Collection->Base)](const FPropertyChangedEvent& PropertyChangedEvent)
+		{
+			// Determine if this is the final setting of the property value.
+			if (PropertyChangedEvent.ChangeType != EPropertyChangeType::ValueSet)
+			{
+				return;
+			}
+
+			// Update all properties to account for the new Base.
+			DetailLayout->ForceRefreshDetails();
+		});
+
+	BaseHandle->SetOnPropertyValueChangedWithData(OnPropertyValueChangedDelegate);
+	BaseHandle->SetOnPropertyResetToDefault(FSimpleDelegate::CreateLambda([DetailLayout = DetailLayout]()
+		{
+			DetailLayout->ForceRefreshDetails();
+		}));
+
+	for (UMaterialParameterCollection* BaseCollection = Collection->Base; BaseCollection; BaseCollection = BaseCollection->Base)
+	{
+		IDetailCategoryBuilder& BaseCategory = DetailLayout->EditCategory(BaseCollection->GetFName());
+		AddParameters(BaseCategory, ScalarParametersName, Collection, BaseCollection->ScalarParameters, BaseCollection);
+		AddParameters(BaseCategory, VectorParametersName, Collection, BaseCollection->VectorParameters, BaseCollection);
+	}
+}
+
+void FMaterialParameterCollectionDetails::CustomizeDetails(IDetailLayoutBuilder& InDetailLayout)
+{
+	DetailLayout = &InDetailLayout;
+
+	FEditorDelegates::PostUndoRedo.AddRaw(DetailLayout, &IDetailLayoutBuilder::ForceRefreshDetails);
+
+	TArray<TWeakObjectPtr<UMaterialParameterCollection>> Collections = DetailLayout->GetObjectsOfTypeBeingCustomized<UMaterialParameterCollection>();
+	for (TWeakObjectPtr<UMaterialParameterCollection>& CollectionWeak : Collections)
+	{
+		if (UMaterialParameterCollection* Collection = CollectionWeak.Get())
+		{
+			CustomizeCollectionDetails(Collection);
+		}
+	}
+}
+
+#undef LOCTEXT_NAMESPACE
diff --git a/Engine/Source/Editor/DetailCustomizations/Private/MaterialParameterCollectionDetails.h b/Engine/Source/Editor/DetailCustomizations/Private/MaterialParameterCollectionDetails.h
new file mode 100644
--- /dev/null
+++ b/Engine/Source/Editor/DetailCustomizations/Private/MaterialParameterCollectionDetails.h
@@ -0,0 +1,57 @@
+// Copyright Epic Games, Inc. All Rights Reserved.
+
+#pragma once
+
+#include "IDetailCustomization.h"
+
+class IDetailGroup;
+class IDetailCategoryBuilder;
+class IDetailLayoutBuilder;
+class IPropertyHandle;
+class UMaterialParameterCollection;
+
+struct FCollectionScalarParameter;
+struct FCollectionVectorParameter;
+struct FGuid;
+struct FLinearColor;
+
+//-----------------------------------------------------------------------------
+//   FMaterialParameterCollectionDetails
+//-----------------------------------------------------------------------------
+
+class FMaterialParameterCollectionDetails : public IDetailCustomization
+{
+public:
+	/** Makes a new instance of this detail layout class for a specific detail view requesting it */
+	static TSharedRef<IDetailCustomization> MakeInstance();
+
+	virtual ~FMaterialParameterCollectionDetails() override;
+
+	/** IDetailCustomization interface */
+	virtual void CustomizeDetails(IDetailLayoutBuilder& DetailLayout) override;
+
+private:
+	FMaterialParameterCollectionDetails() = default;
+
+	FMaterialParameterCollectionDetails(const FMaterialParameterCollectionDetails&) = delete;
+	FMaterialParameterCollectionDetails& operator =(const FMaterialParameterCollectionDetails&) = delete;
+
+	class FBaseOverridesSelector;
+	class FBaseOverridesSelectorConst;
+
+	static FBaseOverridesSelector GetBaseOverridesMap(UMaterialParameterCollection* OverrideCollection);
+	static FBaseOverridesSelectorConst GetBaseOverridesMap(const UMaterialParameterCollection* OverrideCollection);
+
+	template<class FCollectionParameterType>
+	static auto GetParameterValue(const UMaterialParameterCollection* Collection, const FCollectionParameterType& CollectionParameter, const UMaterialParameterCollection* BaseCollection = nullptr);
+
+	template<class FCollectionParameterType>
+	void AddParameter(IDetailGroup& DetailGroup, TSharedPtr<IPropertyHandle> PropertyHandle, UMaterialParameterCollection* Collection, const FCollectionParameterType& CollectionParameter, const UMaterialParameterCollection* BaseCollection);
+
+	template<class FCollectionParameterType>
+	void AddParameters(IDetailCategoryBuilder& DetailCategory, FName PropertyName, UMaterialParameterCollection* Collection, const TArray<FCollectionParameterType>& CollectionParameters, UMaterialParameterCollection* BaseCollection);
+
+	void CustomizeCollectionDetails(UMaterialParameterCollection* Collection);
+
+	IDetailLayoutBuilder* DetailLayout = nullptr;
+};
diff --git a/Engine/Source/Editor/MaterialEditor/Private/MaterialEditorInstanceDetailCustomization.cpp b/Engine/Source/Editor/MaterialEditor/Private/MaterialEditorInstanceDetailCustomization.cpp
--- a/Engine/Source/Editor/MaterialEditor/Private/MaterialEditorInstanceDetailCustomization.cpp
+++ b/Engine/Source/Editor/MaterialEditor/Private/MaterialEditorInstanceDetailCustomization.cpp
@@ -10,9 +10,11 @@
 #include "Widgets/SBoxPanel.h"
 #include "Widgets/Text/STextBlock.h"
 #include "Styling/AppStyle.h"
+#include "MaterialCachedData.h"
 #include "Materials/MaterialInterface.h"
 #include "MaterialEditor/DEditorFontParameterValue.h"
 #include "MaterialEditor/DEditorMaterialLayersParameterValue.h"
+#include "MaterialEditor/DEditorParameterCollectionParameterValue.h"
 #include "MaterialEditor/DEditorRuntimeVirtualTextureParameterValue.h"
 #include "MaterialEditor/DEditorSparseVolumeTextureParameterValue.h"
 #include "MaterialEditor/DEditorScalarParameterValue.h"
@@ -28,6 +30,7 @@
 #include "Materials/MaterialExpressionTextureSampleParameter.h"
 #include "Materials/MaterialExpressionFontSampleParameter.h"
 #include "Materials/MaterialExpressionMaterialAttributeLayers.h"
+#include "Materials/MaterialParameterCollection.h"
 #include "MaterialShared.h"
 #include "EditorSupportDelegates.h"
 #include "DetailWidgetRow.h"
@@ -502,6 +505,8 @@ void FMaterialInstanceParameterDetails::CreateSingleGroupWidget(FEditorParameter
 {
 	TSharedPtr<IPropertyHandle> ParametersArrayProperty = ParameterGroupProperty->GetChildHandle("Parameters");
 
+	TMap<UMaterialParameterCollection*, TArray<int32>> ParameterCollectionSharedOverrideMap;
+
 	// Create a custom widget for each parameter in the group
 	for (int32 ParamIdx = 0; ParamIdx < ParameterGroup.Parameters.Num(); ++ParamIdx)
 	{
@@ -519,6 +524,28 @@ void FMaterialInstanceParameterDetails::CreateSingleGroupWidget(FEditorParameter
 			UDEditorRuntimeVirtualTextureParameterValue* RuntimeVirtualTextureParam = Cast<UDEditorRuntimeVirtualTextureParameterValue>(Parameter);
 			UDEditorSparseVolumeTextureParameterValue* SparseVolumeTextureParam = Cast<UDEditorSparseVolumeTextureParameterValue>(Parameter);
 			UDEditorVectorParameterValue* VectorParam = Cast<UDEditorVectorParameterValue>(Parameter);
+			UDEditorParameterCollectionParameterValue* ParameterCollectionParam = Cast<UDEditorParameterCollectionParameterValue>(Parameter);
+
+			if (ParameterCollectionParam)
+			{
+				if (MaterialEditorInstance && MaterialEditorInstance->SourceInstance)
+				{
+					// Map each parameter collection to its referencing parameters.
+					for (const FMaterialParameterCollectionInfo& CollectionInfo : MaterialEditorInstance->SourceInstance->GetCachedExpressionData().ParameterCollectionInfos)
+					{
+						FMaterialParameterMetadata Result;
+						if (ParameterCollectionParam->GetValue(Result) && ensure(Result.Value.Type == EMaterialParameterType::ParameterCollection) &&
+							Result.Value.ParameterCollection && Result.Value.ParameterCollection->IsInstanceOf(CollectionInfo.StateId))
+						{
+							TArray<int32>& ParameterIndices = ParameterCollectionSharedOverrideMap.FindOrAdd(CollectionInfo.ParameterCollection);
+							ParameterIndices.Add(ParamIdx);
+							break;
+						}
+					}
+				}
+
+				continue;
+			}
 
 			// Don't display custom primitive data parameters in the details panel.
 			// This data is pulled from the primitive and can't be changed on the material.
@@ -591,6 +618,85 @@ void FMaterialInstanceParameterDetails::CreateSingleGroupWidget(FEditorParameter
 			}
 		}
 	}
+
+	for (auto& Entry : ParameterCollectionSharedOverrideMap)
+	{
+		UMaterialParameterCollection* ParameterCollection = Entry.Key;
+		const TArray<int32>& ParameterIndices = Entry.Value;
+		int32 NumParameterIndices = ParameterIndices.Num();
+		if (NumParameterIndices > 0)
+		{
+			UDEditorParameterValue* FirstParameter = ParameterGroup.Parameters[ParameterIndices[0]];
+			TSharedPtr<IPropertyHandle> FirstParameterProperty = ParametersArrayProperty->GetChildHandle(ParameterIndices[0]);
+
+			// Create a parameter value widget for the first parameter.
+			CreateParameterValueWidget(FirstParameter, FirstParameterProperty, DetailGroup);
+
+			if (TSharedPtr<IPropertyHandle> FirstParameterValueProperty = FirstParameterProperty->GetChildHandle("ParameterValue"))
+			{
+				// Display the parameter collection's name on the widget.
+				FirstParameterValueProperty->SetPropertyDisplayName(FText::FromName(ParameterCollection->GetFName()));
+
+				// Gather the relevant property handles of the collection's referencing parameters.
+				TArray<TSharedRef<IPropertyHandle>> PropertyHandles{ FirstParameterValueProperty.ToSharedRef() };
+				for (int32 ParameterIndex = 1; ParameterIndex < NumParameterIndices; ++ParameterIndex)
+				{
+					if (TSharedPtr<IPropertyHandle> ParameterProperty = ParametersArrayProperty->GetChildHandle(ParameterIndices[ParameterIndex]))
+					{
+						if (TSharedPtr<IPropertyHandle> ParameterValueProperty = ParameterProperty->GetChildHandle("ParameterValue"))
+						{
+							PropertyHandles.Add(ParameterValueProperty.ToSharedRef());
+						}
+					}
+				}
+
+				// Create a delegate to invoke when the property value changes.
+				TDelegate<void(const FPropertyChangedEvent&)> OnPropertyValueChangedDelegate = TDelegate<void(const FPropertyChangedEvent&)>::CreateWeakLambda(ParameterCollection, [ParameterCollection, PropertyHandles = MoveTemp(PropertyHandles)](const FPropertyChangedEvent& PropertyChangedEvent)
+					{
+						// Determine if this is the final setting of the property value.
+						if (PropertyChangedEvent.ChangeType != EPropertyChangeType::ValueSet)
+						{
+							return;
+						}
+
+						if (PropertyHandles.IsEmpty())
+						{
+							return;
+						}
+
+						// Get the new value from the first property handle.
+						UObject* Value = nullptr;
+						if (PropertyHandles[0]->GetValue(Value) != FPropertyAccess::Success)
+						{
+							return;
+						}
+
+						// Verify that the new value is valid.
+						if (auto* Collection = Cast<UMaterialParameterCollection>(Value))
+						{
+							if (Collection->IsInstanceOf(ParameterCollection->StateId))
+							{
+								// Set the same value on all of the collection's referencing parameters.
+								for (int32 PropertyHandleIndex = 1; PropertyHandleIndex < PropertyHandles.Num(); ++PropertyHandleIndex)
+								{
+									PropertyHandles[PropertyHandleIndex]->SetValue(Value);
+								}
+
+								return;
+							}
+						}
+
+						// Reset the value to the default parameter collection.
+						Value = ParameterCollection;
+						PropertyHandles[0]->SetValue(Value);
+					});
+
+				// Invoke the delegate when the property value or one of its child property values changes.
+				FirstParameterValueProperty->SetOnPropertyValueChangedWithData(OnPropertyValueChangedDelegate);
+				FirstParameterValueProperty->SetOnChildPropertyValueChangedWithData(OnPropertyValueChangedDelegate);
+			}
+		}
+	}
 }
 
 void FMaterialInstanceParameterDetails::CreateParameterValueWidget(UDEditorParameterValue* Parameter, TSharedPtr<IPropertyHandle> ParameterProperty, IDetailGroup& DetailGroup)
diff --git a/Engine/Source/Editor/MaterialEditor/Private/MaterialEditorUtilities.cpp b/Engine/Source/Editor/MaterialEditor/Private/MaterialEditorUtilities.cpp
--- a/Engine/Source/Editor/MaterialEditor/Private/MaterialEditorUtilities.cpp
+++ b/Engine/Source/Editor/MaterialEditor/Private/MaterialEditorUtilities.cpp
@@ -14,6 +14,7 @@
 #include "Materials/MaterialExpressionStaticBoolParameter.h"
 #include "Materials/MaterialExpressionStaticBool.h"
 #include "Materials/MaterialExpressionStaticSwitch.h"
+#include "Materials/MaterialExpressionCollectionParameter.h"
 #include "Materials/MaterialExpressionComment.h"
 #include "Materials/MaterialExpressionComposite.h"
 #include "Materials/MaterialExpressionTextureCollectionParameter.h"
@@ -461,6 +462,7 @@ void FMaterialEditorUtilities::GetVisibleMaterialParametersFromExpression(
 	UMaterialExpressionRuntimeVirtualTextureSampleParameter* RuntimeVirtualTexParam = Cast<UMaterialExpressionRuntimeVirtualTextureSampleParameter>(MaterialExpressionKey.Expression);
 	UMaterialExpressionSparseVolumeTextureSampleParameter* SparseVolumeTexParam = Cast<UMaterialExpressionSparseVolumeTextureSampleParameter>(MaterialExpressionKey.Expression);
 	UMaterialExpressionFontSampleParameter* FontParam = Cast<UMaterialExpressionFontSampleParameter>( MaterialExpressionKey.Expression );
+	UMaterialExpressionCollectionParameter* ParameterCollectionParam = Cast<UMaterialExpressionCollectionParameter>(MaterialExpressionKey.Expression);
 
 	if (Param)
 	{
@@ -486,8 +488,12 @@ void FMaterialEditorUtilities::GetVisibleMaterialParametersFromExpression(
 	{
 		ParameterInfo.Name = FontParam->ParameterName;
 	}
+	else if (ParameterCollectionParam)
+	{
+		ParameterInfo.Name = ParameterCollectionParam->ParameterName;
+	}
 		
-	if (Param || TexParam || TextureCollectionParam || FontParam || RuntimeVirtualTexParam || SparseVolumeTexParam)
+	if (Param || TexParam || TextureCollectionParam || FontParam || RuntimeVirtualTexParam || SparseVolumeTexParam || ParameterCollectionParam)
 	{
 		VisibleExpressions.AddUnique(ParameterInfo);
 	}
diff --git a/Engine/Source/Editor/UnrealEd/Classes/MaterialEditor/DEditorParameterCollectionParameterValue.h b/Engine/Source/Editor/UnrealEd/Classes/MaterialEditor/DEditorParameterCollectionParameterValue.h
new file mode 100644
--- /dev/null
+++ b/Engine/Source/Editor/UnrealEd/Classes/MaterialEditor/DEditorParameterCollectionParameterValue.h
@@ -0,0 +1,36 @@
+// Copyright Epic Games, Inc. All Rights Reserved.
+
+#pragma once
+
+#include "MaterialEditor/DEditorParameterValue.h"
+#include "DEditorParameterCollectionParameterValue.generated.h"
+
+class UMaterialParameterCollection;
+
+UCLASS(hidecategories=Object, collapsecategories, MinimalAPI)
+class UDEditorParameterCollectionParameterValue : public UDEditorParameterValue
+{
+	GENERATED_UCLASS_BODY()
+
+	UPROPERTY(EditAnywhere, Category=DEditorTextureCollectionParameterValue)
+	TObjectPtr<UMaterialParameterCollection> ParameterValue;
+
+	virtual FName GetDefaultGroupName() const override { return TEXT("Parameter Collection Values"); }
+
+	virtual bool GetValue(FMaterialParameterMetadata& OutResult) const override
+	{
+		UDEditorParameterValue::GetValue(OutResult);
+		OutResult.Value = ParameterValue;
+		return true;
+	}
+
+	virtual bool SetValue(const FMaterialParameterValue& Value) override
+	{
+		if (Value.Type == EMaterialParameterType::ParameterCollection)
+		{
+			ParameterValue = Value.ParameterCollection;
+			return true;
+		}
+		return false;
+	}
+};
diff --git a/Engine/Source/Editor/UnrealEd/Private/EditorParameters.cpp b/Engine/Source/Editor/UnrealEd/Private/EditorParameters.cpp
--- a/Engine/Source/Editor/UnrealEd/Private/EditorParameters.cpp
+++ b/Engine/Source/Editor/UnrealEd/Private/EditorParameters.cpp
@@ -12,6 +12,7 @@
 #include "MaterialEditor/DEditorStaticSwitchParameterValue.h"
 #include "MaterialEditor/DEditorTextureParameterValue.h"
 #include "MaterialEditor/DEditorTextureCollectionParameterValue.h"
+#include "MaterialEditor/DEditorParameterCollectionParameterValue.h"
 #include "MaterialEditor/DEditorVectorParameterValue.h"
 #include "MaterialEditor/DEditorEnumerationParameterValue.h"
 #include "Materials/MaterialEnumeration.h"
@@ -65,6 +66,11 @@ UDEditorTextureCollectionParameterValue::UDEditorTextureCollectionParameterValue
 {
 }
 
+UDEditorParameterCollectionParameterValue::UDEditorParameterCollectionParameterValue(const FObjectInitializer& ObjectInitializer)
+	: Super(ObjectInitializer)
+{
+}
+
 UDEditorVectorParameterValue::UDEditorVectorParameterValue(const FObjectInitializer& ObjectInitializer)
 	: Super(ObjectInitializer)
 {
@@ -208,6 +214,16 @@ static UDEditorParameterValue* CreateParameter_StaticSwitch(UObject* Owner, cons
 	return Parameter;
 }
 
+static UDEditorParameterValue* CreateParameter_ParameterCollection(UObject* Owner, const FMaterialParameterMetadata& Meta)
+{
+	UDEditorParameterCollectionParameterValue* Parameter = NewObject<UDEditorParameterCollectionParameterValue>(Owner);
+	if (Meta.Value.Type == EMaterialParameterType::ParameterCollection)
+	{
+		Parameter->ParameterValue = Meta.Value.ParameterCollection;
+	}
+	return Parameter;
+}
+
 static UDEditorParameterValue* CreateParameter_StaticComponentMask(UObject* Owner, const FMaterialParameterMetadata& Meta)
 {
 	UDEditorStaticComponentMaskParameterValue* Parameter = NewObject<UDEditorStaticComponentMaskParameterValue>(Owner);
@@ -249,6 +265,7 @@ UDEditorParameterValue* UDEditorParameterValue::Create(UObject* Owner,
 	case EMaterialParameterType::SparseVolumeTexture: Parameter = CreateParameter_SparseVolumeTexture(Owner, Meta); break;
 	case EMaterialParameterType::Font: Parameter = CreateParameter_Font(Owner, Meta); break;
 	case EMaterialParameterType::StaticSwitch: Parameter = CreateParameter_StaticSwitch(Owner, Meta); break;
+	case EMaterialParameterType::ParameterCollection: Parameter = CreateParameter_ParameterCollection(Owner, Meta); break;
 	case EMaterialParameterType::StaticComponentMask: Parameter = CreateParameter_StaticComponentMask(Owner, Meta); break;
 	case EMaterialParameterType::Enumeration: Parameter = CreateParameter_Enumeration(Owner, Meta); break;
 	default: break;
diff --git a/Engine/Source/Runtime/Engine/Private/Materials/MaterialCachedData.cpp b/Engine/Source/Runtime/Engine/Private/Materials/MaterialCachedData.cpp
--- a/Engine/Source/Runtime/Engine/Private/Materials/MaterialCachedData.cpp
+++ b/Engine/Source/Runtime/Engine/Private/Materials/MaterialCachedData.cpp
@@ -240,6 +240,10 @@ bool FMaterialCachedExpressionData::AddParameter(const FMaterialParameterInfo& P
 			OutReferencedTextureCollection = ParameterMeta.Value.TextureCollection;
 			break;
 
+		case EMaterialParameterType::ParameterCollection:
+			ParameterCollectionValues.Insert(ParameterMeta.Value.ParameterCollection, Index);
+			break;
+
 		case EMaterialParameterType::Font:
 			FontValues.Insert(ParameterMeta.Value.Font.Value, Index);
 			FontPageValues.Insert(ParameterMeta.Value.Font.Page, Index);
@@ -335,6 +339,14 @@ bool FMaterialCachedExpressionData::AddParameter(const FMaterialParameterInfo& P
 			}
 			break;
 
+		case EMaterialParameterType::ParameterCollection:
+			bSameValue = ParameterCollectionValues[Index] == ParameterMeta.Value.ParameterCollection;
+			if (!bSameValue)
+			{
+				OutErrorMessage = ReportError(ParameterCollectionValues[Index].ToString(), TSoftObjectPtr<UMaterialParameterCollection>(ParameterMeta.Value.ParameterCollection).ToString(), *PreviousEditorInfo, EditorInfo);
+			}
+			break;
+
 		case EMaterialParameterType::Font:
 			bSameValue = FontValues[Index] == ParameterMeta.Value.Font.Value && FontPageValues[Index] == ParameterMeta.Value.Font.Page;
 			if (!bSameValue)
@@ -900,6 +912,9 @@ void FMaterialCachedExpressionData::GetParameterValueByIndex(EMaterialParameterT
 	case EMaterialParameterType::TextureCollection:
 		OutResult.Value = TextureCollectionValues[ParameterIndex].LoadSynchronous();
 		break;
+	case EMaterialParameterType::ParameterCollection:
+		OutResult.Value = ParameterCollectionValues[ParameterIndex].LoadSynchronous();
+		break;
 	case EMaterialParameterType::RuntimeVirtualTexture:
 		OutResult.Value = RuntimeVirtualTextureValues[ParameterIndex].LoadSynchronous();
 		break;
diff --git a/Engine/Source/Runtime/Engine/Private/Materials/MaterialExpressions.cpp b/Engine/Source/Runtime/Engine/Private/Materials/MaterialExpressions.cpp
--- a/Engine/Source/Runtime/Engine/Private/Materials/MaterialExpressions.cpp
+++ b/Engine/Source/Runtime/Engine/Private/Materials/MaterialExpressions.cpp
@@ -31,6 +31,7 @@
 #include "Materials/MaterialFunctionMaterialLayerBlend.h"
 #include "Materials/MaterialFunctionInstance.h"
 #include "Materials/MaterialInstanceSupport.h"
+#include "Materials/MaterialParameterCollection.h"
 #include "Engine/TextureCube.h"
 #include "Engine/Texture2DArray.h"
 #include "Engine/TextureCubeArray.h"
@@ -13644,6 +13645,17 @@ bool UMaterialFunctionInterface::OverrideNamedStaticSwitchParameter(const FHashe
 	return false;
 }
 
+bool UMaterialFunctionInterface::OverrideNamedParameterCollectionParameter(const FHashedMaterialParameterInfo& ParameterInfo, class UMaterialParameterCollection*& OutValue)
+{
+	FMaterialParameterMetadata Meta;
+	if (GetParameterOverrideValue(EMaterialParameterType::ParameterCollection, ParameterInfo.GetName(), Meta))
+	{
+		OutValue = Meta.Value.ParameterCollection;
+		return true;
+	}
+	return false;
+}
+
 bool UMaterialFunctionInterface::OverrideNamedStaticComponentMaskParameter(const FHashedMaterialParameterInfo& ParameterInfo, bool& OutR, bool& OutG, bool& OutB, bool& OutA, FGuid& OutExpressionGuid)
 {
 	FMaterialParameterMetadata Meta;
@@ -15515,6 +15527,7 @@ bool UMaterialFunctionInstance::GetParameterOverrideValue(EMaterialParameterType
 	case EMaterialParameterType::SparseVolumeTexture: bResult = GameThread_GetParameterValue(SparseVolumeTextureParameterValues, ParameterInfo, OutResult); break;
 	case EMaterialParameterType::Font: bResult = GameThread_GetParameterValue(FontParameterValues, ParameterInfo, OutResult); break;
 	case EMaterialParameterType::StaticSwitch: bResult = GameThread_GetParameterValue(StaticSwitchParameterValues, ParameterInfo, OutResult); break;
+	case EMaterialParameterType::ParameterCollection: bResult = GameThread_GetParameterValue(ParameterCollectionParameterValues, ParameterInfo, OutResult); break;
 	case EMaterialParameterType::StaticComponentMask: bResult = GameThread_GetParameterValue(StaticComponentMaskParameterValues, ParameterInfo, OutResult); break;
 	case EMaterialParameterType::Enumeration: bResult = GameThread_GetParameterValue(EnumerationParameterValues, ParameterInfo, OutResult); break;
 	default: checkNoEntry(); break;
@@ -18637,6 +18650,29 @@ void UMaterialExpressionCollectionParameter::GetCaption(TArray<FString>& OutCapt
 	}
 }
 
+bool UMaterialExpressionCollectionParameter::GetParameterValue(FMaterialParameterMetadata& OutMeta) const
+{
+	OutMeta.Value = FMaterialParameterValue(Collection.Get());
+	OutMeta.Description = Desc;
+	OutMeta.ExpressionGuid = ExpressionGUID;
+	OutMeta.Group = Group;
+	OutMeta.SortPriority = SortPriority;
+	OutMeta.AssetPath = GetAssetPathName();
+	return true;
+}
+
+bool UMaterialExpressionCollectionParameter::SetParameterValue(const FName& Name, const FMaterialParameterMetadata& Meta, EMaterialExpressionSetParameterValueFlags Flags)
+{
+	if (Name == ParameterName && Meta.Value.Type == EMaterialParameterType::ParameterCollection)
+	{
+		if (UMaterialParameterCollection* ParameterCollection = Cast<UMaterialParameterCollection>(Meta.Value.ParameterCollection))
+		{
+			Collection = ParameterCollection;
+			return true;
+		}
+	}
+	return false;
+}
 
 bool UMaterialExpressionCollectionParameter::MatchesSearchQuery(const TCHAR* SearchQuery)
 {
diff --git a/Engine/Source/Runtime/Engine/Private/Materials/MaterialIRTypes.cpp b/Engine/Source/Runtime/Engine/Private/Materials/MaterialIRTypes.cpp
--- a/Engine/Source/Runtime/Engine/Private/Materials/MaterialIRTypes.cpp
+++ b/Engine/Source/Runtime/Engine/Private/Materials/MaterialIRTypes.cpp
@@ -115,6 +115,7 @@ const FType* FType::FromMaterialParameterType(EMaterialParameterType Type)
 		case EMaterialParameterType::DoubleVector: UE_MIR_TODO();
 		case EMaterialParameterType::Texture: return FObjectType::GetTexture2D();
 		case EMaterialParameterType::TextureCollection: UE_MIR_TODO();
+		case EMaterialParameterType::ParameterCollection: UE_MIR_TODO();
 		case EMaterialParameterType::Font: UE_MIR_TODO();
 		case EMaterialParameterType::RuntimeVirtualTexture: UE_MIR_TODO();
 		case EMaterialParameterType::SparseVolumeTexture: UE_MIR_TODO();
diff --git a/Engine/Source/Runtime/Engine/Private/Materials/MaterialInstance.cpp b/Engine/Source/Runtime/Engine/Private/Materials/MaterialInstance.cpp
--- a/Engine/Source/Runtime/Engine/Private/Materials/MaterialInstance.cpp
+++ b/Engine/Source/Runtime/Engine/Private/Materials/MaterialInstance.cpp
@@ -24,6 +24,7 @@
 #include "Materials/MaterialExpressionVectorParameter.h"
 #include "Materials/MaterialExpressionDoubleVectorParameter.h"
 #include "Materials/MaterialExpressionTextureCollectionParameter.h"
+#include "Materials/MaterialExpressionCollectionParameter.h"
 #include "Materials/MaterialExpressionTextureSampleParameter.h"
 #include "Materials/MaterialExpressionFontSampleParameter.h"
 #include "Materials/MaterialExpressionMaterialAttributeLayers.h"
@@ -35,6 +36,7 @@
 #include "Materials/MaterialInstanceConstant.h"
 #include "Materials/MaterialInstanceUpdateParameterSet.h"
 #include "Materials/MaterialInstanceSupport.h"
+#include "Materials/MaterialParameterCollection.h"
 #include "Materials/MaterialSharedPrivate.h"
 #include "Engine/SubsurfaceProfile.h"
 #include "Engine/SpecularProfile.h"
@@ -348,6 +350,7 @@ bool FMaterialInstanceResource::GetParameterValue(EMaterialParameterType Type, c
 			case EMaterialParameterType::TextureCollection: bResult = RenderThread_GetParameterValue<const UTextureCollection*>(ParameterInfo, OutValue); break;
 			case EMaterialParameterType::RuntimeVirtualTexture: bResult = RenderThread_GetParameterValue<const URuntimeVirtualTexture*>(ParameterInfo, OutValue); break;
 			case EMaterialParameterType::SparseVolumeTexture: bResult = RenderThread_GetParameterValue<const USparseVolumeTexture*>(ParameterInfo, OutValue); break;
+			case EMaterialParameterType::ParameterCollection: bResult = RenderThread_GetParameterValue<const UMaterialParameterCollection*>(ParameterInfo, OutValue); break;
 			case EMaterialParameterType::Enumeration: bResult = RenderThread_GetParameterValue<FName>(ParameterInfo, OutValue); break;
 			default: ensure(false); break; // other parameter types are not expected on the render thread
 		}
@@ -513,6 +516,7 @@ void FMaterialInstanceResource::InitMIParameters(FMaterialInstanceParameterSet&
 	ParameterSet.TextureCollectionParameters.Sort(SortMaterialInstanceParametersPredicate<const UTextureCollection*>);
 	ParameterSet.RuntimeVirtualTextureParameters.Sort(SortMaterialInstanceParametersPredicate<const URuntimeVirtualTexture*>);
 	ParameterSet.SparseVolumeTextureParameters.Sort(SortMaterialInstanceParametersPredicate<const USparseVolumeTexture*>);
+	ParameterSet.ParameterCollectionParameters.Sort(SortMaterialInstanceParametersPredicate<const UMaterialParameterCollection*>);
 	ParameterSet.EnumerationParameters.Sort(SortMaterialInstanceParametersPredicate<FName>);
 
 	StaticSwitchParameterArray.Array = MoveTemp(ParameterSet.StaticSwitchParameters);
@@ -523,6 +527,7 @@ void FMaterialInstanceResource::InitMIParameters(FMaterialInstanceParameterSet&
 	TextureCollectionParameterArray.Array = MoveTemp(ParameterSet.TextureCollectionParameters);
 	RuntimeVirtualTextureParameterArray.Array = MoveTemp(ParameterSet.RuntimeVirtualTextureParameters);
 	SparseVolumeTextureParameterArray.Array = MoveTemp(ParameterSet.SparseVolumeTextureParameters);
+	ParameterCollectionParameterArray.Array = MoveTemp(ParameterSet.ParameterCollectionParameters);
 	EnumerationParameterArray.Array = MoveTemp(ParameterSet.EnumerationParameters);
 	UserSceneTextureOverrides = MoveTemp(ParameterSet.UserSceneTextureOverrides);
 	PostProcessBlendableOverrides = ParameterSet.PostProcessBlendableOverrides;
@@ -536,9 +541,40 @@ void FMaterialInstanceResource::InitMIParameters(FMaterialInstanceParameterSet&
 	TextureCollectionParameterArray.HashAddAllItems();
 	RuntimeVirtualTextureParameterArray.HashAddAllItems();
 	SparseVolumeTextureParameterArray.HashAddAllItems();
+	ParameterCollectionParameterArray.HashAddAllItems();
 	EnumerationParameterArray.HashAddAllItems();
 }
 
+void FMaterialInstanceResource::EvaluateParameterCollections(FRHICommandListBase& RHICmdList, FUniformExpressionCache& OutUniformExpressionCache, const FMaterialRenderContext& Context, FUniformExpressionCacheAsyncUpdater* Updater) const
+{
+	FMaterialRenderProxy::EvaluateParameterCollections(RHICmdList, OutUniformExpressionCache, Context, Updater);
+
+	for (FGuid& CollectionStateId : OutUniformExpressionCache.ParameterCollections)
+	{
+		for (auto&& Entry : ParameterCollectionParameterArray.Array)
+		{
+			if (Entry.Value && Entry.Value->IsInstanceOf(CollectionStateId))
+			{
+				CollectionStateId = Entry.Value->StateId;
+				break;
+			}
+		}
+	}
+}
+
+bool UMaterialParameterCollection::IsInstanceOf(const FGuid& CollectionStateId) const
+{
+	for (const UMaterialParameterCollection* BaseCollection = this; BaseCollection; BaseCollection = BaseCollection->Base)
+	{
+		if (BaseCollection->StateId == CollectionStateId)
+		{
+			return true;
+		}
+	}
+
+	return false;
+}
+
 /**
 * Updates a parameter on the material instance from the game thread.
 */
@@ -689,6 +725,7 @@ void UMaterialInstance::SwapLayerParameterIndices(int32 OriginalIndex, int32 New
 		SwapLayerParameterIndicesArray(DoubleVectorParameterValues, OriginalIndex, NewIndex);
 		SwapLayerParameterIndicesArray(TextureParameterValues, OriginalIndex, NewIndex);
 		SwapLayerParameterIndicesArray(TextureCollectionParameterValues, OriginalIndex, NewIndex);
+		SwapLayerParameterIndicesArray(ParameterCollectionParameterValues, OriginalIndex, NewIndex);
 		SwapLayerParameterIndicesArray(RuntimeVirtualTextureParameterValues, OriginalIndex, NewIndex);
 		SwapLayerParameterIndicesArray(SparseVolumeTextureParameterValues, OriginalIndex, NewIndex);
 		SwapLayerParameterIndicesArray(FontParameterValues, OriginalIndex, NewIndex);
@@ -709,6 +746,7 @@ void UMaterialInstance::RemoveLayerParameterIndex(int32 Index)
 	RemoveLayerParameterIndicesArray(DoubleVectorParameterValues, Index);
 	RemoveLayerParameterIndicesArray(TextureParameterValues, Index);
 	RemoveLayerParameterIndicesArray(TextureCollectionParameterValues, Index);
+	RemoveLayerParameterIndicesArray(ParameterCollectionParameterValues, Index);
 	RemoveLayerParameterIndicesArray(RuntimeVirtualTextureParameterValues, Index);
 	RemoveLayerParameterIndicesArray(SparseVolumeTextureParameterValues, Index);
 	RemoveLayerParameterIndicesArray(FontParameterValues, Index);
@@ -762,6 +800,9 @@ bool UMaterialInstance::UpdateParameters()
 			// Texture Collection parameters
 			bDirty = UpdateParameterSet<FTextureCollectionParameterValue, UMaterialExpressionTextureCollectionParameter>(TextureCollectionParameterValues, ParentMaterial) || bDirty;
 
+			// Parameter Collection parameters
+			bDirty = UpdateParameterSet<FParameterCollectionParameterValue, UMaterialExpressionCollectionParameter>(ParameterCollectionParameterValues, ParentMaterial) || bDirty;
+
 			// Runtime Virtual Texture parameters
 			bDirty = UpdateParameterSet<FRuntimeVirtualTextureParameterValue, UMaterialExpressionRuntimeVirtualTextureSampleParameter>(RuntimeVirtualTextureParameterValues, ParentMaterial) || bDirty;
 
@@ -798,6 +839,7 @@ bool UMaterialInstance::UpdateParameters()
 					RemapLayerParameterIndicesArray(DoubleVectorParameterValues, RemapLayerIndices);
 					RemapLayerParameterIndicesArray(TextureParameterValues, RemapLayerIndices);
 					RemapLayerParameterIndicesArray(TextureCollectionParameterValues, RemapLayerIndices);
+					RemapLayerParameterIndicesArray(ParameterCollectionParameterValues, RemapLayerIndices);
 					RemapLayerParameterIndicesArray(RuntimeVirtualTextureParameterValues, RemapLayerIndices);
 					RemapLayerParameterIndicesArray(SparseVolumeTextureParameterValues, RemapLayerIndices);
 					RemapLayerParameterIndicesArray(FontParameterValues, RemapLayerIndices);
@@ -912,6 +954,14 @@ void GameThread_InitMIParameters(const UMaterialInstance& Instance)
 		ParamRef.Value = FTextureCollectionParameterValue::GetValue(Parameter);
 	}
 
+	ParameterSet.ParameterCollectionParameters.Reserve(Instance.ParameterCollectionParameterValues.Num());
+	for (const FParameterCollectionParameterValue& Parameter : Instance.ParameterCollectionParameterValues)
+	{
+		auto& ParamRef = ParameterSet.ParameterCollectionParameters.AddDefaulted_GetRef();
+		ParamRef.Info = Parameter.ParameterInfo;
+		ParamRef.Value = FParameterCollectionParameterValue::GetValue(Parameter);
+	}
+
 	// RuntimeVirtualTexture parameters
 	ParameterSet.RuntimeVirtualTextureParameters.Reserve(Instance.RuntimeVirtualTextureParameterValues.Num());
 	for (const FRuntimeVirtualTextureParameterValue& Parameter : Instance.RuntimeVirtualTextureParameterValues)
@@ -1106,6 +1156,7 @@ bool UMaterialInstance::GetParameterOverrideValue(EMaterialParameterType Type, c
 	case EMaterialParameterType::DoubleVector: bResult = GameThread_GetParameterValue(DoubleVectorParameterValues, ParameterInfo, OutResult); break;
 	case EMaterialParameterType::Texture: bResult = GameThread_GetParameterValue(TextureParameterValues, ParameterInfo, OutResult); break;
 	case EMaterialParameterType::TextureCollection: bResult = GameThread_GetParameterValue(TextureCollectionParameterValues, ParameterInfo, OutResult); break;
+	case EMaterialParameterType::ParameterCollection: bResult = GameThread_GetParameterValue(ParameterCollectionParameterValues, ParameterInfo, OutResult); break;
 	case EMaterialParameterType::RuntimeVirtualTexture: bResult = GameThread_GetParameterValue(RuntimeVirtualTextureParameterValues, ParameterInfo, OutResult); break;
 	case EMaterialParameterType::SparseVolumeTexture: bResult = GameThread_GetParameterValue(SparseVolumeTextureParameterValues, ParameterInfo, OutResult); break;
 	case EMaterialParameterType::Font: bResult = GameThread_GetParameterValue(FontParameterValues, ParameterInfo, OutResult); break;
@@ -2234,6 +2285,7 @@ void UMaterialInstance::GetAllParametersOfType(EMaterialParameterType Type, TMap
 		case EMaterialParameterType::SparseVolumeTexture: GameThread_ApplyParameterOverrides(Instance->SparseVolumeTextureParameterValues, LayerIndexRemap, bSetOverride, OverridenParameters, OutParameters); break;
 		case EMaterialParameterType::Font: GameThread_ApplyParameterOverrides(Instance->FontParameterValues, LayerIndexRemap, bSetOverride, OverridenParameters, OutParameters); break;
 		case EMaterialParameterType::StaticSwitch: GameThread_ApplyParameterOverrides(Instance->StaticParametersRuntime.StaticSwitchParameters, LayerIndexRemap, bSetOverride, OverridenParameters, OutParameters); break;
+		case EMaterialParameterType::ParameterCollection: GameThread_ApplyParameterOverrides(Instance->ParameterCollectionParameterValues, LayerIndexRemap, bSetOverride, OverridenParameters, OutParameters); break;
 		case EMaterialParameterType::Enumeration: GameThread_ApplyParameterOverrides(Instance->EnumerationParameterValues, LayerIndexRemap, bSetOverride, OverridenParameters, OutParameters); break;
 #if WITH_EDITORONLY_DATA
 		case EMaterialParameterType::StaticComponentMask:
@@ -3442,6 +3494,17 @@ void UMaterialInstance::PostLoad()
 		}
 	}
 
+	// We have to make sure the resources are created for all used parameter collections.
+	for (int32 ValueIndex = 0; ValueIndex < ParameterCollectionParameterValues.Num(); ValueIndex++)
+	{
+		// Make sure the collection is postloaded so the resource isn't null.
+		UMaterialParameterCollection* Collection = ParameterCollectionParameterValues[ValueIndex].ParameterValue;
+		if (Collection)
+		{
+			Collection->ConditionalPostLoad();
+		}
+	}
+
 	// do the same for runtime virtual textures
 	for (int32 ValueIndex = 0; ValueIndex < RuntimeVirtualTextureParameterValues.Num(); ValueIndex++)
 	{
@@ -3923,6 +3986,7 @@ void UMaterialInstance::ReserveParameterValuesInternal(EMaterialParameterType Ty
 	case EMaterialParameterType::DoubleVector: DoubleVectorParameterValues.Reserve(Capacity); break;
 	case EMaterialParameterType::Texture: TextureParameterValues.Reserve(Capacity); break;
 	case EMaterialParameterType::TextureCollection: TextureCollectionParameterValues.Reserve(Capacity); break;
+	case EMaterialParameterType::ParameterCollection: ParameterCollectionParameterValues.Reserve(Capacity); break;
 	case EMaterialParameterType::Font: FontParameterValues.Reserve(Capacity); break;
 	case EMaterialParameterType::RuntimeVirtualTexture: RuntimeVirtualTextureParameterValues.Reserve(Capacity); break;
 	case EMaterialParameterType::SparseVolumeTexture: SparseVolumeTextureParameterValues.Reserve(Capacity); break;
@@ -3953,6 +4017,7 @@ void UMaterialInstance::AddParameterValueInternal(const FMaterialParameterInfo&
 	case EMaterialParameterType::DoubleVector: DoubleVectorParameterValues.Emplace(ParameterInfo, Value.AsVector4d()); break;
 	case EMaterialParameterType::Texture: TextureParameterValues.Emplace(ParameterInfo, Value.Texture); break;
 	case EMaterialParameterType::TextureCollection: TextureCollectionParameterValues.Emplace(ParameterInfo, Value.TextureCollection); break;
+	case EMaterialParameterType::ParameterCollection: ParameterCollectionParameterValues.Emplace(ParameterInfo, Value.ParameterCollection); break;
 	case EMaterialParameterType::Font: FontParameterValues.Emplace(ParameterInfo, Value.Font.Value, Value.Font.Page); break;
 	case EMaterialParameterType::RuntimeVirtualTexture: RuntimeVirtualTextureParameterValues.Emplace(ParameterInfo, Value.RuntimeVirtualTexture); break;
 	case EMaterialParameterType::SparseVolumeTexture: SparseVolumeTextureParameterValues.Emplace(ParameterInfo, Value.SparseVolumeTexture); break;
@@ -3987,6 +4052,7 @@ void UMaterialInstance::SetParameterValueInternal(const FMaterialParameterInfo&
 	case EMaterialParameterType::Font: SetFontParameterValueInternal(ParameterInfo, Value.Font.Value, Value.Font.Page); break;
 	case EMaterialParameterType::RuntimeVirtualTexture: SetRuntimeVirtualTextureParameterValueInternal(ParameterInfo, Value.RuntimeVirtualTexture); break;
 	case EMaterialParameterType::SparseVolumeTexture: SetSparseVolumeTextureParameterValueInternal(ParameterInfo, Value.SparseVolumeTexture); break;
+	case EMaterialParameterType::ParameterCollection: SetParameterCollectionParameterValueInternal(ParameterInfo, Value.ParameterCollection); break;
 	case EMaterialParameterType::Enumeration: SetEnumerationParameterValueInternal(ParameterInfo, Value.EnumEntry); break;
 	default: checkNoEntry();
 	}
@@ -4214,6 +4280,45 @@ void UMaterialInstance::SetTextureCollectionParameterValueInternal(const FMateri
 	}
 }
 
+void UMaterialInstance::SetParameterCollectionParameterValueInternal(const FMaterialParameterInfo& ParameterInfo, UMaterialParameterCollection* Value)
+{
+	LLM_SCOPE(ELLMTag::MaterialInstance);
+
+	if (Value)
+	{
+		Value->ConditionalPostLoad();
+	}
+
+	FParameterCollectionParameterValue* ParameterValue = GameThread_FindParameterByName(ParameterCollectionParameterValues, ParameterInfo);
+
+	bool bForceUpdate = false;
+	if(!ParameterValue)
+	{
+		// If there's no element for the named parameter in array yet, add one.
+		ParameterValue = new(ParameterCollectionParameterValues) FParameterCollectionParameterValue;
+		ParameterValue->ParameterInfo = ParameterInfo;
+		ParameterValue->ExpressionGUID.Invalidate();
+		bForceUpdate = true;
+	}
+
+	// Don't enqueue an update if it isn't needed
+	if (bForceUpdate || ParameterValue->ParameterValue != Value)
+	{
+		// set as an ensure, because it is somehow possible to accidentally pass non-parameters into here via blueprints...
+		if (Value && ensureMsgf(Value->IsA(UMaterialParameterCollection::StaticClass()), TEXT("Expecting a UMaterialParameterCollection! Value='%s' class='%s'"), *Value->GetName(), *Value->GetClass()->GetName()))
+		{
+			ParameterValue->ParameterValue = Value;
+			Value->AddToCluster(this, true);
+			// Update the material instance data in the rendering thread.
+			GameThread_UpdateMIParameter(this, *ParameterValue);
+
+#if WITH_EDITOR
+			FObjectCacheEventSink::NotifyMaterialChanged_Concurrent(this);
+#endif
+		}		
+	}
+}
+
 void UMaterialInstance::SetRuntimeVirtualTextureParameterValueInternal(const FMaterialParameterInfo& ParameterInfo, URuntimeVirtualTexture* Value)
 {
 	LLM_SCOPE(ELLMTag::MaterialInstance);
@@ -4362,6 +4467,7 @@ void UMaterialInstance::ClearParameterValuesInternal(EMaterialInstanceClearParam
 		ScalarParameterValues.Empty();
 		VectorParameterValues.Empty();
 		DoubleVectorParameterValues.Empty();
+		ParameterCollectionParameterValues.Empty();
 		EnumerationParameterValues.Empty();
 		bUpdateResource = true;
 	}
@@ -4780,6 +4886,7 @@ void UMaterialInstance::GetResourceSizeEx(FResourceSizeEx& CumulativeResourceSiz
 		CumulativeResourceSize.AddDedicatedSystemMemoryBytes(RuntimeVirtualTextureParameterValues.Num() * sizeof(THashedMaterialParameterMap<const URuntimeVirtualTexture*>::TNamedParameter));
 		CumulativeResourceSize.AddDedicatedSystemMemoryBytes(SparseVolumeTextureParameterValues.Num() * sizeof(THashedMaterialParameterMap<const USparseVolumeTexture*>::TNamedParameter));
 		CumulativeResourceSize.AddDedicatedSystemMemoryBytes(FontParameterValues.Num() * sizeof(THashedMaterialParameterMap<const UTexture*>::TNamedParameter));
+		CumulativeResourceSize.AddDedicatedSystemMemoryBytes(ParameterCollectionParameterValues.Num() * sizeof(THashedMaterialParameterMap<const UMaterialParameterCollection*>::TNamedParameter));
 
 		// Record space for hash tables as well..
 		if (ScalarParameterValues.Num()) CumulativeResourceSize.AddDedicatedSystemMemoryBytes(FDefaultSetAllocator::GetNumberOfHashBuckets(ScalarParameterValues.Num()) * sizeof(uint16));
@@ -4789,6 +4896,7 @@ void UMaterialInstance::GetResourceSizeEx(FResourceSizeEx& CumulativeResourceSiz
 		if (RuntimeVirtualTextureParameterValues.Num()) CumulativeResourceSize.AddDedicatedSystemMemoryBytes(FDefaultSetAllocator::GetNumberOfHashBuckets(RuntimeVirtualTextureParameterValues.Num()) * sizeof(uint16));
 		if (SparseVolumeTextureParameterValues.Num()) CumulativeResourceSize.AddDedicatedSystemMemoryBytes(FDefaultSetAllocator::GetNumberOfHashBuckets(SparseVolumeTextureParameterValues.Num()) * sizeof(uint16));
 		if (FontParameterValues.Num()) CumulativeResourceSize.AddDedicatedSystemMemoryBytes(FDefaultSetAllocator::GetNumberOfHashBuckets(FontParameterValues.Num()) * sizeof(uint16));
+		if (ParameterCollectionParameterValues.Num()) CumulativeResourceSize.AddDedicatedSystemMemoryBytes(FDefaultSetAllocator::GetNumberOfHashBuckets(ParameterCollectionParameterValues.Num()) * sizeof(uint16));
 	}
 }
 
@@ -5298,6 +5406,10 @@ bool UMaterialInstance::Equivalent(const UMaterialInstance* CompareTo) const
 	{
 		return false;
 	}
+	if (!CompareValueArraysByExpressionGUID(ParameterCollectionParameterValues, CompareTo->ParameterCollectionParameterValues))
+	{
+		return false;
+	}
 	if (!CompareValueArraysByExpressionGUID(EnumerationParameterValues, CompareTo->EnumerationParameterValues))
 	{
 		return false;
@@ -5368,7 +5480,8 @@ bool UMaterialInstance::IsRedundant() const
 		|| !VectorParameterValues.IsEmpty()
 		|| !DoubleVectorParameterValues.IsEmpty()
 		|| !RuntimeVirtualTextureParameterValues.IsEmpty()
-		|| !FontParameterValues.IsEmpty())
+		|| !FontParameterValues.IsEmpty()
+		|| !ParameterCollectionParameterValues.IsEmpty())
 	{
 		return false;
 	}
@@ -5595,6 +5708,7 @@ void UMaterialInstance::CopyMaterialUniformParametersInternal(UMaterialInterface
 				MergeParameterOverrides(TextureParameterValues, AsInstance->TextureParameterValues);
 				MergeParameterOverrides(RuntimeVirtualTextureParameterValues, AsInstance->RuntimeVirtualTextureParameterValues);
 				MergeParameterOverrides(SparseVolumeTextureParameterValues, AsInstance->SparseVolumeTextureParameterValues);
+				MergeParameterOverrides(ParameterCollectionParameterValues, AsInstance->ParameterCollectionParameterValues);
 				// No fonts?
 			}
 			else if (UMaterial* AsMaterial = Cast<UMaterial>(Interface))
@@ -5606,6 +5720,7 @@ void UMaterialInstance::CopyMaterialUniformParametersInternal(UMaterialInterface
 				checkSlow(TextureParameterValues.Num() == 0);
 				checkSlow(RuntimeVirtualTextureParameterValues.Num() == 0);
 				checkSlow(SparseVolumeTextureParameterValues.Num() == 0);
+				checkSlow(ParameterCollectionParameterValues.Num() == 0);
 
 				const FMaterialResource* MaterialResource = nullptr;
 				if (UWorld* World = AsMaterial->GetWorld())
diff --git a/Engine/Source/Runtime/Engine/Private/Materials/MaterialInstanceSupport.h b/Engine/Source/Runtime/Engine/Private/Materials/MaterialInstanceSupport.h
--- a/Engine/Source/Runtime/Engine/Private/Materials/MaterialInstanceSupport.h
+++ b/Engine/Source/Runtime/Engine/Private/Materials/MaterialInstanceSupport.h
@@ -248,6 +248,7 @@ public:
 		TextureCollectionParameterArray.Empty();
 		RuntimeVirtualTextureParameterArray.Empty();
 		SparseVolumeTextureParameterArray.Empty();
+		ParameterCollectionParameterArray.Empty();
 		UserSceneTextureOverrides.Empty();
 		EnumerationParameterArray.Empty();
 	}
@@ -308,6 +309,8 @@ public:
 	}
 
 private:
+	virtual void EvaluateParameterCollections(FRHICommandListBase& RHICmdList, FUniformExpressionCache& OutUniformExpressionCache, const FMaterialRenderContext& Context, FUniformExpressionCacheAsyncUpdater* Updater) const override;
+
 	/**
 	 * Retrieves the array of values for a given type.
 	 */
@@ -319,6 +322,7 @@ private:
 	static bool IsValidParameterValue(const FVector4d&) { return true; }
 	static bool IsValidParameterValue(const UTexture* Value) { return Value != nullptr; }
 	static bool IsValidParameterValue(const UTextureCollection* Value) { return Value != nullptr; }
+	static bool IsValidParameterValue(const UMaterialParameterCollection* Value) { return Value != nullptr; }
 	static bool IsValidParameterValue(const URuntimeVirtualTexture* Value) { return Value != nullptr; }
 	static bool IsValidParameterValue(const USparseVolumeTexture* Value) { return Value != nullptr; }
 	static bool IsValidParameterValue(FName Value) { return true; }
@@ -413,6 +417,8 @@ private:
 	THashedMaterialParameterMap<const URuntimeVirtualTexture*> RuntimeVirtualTextureParameterArray;
 	/** Sparse Volume Texture parameters for this material instance. */
 	THashedMaterialParameterMap<const USparseVolumeTexture*> SparseVolumeTextureParameterArray;
+	/** Parameter collection parameters for this material instance. */
+	THashedMaterialParameterMap<const UMaterialParameterCollection*> ParameterCollectionParameterArray;
 	/** Enumeration parameters for this material instance. */
 	THashedMaterialParameterMap<FName> EnumerationParameterArray;
 	/** Remap layer indices for parent */
@@ -431,6 +437,7 @@ template <> FORCEINLINE THashedMaterialParameterMap<const UTexture*>& FMaterialI
 template <> FORCEINLINE THashedMaterialParameterMap<const UTextureCollection*>& FMaterialInstanceResource::GetValueArray() { return TextureCollectionParameterArray; }
 template <> FORCEINLINE THashedMaterialParameterMap<const URuntimeVirtualTexture*>& FMaterialInstanceResource::GetValueArray() { return RuntimeVirtualTextureParameterArray; }
 template <> FORCEINLINE THashedMaterialParameterMap<const USparseVolumeTexture*>& FMaterialInstanceResource::GetValueArray() { return SparseVolumeTextureParameterArray; }
+template <> FORCEINLINE THashedMaterialParameterMap<const UMaterialParameterCollection*>& FMaterialInstanceResource::GetValueArray() { return ParameterCollectionParameterArray; }
 template <> FORCEINLINE THashedMaterialParameterMap<FName>& FMaterialInstanceResource::GetValueArray() { return EnumerationParameterArray; }
 template <> FORCEINLINE const THashedMaterialParameterMap<bool>& FMaterialInstanceResource::GetValueArray() const { return StaticSwitchParameterArray; }
 template <> FORCEINLINE const THashedMaterialParameterMap<float>& FMaterialInstanceResource::GetValueArray() const { return ScalarParameterArray; }
@@ -440,21 +447,23 @@ template <> FORCEINLINE const THashedMaterialParameterMap<const UTexture*>& FMat
 template <> FORCEINLINE const THashedMaterialParameterMap<const UTextureCollection*>& FMaterialInstanceResource::GetValueArray() const { return TextureCollectionParameterArray; }
 template <> FORCEINLINE const THashedMaterialParameterMap<const URuntimeVirtualTexture*>& FMaterialInstanceResource::GetValueArray() const { return RuntimeVirtualTextureParameterArray; }
 template <> FORCEINLINE const THashedMaterialParameterMap<const USparseVolumeTexture*>& FMaterialInstanceResource::GetValueArray() const { return SparseVolumeTextureParameterArray; }
+template <> FORCEINLINE const THashedMaterialParameterMap<const UMaterialParameterCollection*>& FMaterialInstanceResource::GetValueArray() const { return ParameterCollectionParameterArray; }
 template <> FORCEINLINE const THashedMaterialParameterMap<FName>& FMaterialInstanceResource::GetValueArray() const { return EnumerationParameterArray; }
 
 struct FMaterialInstanceParameterSet
 {
-	TArray<THashedMaterialParameterMap<bool>::TNamedParameter>							StaticSwitchParameters;
-	TArray<THashedMaterialParameterMap<float>::TNamedParameter>							ScalarParameters;
-	TArray<THashedMaterialParameterMap<FLinearColor>::TNamedParameter>					VectorParameters;
-	TArray<THashedMaterialParameterMap<FVector4d>::TNamedParameter>						DoubleVectorParameters;
-	TArray<THashedMaterialParameterMap<const UTexture*>::TNamedParameter>				TextureParameters;
-	TArray<THashedMaterialParameterMap<const UTextureCollection*>::TNamedParameter>		TextureCollectionParameters;
-	TArray<THashedMaterialParameterMap<const URuntimeVirtualTexture*>::TNamedParameter>	RuntimeVirtualTextureParameters;
-	TArray<THashedMaterialParameterMap<const USparseVolumeTexture*>::TNamedParameter>	SparseVolumeTextureParameters;
-	TArray<THashedMaterialParameterMap<FName>::TNamedParameter>							EnumerationParameters;
-	TArray<FUserSceneTextureOverride>													UserSceneTextureOverrides;
-	FPostProcessBlendableOverrides														PostProcessBlendableOverrides;
+	TArray<THashedMaterialParameterMap<bool>::TNamedParameter>									StaticSwitchParameters;
+	TArray<THashedMaterialParameterMap<float>::TNamedParameter>									ScalarParameters;
+	TArray<THashedMaterialParameterMap<FLinearColor>::TNamedParameter>							VectorParameters;
+	TArray<THashedMaterialParameterMap<FVector4d>::TNamedParameter>								DoubleVectorParameters;
+	TArray<THashedMaterialParameterMap<const UTexture*>::TNamedParameter>						TextureParameters;
+	TArray<THashedMaterialParameterMap<const UTextureCollection*>::TNamedParameter>				TextureCollectionParameters;
+	TArray<THashedMaterialParameterMap<const URuntimeVirtualTexture*>::TNamedParameter>			RuntimeVirtualTextureParameters;
+	TArray<THashedMaterialParameterMap<const USparseVolumeTexture*>::TNamedParameter>			SparseVolumeTextureParameters;
+	TArray<THashedMaterialParameterMap<const UMaterialParameterCollection*>::TNamedParameter>	ParameterCollectionParameters;
+	TArray<THashedMaterialParameterMap<FName>::TNamedParameter>									EnumerationParameters;
+	TArray<FUserSceneTextureOverride>															UserSceneTextureOverrides;
+	FPostProcessBlendableOverrides																PostProcessBlendableOverrides;
 };
 	
 /** Finds a parameter by name from the game thread. */
diff --git a/Engine/Source/Runtime/Engine/Private/Materials/MaterialInterface.cpp b/Engine/Source/Runtime/Engine/Private/Materials/MaterialInterface.cpp
--- a/Engine/Source/Runtime/Engine/Private/Materials/MaterialInterface.cpp
+++ b/Engine/Source/Runtime/Engine/Private/Materials/MaterialInterface.cpp
@@ -971,6 +971,17 @@ bool UMaterialInterface::GetTextureCollectionParameterValue(const FHashedMateria
 	return false;
 }
 
+bool UMaterialInterface::GetParameterCollectionParameterValue(const FHashedMaterialParameterInfo& ParameterInfo, UMaterialParameterCollection*& OutValue, bool bOveriddenOnly) const
+{
+	FMaterialParameterMetadata Result;
+	if (GetParameterValue(EMaterialParameterType::ParameterCollection, ParameterInfo, Result, MakeParameterValueFlags(bOveriddenOnly)))
+	{
+		OutValue = Result.Value.ParameterCollection;
+		return true;
+	}
+	return false;
+}
+
 bool UMaterialInterface::GetRuntimeVirtualTextureParameterValue(const FHashedMaterialParameterInfo& ParameterInfo, URuntimeVirtualTexture*& OutValue, bool bOveriddenOnly) const
 {
 	FMaterialParameterMetadata Result;
@@ -1090,6 +1101,17 @@ bool UMaterialInterface::GetTextureCollectionParameterDefaultValue(const FHashed
 	return false;
 }
 
+bool UMaterialInterface::GetParameterCollectionParameterDefaultValue(const FHashedMaterialParameterInfo& ParameterInfo, class UMaterialParameterCollection*& OutValue) const
+{
+	FMaterialParameterMetadata Result;
+	if (GetParameterDefaultValue(EMaterialParameterType::ParameterCollection, ParameterInfo, Result))
+	{
+		OutValue = Result.Value.ParameterCollection;
+		return true;
+	}
+	return false;
+}
+
 bool UMaterialInterface::GetRuntimeVirtualTextureParameterDefaultValue(const FHashedMaterialParameterInfo& ParameterInfo, class URuntimeVirtualTexture*& OutValue) const
 {
 	FMaterialParameterMetadata Result;
@@ -1203,6 +1225,11 @@ void UMaterialInterface::GetAllTextureCollectionParameterInfo(TArray<FMaterialPa
 	GetAllParameterInfoOfType(EMaterialParameterType::TextureCollection, OutParameterInfo, OutParameterIds);
 }
 
+void UMaterialInterface::GetAllParameterCollectionParameterInfo(TArray<FMaterialParameterInfo>& OutParameterInfo, TArray<FGuid>& OutParameterIds) const
+{
+	GetAllParameterInfoOfType(EMaterialParameterType::ParameterCollection, OutParameterInfo, OutParameterIds);
+}
+
 void UMaterialInterface::GetAllRuntimeVirtualTextureParameterInfo(TArray<FMaterialParameterInfo>& OutParameterInfo, TArray<FGuid>& OutParameterIds) const
 {
 	GetAllParameterInfoOfType(EMaterialParameterType::RuntimeVirtualTexture, OutParameterInfo, OutParameterIds);
diff --git a/Engine/Source/Runtime/Engine/Private/Materials/MaterialRenderProxy.cpp b/Engine/Source/Runtime/Engine/Private/Materials/MaterialRenderProxy.cpp
--- a/Engine/Source/Runtime/Engine/Private/Materials/MaterialRenderProxy.cpp
+++ b/Engine/Source/Runtime/Engine/Private/Materials/MaterialRenderProxy.cpp
@@ -257,6 +257,17 @@ bool FMaterialRenderProxy::GetTextureCollectionValue(const FHashedMaterialParame
 	return false;
 }
 
+bool FMaterialRenderProxy::GetParameterCollectionValue(const FHashedMaterialParameterInfo& ParameterInfo, const UMaterialParameterCollection** OutValue, const FMaterialRenderContext& Context) const
+{
+	FMaterialParameterValue Value;
+	if (GetParameterValue(EMaterialParameterType::ParameterCollection, ParameterInfo, Value, Context))
+	{
+		*OutValue = Value.ParameterCollection;
+		return true;
+	}
+	return false;
+}
+
 bool FMaterialRenderProxy::GetParameterShaderValue(EMaterialParameterType Type, const FHashedMaterialParameterInfo& ParameterInfo, UE::Shader::FValue& OutValue, const FMaterialRenderContext& Context) const
 {
 	FMaterialParameterValue Value;
@@ -520,7 +531,7 @@ void FMaterialRenderProxy::EvaluateUniformExpressions(FRHICommandListBase& RHICm
 		}
 	}
 
-	OutUniformExpressionCache.ParameterCollections = UniformExpressionSet.ParameterCollections;
+	EvaluateParameterCollections(RHICmdList, OutUniformExpressionCache, Context, Updater);
 
 	// Deliberately set this last, see comment above where it's initialized to nullptr
 	OutUniformExpressionCache.CachedUniformExpressionShaderMap = ShaderMap;
@@ -528,6 +539,15 @@ void FMaterialRenderProxy::EvaluateUniformExpressions(FRHICommandListBase& RHICm
 	++UniformExpressionCacheSerialNumber;
 }
 
+void FMaterialRenderProxy::EvaluateParameterCollections(FRHICommandListBase& RHICmdList, FUniformExpressionCache& OutUniformExpressionCache, const FMaterialRenderContext& Context, FUniformExpressionCacheAsyncUpdater* Updater) const
+{
+	// Retrieve the material's uniform expression set.
+	FMaterialShaderMap* ShaderMap = Context.Material.GetRenderingThreadShaderMap();
+	const FUniformExpressionSet& UniformExpressionSet = ShaderMap->GetUniformExpressionSet();
+
+	OutUniformExpressionCache.ParameterCollections = UniformExpressionSet.ParameterCollections;
+}
+
 void FMaterialRenderProxy::CacheUniformExpressions(FRHICommandListBase& RHICmdList, bool bRecreateUniformBuffer)
 {
 	// Register the render proxy's as a render resource so it can receive notifications to free the uniform buffer.
diff --git a/Engine/Source/Runtime/Engine/Private/Materials/MaterialShared.cpp b/Engine/Source/Runtime/Engine/Private/Materials/MaterialShared.cpp
--- a/Engine/Source/Runtime/Engine/Private/Materials/MaterialShared.cpp
+++ b/Engine/Source/Runtime/Engine/Private/Materials/MaterialShared.cpp
@@ -5465,6 +5465,7 @@ UE::Shader::FValue FMaterialParameterValue::AsShaderValue() const
 	case EMaterialParameterType::StaticComponentMask: return UE::Shader::FValue(Bool[0], Bool[1], Bool[2], Bool[3]);
 	case EMaterialParameterType::Texture:
 	case EMaterialParameterType::TextureCollection:
+	case EMaterialParameterType::ParameterCollection:
 	case EMaterialParameterType::Font:
 	case EMaterialParameterType::RuntimeVirtualTexture:
 	case EMaterialParameterType::SparseVolumeTexture:
@@ -5507,6 +5508,7 @@ UE::Shader::FType GetShaderValueType(EMaterialParameterType Type)
 	case EMaterialParameterType::StaticComponentMask: return UE::Shader::EValueType::Bool4;
 	case EMaterialParameterType::Texture:
 	case EMaterialParameterType::TextureCollection:
+	case EMaterialParameterType::ParameterCollection:
 	case EMaterialParameterType::RuntimeVirtualTexture:
 	case EMaterialParameterType::Font:
 		return FMaterialTextureValue::GetTypeName();
diff --git a/Engine/Source/Runtime/Engine/Private/Materials/ParameterCollection.cpp b/Engine/Source/Runtime/Engine/Private/Materials/ParameterCollection.cpp
--- a/Engine/Source/Runtime/Engine/Private/Materials/ParameterCollection.cpp
+++ b/Engine/Source/Runtime/Engine/Private/Materials/ParameterCollection.cpp
@@ -12,6 +12,7 @@
 #include "Materials/MaterialParameterCollection.h"
 #include "Misc/App.h"
 #include "RenderGraphBuilder.h"
+#include "Templates/Tuple.h"
 
 int32 GDeferUpdateRenderStates = 1;
 FAutoConsoleVariableRef CVarDeferUpdateRenderStates(
@@ -21,6 +22,14 @@ FAutoConsoleVariableRef CVarDeferUpdateRenderStates(
 	ECVF_RenderThreadSafe
 );
 
+int32 GMaterialParameterCollectionMaxVectorStorage = 1280;
+FAutoConsoleVariableRef CVarMaterialParameterCollectionMaxVectorStorage(
+	TEXT("r.MPC.MaxVectorStorage"),
+	GMaterialParameterCollectionMaxVectorStorage,
+	TEXT("The maximum number of vectors allowed in a parameter collection without generating a warning."),
+	ECVF_RenderThreadSafe
+);
+
 TMultiMap<FGuid, FMaterialParameterCollectionInstanceResource*> GDefaultMaterialParameterCollectionInstances;
 
 UMaterialParameterCollection::UMaterialParameterCollection(const FObjectInitializer& ObjectInitializer)
@@ -43,6 +52,26 @@ void UMaterialParameterCollection::PostInitProperties()
 void UMaterialParameterCollection::PostLoad()
 {
 	Super::PostLoad();
+
+	if (Base)
+	{
+		// Determine if the structure of the Base has changed.
+		if (BaseStateId != Base->StateId)
+		{
+			UpdateOverrides(Base);
+		}
+	}
+	else
+	{
+		// Determine if there used to be a Base.
+		if (BaseStateId != FGuid())
+		{
+			BaseStateId = FGuid();
+			StateId = FGuid::NewGuid();
+			ScalarParameterBaseOverrides.Empty();
+			VectorParameterBaseOverrides.Empty();
+		}
+	}
 	
 	if (!StateId.IsValid())
 	{
@@ -110,156 +139,128 @@ void UMaterialParameterCollection::FinishDestroy()
 
 #if WITH_EDITOR
 
-template<typename ParameterType>
-FName CreateUniqueName(TArray<ParameterType>& Parameters, int32 RenameParameterIndex)
+bool UMaterialParameterCollection::SetScalarParameterDefaultValueByInfo(FCollectionScalarParameter ScalarParameter)
 {
-	FString RenameString;
-	Parameters[RenameParameterIndex].ParameterName.ToString(RenameString);
-
-	int32 NumberStartIndex = RenameString.FindLastCharByPredicate([](TCHAR Letter){ return !FChar::IsDigit(Letter); }) + 1;
-	
-	int32 RenameNumber = 0;
-	if (NumberStartIndex < RenameString.Len() - 1)
-	{
-		FString RenameStringNumberPart = RenameString.RightChop(NumberStartIndex);
-		ensure(RenameStringNumberPart.IsNumeric());
-
-		TTypeFromString<int32>::FromString(RenameNumber, *RenameStringNumberPart);
-	}
-
-	FString BaseString = RenameString.Left(NumberStartIndex);
-
-	FName Renamed = FName(*FString::Printf(TEXT("%s%u"), *BaseString, ++RenameNumber));
-
-	bool bMatchFound = false;
-	
-	do
-	{
-		bMatchFound = false;
-
-		for (int32 i = 0; i < Parameters.Num(); ++i)
-		{
-			if (Parameters[i].ParameterName == Renamed && RenameParameterIndex != i)
-			{
-				Renamed = FName(*FString::Printf(TEXT("%s%u"), *BaseString, ++RenameNumber));
-				bMatchFound = true;
-				break;
-			}
-		}
-	} while (bMatchFound);
-	
-	return Renamed;
+	// if the input parameter exists, pass the name and value down to SetScalarParameterDefaultValue
+	// since we want to preserve the Guid of the parameter that's already on the asset
+	return SetScalarParameterDefaultValue(ScalarParameter.ParameterName, ScalarParameter.DefaultValue);
 }
 
-template<typename ParameterType>
-void SanitizeParameters(TArray<ParameterType>& Parameters)
+bool UMaterialParameterCollection::SetScalarParameterDefaultValue(FName ParameterName, const float Value)
 {
-	for (int32 i = 0; i < Parameters.Num() - 1; ++i)
+	for (UMaterialParameterCollection* BaseCollection = this; BaseCollection; BaseCollection = BaseCollection->Base)
 	{
-		for (int32 j = i + 1; j < Parameters.Num(); ++j)
+		// Search each collection's scalar parameters for the named parameter.
+		for (int32 ParameterIndex = 0; ParameterIndex < BaseCollection->ScalarParameters.Num(); ++ParameterIndex)
 		{
-			if (Parameters[i].Id == Parameters[j].Id)
+			FCollectionScalarParameter& CollectionParameter = BaseCollection->ScalarParameters[ParameterIndex];
+			if (CollectionParameter.ParameterName != ParameterName)
 			{
-				FPlatformMisc::CreateGuid(Parameters[j].Id);
+				continue;
 			}
 
-			if (Parameters[i].ParameterName == Parameters[j].ParameterName)
+			if (BaseCollection == this)
 			{
-				Parameters[j].ParameterName = CreateUniqueName(Parameters, j);
+				// If this collections own the parameter,
+				// change the value in the array itself to maintain its GUID.
+				CollectionParameter.DefaultValue = Value;
 			}
+			else
+			{
+				// Otherwise, set an override value for the parameter in this collection.
+				float& OverrideValue = ScalarParameterBaseOverrides.FindOrAdd(CollectionParameter.Id);
+				OverrideValue = Value;
+			}
+
+			return true;
 		}
 	}
-}
 
-bool UMaterialParameterCollection::SetScalarParameterDefaultValueByInfo(FCollectionScalarParameter ScalarParameter)
-{
-	if(GetScalarParameterByName(ScalarParameter.ParameterName))
-	{
-		// if the input parameter exists, pass the name and value down to SetScalarParameterDefaultValue
-		// since we want to preserve the Guid of the parameter that's already on the asset
-		return SetScalarParameterDefaultValue(ScalarParameter.ParameterName, ScalarParameter.DefaultValue);
-	}
-	// otherwise, return false
 	return false;
 }
 
-bool UMaterialParameterCollection::SetScalarParameterDefaultValue(FName ParameterName, const float Value)
-{
-	// make sure the input parameter name exists in the array
-	const int32 ParameterIndex = GetScalarParameterIndexByName(ParameterName);
-	if(ParameterIndex == -1)
-	{
-		return false;
-	}
-	// if so, change the value on the parameter in the array itself to maintain its GUID
-	ScalarParameters[ParameterIndex].DefaultValue = Value;
-
-	// and return a positive result
-	return true;
-}
-
 bool UMaterialParameterCollection::SetVectorParameterDefaultValueByInfo(FCollectionVectorParameter VectorParameter)
 {
-	if(GetVectorParameterByName(VectorParameter.ParameterName))
-    {
-    	// if the input parameter exists, pass the name and value down to SetVectorParameterDefaultValue
-    	// since we want to preserve the Guid of the parameter that's already on the asset
-    	return SetVectorParameterDefaultValue(VectorParameter.ParameterName, VectorParameter.DefaultValue);
-    }
-	// otherwise, return false
-    return false;
+	// if the input parameter exists, pass the name and value down to SetVectorParameterDefaultValue
+	// since we want to preserve the Guid of the parameter that's already on the asset
+	return SetVectorParameterDefaultValue(VectorParameter.ParameterName, VectorParameter.DefaultValue);
 }
 
 bool UMaterialParameterCollection::SetVectorParameterDefaultValue(FName ParameterName, const FLinearColor& Value)
 {
-	// make sure the input parameter name exists in the array
-	const int32 ParameterIndex = GetVectorParameterIndexByName(ParameterName);
-	if(ParameterIndex == -1)
+	for (UMaterialParameterCollection* BaseCollection = this; BaseCollection; BaseCollection = BaseCollection->Base)
 	{
-		return false;
+		// Search each collection's vector parameters for the named parameter.
+		for (int32 ParameterIndex = 0; ParameterIndex < BaseCollection->VectorParameters.Num(); ++ParameterIndex)
+		{
+			FCollectionVectorParameter& CollectionParameter = BaseCollection->VectorParameters[ParameterIndex];
+			if (CollectionParameter.ParameterName != ParameterName)
+			{
+				continue;
+			}
+
+			if (BaseCollection == this)
+			{
+				// If this collections own the parameter,
+				// change the value in the array itself to maintain its GUID.
+				CollectionParameter.DefaultValue = Value;
+			}
+			else
+			{
+				// Otherwise, set an override value for the parameter in this collection.
+				FLinearColor& OverrideValue = VectorParameterBaseOverrides.FindOrAdd(CollectionParameter.Id);
+				OverrideValue = Value;
+			}
+
+			return true;
+		}
 	}
-	// if so, change the value on the parameter in the array itself to maintain its GUID
-	VectorParameters[ParameterIndex].DefaultValue = Value;
 
-	// and return a positive result
-	return true;
+	return false;
 }
 
-int32 PreviousNumScalarParameters = 0;
-int32 PreviousNumVectorParameters = 0;
+int32 PreviousTotalVectorStorage = 0;
 
 void UMaterialParameterCollection::PreEditChange(FProperty* PropertyThatWillChange)
 {
 	Super::PreEditChange(PropertyThatWillChange);
 
-	PreviousNumScalarParameters = ScalarParameters.Num();
-	PreviousNumVectorParameters = VectorParameters.Num();
+	PreviousTotalVectorStorage = GetTotalVectorStorage();
 }
 
 void UMaterialParameterCollection::PostEditChangeProperty(FPropertyChangedEvent& PropertyChangedEvent)
 {
-	SanitizeParameters(ScalarParameters);
-	SanitizeParameters(VectorParameters);
-
-	// If the array counts have changed, an element has been added or removed, and we need to update the uniform buffer layout,
-	// Which also requires recompiling any referencing materials
-	if (ScalarParameters.Num() != PreviousNumScalarParameters || VectorParameters.Num() != PreviousNumVectorParameters)
+	if (Base && (!BaseStateId.IsValid() || BaseStateId != Base->StateId))
 	{
-		// Limit the count of parameters to fit within uniform buffer limits
-		const uint32 MaxScalarParameters = 1024;
+		BaseStateId = Base->StateId;
 
-		if (ScalarParameters.Num() > MaxScalarParameters)
+		// Verify that a new Base does not create a circular dependency.
+		for (UMaterialParameterCollection* BaseCollection = Base; BaseCollection; BaseCollection = BaseCollection->Base)
 		{
-			ScalarParameters.RemoveAt(MaxScalarParameters, ScalarParameters.Num() - MaxScalarParameters);
+			if (BaseCollection == this)
+			{
+				Base = nullptr;
+				break;
+			}
 		}
+	}
 
-		const uint32 MaxVectorParameters = 1024;
+	if (!Base)
+	{
+		// Clear Base-dependent state.
+		BaseStateId = FGuid();
+		ScalarParameterBaseOverrides.Empty();
+		VectorParameterBaseOverrides.Empty();
+	}
 
-		if (VectorParameters.Num() > MaxVectorParameters)
-		{
-			VectorParameters.RemoveAt(MaxVectorParameters, VectorParameters.Num() - MaxVectorParameters);
-		}
+	SanitizeParameters(&UMaterialParameterCollection::ScalarParameters);
+	SanitizeParameters(&UMaterialParameterCollection::VectorParameters);
 
+	// If the storage total has changed, an element has been added or removed, and we need to update the uniform buffer layout,
+	// Which also requires recompiling any referencing materials
+	int32 TotalVectorStorage = GetTotalVectorStorage();
+	if (TotalVectorStorage != PreviousTotalVectorStorage)
+	{
 		// Generate a new Id so that unloaded materials that reference this collection will update correctly on load
 		// Now that we changed the guid, we must recompile all materials which reference this collection
 		StateId = FGuid::NewGuid();
@@ -267,6 +268,18 @@ void UMaterialParameterCollection::PostEditChangeProperty(FPropertyChangedEvent&
 		// Update the uniform buffer layout
 		CreateBufferStruct();
 
+		// If this collection is the base of another collection, the other collection must be updated.
+		TArray<UMaterialParameterCollection*> DerivedCollections;
+		for (TObjectIterator<UMaterialParameterCollection> It(RF_ClassDefaultObject, true, EInternalObjectFlags::Garbage); It; ++It)
+		{
+			UMaterialParameterCollection* Collection = *It;
+			if (Collection != this && Collection->UpdateOverrides(this))
+			{
+				Collection->CreateBufferStruct();
+				DerivedCollections.Add(Collection);
+			}
+		}
+
 		// Create a material update context so we can safely update materials using this parameter collection.
 		{
 			FMaterialUpdateContext UpdateContext;
@@ -286,9 +299,16 @@ void UMaterialParameterCollection::PostEditChangeProperty(FPropertyChangedEvent&
 				}
 				else
 				{
-					for (int32 FunctionIndex = 0; FunctionIndex < CurrentMaterial->GetCachedExpressionData().ParameterCollectionInfos.Num() && !bRecompile; FunctionIndex++)
+					for (int32 CollectionInfoIndex = 0; CollectionInfoIndex < CurrentMaterial->GetCachedExpressionData().ParameterCollectionInfos.Num(); ++CollectionInfoIndex)
 					{
-						if (CurrentMaterial->GetCachedExpressionData().ParameterCollectionInfos[FunctionIndex].ParameterCollection == this)
+						// If this collection is referenced by a material, or is a base of a collection referenced by a material, the material must be recompiled.
+						UMaterialParameterCollection* Collection = CurrentMaterial->GetCachedExpressionData().ParameterCollectionInfos[CollectionInfoIndex].ParameterCollection;
+						while (Collection != this && Collection)
+						{
+							Collection = Collection->Base;
+						}
+
+						if (Collection)
 						{
 							bRecompile = true;
 							break;
@@ -307,7 +327,7 @@ void UMaterialParameterCollection::PostEditChangeProperty(FPropertyChangedEvent&
 				}
 			}
 
-			// Recreate all uniform buffers based off of this collection
+			// Recreate all uniform buffers based on this collection
 			for (TObjectIterator<UWorld> It(/*AdditionalExclusionFlags = */RF_ClassDefaultObject, /*bIncludeDerivedClasses = */true, /*InInternalExclusionFlags = */EInternalObjectFlags::Garbage); It; ++It)
 			{
 				UWorld* CurrentWorld = *It;
@@ -315,6 +335,12 @@ void UMaterialParameterCollection::PostEditChangeProperty(FPropertyChangedEvent&
 			}
 
 			UpdateDefaultResource(true);
+
+			// Updated the default resource of any derived collections that are loaded.
+			for (UMaterialParameterCollection* DerivedCollection : DerivedCollections)
+			{
+				DerivedCollection->UpdateDefaultResource(true);
+			}
 		}
 	}
 	else
@@ -334,6 +360,132 @@ void UMaterialParameterCollection::PostEditChangeProperty(FPropertyChangedEvent&
 
 #endif // WITH_EDITOR
 
+TTuple<FString, int32> SplitParameterName(FName ParameterName)
+{
+	FString NameString;
+	ParameterName.ToString(NameString);
+
+	int32 Number = 0;
+	int32 NumberStartIndex = NameString.FindLastCharByPredicate([](TCHAR Letter){ return !FChar::IsDigit(Letter); }) + 1;
+	if (NumberStartIndex < NameString.Len() - 1)
+	{
+		ensure(TCString<TCHAR>::IsNumeric(&NameString[NumberStartIndex]));
+		TTypeFromString<int32>::FromString(Number, &NameString[NumberStartIndex]);
+	}
+	
+	return { NameString.Left(NumberStartIndex), Number };
+}
+
+template<typename FCollectionParameterType>
+void UMaterialParameterCollection::SanitizeParameters(TArray<FCollectionParameterType> UMaterialParameterCollection::* CollectionParameters)
+{
+	TSet<FGuid> ActiveParameterIds;
+	TMap<FString, TSet<int32>> ActiveParameterNames;
+
+	// Collect active parameter ids and names from any base collections.
+	// SanitizeParameters should be called when base collections have already been sanitized.
+	for (const UMaterialParameterCollection* BaseCollection = Base; BaseCollection; BaseCollection = BaseCollection->Base)
+	{
+		// Add each parameter in each base.
+		for (int32 BaseParameterIndex = 0; BaseParameterIndex < (BaseCollection->*CollectionParameters).Num(); ++BaseParameterIndex)
+		{
+			const FCollectionParameterType& BaseCollectionParameter = (BaseCollection->*CollectionParameters)[BaseParameterIndex];
+
+			// Update the set of active ids.
+			ActiveParameterIds.Add(BaseCollectionParameter.Id);
+
+			// Split the parameter name into a name and number.
+			TTuple<FString, int32> ParameterNameAndNumber = SplitParameterName(BaseCollectionParameter.ParameterName);
+
+			// Update the map of active parameter numbers.
+			ActiveParameterNames.FindOrAdd(get<0>(ParameterNameAndNumber)).Add(get<1>(ParameterNameAndNumber));
+		}
+	}
+
+	// Sanitize each parameter in the collection.
+	for (int32 ParameterIndex = 0; ParameterIndex < (this->*CollectionParameters).Num(); ++ParameterIndex)
+	{
+		FCollectionParameterType& CollectionParameter = (this->*CollectionParameters)[ParameterIndex];
+
+		// Ensure the parameter has a unique id.
+		while (ActiveParameterIds.Contains(CollectionParameter.Id))
+		{
+			FPlatformMisc::CreateGuid(CollectionParameter.Id);
+		}
+		ActiveParameterIds.Add(CollectionParameter.Id);
+
+		// Split the parameter name into a name and number.
+		TTuple<FString, int32> ParameterNameAndNumber = SplitParameterName(CollectionParameter.ParameterName);
+
+		// Find the next available parameter number for the parameter name.
+		int32 ParameterNumber = get<1>(ParameterNameAndNumber);
+		TSet<int32>& ActiveParameterNumbers = ActiveParameterNames.FindOrAdd(get<0>(ParameterNameAndNumber));
+		while (ActiveParameterNumbers.Contains(ParameterNumber))
+		{
+			++ParameterNumber;
+		}
+		ActiveParameterNumbers.Add(ParameterNumber);
+
+		// If the parameter number has changed, update the parameter name.
+		if (ParameterNumber != get<1>(ParameterNameAndNumber))
+		{
+			CollectionParameter.ParameterName = FName(*FString::Printf(TEXT("%s%u"), *get<0>(ParameterNameAndNumber), ParameterNumber));
+		}
+	}
+}
+
+bool UMaterialParameterCollection::UpdateOverrides(UMaterialParameterCollection* BaseCollection)
+{
+	if (!Base || (Base != BaseCollection && !Base->UpdateOverrides(BaseCollection)))
+	{
+		// Indicate that this collection is not based on the BaseCollection.
+		return false;
+	}
+
+	// If the BaseStateId is up-to-date, the base overrides should be as well.
+	if (BaseStateId == Base->StateId)
+	{
+		return true;
+	}
+
+	BaseStateId = Base->StateId;
+
+	// Generate a new Id so that unloaded materials that reference this collection will update correctly on load.
+	// All materials that reference this collection must be recompiled.
+	StateId = FGuid::NewGuid();
+
+	// Remove missing scalar parameter overrides.
+	TMap<FGuid, float> ScalarParameterBaseOverridesCache;
+	Swap(ScalarParameterBaseOverridesCache, ScalarParameterBaseOverrides);
+	ScalarParameterBaseOverrides.Empty(ScalarParameterBaseOverridesCache.Num());
+	for (auto& Entry : ScalarParameterBaseOverridesCache)
+	{
+		if (Base->GetParameterName(Entry.Key) != NAME_None)
+		{
+			ScalarParameterBaseOverrides.Add(Entry);
+		}
+	}
+
+	// Remove missing vector parameter overrides.
+	TMap<FGuid, FLinearColor> VectorParameterBaseOverridesCache;
+	Swap(VectorParameterBaseOverridesCache, VectorParameterBaseOverrides);
+	VectorParameterBaseOverrides.Empty(VectorParameterBaseOverridesCache.Num());
+	for (auto& Entry : VectorParameterBaseOverridesCache)
+	{
+		if (Base->GetParameterName(Entry.Key) != NAME_None)
+		{
+			VectorParameterBaseOverrides.Add(Entry);
+		}
+	}
+
+	// Ensure the collection's parameter names account for the new base.
+	SanitizeParameters(&UMaterialParameterCollection::ScalarParameters);
+	SanitizeParameters(&UMaterialParameterCollection::VectorParameters);
+
+	// Indicate that this collection is based on the BaseCollection.
+	return true;
+}
+
 int32 UMaterialParameterCollection::GetScalarParameterIndexByName(FName ParameterName) const
 {
 	// loop over all the available scalar parameters and look for a name match
@@ -378,50 +530,90 @@ TArray<FName> UMaterialParameterCollection::GetVectorParameterNames() const
 
 float UMaterialParameterCollection::GetScalarParameterDefaultValue(FName ParameterName, bool& bParameterFound) const
 {
-	const int32 ParameterIndex = GetScalarParameterIndexByName(ParameterName);
-	bParameterFound = true;
-	if(ParameterIndex == -1)
+	for (const UMaterialParameterCollection* BaseCollection = this; BaseCollection; BaseCollection = BaseCollection->Base)
 	{
-		bParameterFound = false;
-		return 0.0;
+		// Search each base collection's scalar parameters for the named parameter.
+		for (int32 ParameterIndex = 0; ParameterIndex < BaseCollection->ScalarParameters.Num(); ++ParameterIndex)
+		{
+			const FCollectionScalarParameter& CollectionParameter = BaseCollection->ScalarParameters[ParameterIndex];
+			if (CollectionParameter.ParameterName == ParameterName)
+			{
+				bParameterFound = true;
+
+				// Search each derived collection's overrides for the parameter's id.
+				for (const UMaterialParameterCollection* OverrideCollection = this; OverrideCollection != BaseCollection; OverrideCollection = OverrideCollection->Base)
+				{
+					if (const float* OverrideValue = OverrideCollection->ScalarParameterBaseOverrides.Find(CollectionParameter.Id))
+					{
+						// Use the overridden value.
+						return *OverrideValue;
+					}
+				}
+
+				// Use the default value.
+				return CollectionParameter.DefaultValue;
+			}
+		}
 	}
-	
-	return ScalarParameters[ParameterIndex].DefaultValue;
+
+	bParameterFound = false;
+	return 0.f;
 }
 
 FLinearColor UMaterialParameterCollection::GetVectorParameterDefaultValue(FName ParameterName, bool& bParameterFound) const
 {
-	const int32 ParameterIndex = GetVectorParameterIndexByName(ParameterName);
-	bParameterFound = true;
-	if(ParameterIndex == -1)
+	for (const UMaterialParameterCollection* BaseCollection = this; BaseCollection; BaseCollection = BaseCollection->Base)
 	{
-		bParameterFound = false;
-		return FLinearColor::Black;
+		// Search each base collection's vector parameters for the named parameter.
+		for (int32 ParameterIndex = 0; ParameterIndex < BaseCollection->VectorParameters.Num(); ++ParameterIndex)
+		{
+			const FCollectionVectorParameter& CollectionParameter = BaseCollection->VectorParameters[ParameterIndex];
+			if (CollectionParameter.ParameterName == ParameterName)
+			{
+				bParameterFound = true;
+
+				// Search each derived collection's overrides for the parameter's id.
+				for (const UMaterialParameterCollection* OverrideCollection = this; OverrideCollection != BaseCollection; OverrideCollection = OverrideCollection->Base)
+				{
+					if (const FLinearColor* OverrideValue = OverrideCollection->VectorParameterBaseOverrides.Find(CollectionParameter.Id))
+					{
+						// Use the overridden value.
+						return *OverrideValue;
+					}
+				}
+
+				// Use the default value.
+				return CollectionParameter.DefaultValue;
+			}
+		}
 	}
-	
-	return VectorParameters[ParameterIndex].DefaultValue;
-}
 
+	bParameterFound = false;
+	return FLinearColor::Black;
+}
 
 FName UMaterialParameterCollection::GetParameterName(const FGuid& Id) const
 {
-	for (int32 ParameterIndex = 0; ParameterIndex < ScalarParameters.Num(); ParameterIndex++)
+	for (const UMaterialParameterCollection* Collection = this; Collection; Collection = Collection->Base)
 	{
-		const FCollectionScalarParameter& Parameter = ScalarParameters[ParameterIndex];
-
-		if (Parameter.Id == Id)
+		// Search each base collection's scalar parameters for the parameter's id.
+		for (int32 ParameterIndex = 0; ParameterIndex < Collection->ScalarParameters.Num(); ++ParameterIndex)
 		{
-			return Parameter.ParameterName;
+			const FCollectionScalarParameter& CollectionParameter = Collection->ScalarParameters[ParameterIndex];
+			if (CollectionParameter.Id == Id)
+			{
+				return CollectionParameter.ParameterName;
+			}
 		}
-	}
-
-	for (int32 ParameterIndex = 0; ParameterIndex < VectorParameters.Num(); ParameterIndex++)
-	{
-		const FCollectionVectorParameter& Parameter = VectorParameters[ParameterIndex];
 
-		if (Parameter.Id == Id)
+		// Search each base collection's vector parameters for the parameter's id.
+		for (int32 ParameterIndex = 0; ParameterIndex < Collection->VectorParameters.Num(); ++ParameterIndex)
 		{
-			return Parameter.ParameterName;
+			const FCollectionVectorParameter& CollectionParameter = Collection->VectorParameters[ParameterIndex];
+			if (CollectionParameter.Id == Id)
+			{
+				return CollectionParameter.ParameterName;
+			}
 		}
 	}
 
@@ -430,23 +622,26 @@ FName UMaterialParameterCollection::GetParameterName(const FGuid& Id) const
 
 FGuid UMaterialParameterCollection::GetParameterId(FName ParameterName) const
 {
-	for (int32 ParameterIndex = 0; ParameterIndex < ScalarParameters.Num(); ParameterIndex++)
+	for (const UMaterialParameterCollection* Collection = this; Collection; Collection = Collection->Base)
 	{
-		const FCollectionScalarParameter& Parameter = ScalarParameters[ParameterIndex];
-
-		if (Parameter.ParameterName == ParameterName)
+		// Search each base collection's scalar parameters for the named parameter.
+		for (int32 ParameterIndex = 0; ParameterIndex < Collection->ScalarParameters.Num(); ++ParameterIndex)
 		{
-			return Parameter.Id;
+			const FCollectionScalarParameter& CollectionParameter = Collection->ScalarParameters[ParameterIndex];
+			if (CollectionParameter.ParameterName == ParameterName)
+			{
+				return CollectionParameter.Id;
+			}
 		}
-	}
-
-	for (int32 ParameterIndex = 0; ParameterIndex < VectorParameters.Num(); ParameterIndex++)
-	{
-		const FCollectionVectorParameter& Parameter = VectorParameters[ParameterIndex];
 
-		if (Parameter.ParameterName == ParameterName)
+		// Search each base collection's vector parameters for the named parameter.
+		for (int32 ParameterIndex = 0; ParameterIndex < Collection->VectorParameters.Num(); ++ParameterIndex)
 		{
-			return Parameter.Id;
+			const FCollectionVectorParameter& CollectionParameter = Collection->VectorParameters[ParameterIndex];
+			if (CollectionParameter.ParameterName == ParameterName)
+			{
+				return CollectionParameter.Id;
+			}
 		}
 	}
 
@@ -455,48 +650,69 @@ FGuid UMaterialParameterCollection::GetParameterId(FName ParameterName) const
 
 void UMaterialParameterCollection::GetParameterIndex(const FGuid& Id, int32& OutIndex, int32& OutComponentIndex) const
 {
-	// The parameter and component index allocated in this function must match the memory layout in UMaterialParameterCollectionInstance::GetParameterData
-
 	OutIndex = -1;
 	OutComponentIndex = -1;
 
-	for (int32 ParameterIndex = 0; ParameterIndex < ScalarParameters.Num(); ParameterIndex++)
+	int32 ScalarParameterBase = GetTotalVectorStorage();
+	for (const UMaterialParameterCollection* Collection = this; Collection; Collection = Collection->Base)
 	{
-		const FCollectionScalarParameter& Parameter = ScalarParameters[ParameterIndex];
+		// Find this collection's scalar and vector parameter offsets into the vector storage.
+		int32 VectorParameterBase = ScalarParameterBase - Collection->VectorParameters.Num();
+		ScalarParameterBase = VectorParameterBase - FMath::DivideAndRoundUp(Collection->ScalarParameters.Num(), 4);
 
-		if (Parameter.Id == Id)
+		// Search the scalar parameters for the parameter id.
+		for (int32 ParameterIndex = 0; ParameterIndex < Collection->ScalarParameters.Num(); ++ParameterIndex)
 		{
-			// Scalar parameters are packed into float4's
-			OutIndex = ParameterIndex / 4;
-			OutComponentIndex = ParameterIndex % 4;
-			break;
+			const FCollectionScalarParameter& CollectionParameter = Collection->ScalarParameters[ParameterIndex];
+			if (CollectionParameter.Id == Id)
+			{
+				// Scalar parameters are packed into float4's, so derive the indices accordingly.
+				OutIndex = ScalarParameterBase + ParameterIndex / 4;
+				OutComponentIndex = ParameterIndex % 4;
+				return;
+			}
 		}
-	}
-
-	const int32 VectorParameterBase = FMath::DivideAndRoundUp(ScalarParameters.Num(), 4);
 
-	for (int32 ParameterIndex = 0; ParameterIndex < VectorParameters.Num(); ParameterIndex++)
-	{
-		const FCollectionVectorParameter& Parameter = VectorParameters[ParameterIndex];
-
-		if (Parameter.Id == Id)
+		// Search the vector parameters for the parameter id.
+		for (int32 ParameterIndex = 0; ParameterIndex < Collection->VectorParameters.Num(); ++ParameterIndex)
 		{
-			OutIndex = ParameterIndex + VectorParameterBase;
-			break;
+			const FCollectionVectorParameter& CollectionParameter = Collection->VectorParameters[ParameterIndex];
+			if (CollectionParameter.Id == Id)
+			{
+				// Vector parameters don't use the component index.
+				OutIndex = VectorParameterBase + ParameterIndex;
+				return;
+			}
 		}
 	}
 }
 
 int32 UMaterialParameterCollection::GetTotalVectorStorage() const
 {
-	return FMath::DivideAndRoundUp(ScalarParameters.Num(), 4) + VectorParameters.Num();
+	// Sum the vector-aligned storage required for each collection's scalar and vector parameters.
+	int32 NumVectors = FMath::DivideAndRoundUp(ScalarParameters.Num(), 4) + VectorParameters.Num();
+	for (const UMaterialParameterCollection* Collection = Base; Collection; Collection = Collection->Base)
+	{
+		NumVectors += FMath::DivideAndRoundUp(Collection->ScalarParameters.Num(), 4) + Collection->VectorParameters.Num();
+	}
+	UE_CLOG(NumVectors > GMaterialParameterCollectionMaxVectorStorage, LogMaterial, Warning,
+		TEXT("'%s' requires more than the maximum configured number of vectors of storage (%u) specified by the 'r.MPC.MaxVectorStorage' cvar."),
+		*GetPathName(), GMaterialParameterCollectionMaxVectorStorage);
+	return NumVectors;
 }
 
 void UMaterialParameterCollection::GetParameterNames(TArray<FName>& OutParameterNames, bool bVectorParameters) const
 {
+	if (Base)
+	{
+		// Add base collection parameter names.
+		Base->GetParameterNames(OutParameterNames, bVectorParameters);
+	}
+
 	if (bVectorParameters)
 	{
-		for (int32 ParameterIndex = 0; ParameterIndex < VectorParameters.Num(); ParameterIndex++)
+		// Add each vector parameter name.
+		for (int32 ParameterIndex = 0; ParameterIndex < VectorParameters.Num(); ++ParameterIndex)
 		{
 			const FCollectionVectorParameter& Parameter = VectorParameters[ParameterIndex];
 			OutParameterNames.Add(Parameter.ParameterName);
@@ -504,7 +720,8 @@ void UMaterialParameterCollection::GetParameterNames(TArray<FName>& OutParameter
 	}
 	else
 	{
-		for (int32 ParameterIndex = 0; ParameterIndex < ScalarParameters.Num(); ParameterIndex++)
+		// Add each scalar parameter name.
+		for (int32 ParameterIndex = 0; ParameterIndex < ScalarParameters.Num(); ++ParameterIndex)
 		{
 			const FCollectionScalarParameter& Parameter = ScalarParameters[ParameterIndex];
 			OutParameterNames.Add(Parameter.ParameterName);
@@ -514,13 +731,16 @@ void UMaterialParameterCollection::GetParameterNames(TArray<FName>& OutParameter
 
 const FCollectionScalarParameter* UMaterialParameterCollection::GetScalarParameterByName(FName ParameterName) const
 {
-	for (int32 ParameterIndex = 0; ParameterIndex < ScalarParameters.Num(); ParameterIndex++)
+	// Search each collection's scalar parameters for the named parameter.
+	for (const UMaterialParameterCollection* Collection = this; Collection; Collection = Collection->Base)
 	{
-		const FCollectionScalarParameter& Parameter = ScalarParameters[ParameterIndex];
-
-		if (Parameter.ParameterName == ParameterName)
+		for (int32 ParameterIndex = 0; ParameterIndex < Collection->ScalarParameters.Num(); ++ParameterIndex)
 		{
-			return &Parameter;
+			const FCollectionScalarParameter& CollectionParameter = Collection->ScalarParameters[ParameterIndex];
+			if (CollectionParameter.ParameterName == ParameterName)
+			{
+				return &CollectionParameter;
+			}
 		}
 	}
 
@@ -529,13 +749,16 @@ const FCollectionScalarParameter* UMaterialParameterCollection::GetScalarParamet
 
 const FCollectionVectorParameter* UMaterialParameterCollection::GetVectorParameterByName(FName ParameterName) const
 {
-	for (int32 ParameterIndex = 0; ParameterIndex < VectorParameters.Num(); ParameterIndex++)
+	// Search each collection's vector parameters for the named parameter.
+	for (const UMaterialParameterCollection* Collection = this; Collection; Collection = Collection->Base)
 	{
-		const FCollectionVectorParameter& Parameter = VectorParameters[ParameterIndex];
-
-		if (Parameter.ParameterName == ParameterName)
+		for (int32 ParameterIndex = 0; ParameterIndex < Collection->VectorParameters.Num(); ++ParameterIndex)
 		{
-			return &Parameter;
+			const FCollectionVectorParameter& CollectionParameter = Collection->VectorParameters[ParameterIndex];
+			if (CollectionParameter.ParameterName == ParameterName)
+			{
+				return &CollectionParameter;
+			}
 		}
 	}
 
@@ -551,8 +774,8 @@ void UMaterialParameterCollection::CreateBufferStruct()
 
 	TArray<FShaderParametersMetadata::FMember> Members;
 	uint32 NextMemberOffset = 0;
-
-	const uint32 NumVectors = FMath::DivideAndRoundUp(ScalarParameters.Num(), 4) + VectorParameters.Num();
+	
+	const uint32 NumVectors = GetTotalVectorStorage();
 	new(Members) FShaderParametersMetadata::FMember(TEXT("Vectors"),TEXT(""),__LINE__,NextMemberOffset,UBMT_FLOAT32,EShaderPrecisionModifier::Half,1,4,NumVectors, nullptr);
 	const uint32 VectorArraySize = NumVectors * sizeof(FVector4f);
 	NextMemberOffset += VectorArraySize;
@@ -576,30 +799,93 @@ void UMaterialParameterCollection::CreateBufferStruct()
 }
 
 void UMaterialParameterCollection::GetDefaultParameterData(TArray<FVector4f>& ParameterData) const
+{
+	GetParameterData(ParameterData, nullptr, nullptr);
+}
+
+void UMaterialParameterCollection::GetParameterData(TArray<FVector4f>& ParameterData, const TMap<FName, float>* ScalarParameterInstanceOverrides, const TMap<FName, FLinearColor>* VectorParameterInstanceOverrides) const
 {
 	// The memory layout created here must match the index assignment in UMaterialParameterCollection::GetParameterIndex
 
-	ParameterData.Empty(FMath::DivideAndRoundUp(ScalarParameters.Num(), 4) + VectorParameters.Num());
+	// Allocate the vector storage.
+	ParameterData.SetNum(GetTotalVectorStorage());
 
-	for (int32 ParameterIndex = 0; ParameterIndex < ScalarParameters.Num(); ParameterIndex++)
+	int32 ScalarParameterBase = ParameterData.Num();
+	for (const UMaterialParameterCollection* BaseCollection = this; BaseCollection; BaseCollection = BaseCollection->Base)
 	{
-		const FCollectionScalarParameter& Parameter = ScalarParameters[ParameterIndex];
+		// Find this collection's scalar and vector parameter offsets into the vector storage.
+		int32 VectorParameterBase = ScalarParameterBase - BaseCollection->VectorParameters.Num();
+		ScalarParameterBase = VectorParameterBase - FMath::DivideAndRoundUp(BaseCollection->ScalarParameters.Num(), 4);
 
-		// Add a new vector for each packed vector
-		if (ParameterIndex % 4 == 0)
+		// Set the value of each scalar parameter in this collection.
+		for (int32 ParameterIndex = 0; ParameterIndex < BaseCollection->ScalarParameters.Num(); ++ParameterIndex)
 		{
-			ParameterData.Add(FVector4f(0, 0, 0, 0));
+			const FCollectionScalarParameter& CollectionParameter = BaseCollection->ScalarParameters[ParameterIndex];
+			FVector4f& CurrentVector = ParameterData[ScalarParameterBase + ParameterIndex / 4];
+
+			// Pack into the appropriate component of this packed vector
+			float& CurrentValue = CurrentVector[ParameterIndex % 4];
+
+			// Determine if scalar instance parameters are available.
+			if (ScalarParameterInstanceOverrides)
+			{
+				// Search the instance's scalar parameters for the named parameter.
+				if (const float* OverrideValue = ScalarParameterInstanceOverrides->Find(CollectionParameter.ParameterName))
+				{
+					// Use the instance's parameter value.
+					CurrentValue = *OverrideValue;
+					continue;
+				}
+			}
+
+			// Speculatively set the parameter value to the default.
+			CurrentValue = CollectionParameter.DefaultValue;
+
+			// Search the derived collections' scalar parameter overrides for the parameter id.
+			for (const UMaterialParameterCollection* OverrideCollection = this; OverrideCollection != BaseCollection; OverrideCollection = OverrideCollection->Base)
+			{
+				if (const float* OverrideValue = OverrideCollection->ScalarParameterBaseOverrides.Find(CollectionParameter.Id))
+				{
+					// Use the overridden parameter value.
+					CurrentValue = *OverrideValue;
+					break;
+				}
+			}
 		}
 
-		FVector4f& CurrentVector = ParameterData.Last();
-		// Pack into the appropriate component of this packed vector
-		CurrentVector[ParameterIndex % 4] = Parameter.DefaultValue;
-	}
+		// Set the value of each vector parameter in this collection.
+		for (int32 ParameterIndex = 0; ParameterIndex < BaseCollection->VectorParameters.Num(); ++ParameterIndex)
+		{
+			const FCollectionVectorParameter& CollectionParameter = BaseCollection->VectorParameters[ParameterIndex];
 
-	for (int32 ParameterIndex = 0; ParameterIndex < VectorParameters.Num(); ParameterIndex++)
-	{
-		const FCollectionVectorParameter& Parameter = VectorParameters[ParameterIndex];
-		ParameterData.Add(FVector4f(Parameter.DefaultValue));
+			FVector4f& CurrentValue = ParameterData[VectorParameterBase + ParameterIndex];
+
+			// Determine if vector instance parameters are available.
+			if (VectorParameterInstanceOverrides)
+			{
+				// Search the instance's vector parameters for the named parameter.
+				if (const FLinearColor* OverrideValue = VectorParameterInstanceOverrides->Find(CollectionParameter.ParameterName))
+				{
+					// Use the instance's parameter value.
+					CurrentValue = *OverrideValue;
+					continue;
+				}
+			}
+
+			// Speculatively set the parameter value to the default.
+			CurrentValue = CollectionParameter.DefaultValue;
+
+			// Search the derived collections' vector parameter overrides for the parameter id.
+			for (const UMaterialParameterCollection* OverrideCollection = this; OverrideCollection != BaseCollection; OverrideCollection = OverrideCollection->Base)
+			{
+				if (const FLinearColor* OverrideValue = OverrideCollection->VectorParameterBaseOverrides.Find(CollectionParameter.Id))
+				{
+					// Use the overridden parameter value.
+					CurrentValue = *OverrideValue;
+					break;
+				}
+			}
+		}
 	}
 }
 
@@ -804,34 +1090,9 @@ void UMaterialParameterCollectionInstance::ForceReturnToDefaultValues()
 
 void UMaterialParameterCollectionInstance::GetParameterData(TArray<FVector4f>& ParameterData) const
 {
-	// The memory layout created here must match the index assignment in UMaterialParameterCollection::GetParameterIndex
-
-	if (Collection.IsValid())
+	if (const UMaterialParameterCollection* BaseCollection = Collection.Get())
 	{
-		ParameterData.Empty(FMath::DivideAndRoundUp(Collection->ScalarParameters.Num(), 4) + Collection->VectorParameters.Num());
-
-		for (int32 ParameterIndex = 0; ParameterIndex < Collection->ScalarParameters.Num(); ParameterIndex++)
-		{
-			const FCollectionScalarParameter& Parameter = Collection->ScalarParameters[ParameterIndex];
-
-			// Add a new vector for each packed vector
-			if (ParameterIndex % 4 == 0)
-			{
-				ParameterData.Add(FVector4f(0, 0, 0, 0));
-			}
-
-			FVector4f& CurrentVector = ParameterData.Last();
-			const float* InstanceData = ScalarParameterValues.Find(Parameter.ParameterName);
-			// Pack into the appropriate component of this packed vector
-			CurrentVector[ParameterIndex % 4] = InstanceData ? *InstanceData : Parameter.DefaultValue;
-		}
-
-		for (int32 ParameterIndex = 0; ParameterIndex < Collection->VectorParameters.Num(); ParameterIndex++)
-		{
-			const FCollectionVectorParameter& Parameter = Collection->VectorParameters[ParameterIndex];
-			const FLinearColor* InstanceData = VectorParameterValues.Find(Parameter.ParameterName);
-			ParameterData.Add(InstanceData ? FVector4f(*InstanceData) : FVector4f(Parameter.DefaultValue));
-		}
+		BaseCollection->GetParameterData(ParameterData, &ScalarParameterValues, &VectorParameterValues);
 	}
 }
 
diff --git a/Engine/Source/Runtime/Engine/Public/MaterialCachedData.h b/Engine/Source/Runtime/Engine/Public/MaterialCachedData.h
--- a/Engine/Source/Runtime/Engine/Public/MaterialCachedData.h
+++ b/Engine/Source/Runtime/Engine/Public/MaterialCachedData.h
@@ -20,6 +20,7 @@ class UFont;
 class UMaterial;
 class UMaterialExpression;
 class UMaterialExpressionExternalCodeBase;
+class UMaterialParameterCollection;
 class URuntimeVirtualTexture;
 class ULandscapeGrassType;
 class UMaterialFunctionInterface;
@@ -62,7 +63,7 @@ struct FMaterialParameterCollectionInfo
 
 	/** The collection which this material has a dependency on. */
 	UPROPERTY()
-	TObjectPtr<class UMaterialParameterCollection> ParameterCollection = nullptr;
+	TObjectPtr<UMaterialParameterCollection> ParameterCollection = nullptr;
 
 	bool operator==(const FMaterialParameterCollectionInfo& Other) const
 	{
@@ -311,6 +312,9 @@ struct FMaterialCachedExpressionData
 	UPROPERTY()
 	TArray<TSoftObjectPtr<UTextureCollection>> TextureCollectionValues;
 
+	UPROPERTY()
+	TArray<TSoftObjectPtr<UMaterialParameterCollection>> ParameterCollectionValues;
+
 	UPROPERTY()
 	TArray<TSoftObjectPtr<UFont>> FontValues;
 
diff --git a/Engine/Source/Runtime/Engine/Public/Materials/MaterialExpressionCollectionParameter.h b/Engine/Source/Runtime/Engine/Public/Materials/MaterialExpressionCollectionParameter.h
--- a/Engine/Source/Runtime/Engine/Public/Materials/MaterialExpressionCollectionParameter.h
+++ b/Engine/Source/Runtime/Engine/Public/Materials/MaterialExpressionCollectionParameter.h
@@ -27,6 +27,18 @@ class UMaterialExpressionCollectionParameter : public UMaterialExpression
 	UPROPERTY(EditAnywhere, Category=MaterialExpressionCollectionParameter)
 	FName ParameterName;
 
+	/** GUID that should be unique within the material, this is used for parameter renaming. */
+	UPROPERTY()
+	FGuid ExpressionGUID;
+
+	/** The name of the parameter Group to display in MaterialInstance Editor. Default is None group */
+	UPROPERTY(EditAnywhere, Category=MaterialExpressionCollectionParameter)
+	FName Group;
+
+	/** Controls where the this parameter is displayed in a material instance parameter list. The lower the number the higher up in the parameter list. */
+	UPROPERTY(EditAnywhere, Category=MaterialExpressionCollectionParameter)
+	int32 SortPriority = 32;
+
 	/** Id that is set from the name, and used to handle renaming of collection parameters. */
 	UPROPERTY()
 	FGuid ParameterId;
@@ -46,6 +58,10 @@ class UMaterialExpressionCollectionParameter : public UMaterialExpression
 	virtual bool HasAParameterName() const override { return true; }
 	virtual FName GetParameterName() const override { return ParameterName; }
 	virtual void SetParameterName(const FName& Name) override { ParameterName = Name; }
+
+	ENGINE_API virtual bool GetParameterValue(FMaterialParameterMetadata& OutMeta) const;
+	ENGINE_API virtual bool SetParameterValue(const FName& Name, const FMaterialParameterMetadata& Meta, EMaterialExpressionSetParameterValueFlags Flags = EMaterialExpressionSetParameterValueFlags::None);
+
 	virtual void ValidateParameterName(const bool bAllowDuplicateName) override {};
 
 	virtual bool MatchesSearchQuery( const TCHAR* SearchQuery ) override;
diff --git a/Engine/Source/Runtime/Engine/Public/Materials/MaterialFunctionInstance.h b/Engine/Source/Runtime/Engine/Public/Materials/MaterialFunctionInstance.h
--- a/Engine/Source/Runtime/Engine/Public/Materials/MaterialFunctionInstance.h
+++ b/Engine/Source/Runtime/Engine/Public/Materials/MaterialFunctionInstance.h
@@ -49,6 +49,10 @@ class UMaterialFunctionInstance : public UMaterialFunctionInterface
 	UPROPERTY(EditAnywhere, Category=MaterialFunctionInstance)
 	TArray<struct FTextureCollectionParameterValue> TextureCollectionParameterValues;
 
+	/** Parameter Collection parameters. */
+	UPROPERTY(EditAnywhere, Category=MaterialFunctionInstance)
+	TArray<struct FParameterCollectionParameterValue> ParameterCollectionParameterValues;
+
 	/** Font parameters. */
 	UPROPERTY(EditAnywhere, Category=MaterialFunctionInstance)
 	TArray<struct FFontParameterValue> FontParameterValues;
diff --git a/Engine/Source/Runtime/Engine/Public/Materials/MaterialFunctionInterface.h b/Engine/Source/Runtime/Engine/Public/Materials/MaterialFunctionInterface.h
--- a/Engine/Source/Runtime/Engine/Public/Materials/MaterialFunctionInterface.h
+++ b/Engine/Source/Runtime/Engine/Public/Materials/MaterialFunctionInterface.h
@@ -25,6 +25,7 @@ class FMaterialUpdateContext;
 class UMaterialInterface;
 class UMaterialExpression;
 struct FFunctionExpressionOutput;
+class UMaterialParameterCollection;
 
 /** Usage set on a material function determines feature compatibility and validation. */
 UENUM()
@@ -338,6 +339,7 @@ public:
 	ENGINE_API bool OverrideNamedSparseVolumeTextureParameter(const FHashedMaterialParameterInfo& ParameterInfo, class USparseVolumeTexture*& OutValue);
 	ENGINE_API bool OverrideNamedFontParameter(const FHashedMaterialParameterInfo& ParameterInfo, class UFont*& OutFontValue, int32& OutFontPage);
 	ENGINE_API bool OverrideNamedStaticSwitchParameter(const FHashedMaterialParameterInfo& ParameterInfo, bool& OutValue, FGuid& OutExpressionGuid);
+	ENGINE_API bool OverrideNamedParameterCollectionParameter(const FHashedMaterialParameterInfo& ParameterInfo, class UMaterialParameterCollection*& OutValue);
 	ENGINE_API bool OverrideNamedStaticComponentMaskParameter(const FHashedMaterialParameterInfo& ParameterInfo, bool& OutR, bool& OutG, bool& OutB, bool& OutA, FGuid& OutExpressionGuid);
 
 	virtual bool IsUsingNewHLSLGenerator() const { return false; }
diff --git a/Engine/Source/Runtime/Engine/Public/Materials/MaterialInstance.h b/Engine/Source/Runtime/Engine/Public/Materials/MaterialInstance.h
--- a/Engine/Source/Runtime/Engine/Public/Materials/MaterialInstance.h
+++ b/Engine/Source/Runtime/Engine/Public/Materials/MaterialInstance.h
@@ -323,6 +323,54 @@ struct FTextureCollectionParameterValue
 	}
 };
 
+/** Editable parameter collection parameter. */
+USTRUCT(BlueprintType)
+struct FParameterCollectionParameterValue
+{
+	GENERATED_USTRUCT_BODY()
+
+	UPROPERTY(EditAnywhere, BlueprintReadWrite, Category=ParameterCollectionParameterValue)
+	FMaterialParameterInfo ParameterInfo;
+
+	UPROPERTY(EditAnywhere, BlueprintReadWrite, Category=ParameterCollectionParameterValue)
+	TObjectPtr<UMaterialParameterCollection> ParameterValue;
+
+	UPROPERTY()
+	FGuid ExpressionGUID;
+
+	explicit FParameterCollectionParameterValue(const FMaterialParameterInfo& InParameterInfo = FMaterialParameterInfo(), UMaterialParameterCollection* InValue = nullptr)
+		: ParameterInfo(InParameterInfo), ParameterValue(InValue)
+	{
+	}
+
+	bool IsOverride() const { return true; }
+
+	bool IsValid() const { return GetValue(*this) != nullptr; }
+
+	bool operator==(const FParameterCollectionParameterValue& Other) const
+	{
+		return
+			ParameterInfo == Other.ParameterInfo &&
+			ParameterValue == Other.ParameterValue &&
+			ExpressionGUID == Other.ExpressionGUID;
+	}
+	bool operator!=(const FParameterCollectionParameterValue& Other) const
+	{
+		return !((*this) == Other);
+	}
+
+	typedef const UMaterialParameterCollection* ValueType;
+	static ValueType GetValue(const FParameterCollectionParameterValue& Parameter) { return Parameter.ParameterValue; }
+
+	void GetValue(FMaterialParameterMetadata& OutResult) const
+	{
+		OutResult.Value = ParameterValue;
+#if WITH_EDITORONLY_DATA
+		OutResult.ExpressionGuid = ExpressionGUID;
+#endif
+	}
+};
+
 /** Editable runtime virtual texture parameter. */
 USTRUCT(BlueprintType)
 struct FRuntimeVirtualTextureParameterValue
@@ -765,6 +813,10 @@ public:
 	UPROPERTY(EditAnywhere, BlueprintReadOnly, Category=MaterialInstance, meta = (EditFixedOrder))
 	TArray<struct FTextureCollectionParameterValue> TextureCollectionParameterValues;
 
+	/** Texture Collection parameters. */
+	UPROPERTY(EditAnywhere, BlueprintReadOnly, Category=MaterialInstance, meta = (EditFixedOrder))
+	TArray<struct FParameterCollectionParameterValue> ParameterCollectionParameterValues;
+
 	/** RuntimeVirtualTexture parameters. */
 	UPROPERTY(EditAnywhere, BlueprintReadOnly, Category = MaterialInstance, meta = (EditFixedOrder))
 	TArray<struct FRuntimeVirtualTextureParameterValue> RuntimeVirtualTextureParameterValues;
@@ -1198,6 +1250,7 @@ protected:
 #endif
 	void SetTextureParameterValueInternal(const FMaterialParameterInfo& ParameterInfo, class UTexture* Value);
 	void SetTextureCollectionParameterValueInternal(const FMaterialParameterInfo& ParameterInfo, class UTextureCollection* Value);
+	void SetParameterCollectionParameterValueInternal(const FMaterialParameterInfo& ParameterInfo, class UMaterialParameterCollection* Value);
 	void SetRuntimeVirtualTextureParameterValueInternal(const FMaterialParameterInfo& ParameterInfo, class URuntimeVirtualTexture* Value);
 	void SetSparseVolumeTextureParameterValueInternal(const FMaterialParameterInfo& ParameterInfo, class USparseVolumeTexture* Value);
 	void SetFontParameterValueInternal(const FMaterialParameterInfo& ParameterInfo, class UFont* FontValue, int32 FontPage);
diff --git a/Engine/Source/Runtime/Engine/Public/Materials/MaterialInstanceUpdateParameterSet.h b/Engine/Source/Runtime/Engine/Public/Materials/MaterialInstanceUpdateParameterSet.h
--- a/Engine/Source/Runtime/Engine/Public/Materials/MaterialInstanceUpdateParameterSet.h
+++ b/Engine/Source/Runtime/Engine/Public/Materials/MaterialInstanceUpdateParameterSet.h
@@ -297,7 +297,10 @@ bool UpdateParameterSet(TArray<FStaticComponentMaskParameter>& Parameters, UMate
 {
 	return MaterialInstance_Private::UpdateParameterSet_WithCachedData<FStaticComponentMaskParameter, ExpressionType>(EMaterialParameterType::StaticComponentMask, Parameters, ParentMaterial);
 }
-
-
+template <typename ParameterType, typename ExpressionType>
+bool UpdateParameterSet(TArray<FParameterCollectionParameterValue>& Parameters, UMaterial* ParentMaterial)
+{
+	return MaterialInstance_Private::UpdateParameterSet_WithCachedData<FParameterCollectionParameterValue, ExpressionType>(EMaterialParameterType::ParameterCollection, Parameters, ParentMaterial);
+}
 
 #endif // WITH_EDITOR
diff --git a/Engine/Source/Runtime/Engine/Public/Materials/MaterialInterface.h b/Engine/Source/Runtime/Engine/Public/Materials/MaterialInterface.h
--- a/Engine/Source/Runtime/Engine/Public/Materials/MaterialInterface.h
+++ b/Engine/Source/Runtime/Engine/Public/Materials/MaterialInterface.h
@@ -712,6 +712,7 @@ public:
 	ENGINE_API void GetAllDoubleVectorParameterInfo(TArray<FMaterialParameterInfo>& OutParameterInfo, TArray<FGuid>& OutParameterIds) const;
 	ENGINE_API void GetAllTextureParameterInfo(TArray<FMaterialParameterInfo>& OutParameterInfo, TArray<FGuid>& OutParameterIds) const;
 	ENGINE_API void GetAllTextureCollectionParameterInfo(TArray<FMaterialParameterInfo>& OutParameterInfo, TArray<FGuid>& OutParameterIds) const;
+	ENGINE_API void GetAllParameterCollectionParameterInfo(TArray<FMaterialParameterInfo>& OutParameterInfo, TArray<FGuid>& OutParameterIds) const;
 	ENGINE_API void GetAllRuntimeVirtualTextureParameterInfo(TArray<FMaterialParameterInfo>& OutParameterInfo, TArray<FGuid>& OutParameterIds) const;
 	ENGINE_API void GetAllSparseVolumeTextureParameterInfo(TArray<FMaterialParameterInfo>& OutParameterInfo, TArray<FGuid>& OutParameterIds) const;
 	ENGINE_API void GetAllFontParameterInfo(TArray<FMaterialParameterInfo>& OutParameterInfo, TArray<FGuid>& OutParameterIds) const;
@@ -733,6 +734,7 @@ public:
 	ENGINE_API bool GetDoubleVectorParameterDefaultValue(const FHashedMaterialParameterInfo& ParameterInfo, FVector4d& OutValue) const;
 	ENGINE_API bool GetTextureParameterDefaultValue(const FHashedMaterialParameterInfo& ParameterInfo, class UTexture*& OutValue) const;
 	ENGINE_API bool GetTextureCollectionParameterDefaultValue(const FHashedMaterialParameterInfo& ParameterInfo, class UTextureCollection*& OutValue) const;
+	ENGINE_API bool GetParameterCollectionParameterDefaultValue(const FHashedMaterialParameterInfo& ParameterInfo, class UMaterialParameterCollection*& OutValue) const;
 	ENGINE_API bool GetRuntimeVirtualTextureParameterDefaultValue(const FHashedMaterialParameterInfo& ParameterInfo, class URuntimeVirtualTexture*& OutValue) const;
 	ENGINE_API bool GetSparseVolumeTextureParameterDefaultValue(const FHashedMaterialParameterInfo& ParameterInfo, class USparseVolumeTexture*& OutValue) const;
 	ENGINE_API bool GetFontParameterDefaultValue(const FHashedMaterialParameterInfo& ParameterInfo, class UFont*& OutFontValue, int32& OutFontPage) const;
@@ -961,6 +963,7 @@ public:
 #endif // WITH_EDITOR
 	ENGINE_API virtual bool GetTextureParameterValue(const FHashedMaterialParameterInfo& ParameterInfo, class UTexture*& OutValue, bool bOveriddenOnly = false) const;
 	ENGINE_API virtual bool GetTextureCollectionParameterValue(const FHashedMaterialParameterInfo& ParameterInfo, class UTextureCollection*& OutValue, bool bOveriddenOnly = false) const;
+	ENGINE_API virtual bool GetParameterCollectionParameterValue(const FHashedMaterialParameterInfo& ParameterInfo, class UMaterialParameterCollection*& OutValue, bool bOveriddenOnly = false) const;
 	ENGINE_API bool GetRuntimeVirtualTextureParameterValue(const FHashedMaterialParameterInfo& ParameterInfo, class URuntimeVirtualTexture*& OutValue, bool bOveriddenOnly = false) const;
 	ENGINE_API bool GetSparseVolumeTextureParameterValue(const FHashedMaterialParameterInfo& ParameterInfo, class USparseVolumeTexture*& OutValue, bool bOveriddenOnly = false) const;
 #if WITH_EDITOR
diff --git a/Engine/Source/Runtime/Engine/Public/Materials/MaterialParameterCollection.h b/Engine/Source/Runtime/Engine/Public/Materials/MaterialParameterCollection.h
--- a/Engine/Source/Runtime/Engine/Public/Materials/MaterialParameterCollection.h
+++ b/Engine/Source/Runtime/Engine/Public/Materials/MaterialParameterCollection.h
@@ -79,6 +79,10 @@ class UMaterialParameterCollection : public UObject
 {
 	GENERATED_UCLASS_BODY()
 
+	friend class FMaterialParameterCollectionDetails;
+
+public:
+
 	/** Used by materials using this collection to know when to recompile. */
 	UPROPERTY(duplicatetransient)
 	FGuid StateId;
@@ -182,10 +186,33 @@ class UMaterialParameterCollection : public UObject
 	/** Create an instance for this collection in every world. */
 	ENGINE_API void SetupWorldParameterCollectionInstances();
 
+	ENGINE_API void GetParameterData(TArray<FVector4f>& ParameterData, const TMap<FName, float>* ScalarParameterInstanceOverrides, const TMap<FName, FLinearColor>* VectorParameterInstanceOverrides) const;
+
+	ENGINE_API bool IsInstanceOf(const FGuid& CollectionStateId) const;
+
 private:
 	virtual ENGINE_API void FinishDestroy() override;
 	virtual ENGINE_API bool IsReadyForFinishDestroy() override;
 
+	template<typename FCollectionParameterType>
+	void SanitizeParameters(TArray<FCollectionParameterType> UMaterialParameterCollection::* CollectionParameters);
+
+	bool UpdateOverrides(UMaterialParameterCollection* BaseCollection);
+
+	UPROPERTY(EditAnywhere, Category=Material)
+	TObjectPtr<UMaterialParameterCollection> Base;
+
+	UPROPERTY()
+	FGuid BaseStateId;
+
+	/** Overrides for scalar parameter values. */
+	UPROPERTY()
+	TMap<FGuid, float> ScalarParameterBaseOverrides;
+
+	/** Overrides for vector parameter values. */
+	UPROPERTY()
+	TMap<FGuid, FLinearColor> VectorParameterBaseOverrides;
+
 	/** Flag used to guarantee that the RT is finished using various resources in this UMaterial before cleanup. */
 	FThreadSafeBool ReleasedByRT;
 
diff --git a/Engine/Source/Runtime/Engine/Public/Materials/MaterialParameters.h b/Engine/Source/Runtime/Engine/Public/Materials/MaterialParameters.h
--- a/Engine/Source/Runtime/Engine/Public/Materials/MaterialParameters.h
+++ b/Engine/Source/Runtime/Engine/Public/Materials/MaterialParameters.h
@@ -18,6 +18,7 @@ class UCurveLinearColorAtlas;
 class UFont;
 class URuntimeVirtualTexture;
 class USparseVolumeTexture;
+class UMaterialParameterCollection;
 
 UENUM()
 enum EMaterialParameterAssociation : int
@@ -193,6 +194,7 @@ enum class EMaterialParameterType : uint8
 	RuntimeVirtualTexture,
 	SparseVolumeTexture,
 	StaticSwitch,
+	ParameterCollection,
 	Enumeration,
 
 	NumRuntime, // Runtime parameter types must go above here, and editor-only ones below
@@ -351,6 +353,7 @@ struct FMaterialParameterValue
 	FMaterialParameterValue(const TObjectPtr<URuntimeVirtualTexture>& InValue) : Type(EMaterialParameterType::RuntimeVirtualTexture) { RuntimeVirtualTexture = InValue; }
 	FMaterialParameterValue(USparseVolumeTexture* InValue) : Type(EMaterialParameterType::SparseVolumeTexture) { SparseVolumeTexture = InValue; }
 	FMaterialParameterValue(const TObjectPtr<USparseVolumeTexture>& InValue) : Type(EMaterialParameterType::SparseVolumeTexture) { SparseVolumeTexture = InValue; }
+	FMaterialParameterValue(const TObjectPtr<UMaterialParameterCollection>& InValue) : Type(EMaterialParameterType::ParameterCollection) { ParameterCollection = InValue; }
 
 	// Gamethread parameters are typically non-const, but renderthread parameters are const
 	// Would be possible to store an additional const-flag member, and provide runtime checks to ensure constness is not violated...maybe worth doing in the future
@@ -362,6 +365,7 @@ struct FMaterialParameterValue
 	FMaterialParameterValue(const TObjectPtr<const URuntimeVirtualTexture>& InValue) : Type(EMaterialParameterType::RuntimeVirtualTexture) { RuntimeVirtualTexture = const_cast<URuntimeVirtualTexture*>(InValue.Get()); }
 	FMaterialParameterValue(const USparseVolumeTexture* InValue) : Type(EMaterialParameterType::SparseVolumeTexture) { SparseVolumeTexture = const_cast<USparseVolumeTexture*>(InValue); }
 	FMaterialParameterValue(const TObjectPtr<const USparseVolumeTexture>& InValue) : Type(EMaterialParameterType::SparseVolumeTexture) { SparseVolumeTexture = const_cast<USparseVolumeTexture*>(InValue.Get()); }
+	FMaterialParameterValue(const UMaterialParameterCollection* InValue) : Type(EMaterialParameterType::ParameterCollection) { ParameterCollection = const_cast<UMaterialParameterCollection*>(InValue); }
 
 	FMaterialParameterValue(UFont* InValue, int32 InPage) : Type(EMaterialParameterType::Font) { Font.Value = InValue; Font.Page = InPage; }
 	FMaterialParameterValue(bool InValue) : Type(EMaterialParameterType::StaticSwitch) { Bool[0] = InValue; }
@@ -392,6 +396,7 @@ struct FMaterialParameterValue
 			UFont* Value;
 			int32 Page;
 		} Font;
+		UMaterialParameterCollection* ParameterCollection;
 		FName EnumEntry;
 	};
 	EMaterialParameterType Type;
@@ -423,6 +428,7 @@ struct FMaterialParameterValue
 		case EMaterialParameterType::RuntimeVirtualTexture: return Lhs.RuntimeVirtualTexture == Rhs.RuntimeVirtualTexture;
 		case EMaterialParameterType::SparseVolumeTexture: return Lhs.SparseVolumeTexture == Rhs.SparseVolumeTexture;
 		case EMaterialParameterType::StaticSwitch: return Lhs.Bool[0] == Rhs.Bool[0];
+		case EMaterialParameterType::ParameterCollection: return Lhs.ParameterCollection == Rhs.ParameterCollection;
 		case EMaterialParameterType::StaticComponentMask: return
 			Lhs.Bool[0] == Rhs.Bool[0] &&
 			Lhs.Bool[1] == Rhs.Bool[1] &&
diff --git a/Engine/Source/Runtime/Engine/Public/Materials/MaterialRenderProxy.h b/Engine/Source/Runtime/Engine/Public/Materials/MaterialRenderProxy.h
--- a/Engine/Source/Runtime/Engine/Public/Materials/MaterialRenderProxy.h
+++ b/Engine/Source/Runtime/Engine/Public/Materials/MaterialRenderProxy.h
@@ -29,6 +29,7 @@ class USpecularProfile;
 class UNeuralProfile;
 class UTexture;
 class UTextureCollection;
+class UMaterialParameterCollection;
 
 struct FMaterialParameterValue;
 struct FMaterialRenderContext;
@@ -171,6 +172,7 @@ public:
 	ENGINE_API bool GetTextureValue(const FHashedMaterialParameterInfo& ParameterInfo, const URuntimeVirtualTexture** OutValue, const FMaterialRenderContext& Context) const;
 	ENGINE_API bool GetTextureValue(const FHashedMaterialParameterInfo& ParameterInfo, const USparseVolumeTexture** OutValue, const FMaterialRenderContext& Context) const;
 	ENGINE_API bool GetTextureCollectionValue(const FHashedMaterialParameterInfo& ParameterInfo, const UTextureCollection** OutValue, const FMaterialRenderContext& Context) const;
+	ENGINE_API bool GetParameterCollectionValue(const FHashedMaterialParameterInfo& ParameterInfo, const UMaterialParameterCollection** OutValue, const FMaterialRenderContext& Context) const;
 	ENGINE_API bool GetParameterShaderValue(EMaterialParameterType Type, const FHashedMaterialParameterInfo& ParameterInfo, UE::Shader::FValue& OutValue, const FMaterialRenderContext& Context) const;
 	virtual bool GetParameterValue(EMaterialParameterType Type, const FHashedMaterialParameterInfo& ParameterInfo, FMaterialParameterValue& OutValue, const FMaterialRenderContext& Context) const = 0;
 	virtual bool GetUserSceneTextureOverride(FName& InOutValue) const { return false; }
@@ -248,6 +250,9 @@ public:
 
 	const FString& GetMaterialName() const { return MaterialName; }
 
+protected:
+	ENGINE_API virtual void EvaluateParameterCollections(FRHICommandListBase& RHICmdList, FUniformExpressionCache& OutUniformExpressionCache, const FMaterialRenderContext& Context, FUniformExpressionCacheAsyncUpdater* Updater) const;
+
 private:
 	ENGINE_API IAllocatedVirtualTexture* GetPreallocatedVTStack(const FMaterialRenderContext& Context, const FUniformExpressionSet& UniformExpressionSet, const FMaterialVirtualTextureStack& VTStack) const;
 	ENGINE_API IAllocatedVirtualTexture* AllocateVTStack(FRHICommandListBase& RHICmdList, const FMaterialRenderContext& Context, const FUniformExpressionSet& UniformExpressionSet, const FMaterialVirtualTextureStack& VTStack) const;
