[TSAN]
- Fix benign race on time stored in FPlatformCompressionUtilities::GetCookOverrides

diff --git a/Engine/Source/Runtime/Engine/Private/AudioCompressionSettingsUtils.cpp b/Engine/Source/Runtime/Engine/Private/AudioCompressionSettingsUtils.cpp
--- a/Engine/Source/Runtime/Engine/Private/AudioCompressionSettingsUtils.cpp
+++ b/Engine/Source/Runtime/Engine/Private/AudioCompressionSettingsUtils.cpp
@@ -299,7 +299,7 @@ const FPlatformAudioCookOverrides* FPlatformCompressionUtilities::GetCookOverrid
 
 	if (GIsEditor && !IsRunningCommandlet())
 	{
-		static double LastCacheTime = 0.0;
+		static TSAN_ATOMIC(double) LastCacheTime = 0.0;
 		double CurrentTime = FPlatformTime::Seconds();
 		double TimeSinceLastCache = CurrentTime - LastCacheTime;
 
