Initial prototype of the LayerStack node. Requires available layers and blends to be set in the base material parent to generate a read only override input pin list (aka payloads) which then propagates from the base material through the layer stack in the material and any of it's instances.
This change also lays the groundwork for allowing read/write mapping of payloads, and creating a healthier compilation environment for layers in general in the new translator.
Additionally, FunctionInput type mapping has been migrated to helper functions to allow for easier re-usability.

diff --git a/Engine/Source/Editor/MaterialEditor/Private/MaterialEditor.cpp b/Engine/Source/Editor/MaterialEditor/Private/MaterialEditor.cpp
--- a/Engine/Source/Editor/MaterialEditor/Private/MaterialEditor.cpp
+++ b/Engine/Source/Editor/MaterialEditor/Private/MaterialEditor.cpp
@@ -102,6 +102,7 @@
 #include "Materials/MaterialExpressionShadingModel.h"
 #include "Materials/MaterialFunction.h"
 #include "Materials/MaterialFunctionInstance.h"
+#include "Materials/MaterialFunctionMaterialLayerBlend.h"
 #include "Materials/MaterialParameterCollection.h"
 
 #include "MaterialGraphNode_Knot.h"
@@ -942,8 +943,9 @@ void FMaterialEditor::InitMaterialEditor( const EToolkitMode::Type Mode, const T
 						{
 							UMaterialExpressionFunctionInput* BaseAttributesInput = Cast<UMaterialExpressionFunctionInput>(InputTop);
 							BaseAttributesInput->InputType = FunctionInput_MaterialAttributes;
-							BaseAttributesInput->InputName = TEXT("Top Layer");
+							BaseAttributesInput->InputName = TopMaterialBlendInputName;
 							BaseAttributesInput->bUsePreviewValueAsDefault = true;
+							BaseAttributesInput->BlendInputRelevance = EBlendInputRelevance::Top;
 						}
 
 						UMaterialExpression* InputBottom = CreateNewMaterialExpression(UMaterialExpressionFunctionInput::StaticClass(), FVector2D(-300, 200), false, true);
@@ -951,8 +953,9 @@ void FMaterialEditor::InitMaterialEditor( const EToolkitMode::Type Mode, const T
 						{
 							UMaterialExpressionFunctionInput* BaseAttributesInput = Cast<UMaterialExpressionFunctionInput>(InputBottom);
 							BaseAttributesInput->InputType = FunctionInput_MaterialAttributes;
-							BaseAttributesInput->InputName = TEXT("Bottom Layer");
+							BaseAttributesInput->InputName = BottomMaterialBlendInputName;
 							BaseAttributesInput->bUsePreviewValueAsDefault = true;
+							BaseAttributesInput->BlendInputRelevance = EBlendInputRelevance::Bottom;
 						}
 						if (GetDefault<UEditorExperimentalSettings>()->bExampleLayersAndBlends)
 						{
@@ -3075,118 +3078,9 @@ void FMaterialEditor::UpdateMaterialinfoList_Old()
 		{
 			continue;
 		}
-
-		if (MaterialFunction && ExpressionPreviewMaterial)
-		{
-			bool bHasValidOutput = true;
-			int32 NumInputs = 0;
-			int32 NumOutputs = 0;
-			// For Material Layers
-
-			if (!Substrate::IsMaterialLayeringSupportEnabled() && MaterialFunction->GetMaterialFunctionUsage() == EMaterialFunctionUsage::MaterialLayer)
-			{
-				// Material layers must have a single MA input and output only
-				for (UMaterialExpression* Expression : MaterialFunction->GetExpressions())
-				{
-					if (UMaterialExpressionFunctionInput* InputExpression = Cast<UMaterialExpressionFunctionInput>(Expression))
-					{
-						++NumInputs;
-						if (NumInputs > 1 || !InputExpression->IsResultMaterialAttributes(0))
-						{
-							CompileErrors.Add(TEXT("Layer graphs only support a single material attributes input."));
-							FailingExpression.Add(nullptr);
-						}
-					}
-					else if (UMaterialExpressionFunctionOutput* OutputExpression = Cast<UMaterialExpressionFunctionOutput>(Expression))
-					{
-						++NumOutputs;
-						if (NumOutputs > 1 || !OutputExpression->IsResultMaterialAttributes(0))
-						{
-							CompileErrors.Add(TEXT("Layer graphs only support a single material attributes output."));
-							FailingExpression.Add(nullptr);
-						}
-					}
-					else if (UMaterialExpressionMaterialAttributeLayers* RecursiveLayer = Cast<UMaterialExpressionMaterialAttributeLayers>(Expression))
-					{
-						CompileErrors.Add(TEXT("Layer graphs do not support layers within layers."));
-						FailingExpression.Add(nullptr);
-					}
-				}
-
-				if (NumInputs > 1 || NumOutputs < 1)
-				{
-					CompileErrors.Add(TEXT("Layer graphs require a single material attributes output and optionally, a single material attributes input."));
-					FailingExpression.Add(nullptr);
-				}
-			}
-			else if (!Substrate::IsMaterialLayeringSupportEnabled() && MaterialFunction->GetMaterialFunctionUsage() == EMaterialFunctionUsage::MaterialLayerBlend)
-			{
-				// Material layer blends can have two MA inputs and single MA output only
-				for (UMaterialExpression* Expression : MaterialFunction->GetExpressions())
-				{
-					if (UMaterialExpressionFunctionInput* InputExpression = Cast<UMaterialExpressionFunctionInput>(Expression))
-					{
-						++NumInputs;
-						if (NumInputs > 2 || !InputExpression->IsResultMaterialAttributes(0))
-						{
-							CompileErrors.Add(TEXT("Layer blend graphs only support two material attributes inputs."));
-							FailingExpression.Add(nullptr);
-						}
-					}
-					else if (UMaterialExpressionFunctionOutput* OutputExpression = Cast<UMaterialExpressionFunctionOutput>(Expression))
-					{
-						++NumOutputs;
-						if (NumOutputs > 1 || !OutputExpression->IsResultMaterialAttributes(0))
-						{
-							CompileErrors.Add(TEXT("Layer blend graphs only support a single material attributes output."));
-							FailingExpression.Add(nullptr);
-						}
-					}
-					else if (UMaterialExpressionMaterialAttributeLayers* RecursiveLayer = Cast<UMaterialExpressionMaterialAttributeLayers>(Expression))
-					{
-						CompileErrors.Add(TEXT("Layer blend graphs do not support layers within layers."));
-						FailingExpression.Add(nullptr);
-					}
-				}
-
-				if (NumOutputs < 1)
-				{
-					CompileErrors.Add(TEXT("Layer blend graphs can have up to two material attributes inputs and a single output."));
-					FailingExpression.Add(nullptr);
-				}
-			}
-			else
-			{
-				// Add a compile error message for functions missing an output
-				FMaterialResource* CurrentResource = ExpressionPreviewMaterial->GetMaterialResource(FeatureLevel);
-				if (CurrentResource)
-				{
-					CompileErrors = CurrentResource->GetCompileErrors();
-					FailingExpression = CurrentResource->GetErrorExpressions();
-				}
-
-				bool bFoundFunctionOutput = false;
-				for (UMaterialExpression* MaterialExpression : Material->GetExpressions())
-				{
-					if (MaterialExpression->IsA(UMaterialExpressionFunctionOutput::StaticClass()))
-					{
-						bFoundFunctionOutput = true;
-						break;
-					}
-				}
-
-				if (!bFoundFunctionOutput)
-				{
-					CompileErrors.Add(TEXT("Missing a function output"));
-					FailingExpression.Add(nullptr);
-				}
-			}
-		}
-		else
-		{
-			CompileErrors = MaterialResource->GetCompileErrors();
-			FailingExpression = MaterialResource->GetErrorExpressions();
-		}
+		
+		CompileErrors = MaterialResource->GetCompileErrors();
+		FailingExpression = MaterialResource->GetErrorExpressions();
 
 		// Only show general info if there are no errors and stats are enabled - Stats show for Materials, layers and blends
 		if (CompileErrors.Num() == 0 && (!MaterialFunction || MaterialFunction->GetMaterialFunctionUsage() != EMaterialFunctionUsage::Default))
diff --git a/Engine/Source/Editor/MaterialEditor/Private/MaterialPropertyHelpers.cpp b/Engine/Source/Editor/MaterialEditor/Private/MaterialPropertyHelpers.cpp
--- a/Engine/Source/Editor/MaterialEditor/Private/MaterialPropertyHelpers.cpp
+++ b/Engine/Source/Editor/MaterialEditor/Private/MaterialPropertyHelpers.cpp
@@ -159,93 +159,100 @@ void FMaterialPropertyHelpers::OnMaterialLayerAssetChanged(const struct FAssetDa
 
 bool FMaterialPropertyHelpers::FilterLayerAssets(const struct FAssetData& InAssetData, FMaterialLayersFunctions* LayerFunction, EMaterialParameterAssociation MaterialType, int32 Index)
 {
-	bool ShouldAssetBeFilteredOut = false;
+	bool bShouldAssetBeFilteredOut = true;
+
 	const FName FilterTag = FName(TEXT("MaterialFunctionUsage"));
-	const FName BaseTag = FName(TEXT("Base"));
 	FAssetDataTagMapSharedView::FFindTagResult MaterialFunctionUsage = InAssetData.TagsAndValues.FindTag(FilterTag);
-
-	FName BaseClassName;
-	FName InstanceClassName;
-
-	FString CompareString;
 	if (MaterialFunctionUsage.IsSet())
 	{
-		FAssetDataTagMapSharedView::FFindTagResult Base = InAssetData.TagsAndValues.FindTag(BaseTag);
-
-		FString BaseString;
-		FString DiscardString;
-		FString CleanString;
-		if (Base.IsSet())
+		static const UEnum* FunctionUsageEnum = StaticEnum<EMaterialFunctionUsage>();
+		if (!FunctionUsageEnum)
 		{
-			BaseString = Base.GetValue();
-			BaseString.Split(".", &DiscardString, &CleanString);
-			CleanString.Split("'", &CleanString, &DiscardString);
-		}
-		else
-		{
-			CleanString = InAssetData.AssetName.ToString();
+			return bShouldAssetBeFilteredOut;
 		}
 
-		FString LeftPath;
+		const FName BaseTag = FName(TEXT("Base"));
+		FAssetDataTagMapSharedView::FFindTagResult Base = InAssetData.TagsAndValues.FindTag(BaseTag);
+		FSoftObjectPath AssetPath = LayerFunction->LayerStackCache != nullptr && Base.IsSet() ? FSoftObjectPath(Base.GetValue()) : InAssetData.ToSoftObjectPath();
+
 		FString RightPath;
-		bool bShouldFilter = false;
+		bool bLegacyShouldRestrict = false;
 		switch (MaterialType)
 		{
 		case EMaterialParameterAssociation::LayerParameter:
 		{
-			CompareString = "MaterialLayer";
-			if (LayerFunction->Layers[Index] != nullptr)
+			if (MaterialFunctionUsage.GetValue() == FunctionUsageEnum->GetNameStringByValue((int64)EMaterialFunctionUsage::Default) 
+				|| MaterialFunctionUsage.GetValue() == FunctionUsageEnum->GetNameStringByValue((int64)EMaterialFunctionUsage::MaterialLayer))
 			{
-				RightPath = LayerFunction->Layers[Index]->GetBaseFunction()->GetFName().ToString();
-				if (RightPath.IsEmpty())
+				//If we are using the new layer stack node, we only need to ensure whether the asset is usable.
+				if (LayerFunction->LayerStackCache)
 				{
-					RightPath = LayerFunction->Layers[Index]->GetFName().ToString();
+					bShouldAssetBeFilteredOut = !LayerFunction->LayerStackCache->AvailableLayerPaths.Contains(AssetPath);
+				}
+				else
+				{
+					bShouldAssetBeFilteredOut = false;
+					bLegacyShouldRestrict = LayerFunction->EditorOnly.RestrictToLayerRelatives[Index];
+					if (bLegacyShouldRestrict && LayerFunction->Layers[Index] != nullptr)
+					{
+						RightPath = LayerFunction->Layers[Index]->GetBaseFunction()->GetFName().ToString();
+						if (RightPath.IsEmpty())
+						{
+							RightPath = LayerFunction->Layers[Index]->GetFName().ToString();
+						}
+					}
 				}
 			}
-			bShouldFilter = LayerFunction->EditorOnly.RestrictToLayerRelatives[Index];
-			BaseClassName = UMaterialFunctionMaterialLayer::StaticClass()->GetFName();
-			InstanceClassName = UMaterialFunctionMaterialLayerInstance::StaticClass()->GetFName();
+			break;
 		}
-		break;
 		case EMaterialParameterAssociation::BlendParameter:
 		{
-			CompareString = "MaterialLayerBlend";
-			if (LayerFunction->Blends[Index] != nullptr)
+			if (MaterialFunctionUsage.GetValue() == FunctionUsageEnum->GetNameStringByValue((int64)EMaterialFunctionUsage::MaterialLayerBlend))
 			{
-				RightPath = LayerFunction->Blends[Index]->GetBaseFunction()->GetFName().ToString();
-				if (RightPath.IsEmpty())
+				if (LayerFunction->LayerStackCache)
+				{
+					bShouldAssetBeFilteredOut = !LayerFunction->LayerStackCache->AvailableBlendPaths.Contains(AssetPath);
+				}
+				else
 				{
-					RightPath = LayerFunction->Blends[Index]->GetFName().ToString();
+					bShouldAssetBeFilteredOut = false;
+					bLegacyShouldRestrict = LayerFunction->EditorOnly.RestrictToBlendRelatives[Index];
+					if (bLegacyShouldRestrict && LayerFunction->Blends[Index] != nullptr)
+					{
+						RightPath = LayerFunction->Blends[Index]->GetBaseFunction()->GetFName().ToString();
+						if (RightPath.IsEmpty())
+						{
+							RightPath = LayerFunction->Blends[Index]->GetFName().ToString();
+						}
+					}
 				}
 			}
-			bShouldFilter = LayerFunction->EditorOnly.RestrictToBlendRelatives[Index];
-			BaseClassName = UMaterialFunctionMaterialLayerBlend::StaticClass()->GetFName();
-			InstanceClassName = UMaterialFunctionMaterialLayerBlendInstance::StaticClass()->GetFName();
+			break;
 		}
-		break;
 		default:
-		break;
+			break;
 		}
 
-		if (MaterialFunctionUsage.GetValue() != CompareString)
+		if (bLegacyShouldRestrict)
 		{
-
-			ShouldAssetBeFilteredOut = true;
-		}
-		else
-		{
-			bool bSameBase = CleanString == RightPath;
-			if (!RightPath.IsEmpty() && !bSameBase && bShouldFilter)
+			FString BaseString;
+			FString DiscardString;
+			FString CleanString;
+			if (Base.IsSet())
+			{
+				BaseString = Base.GetValue();
+				BaseString.Split(".", &DiscardString, &CleanString);
+				CleanString.Split("'", &CleanString, &DiscardString);
+			}
+			else
 			{
-				ShouldAssetBeFilteredOut = true;
+				CleanString = InAssetData.AssetName.ToString();
 			}
+
+			bShouldAssetBeFilteredOut = CleanString != RightPath;
 		}
 	}
-	else
-	{
-		ShouldAssetBeFilteredOut = true;
-	}
-	return ShouldAssetBeFilteredOut;
+	return bShouldAssetBeFilteredOut;
 }
 
 FReply FMaterialPropertyHelpers::OnClickedSaveNewMaterialInstance(UMaterialInterface* Parent, UObject* EditorObject)
diff --git a/Engine/Source/Runtime/Core/Public/UObject/FortniteMainBranchObjectVersions.inl b/Engine/Source/Runtime/Core/Public/UObject/FortniteMainBranchObjectVersions.inl
--- a/Engine/Source/Runtime/Core/Public/UObject/FortniteMainBranchObjectVersions.inl
+++ b/Engine/Source/Runtime/Core/Public/UObject/FortniteMainBranchObjectVersions.inl
@@ -640,5 +640,8 @@ FORTNITE_MAIN_VERSION(MediaPlateOverlayTechniqueRemoval, 208)
 // Added particle flag to allow/disallow partial island sleeping in the island the particle is in
 FORTNITE_MAIN_VERSION(PerParticleFlagToAllowPartialIslandSleepInConnectedIsland, 209)
 
+// Material Function Blend Deserialize Top/Bottom input nodes with clearer enum marker.
+FORTNITE_MAIN_VERSION(MaterialFunctionBlendTopBottomInputEnum, 210)
+
 // -----<new versions can be added above this line>-------------------------------------------------
 #undef FORTNITE_MAIN_VERSION
\ No newline at end of file
diff --git a/Engine/Source/Runtime/Engine/Private/Materials/MaterialCachedData.cpp b/Engine/Source/Runtime/Engine/Private/Materials/MaterialCachedData.cpp
--- a/Engine/Source/Runtime/Engine/Private/Materials/MaterialCachedData.cpp
+++ b/Engine/Source/Runtime/Engine/Private/Materials/MaterialCachedData.cpp
@@ -20,6 +20,7 @@
 #include "Materials/MaterialExpressionSetMaterialAttributes.h"
 #include "Materials/MaterialExpressionMakeMaterialAttributes.h"
 #include "Materials/MaterialExpressionMaterialAttributeLayers.h"
+#include "Materials/MaterialExpressionLayerStack.h"
 #include "Materials/MaterialExpressionMaterialFunctionCall.h"
 #include "Materials/MaterialExpressionStaticSwitchParameter.h"
 #include "Materials/MaterialExpressionStaticSwitch.h"
@@ -673,6 +674,10 @@ void FMaterialCachedExpressionData::UpdateForExpressions(const FMaterialCachedEx
 			{
 				const FMaterialLayersFunctions& Layers = Context.LayerOverrides ? *Context.LayerOverrides : LayersExpression->DefaultLayers;
 				UpdateForLayerFunctions(Context, Layers);
+				if (UMaterialExpressionLayerStack* LayerStackExpression = Cast<UMaterialExpressionLayerStack>(LayersExpression))
+				{
+					Layers.LayerStackCache = LayerStackExpression->GetSharedAvailableFunctionsCache();
+				}
 
 				// TODO(?) - Layers for MIs are currently duplicated here and in FStaticParameterSet
 				bHasMaterialLayers = true;
diff --git a/Engine/Source/Runtime/Engine/Private/Materials/MaterialExpressionLayerStack.cpp b/Engine/Source/Runtime/Engine/Private/Materials/MaterialExpressionLayerStack.cpp
new file mode 100644
--- /dev/null
+++ b/Engine/Source/Runtime/Engine/Private/Materials/MaterialExpressionLayerStack.cpp
@@ -0,0 +1,740 @@
+// Copyright Epic Games, Inc. All Rights Reserved.
+
+/*=============================================================================
+	MaterialExpressionLayerStack.cpp - Material expression Layer Stacks implementation.
+=============================================================================*/
+#include "Materials/MaterialExpressionLayerStack.h"
+#include "MaterialShared.h"
+#include "Materials/Material.h"
+#include "Materials/MaterialExpressionFunctionOutput.h"
+#include "Materials/MaterialFunction.h"
+#include "Materials/MaterialFunctionInstance.h"
+#include "UObject/ObjectSaveContext.h"
+
+#if WITH_EDITOR
+#include "MaterialGraph/MaterialGraphSchema.h"
+#endif
+
+//
+// FMaterialLayerInput
+//
+FMaterialLayerInput::FMaterialLayerInput(FName NewInputName, EFunctionInputType NewInputType)
+{
+	InputName = NewInputName;
+	InputType = NewInputType;
+}
+
+FString FMaterialLayerInput::GetInputName() const
+{
+	FString OutName;
+	//Statics are disallowed allowed in the layer stack system.
+	if (InputType == FunctionInput_StaticBool
+		|| InputType == FunctionInput_MAX)
+	{
+		return OutName;
+	}
+
+	const FString TypeString = UMaterialExpressionFunctionInput::GetInputTypeDisplayName(InputType);
+	if (InputName.IsValid() && !TypeString.IsEmpty())
+	{
+		FString NameString = InputName.ToString();
+		OutName = NameString.Appendf(TEXT(" (%s)"), *TypeString);
+	}
+	return OutName;
+}
+
+//
+// UMaterialExpressionLayerStack
+//
+UMaterialExpressionLayerStack::UMaterialExpressionLayerStack(const FObjectInitializer& ObjectInitializer)
+	: Super(ObjectInitializer)
+{
+#if WITH_EDITOR
+	CachedInputs.Empty();
+#endif
+}
+
+void UMaterialExpressionLayerStack::PostLoad()
+{
+#if WITH_EDITOR
+	bAreAvailableLayersValid = false;
+	ResolveLayerInputs();
+	GetSharedAvailableFunctionsCache();
+#endif
+	Super::PostLoad();
+}
+
+void UMaterialExpressionLayerStack::PreSave(FObjectPreSaveContext ObjectSaveContext)
+{
+	Super::PreSave(ObjectSaveContext);
+#if WITH_EDITOR
+	GetSharedAvailableFunctionsCache();
+#endif
+}
+
+TArray<FSoftObjectPath> UMaterialExpressionLayerStack::GetPathsFromAvailableFunctions(TSet<TObjectPtr<UMaterialFunctionInterface>>& InFunctions)
+{
+	//We use paths to manage the available layers/blends so we can evaluate against asset (meta)data instead of requiring loaded objects.
+	TArray<FSoftObjectPath> OutPaths;
+	for (TObjectPtr<UMaterialFunctionInterface> Function : InFunctions)
+	{
+		if (!Function)
+		{
+			continue;
+		}
+
+		/**
+		* We always map the base function to ensure all functions and instances in the vertical can be referenced regardless of heirarchy
+		* This should be handled by the existing code which manages the available layers/blends list but no harm in ensuring this
+		* is the only info we pass to the UI for asset management.
+		*/
+		if (UMaterialFunctionInterface* BaseFunction = Function->GetBaseFunction())
+		{
+			OutPaths.AddUnique(FSoftObjectPath(BaseFunction).ToString());
+		}
+	}
+	return MoveTemp(OutPaths);
+}
+
+#if WITH_EDITOR
+TSharedPtr<FMaterialLayerStackFunctionsCache> UMaterialExpressionLayerStack::GetSharedAvailableFunctionsCache()
+{
+	/**
+	* Managing available layers across expressions and the UI requires an interface via the FMaterialLayersFunctions class
+	* However modifying the runtime data directly is difficult without causing issues with managing lifetime of the referenced items.
+	* Utilising a dedicated shared cache between parents and instances ensures the available layers/blends are always available for reference.
+	*/
+	if (!SharedCache)
+	{
+		SharedCache = MakeShared<FMaterialLayerStackFunctionsCache>();
+	}
+
+	SharedCache->AvailableLayerPaths = GetAvailableLayerPaths();
+	SharedCache->AvailableBlendPaths = GetAvailableBlendPaths();
+	return SharedCache;
+}
+
+void UMaterialExpressionLayerStack::ResolveLayerInputs()
+{
+	bool bHasChanged = false;
+	bool bClearAllInputs = false;
+	if (AvailableLayers.IsEmpty() && AvailableBlends.IsEmpty())
+	{
+		bClearAllInputs = !LayerInputs.IsEmpty();
+	}
+	else
+	{
+		//Helper function for retrieving per layer inputs.
+		auto GetLayerInputs = [](UMaterialFunctionInterface* CurrentFunction, TMap<FString, FMaterialLayerInput>& CollectedInputs) -> void
+			{
+				//If null or if the expression has already been recursed, skip.
+				if (!CurrentFunction)
+				{
+					return;
+				}
+
+				TArray<FFunctionExpressionInput> FunctionInputs;
+				TArray<FFunctionExpressionOutput> FunctionOutputs;
+				CurrentFunction->GetInputsAndOutputs(FunctionInputs, FunctionOutputs);
+
+				for (FFunctionExpressionInput& CurrentInput : FunctionInputs)
+				{
+					//For now we have to only process blend input types that are not marked top/bottom.
+					if (!CurrentInput.ExpressionInput || CurrentInput.ExpressionInput->BlendInputRelevance != EBlendInputRelevance::General)
+					{
+						continue;
+					}
+
+					/**
+					* We create a map of the existing name which contains the Name and the Type to ensure we allow the maximum number of unique entries
+					* even if 2 inputs have the same name but different types, this should be allowed.
+					*/
+					FMaterialLayerInput NewInput(CurrentInput.Input.InputName, CurrentInput.ExpressionInput->InputType);
+					FString CurrentName = NewInput.GetInputName();
+					if (!CurrentName.IsEmpty() && !CollectedInputs.Contains(CurrentName))
+					{
+						CollectedInputs.Add(CurrentName, NewInput);
+					}
+				}
+			};
+
+		TMap<FString, FMaterialLayerInput> AllInputs;
+		for (UMaterialFunctionInterface* Layer : AvailableLayers)
+		{
+			GetLayerInputs(Layer, AllInputs);
+		}
+
+		for (UMaterialFunctionInterface* Blend : AvailableBlends)
+		{
+			GetLayerInputs(Blend, AllInputs);
+		}
+
+		if (AllInputs.IsEmpty())
+		{
+			bClearAllInputs = !LayerInputs.IsEmpty();
+		}
+		else
+		{
+			/**
+			* This loop ensures we retain existing inputs that already have node connections,
+			* whilst also tracking any inputs that need to be removed from the existing list.
+			*/
+			TArray<FString> RemoveList;
+			for (FMaterialLayerInput& ExistingInput : LayerInputs)
+			{
+				FString ExistingInputName = ExistingInput.GetInputName();
+				FMaterialLayerInput* NewInput = AllInputs.Find(ExistingInputName);
+				if (NewInput)
+				{
+					AllInputs.Remove(ExistingInputName);
+				}
+				else
+				{
+					if (GraphNode)
+					{
+						UEdGraphPin* ThisPin = GraphNode->FindPin(*ExistingInputName, EGPD_Input);
+						GraphNode->RemovePin(ThisPin);
+					}
+					RemoveList.Add(ExistingInputName);
+				}
+			}
+
+			bool bRemovingInputs = !RemoveList.IsEmpty();
+			bHasChanged = !AllInputs.IsEmpty() || bRemovingInputs;
+			if (bHasChanged)
+			{
+				Modify();
+				if (bRemovingInputs)
+				{
+					//Remove any inputs that no longer exist in the available layers/blends
+					LayerInputs.RemoveAll([&RemoveList](FMaterialLayerInput& LayerInput)
+						{
+							return RemoveList.Contains(LayerInput.GetInputName());
+						}
+					);
+				}
+
+				//Add any new inputs
+				for (TPair<FString, FMaterialLayerInput>& InputPair : AllInputs)
+				{
+					LayerInputs.Add(InputPair.Value);
+				}
+			}
+		}
+	}
+
+	if (bClearAllInputs)
+	{
+		Modify();
+		if (GraphNode)
+		{
+			TArrayView<UEdGraphPin*> ExistingPins = GraphNode->Pins;
+			for (UEdGraphPin* InputPin : ExistingPins)
+			{
+				if (InputPin->Direction == EGPD_Input)
+				{
+					GraphNode->RemovePin(InputPin);
+				}
+			}
+		}
+		LayerInputs.Empty();
+		bHasChanged = true;
+	}
+
+	if (bHasChanged)
+	{
+		if (LayerInputs.Num() > 1)
+		{
+			//Sort the layer inputs, first by type, then in alphabetical order
+			LayerInputs.Sort([](const FMaterialLayerInput& A, const FMaterialLayerInput& B) {
+				if (A.InputType < B.InputType)
+				{
+					return true;
+				}
+				else if (A.InputType == B.InputType)
+				{
+					return A.InputName.ToString() < B.InputName.ToString();
+				}
+				return false;
+				});
+		}
+
+		CacheLayerInputs();
+
+		if (GraphNode)
+		{
+			GraphNode->ReconstructNode();
+		}
+		bAreAvailableLayersValid = false;
+	}
+}
+
+UMaterialFunctionInterface* UMaterialExpressionLayerStack::ExtractParentFunctionFromInstance(FMaterialCompiler* Compiler, UMaterialFunctionInterface* CurrentFunction)
+{
+	/**
+	* Instances need to be mapped to their parent, we should allow all instances in a vertical of an allowed, validated base.
+	*/
+	if (UMaterialFunctionInstance* Instance = Cast<UMaterialFunctionInstance>(CurrentFunction))
+	{
+		if (UMaterialFunctionInterface* Parent = Instance->GetBaseFunction())
+		{
+			CurrentFunction = Parent;
+		}
+		else
+		{
+			LogError(Compiler, TEXT("Function %s: instance has no parent set."), *CurrentFunction->GetName());
+			CurrentFunction = nullptr;
+		}
+	}
+	return CurrentFunction;
+}
+
+bool UMaterialExpressionLayerStack::PollFunctionExpressionsForLayerUsage(FMaterialCompiler* Compiler, UMaterialFunctionInterface* CurrentFunction, FValidLayerUsageTracker& Tracker, bool bCheckStatics)
+{
+	CurrentFunction = ExtractParentFunctionFromInstance(Compiler, CurrentFunction);
+	if (!CurrentFunction || !CurrentFunction->IsA<UMaterialFunction>())
+	{
+		LogError(Compiler, TEXT("Function %s: Invalid base material function being validated. Instances cannot be validated, only their base functions."), *CurrentFunction->GetName());
+		return false;
+	}
+
+	TConstArrayView<TObjectPtr<UMaterialExpression>> AllExpressions = CurrentFunction->GetExpressions();
+	for (TObjectPtr<UMaterialExpression> Expression : AllExpressions)
+	{
+		if (UMaterialExpressionFunctionInput* InputExpression = Cast<UMaterialExpressionFunctionInput>(Expression))
+		{
+			if (InputExpression->GetInputValueType(0) == MCT_MaterialAttributes)
+			{
+				Tracker.MAInputCount++;
+			}
+		}
+
+		if (bCheckStatics && !Tracker.bContainsStatics && Expression->IsStaticExpression())
+		{
+			Tracker.bContainsStatics = true;
+		}
+
+		if (UMaterialExpressionFunctionOutput* OutputExpression = Cast<UMaterialExpressionFunctionOutput>(Expression))
+		{
+			if (OutputExpression->GetOutputValueType(0) == MCT_MaterialAttributes)
+			{
+				Tracker.MAOutputCount++;
+			}
+		}
+	}
+
+	bool bResult = true;
+	if (bCheckStatics && Tracker.bContainsStatics)
+	{
+		LogError(Compiler, TEXT("Function %s: cannot contain any static nodes as these are incompatible with the layering."), *CurrentFunction->GetName());
+		bResult = false;
+	}
+
+	if (Tracker.MAOutputCount != AcceptableNumLayerBlendMAOutputs)
+	{
+		LogError(Compiler, TEXT("Function %i: %s, must have exactly %d Material Attributes type output node. It has %d"), *CurrentFunction->GetName(), AcceptableNumLayerBlendMAOutputs, Tracker.MAOutputCount);
+		bResult = false;
+	}
+
+	return bResult;
+}
+
+bool UMaterialExpressionLayerStack::ValidateFunctionForLayerUsage(FMaterialCompiler* Compiler, UMaterialFunctionInterface* CurrentFunction, bool bCheckStatics)
+{
+	EMaterialFunctionUsage Usage = CurrentFunction->GetMaterialFunctionUsage();
+	if (!(Usage == EMaterialFunctionUsage::Default || Usage == EMaterialFunctionUsage::MaterialLayer))
+	{
+		LogError(Compiler, TEXT("Function %s: function usage is not set for use as a layer."), *CurrentFunction->GetName());
+		return false;
+	}
+
+	FValidLayerUsageTracker Tracker;
+	bool bResult = PollFunctionExpressionsForLayerUsage(Compiler, CurrentFunction, Tracker, bCheckStatics);
+
+	if (Tracker.MAInputCount > AcceptableNumLayerMAInputs)
+	{
+		LogError(Compiler, TEXT("Layer %s: can only have up to %d Material Attributes type output node. It has %d"), *CurrentFunction->GetName(), AcceptableNumLayerMAInputs, Tracker.MAInputCount);
+		bResult = false;
+	}
+
+	return bResult;
+}
+
+bool UMaterialExpressionLayerStack::ValidateFunctionForBlendUsage(FMaterialCompiler* Compiler, UMaterialFunctionInterface* CurrentFunction, bool bCheckStatics)
+{
+	EMaterialFunctionUsage Usage = CurrentFunction->GetMaterialFunctionUsage();
+	if (Usage != EMaterialFunctionUsage::MaterialLayerBlend)
+	{
+		LogError(Compiler, TEXT("Function %s: function usage is not set for use as a blend."), *CurrentFunction->GetName());
+		return false;
+	}
+
+	FValidLayerUsageTracker Tracker;
+	bool bResult = PollFunctionExpressionsForLayerUsage(Compiler, CurrentFunction, Tracker, bCheckStatics);
+	if (Tracker.MAInputCount != AcceptableNumBlendMAInputs)
+	{
+		LogError(Compiler, TEXT("Blend %s: must have exactly %d Material Attributes type input node. It has %d"), *CurrentFunction->GetName(), AcceptableNumBlendMAInputs, Tracker.MAInputCount);
+		bResult = false;
+	}
+
+	return true;
+}
+
+bool UMaterialExpressionLayerStack::ValidateLayerConfiguration(FMaterialCompiler* Compiler, bool bReportErrors)
+{
+	bool bIsValid = true;
+
+	if (!bAreAvailableLayersValid)
+	{
+		//Helper function for retrieving per layer inputs.
+		auto ValidateAvailableFunctions = [Compiler](TSet<TObjectPtr<UMaterialFunctionInterface>>& AvailableFunctions, bool& bIsValid, const bool bIsBlendArray) -> void
+			{
+				for (TObjectPtr<UMaterialFunctionInterface>& Function : AvailableFunctions)
+				{
+					if (Function)
+					{
+						UMaterialFunctionInterface* BaseFunction = ExtractParentFunctionFromInstance(Compiler, Function);
+						bool bRequiresValidation = false;
+						if (BaseFunction == Function)
+						{
+							bRequiresValidation = true;
+						}
+						else
+						{
+							/**
+							* The validation process involves resolving the parent to ensure any function instance can be used
+							* so to save having to check each entry in the available list constantly, we only want parents in the available layers.
+							* This way once we have the parent of an actually used layer, we can speed up the comparison.
+							*
+							* Additionally, we need to ensure we aren't duplicating existing entries.
+							*/
+							const FString FunctionType = bIsBlendArray ? "blends" : "layers";
+							if (AvailableFunctions.Contains(BaseFunction))
+							{
+								UE_LOG(LogMaterial, Warning, TEXT("%s not appended because available %s already contain the base (%s)."), *Function->GetName(), *FunctionType, *BaseFunction->GetName());
+								Function = nullptr;
+							}
+							else
+							{
+								UE_LOG(LogMaterial, Warning, TEXT("Resolving base (%s) of %s for the available %s list."), *BaseFunction->GetName(), *Function->GetName(), *FunctionType);
+								Function = BaseFunction;
+								bRequiresValidation = true;
+							}
+						}
+
+						if (bRequiresValidation)
+						{
+							bIsValid = bIsBlendArray ? ValidateFunctionForBlendUsage(Compiler, BaseFunction, true) : ValidateFunctionForLayerUsage(Compiler, BaseFunction, true);
+						}
+					}
+				}
+			};
+		ValidateAvailableFunctions(AvailableLayers, bIsValid, false);
+		ValidateAvailableFunctions(AvailableBlends, bIsValid, true);
+		bAreAvailableLayersValid = bIsValid;
+
+		if (!bIsValid)
+		{
+			LogError(Compiler, TEXT("LayerStack Available Layers/Blends contain invalid functions."));
+			return false;
+		}
+	}
+
+	for (FMaterialLayerInput& LayerInput : LayerInputs)
+	{
+		if (!LayerInput.IsConnected())
+		{
+			LogError(Compiler, TEXT("LayerStack \"%s\" input pin must be connected."), *LayerInput.GetInputName());
+			bIsValid = false;
+		}
+	}
+
+	// If layer inputs are not connected, early out. We want to force users to apply default values at least at the base material level for now.
+	if (!bIsValid)
+	{
+		return false;
+	}
+
+	/**
+	* Despite available layers / blends having their validation cached, we still have to then validate the utilised blends / layers.
+	* Fortunately the cached available functions allows us to just ensure they are present in the existing list to reduce check logic required.
+	*/
+	const TArray<UMaterialFunctionInterface*>& Layers = GetLayers();
+	int32 NumActiveLayers = 0;
+	for (int32 LayerIndex = 0; LayerIndex < Layers.Num(); ++LayerIndex)
+	{
+		UMaterialFunctionInterface* Layer = Layers[LayerIndex];
+		if (Layer)
+		{
+			Layer = ExtractParentFunctionFromInstance(Compiler, Layer);
+			if (AvailableLayers.Contains(Layer))
+			{
+				NumActiveLayers++;
+			}
+			else
+			{
+				LogError(Compiler, TEXT("Layer %i, %s, is not set as an allowed layer in the base layer stack node."), LayerIndex, *Layer->GetName());
+				bIsValid = false;
+			}
+		}
+	}
+
+	const TArray<UMaterialFunctionInterface*>& Blends = GetBlends();
+	int32 NumActiveBlends = 0;
+	for (int32 BlendIndex = 0; BlendIndex < Blends.Num(); ++BlendIndex)
+	{
+		UMaterialFunctionInterface* Blend = Blends[BlendIndex];
+		if (Blend)
+		{
+			Blend = ExtractParentFunctionFromInstance(Compiler, Blend);
+			if (AvailableBlends.Contains(Blend))
+			{
+				NumActiveBlends++;
+			}
+			else
+			{
+				LogError(Compiler, TEXT("Blend %i, %s, is not set as an allowed blend in the base layer stack node."), BlendIndex, *Blend->GetName());
+				bIsValid = false;
+			}
+		}
+	}
+
+	if (!bIsValid)
+	{
+		return false;
+	}
+
+	// Currently we want only allow these configurations but in future we will unlock blends and layer only lists being usable together, however a new UI is likely required for this.
+	bool bValidConfiguration = (NumActiveLayers >= 0 && NumActiveBlends == 0)						// Default behaviour if both 0, otherwise layers only configuration
+							|| (NumActiveLayers >= 2 && NumActiveBlends == NumActiveLayers - 1);	// Blend graph
+
+	if (!bValidConfiguration)
+	{
+		LogError(Compiler, TEXT("Invalid number of layers (%i) or blends (%i) assigned. Number of blends must be equal to 0, or 1 less than the number of active layers."), NumActiveLayers, NumActiveBlends);
+	}
+	return bValidConfiguration;
+}
+
+void UMaterialExpressionLayerStack::CacheLayerInputs()
+{
+	CachedInputs.Empty();
+	if (!LayerInputs.IsEmpty())
+	{
+		CachedInputs.Reserve(LayerInputs.Num());
+		for (FExpressionInput& CurrentInput : LayerInputs)
+		{
+			CachedInputs.Push(&CurrentInput);
+		}
+	}
+}
+
+void UMaterialExpressionLayerStack::RebuildLayerGraph(bool bReportErrors)
+{
+	if (ValidateLayerConfiguration(nullptr, bReportErrors))
+	{
+		// Reset graph connectivity
+		bIsLayerGraphBuilt = false;
+
+		//Similar to the parent implementation, this helper function creates a new Function call for each referenced function to create the new graph
+		auto ProcessFunctionCallers = [this](const TArray<UMaterialFunctionInterface*>& Functions, TArray<TObjectPtr<UMaterialExpressionMaterialFunctionCall>>& CallerArray, int32& CallerCount, EMaterialParameterAssociation InAssociation) -> void
+			{
+				CallerArray.Empty();
+				CallerCount = 0;
+				for (int32 LayerIndex = 0; LayerIndex < Functions.Num(); LayerIndex++)
+				{
+					//Create the new function call
+					TObjectPtr<UMaterialExpressionMaterialFunctionCall> NewFunctionCall = NewObject<UMaterialExpressionMaterialFunctionCall>(GetTransientPackage());
+					NewFunctionCall->SetMaterialFunction(Functions[LayerIndex]);
+					NewFunctionCall->FunctionParameterInfo.Association = InAssociation;
+					NewFunctionCall->FunctionParameterInfo.Index = LayerIndex;
+					NewFunctionCall->UpdateFromFunctionResource();
+					if (!this->LayerInputs.IsEmpty() && !NewFunctionCall->FunctionInputs.IsEmpty())
+					{
+						for (FMaterialLayerInput& LayerInput : this->LayerInputs)
+						{
+							if (!LayerInput.Expression)
+							{
+								continue;
+							}
+
+							for (int32 CallInputIndex = 0; CallInputIndex < NewFunctionCall->FunctionInputs.Num(); CallInputIndex++)
+							{
+								FFunctionExpressionInput& ThisInput = NewFunctionCall->FunctionInputs[CallInputIndex];
+								if (!ThisInput.ExpressionInput)
+								{
+									continue;
+								}
+
+								//If an exposed layer input at the parent material is exposed and connected, map the overrides
+								if (ThisInput.Input.InputName == LayerInput.InputName
+									&& ThisInput.ExpressionInput->InputType == LayerInput.InputType)
+								{
+									//Sets the read only inputs, but will be overridden through the stack if any sublayers have matching outputs
+									ThisInput.Input.Connect(LayerInput.OutputIndex, LayerInput.Expression);
+								}
+							}
+						}
+					}
+					CallerArray.Add(NewFunctionCall);
+					if (NewFunctionCall->MaterialFunction)
+					{
+						CallerCount++;
+					}
+				}
+			};
+
+		//ProcessLayers
+		ProcessFunctionCallers(GetLayers(), LayerCallers, NumActiveLayerCallers, EMaterialParameterAssociation::LayerParameter);
+
+		//ProcessBlends
+		ProcessFunctionCallers(GetBlends(), BlendCallers, NumActiveBlendCallers, EMaterialParameterAssociation::BlendParameter);
+
+		// Assemble function chain so each layer blends with the previous
+		if (NumActiveLayerCallers >= 2 && NumActiveBlendCallers == NumActiveLayerCallers - 1)
+		{
+			int32 CurrentLayerIndex = 0;
+			TObjectPtr<UMaterialExpressionMaterialFunctionCall> Bottom = LayerCallers[CurrentLayerIndex++];
+			TObjectPtr<UMaterialExpressionMaterialFunctionCall> Top = LayerCallers[CurrentLayerIndex++];
+			for (int32 BlendIndex = 0; BlendIndex < NumActiveBlendCallers; BlendIndex++)
+			{
+				TObjectPtr<UMaterialExpressionMaterialFunctionCall> Blend = BlendCallers[BlendIndex];
+				bool bBottomSet = false;
+				bool bTopSet = false;
+				for (FFunctionExpressionInput& ThisInput : Blend->FunctionInputs)
+				{
+					/**
+					* For now we only want to connect MA types, and we restrict the number we accept in the entries
+					* as we can't account for name matching with the required blend inputs (hence the EBlendInputRelevance setting).
+					* However this will allow us to assign other input/output matching types at a later stage.
+					*/
+					if (ThisInput.ExpressionInput && ThisInput.ExpressionInput->InputType == FunctionInput_MaterialAttributes)
+					{
+						if (!bBottomSet && ThisInput.ExpressionInput->BlendInputRelevance == EBlendInputRelevance::Bottom)
+						{
+							ThisInput.Input.Connect(0, Bottom);
+							bBottomSet = true;
+						}
+
+						if (!bTopSet && ThisInput.ExpressionInput->BlendInputRelevance == EBlendInputRelevance::Top)
+						{
+							ThisInput.Input.Connect(0, Top);
+							bTopSet = true;
+						}
+					}
+
+					if (bBottomSet && bTopSet)
+					{
+						break;
+					}
+				}
+
+				if (++CurrentLayerIndex < NumActiveLayerCallers)
+				{
+					Bottom = LayerCallers[CurrentLayerIndex];
+					Top = Blend;
+				}
+				else
+				{
+					break;
+				}
+			}
+			bIsLayerGraphBuilt = true;
+		}
+		else if (NumActiveBlendCallers == 0)
+		{
+			//If no blends are present, create a layer chain
+			if (NumActiveLayerCallers > 1)
+			{
+				TObjectPtr<UMaterialExpressionMaterialFunctionCall> CurrentLayer = LayerCallers[NumActiveLayerCallers - 1];
+				for (int32 LayerIndex = NumActiveLayerCallers - 2; LayerIndex >= 0; LayerIndex--)
+				{
+					TObjectPtr<UMaterialExpressionMaterialFunctionCall> PreviousLayer = LayerCallers[LayerIndex];
+					for (int32 CallInputIndex = 0; CallInputIndex < CurrentLayer->FunctionInputs.Num(); CallInputIndex++)
+					{
+						FFunctionExpressionInput& CurrentInput = CurrentLayer->FunctionInputs[CallInputIndex];
+						EMaterialValueType CurrentInputType = CurrentInput.ExpressionInput->GetInputValueType(0);
+						if (!CurrentInput.ExpressionInput || CurrentInputType != MCT_MaterialAttributes)
+						{
+							//If we hit an input that can't be connected, break the chain connection.
+							if (CallInputIndex < CurrentLayer->FunctionInputs.Num() - 1)
+							{
+								UE_LOG(LogMaterial, Warning, TEXT("LayerStack's layer-only graph chain in \"%s\" is cut short due to missing MaterialAttributes input in %s."),
+									Material ? *(Material->GetName()) : TEXT("Unknown"), *CurrentLayer->GetName());
+							}
+							break;
+						}
+
+						//For each input, iterate through the outputs and connect if there's a match
+						//This may seem cumbersome, but this is laying the groundwork for mapping matching inputs and outputs outside of the MA type.
+						for (int32 CallOutputIndex = 0; CallOutputIndex < PreviousLayer->FunctionOutputs.Num(); CallOutputIndex++)
+						{
+							FFunctionExpressionOutput& PreviousOutput = PreviousLayer->FunctionOutputs[CallOutputIndex];
+							if (PreviousOutput.ExpressionOutput
+								&& CurrentInput.Input.InputName == PreviousOutput.ExpressionOutput->OutputName
+								&& CanConnectMaterialValueTypes(CurrentInputType, PreviousOutput.ExpressionOutput->GetOutputValueType(0)))
+							{
+								CurrentInput.Input.Connect(0, PreviousOutput.ExpressionOutput);
+								break;
+							}
+						}
+					}
+					CurrentLayer = PreviousLayer;
+				}
+			}
+			//We always say true at this point because regardless of having a layer chain, a single layer, or no entries we want to compile.
+			bIsLayerGraphBuilt = true;
+		}
+	}
+
+	if (!bIsLayerGraphBuilt && bReportErrors)
+	{
+		UE_LOG(LogMaterial, Warning, TEXT("Failed to build LayerStack graph for %s."), Material ? *(Material->GetName()) : TEXT("Unknown"));
+	}
+}
+
+void UMaterialExpressionLayerStack::PostEditChangeProperty(FPropertyChangedEvent& PropertyChangedEvent)
+{
+	if (PropertyChangedEvent.Property)
+	{
+		FName PropertyName = PropertyChangedEvent.Property->GetFName();
+		if (PropertyName == GET_MEMBER_NAME_CHECKED(UMaterialExpressionLayerStack, AvailableLayers)
+			|| PropertyName == GET_MEMBER_NAME_CHECKED(UMaterialExpressionLayerStack, AvailableBlends))
+		{
+			//If we change the available layers/blends sets, the validation of their usage needs to be re-evaluated.
+			bAreAvailableLayersValid = false;
+			ResolveLayerInputs();
+		}
+	}
+	Super::PostEditChangeProperty(PropertyChangedEvent);
+}
+
+void UMaterialExpressionLayerStack::GetCaption(TArray<FString>& OutCaptions) const
+{
+	OutCaptions.Add(TEXT("Layer Stack"));
+}
+
+void UMaterialExpressionLayerStack::GetExpressionToolTip(TArray<FString>& OutToolTip)
+{
+	ConvertToMultilineToolTip(TEXT("Evaluates the material layer stack against it's available layer functions, and outputs the merged attributes via the specified blends."), 40, OutToolTip);
+}
+
+FName UMaterialExpressionLayerStack::GetInputName(int32 InputIndex) const
+{
+	return LayerInputs.IsValidIndex(InputIndex) ? FName(LayerInputs[InputIndex].GetInputName()) : NAME_None;
+}
+
+FExpressionInput* UMaterialExpressionLayerStack::GetInput(int32 InputIndex)
+{
+	return LayerInputs.IsValidIndex(InputIndex) ? &LayerInputs[InputIndex] : nullptr;
+}
+
+TArrayView<FExpressionInput*> UMaterialExpressionLayerStack::GetInputsView()
+{
+	return CachedInputs;
+}
+
+EMaterialValueType UMaterialExpressionLayerStack::GetInputValueType(int32 InputIndex)
+{
+	return LayerInputs.IsValidIndex(InputIndex) ? UMaterialExpressionFunctionInput::GetMaterialTypeFromInputType(LayerInputs[InputIndex].InputType) : MCT_Unknown;
+}
+#endif //WITH_EDITOR
\ No newline at end of file
diff --git a/Engine/Source/Runtime/Engine/Private/Materials/MaterialExpressions.cpp b/Engine/Source/Runtime/Engine/Private/Materials/MaterialExpressions.cpp
--- a/Engine/Source/Runtime/Engine/Private/Materials/MaterialExpressions.cpp
+++ b/Engine/Source/Runtime/Engine/Private/Materials/MaterialExpressions.cpp
@@ -8,6 +8,7 @@
 #include "Misc/MessageDialog.h"
 #include "UObject/AssetRegistryTagsContext.h"
 #include "UObject/LinkerLoad.h"
+#include "UObject/FortniteMainBranchObjectVersion.h"
 #include "UObject/UE5MainStreamObjectVersion.h"
 #include "Internationalization/LocKeyFuncs.h"
 #include "UObject/NaniteResearchStreamObjectVersion.h"
@@ -24,8 +25,9 @@
 #include "MaterialShared.h"
 #include "Materials/HLSLMaterialTranslator.h"
 #include "Materials/MaterialAttributeDefinitionMap.h"
-#include "Materials/MaterialExpressionMaterialFunctionCall.h"
 #include "Materials/MaterialExpressionMaterialAttributeLayers.h"
+#include "Materials/MaterialExpressionLayerStack.h"
+#include "Materials/MaterialExpressionMaterialFunctionCall.h"
 #include "Materials/MaterialFunction.h"
 #include "Materials/MaterialFunctionMaterialLayer.h"
 #include "Materials/MaterialFunctionMaterialLayerBlend.h"
@@ -149,7 +151,6 @@
 #include "Materials/MaterialExpressionLogarithm10.h"
 #include "Materials/MaterialExpressionMakeMaterialAttributes.h"
 #include "Materials/MaterialExpressionMax.h"
-#include "Materials/MaterialExpressionMaterialAttributeLayers.h"
 #include "Materials/MaterialExpressionMaterialProxyReplace.h"
 #include "Materials/MaterialExpressionMin.h"
 #include "Materials/MaterialExpressionModulo.h"
@@ -6629,9 +6630,8 @@ EMaterialValueType UMaterialExpressionSetMaterialAttributes::GetInputValueType(i
 	{
 		InputType = MCT_MaterialAttributes;
 	}
-	else
+	else if(InputIndex > GET_SET_MA_MATERIALATTRIBUTESINDEX && InputIndex < AttributeSetTypes.Num() + 1)
 	{
-		ensure(InputIndex > GET_SET_MA_MATERIALATTRIBUTESINDEX && InputIndex < AttributeSetTypes.Num() + 1);
 		InputType = FMaterialAttributeDefinitionMap::GetValueType(AttributeSetTypes[InputIndex - 1]);
 		if (InputType == MCT_ShadingModel)
 		{
@@ -7437,8 +7437,6 @@ void UMaterialExpressionMaterialAttributeLayers::RebuildLayerGraph(bool bReportE
 
 		if (ValidateLayerConfiguration(nullptr, bReportErrors))
 		{
-
-
 			// Initialize layer function callers
 			for (int32 LayerIndex = 0; LayerIndex < Layers.Num(); ++LayerIndex)
 			{
@@ -7553,13 +7551,32 @@ void UMaterialExpressionMaterialAttributeLayers::OverrideLayerGraph(const FMater
 	}
 }
 
+
+void UMaterialExpressionMaterialAttributeLayers::LogError(FMaterialCompiler* Compiler, const TCHAR* Format, ...)
+{
+	TStringBuilder<256> Builder;
+	Builder.Append(TEXT("LayerValidation:"));
+	va_list Args;
+	va_start(Args, Format);
+	const TCHAR* ErrorMsg = Builder.AppendV(Format, Args).GetData();
+	va_end(Args);
+
+	if (Compiler)
+	{
+		Compiler->Errorf(ErrorMsg);
+	}
+	else
+	{
+		UE_LOG(LogMaterial, Warning, TEXT("%s"), ErrorMsg);
+	}
+}
+
 bool UMaterialExpressionMaterialAttributeLayers::ValidateLayerConfiguration(FMaterialCompiler* Compiler, bool bReportErrors)
 {
-#define COMPILER_OR_LOG_ERROR(Format, ...)								\
-	if (bReportErrors)													\
-	{																	\
-		if (Compiler) { Compiler->Errorf(Format, ##__VA_ARGS__); }		\
-		else { UE_LOG(LogMaterial, Warning, Format, ##__VA_ARGS__); }	\
+#define COMPILER_OR_LOG_ERROR(Format, ...)				\
+	if (bReportErrors)									\
+	{													\
+		LogError(Compiler, Format, ##__VA_ARGS__);		\
 	}
 
 	const TArray<UMaterialFunctionInterface*>& Layers = GetLayers();
@@ -7596,7 +7613,7 @@ bool UMaterialExpressionMaterialAttributeLayers::ValidateLayerConfiguration(FMat
 			{
 				TArray<UMaterialExpressionFunctionInput*> InputExpressions;
 				Layer->GetAllExpressionsOfType<UMaterialExpressionFunctionInput>(InputExpressions, false);
-				if (InputExpressions.Num() > 1)
+				if (InputExpressions.Num() > AcceptableNumLayerMAInputs)
 				{
 					COMPILER_OR_LOG_ERROR(TEXT("Layer %i, %s, must have one MaterialAttributes input only."), LayerIndex, *Layer->GetName());
 					bIsValid = false;
@@ -7634,7 +7651,7 @@ bool UMaterialExpressionMaterialAttributeLayers::ValidateLayerConfiguration(FMat
 				TArray<UMaterialExpressionFunctionInput*> InputExpressions;
 				Blend->GetAllExpressionsOfType<UMaterialExpressionFunctionInput>(InputExpressions, false);
 
-				if (Substrate::IsMaterialLayeringSupportEnabled() ? InputExpressions.Num() < 2 : InputExpressions.Num() != 2)
+				if (Substrate::IsMaterialLayeringSupportEnabled() ? InputExpressions.Num() < AcceptableNumBlendMAInputs : InputExpressions.Num() != AcceptableNumBlendMAInputs)
 				{
 					COMPILER_OR_LOG_ERROR(TEXT("Blend %i, %s, must have two MaterialAttributes inputs."), BlendIndex, *Blend->GetName());
 					bIsValid = false;
@@ -7869,10 +7886,7 @@ EMaterialValueType UMaterialExpressionMaterialAttributeLayers::GetInputValueType
 }
 #endif // WITH_EDITOR
 
-
-// -----
 #if WITH_EDITOR
-
 int32 UMaterialExpressionFloatToUInt::Compile(class FMaterialCompiler* Compiler, int32 OutputIndex)
 {
 	if (!Input.GetTracedInput().Expression)
@@ -13843,6 +13857,8 @@ void UMaterialFunction::Serialize(FArchive& Ar)
 {
 	Super::Serialize(Ar);
 
+	Ar.UsingCustomVersion(FFortniteMainBranchObjectVersion::GUID);
+
 #if WITH_EDITOR
 	if (Ar.UEVer() < VER_UE4_FLIP_MATERIAL_COORDS)
 	{
@@ -14707,39 +14723,7 @@ bool UMaterialFunction::ValidateFunctionUsage(FMaterialCompiler* Compiler, const
 		}
 		else
 		{
-			// Material layers must have a single MA input and output only
-			for (UMaterialExpression* Expression : GetExpressions())
-			{
-				if (UMaterialExpressionFunctionInput* InputExpression = Cast<UMaterialExpressionFunctionInput>(Expression))
-				{
-					++NumInputs;
-					if (NumInputs > 1 || !InputExpression->IsResultMaterialAttributes(0))
-					{
-						Compiler->Errorf(TEXT("Layer graphs only support a single material attributes input."));
-						bHasValidOutput = false;
-					}
-				}
-				else if (UMaterialExpressionFunctionOutput* OutputExpression = Cast<UMaterialExpressionFunctionOutput>(Expression))
-				{
-					++NumOutputs;
-					if (NumOutputs > 1 || !OutputExpression->IsResultMaterialAttributes(0))
-					{
-						Compiler->Errorf(TEXT("Layer graphs only support a single material attributes output."));
-						bHasValidOutput = false;
-					}
-				}
-				else if (UMaterialExpressionMaterialAttributeLayers* RecursiveLayer = Cast<UMaterialExpressionMaterialAttributeLayers>(Expression))
-				{
-					Compiler->Errorf(TEXT("Layer graphs do not support layers within layers."));
-						bHasValidOutput = false;
-				}
-			}
-
-			if ( NumInputs > 1 || NumOutputs < 1)
-			{
-				Compiler->Errorf(TEXT("Layer graphs require a single material attributes output and optionally, a single material attributes input."));
-				bHasValidOutput = false;
-			}
+			bHasValidOutput = UMaterialExpressionLayerStack::ValidateFunctionForLayerUsage(Compiler, this);
 		}
 	}
 	else if (GetMaterialFunctionUsage() == EMaterialFunctionUsage::MaterialLayerBlend)
@@ -14789,39 +14773,7 @@ bool UMaterialFunction::ValidateFunctionUsage(FMaterialCompiler* Compiler, const
 		}
 		else
 		{
-			// Material layer blends can have up to two MA inputs and single MA output only
-			for (UMaterialExpression* Expression : GetExpressions())
-			{
-				if (UMaterialExpressionFunctionInput* InputExpression = Cast<UMaterialExpressionFunctionInput>(Expression))
-				{
-					++NumInputs;
-					if (NumInputs > 2 || !InputExpression->IsResultMaterialAttributes(0))
-					{
-						Compiler->Errorf(TEXT("Layer blend graphs only support two material attributes inputs."));
-						bHasValidOutput = false;
-					}
-				}
-				else if (UMaterialExpressionFunctionOutput* OutputExpression = Cast<UMaterialExpressionFunctionOutput>(Expression))
-				{
-					++NumOutputs;
-					if (NumOutputs > 1 || !OutputExpression->IsResultMaterialAttributes(0))
-					{
-						Compiler->Errorf(TEXT("Layer blend graphs only support a single material attributes output."));
-						bHasValidOutput = false;
-					}
-				}
-				else if (UMaterialExpressionMaterialAttributeLayers* RecursiveLayer = Cast<UMaterialExpressionMaterialAttributeLayers>(Expression))
-				{
-					Compiler->Errorf(TEXT("Layer blend graphs do not support layers within layers."));
-						bHasValidOutput = false;
-				}
-			}
-
-			if (NumOutputs < 1)
-			{
-				Compiler->Errorf(TEXT("Layer blend graphs can have up to two material attributes inputs and a single output."));
-				bHasValidOutput = false;
-			}
+			bHasValidOutput = UMaterialExpressionLayerStack::ValidateFunctionForBlendUsage(Compiler, this);
 		}
 	}
 #endif
@@ -15105,6 +15057,63 @@ bool UMaterialFunction::IsUsingNewHLSLGenerator() const
 
 #endif // WITH_EDITOR
 
+///////////////////////////////////////////////////////////////////////////////
+// UMaterialFunctionMaterialLayerBlend
+///////////////////////////////////////////////////////////////////////////////
+void UMaterialFunctionMaterialLayerBlend::PostLoad()
+{
+	Super::PostLoad();
+#if WITH_EDITORONLY_DATA
+	const int32 FortMainVer = GetLinkerCustomVersion(FFortniteMainBranchObjectVersion::GUID);
+	if (FortMainVer < FFortniteMainBranchObjectVersion::MaterialFunctionBlendTopBottomInputEnum)
+	{
+		/**
+		* This code takes the legacy blend nodes which were based on an unreliable method of node generation
+		* and converts it to assign designated blend input enum types which explicitly defines whether the input node
+		* is relevant to the top or the bottom attributes input. No other inputs need updating as legacy blends have
+		* a strict rule of only 2 inputs exactly.
+		*/
+		bool TopAssigned = false;
+		bool BottomAssigned = false;
+
+		TArray<TObjectPtr<UMaterialExpression>>& BlendExpressions = GetExpressionCollection().Expressions;
+		for (UMaterialExpression* Expression : BlendExpressions)
+		{
+			UMaterialExpressionFunctionInput* InputNode = Cast<UMaterialExpressionFunctionInput>(Expression);
+			if (!InputNode)
+			{
+				continue;
+			}
+
+			if (InputNode->InputType == FunctionInput_MaterialAttributes)
+			{
+				if (!TopAssigned && InputNode->InputName == TopMaterialBlendInputName)
+				{
+					InputNode->BlendInputRelevance = EBlendInputRelevance::Top;
+					TopAssigned = true;
+				}
+				if (!BottomAssigned && InputNode->InputName == BottomMaterialBlendInputName)
+				{
+					InputNode->BlendInputRelevance = EBlendInputRelevance::Bottom;
+					BottomAssigned = true;
+				}
+
+				//Historical Blends likely only have 2 inputs anyway so early out.
+				if (TopAssigned && BottomAssigned)
+				{
+					break;
+				}
+			}
+		}
+
+		if (!TopAssigned || !BottomAssigned)
+		{
+			UE_LOG(LogMaterial, Warning, TEXT("Legacy Material Layer Blend Function \"%s\" has renamed Top and Bottom input nodes. Please manually set BlendInputRelevance on these nodes to their relative type."), *GetFullName());
+		}
+	}
+#endif
+}
+
 ///////////////////////////////////////////////////////////////////////////////
 // UMaterialFunctionInstance
 ///////////////////////////////////////////////////////////////////////////////
@@ -16768,6 +16777,9 @@ bool FMaterialLayersFunctions::ResolveParent(const FMaterialLayersFunctionsRunti
 	TArray<int32> ParentLayerIndices;
 
 	ResolvedLayers.Empty();
+	
+	//Ensure we retain the layer stack cache before we do anything else (only applies for LayerStack node usage, otherwise null and inconsequential).
+	ResolvedLayers.LayerStackCache = ParentRuntime.LayerStackCache;
 
 	bool bHasUninitializedLinks = false;
 	for (int32 LayerIndex = 0; LayerIndex < Runtime.Layers.Num(); ++LayerIndex)
@@ -18138,6 +18150,86 @@ void UMaterialExpressionFunctionInput::ConditionallyGenerateId(bool bForce)
 	}
 }
 
+FString UMaterialExpressionFunctionInput::GetInputTypeDisplayName(EFunctionInputType InputType)
+{
+	static const UEnum* FunctionInputEnum = StaticEnum<EFunctionInputType>();
+	return FunctionInputEnum->GetDisplayNameTextByValue((int64)InputType).ToString();
+}
+
+EMaterialValueType UMaterialExpressionFunctionInput::GetMaterialTypeFromInputType(EFunctionInputType InputType)
+{
+	static_assert(FunctionInput_MAX == EXPECTED_NUM_INPUT_TYPES, "Number of expected Function Input types is incorrect. Need to update");
+	switch (InputType)
+	{
+	case FunctionInput_Scalar:
+		return MCT_Float;
+	case FunctionInput_Vector2:
+		return MCT_Float2;
+	case FunctionInput_Vector3:
+		return MCT_Float3;
+	case FunctionInput_Vector4:
+		return MCT_Float4;
+	case FunctionInput_Texture2D:
+		return MCT_Texture2D;
+	case FunctionInput_TextureCube:
+		return MCT_TextureCube;
+	case FunctionInput_Texture2DArray:
+		return MCT_Texture2DArray;
+	case FunctionInput_TextureExternal:
+		return MCT_TextureExternal;
+	case FunctionInput_VolumeTexture:
+		return MCT_VolumeTexture;
+	case FunctionInput_StaticBool:
+		return MCT_StaticBool;
+	case FunctionInput_Bool:
+		return MCT_Bool;
+	case FunctionInput_MaterialAttributes:
+		return MCT_MaterialAttributes;
+	case FunctionInput_Substrate:
+		return MCT_Substrate;
+	default:
+		return MCT_Unknown;
+	}
+}
+
+EFunctionInputType UMaterialExpressionFunctionInput::GetInputTypeFromMaterialType(EMaterialValueType MaterialType)
+{
+	static_assert(FunctionInput_MAX == EXPECTED_NUM_INPUT_TYPES, "Number of expected Function Input types is incorrect. Need to update");
+	switch (MaterialType)
+	{
+	case MCT_Float:
+	case MCT_Float1:
+		return FunctionInput_Scalar;
+	case MCT_Float2:
+		return FunctionInput_Vector2;
+	case MCT_Float3:
+		return FunctionInput_Vector3;
+	case MCT_Float4:
+		return FunctionInput_Vector4;
+	case MCT_Texture:
+	case MCT_Texture2D:
+		return FunctionInput_Texture2D;
+	case MCT_TextureCube:
+		return FunctionInput_TextureCube;
+	case MCT_Texture2DArray:
+		return FunctionInput_Texture2DArray;
+	case MCT_TextureExternal:
+		return FunctionInput_TextureExternal;
+	case MCT_VolumeTexture:
+		return FunctionInput_VolumeTexture;
+	case MCT_StaticBool:
+		return FunctionInput_StaticBool;
+	case MCT_Bool:
+		return FunctionInput_Bool;
+	case MCT_MaterialAttributes:
+		return FunctionInput_MaterialAttributes;
+	case MCT_Substrate:
+		return FunctionInput_Substrate;
+	default:
+		return FunctionInput_MAX;
+	}
+}
+
 #if WITH_EDITOR
 void UMaterialExpressionFunctionInput::ValidateName()
 {
@@ -18287,37 +18379,7 @@ FSubstrateOperator* UMaterialExpressionFunctionInput::SubstrateGenerateMaterialT
 
 EMaterialValueType UMaterialExpressionFunctionInput::GetInputValueType(int32 InputIndex)
 {
-	switch (InputType)
-	{
-	case FunctionInput_Scalar:
-		return MCT_Float;
-	case FunctionInput_Vector2:
-		return MCT_Float2;
-	case FunctionInput_Vector3:
-		return MCT_Float3;
-	case FunctionInput_Vector4:
-		return MCT_Float4;
-	case FunctionInput_Texture2D:
-		return MCT_Texture2D;
-	case FunctionInput_TextureCube:
-		return MCT_TextureCube;
-	case FunctionInput_Texture2DArray:
-		return MCT_Texture2DArray;
-	case FunctionInput_TextureExternal:
-		return MCT_TextureExternal;
-	case FunctionInput_VolumeTexture:
-		return MCT_VolumeTexture;
-	case FunctionInput_StaticBool:
-		return MCT_StaticBool;
-	case FunctionInput_Bool:
-		return MCT_Bool;
-	case FunctionInput_MaterialAttributes:
-		return MCT_MaterialAttributes;
-	case FunctionInput_Substrate:
-		return MCT_Substrate;
-	default:
-		return MCT_Unknown;
-	}
+	return GetMaterialTypeFromInputType(InputType);
 }
 
 EMaterialValueType UMaterialExpressionFunctionInput::GetOutputValueType(int32 OutputIndex)
diff --git a/Engine/Source/Runtime/Engine/Public/Materials/MaterialExpression.h b/Engine/Source/Runtime/Engine/Public/Materials/MaterialExpression.h
--- a/Engine/Source/Runtime/Engine/Public/Materials/MaterialExpression.h
+++ b/Engine/Source/Runtime/Engine/Public/Materials/MaterialExpression.h
@@ -363,6 +363,7 @@ class UMaterialExpression : public UObject
 		PRAGMA_ENABLE_DEPRECATION_WARNINGS
 	}
 
+	virtual bool IsStaticExpression() const { return false; }
 	virtual bool CanUserDeleteExpression() const { return true; }
 
 	/** Find the property that is associated with the input pin. */
diff --git a/Engine/Source/Runtime/Engine/Public/Materials/MaterialExpressionFunctionInput.h b/Engine/Source/Runtime/Engine/Public/Materials/MaterialExpressionFunctionInput.h
--- a/Engine/Source/Runtime/Engine/Public/Materials/MaterialExpressionFunctionInput.h
+++ b/Engine/Source/Runtime/Engine/Public/Materials/MaterialExpressionFunctionInput.h
@@ -17,20 +17,30 @@ struct FPropertyChangedEvent;
 UENUM(BlueprintType)
 enum EFunctionInputType : int
 {
-	FunctionInput_Scalar,
-	FunctionInput_Vector2,
-	FunctionInput_Vector3,
-	FunctionInput_Vector4,
-	FunctionInput_Texture2D,
-	FunctionInput_TextureCube,
-	FunctionInput_Texture2DArray,
-	FunctionInput_VolumeTexture,
-	FunctionInput_StaticBool,
-	FunctionInput_MaterialAttributes,
-	FunctionInput_TextureExternal,
-	FunctionInput_Bool,
-	FunctionInput_Substrate,
-	FunctionInput_MAX,
+	FunctionInput_Scalar UMETA(DisplayName = "Scalar"),
+	FunctionInput_Vector2 UMETA(DisplayName = "Vector2"),
+	FunctionInput_Vector3 UMETA(DisplayName = "Vector3"),
+	FunctionInput_Vector4 UMETA(DisplayName = "Vector4"),
+	FunctionInput_Texture2D UMETA(DisplayName = "Texture2D"),
+	FunctionInput_TextureCube UMETA(DisplayName = "TextureCube"),
+	FunctionInput_Texture2DArray UMETA(DisplayName = "Texture2DArray"),
+	FunctionInput_VolumeTexture UMETA(DisplayName = "VolumeTexture"),
+	FunctionInput_StaticBool UMETA(DisplayName = "StaticBool"),
+	FunctionInput_MaterialAttributes UMETA(DisplayName = "MaterialAttributes"),
+	FunctionInput_TextureExternal UMETA(DisplayName = "TextureExternal"),
+	FunctionInput_Bool UMETA(DisplayName = "Bool"),
+	FunctionInput_Substrate UMETA(DisplayName = "Substrate"),
+	FunctionInput_MAX UMETA(DisplayName = "NULL"),
+};
+#define EXPECTED_NUM_INPUT_TYPES 13
+
+/** Supported input types */
+UENUM()
+enum EBlendInputRelevance : int
+{
+	General,
+	Top,
+	Bottom
 };
 
 UCLASS(hidecategories=object, MinimalAPI)
@@ -73,6 +83,10 @@ class UMaterialExpressionFunctionInput : public UMaterialExpression
 	UPROPERTY(EditAnywhere, BlueprintReadWrite, Category=MaterialExpression)
 	int32 SortPriority;
 
+	/** Marks that this input is necessary for blend function use **/
+	UPROPERTY(EditAnywhere, Category = MaterialExpressionFunctionInput)
+	TEnumAsByte<enum EBlendInputRelevance> BlendInputRelevance = EBlendInputRelevance::General;
+
 	UE_DEPRECATED(5.6, "bCompilingFunctionPreview has been deprecated.")
 	UPROPERTY(transient, meta = (DeprecatedProperty, DeprecationMessage = "bCompilingFunctionPreview has been removed from function inputs and has no functional usage, it has been replaced by automatic state tracking between function calls."))
 	uint32 bCompilingFunctionPreview_DEPRECATED : 1;
@@ -118,12 +132,14 @@ class UMaterialExpressionFunctionInput : public UMaterialExpression
 	/** Generate the Id for this input. */
 	ENGINE_API void ConditionallyGenerateId(bool bForce);
 
+	ENGINE_API static FString GetInputTypeDisplayName(EFunctionInputType InputType);
+	ENGINE_API static EMaterialValueType GetMaterialTypeFromInputType(EFunctionInputType InputType);
+	ENGINE_API static EFunctionInputType GetInputTypeFromMaterialType(EMaterialValueType MaterialType);
 #if WITH_EDITOR
 	/** Validate InputName.  Must be called after InputName is changed to prevent duplicate inputs. */
 	ENGINE_API void ValidateName();
 #endif // WITH_EDITOR
 
-
 private:
 #if WITH_EDITOR
 	/** Helper function which compiles this expression for previewing. */
diff --git a/Engine/Source/Runtime/Engine/Public/Materials/MaterialExpressionLayerStack.h b/Engine/Source/Runtime/Engine/Public/Materials/MaterialExpressionLayerStack.h
new file mode 100644
--- /dev/null
+++ b/Engine/Source/Runtime/Engine/Public/Materials/MaterialExpressionLayerStack.h
@@ -0,0 +1,90 @@
+// Copyright Epic Games, Inc. All Rights Reserved.
+
+#pragma once
+
+#include "CoreMinimal.h"
+#include "UObject/ObjectMacros.h"
+#include "UObject/SoftObjectPath.h"
+#include "MaterialExpressionIO.h"
+#include "Materials/MaterialExpression.h"
+#include "Materials/MaterialExpressionFunctionInput.h"
+#include "Materials/MaterialExpressionMaterialAttributeLayers.h"
+#include "Materials/MaterialLayersFunctions.h"
+#include "MaterialExpressionLayerStack.generated.h"
+
+class FMaterialCompiler;
+class UMaterialFunctionInterface;
+
+USTRUCT()
+struct FMaterialLayerInput : public FExpressionInput
+{
+	GENERATED_BODY()
+
+	UPROPERTY()
+	TEnumAsByte<EFunctionInputType> InputType;
+
+	FMaterialLayerInput(FName NewInputName = NAME_None, EFunctionInputType NewInputType = FunctionInput_MAX);
+	FString GetInputName() const;
+};
+
+UCLASS(hidecategories = Object, MinimalAPI)
+class UMaterialExpressionLayerStack : public UMaterialExpressionMaterialAttributeLayers
+{
+	GENERATED_UCLASS_BODY()
+
+	UPROPERTY()
+	TArray<FMaterialLayerInput> LayerInputs;
+
+	UPROPERTY(EditAnywhere, Category = AvailableFunctions, meta = (DisallowedClasses = "/Script/Engine.MaterialFunctionMaterialLayerBlend, /Script/Engine.MaterialFunctionMaterialLayerBlendInstance"))
+	TSet<TObjectPtr<UMaterialFunctionInterface>> AvailableLayers;
+
+	UPROPERTY(EditAnywhere, Category = AvailableFunctions, meta = (AllowedClasses = "/Script/Engine.MaterialFunctionMaterialLayerBlend, /Script/Engine.MaterialFunctionMaterialLayerBlendInstance"))
+	TSet<TObjectPtr<UMaterialFunctionInterface>> AvailableBlends;
+
+	//~ Begin UObject Interface
+	virtual void PostLoad() override;
+	virtual void PreSave(FObjectPreSaveContext ObjectSaveContext) override;
+	//~ End UObject Interface
+
+#if WITH_EDITOR
+	//~ Begin UMaterialExpression Interface
+	virtual void PostEditChangeProperty(FPropertyChangedEvent& PropertyChangedEvent) override;
+	virtual void GetCaption(TArray<FString>& OutCaptions) const override;
+	virtual void GetExpressionToolTip(TArray<FString>& OutToolTip) override;
+	virtual FName GetInputName(int32 InputIndex) const override;
+	virtual FExpressionInput* GetInput(int32 InputIndex) override;
+	virtual TArrayView<FExpressionInput*> GetInputsView() override;
+	virtual EMaterialValueType GetInputValueType(int32 InputIndex) override;
+	//~ End UMaterialExpression Interface
+
+	//~ Begin UMaterialExpressionMaterialAttributeLayers Interface
+	virtual bool ValidateLayerConfiguration(FMaterialCompiler* Compiler, bool bReportErrors) override;
+	ENGINE_API virtual void RebuildLayerGraph(bool bReportErrors);
+	//~ End UMaterialExpressionMaterialAttributeLayers Interface
+
+	//This struct is to circumvent having to extend the below lambda in future if we want to validate against additional node types.
+	struct FValidLayerUsageTracker
+	{
+		int32 MAInputCount = 0;
+		int32 MAOutputCount = 0;
+		bool bContainsStatics = false;
+	};
+
+	static UMaterialFunctionInterface* ExtractParentFunctionFromInstance(class FMaterialCompiler* Compiler, UMaterialFunctionInterface* CurrentFunction);
+	static bool PollFunctionExpressionsForLayerUsage(class FMaterialCompiler* Compiler, UMaterialFunctionInterface* CurrentFunction, FValidLayerUsageTracker& Tracker, bool bCheckStatics);
+	static bool ValidateFunctionForLayerUsage(class FMaterialCompiler* Compiler, UMaterialFunctionInterface* CurrentFunction, bool bCheckStatics = false);
+	static bool ValidateFunctionForBlendUsage(class FMaterialCompiler* Compiler, UMaterialFunctionInterface* CurrentFunction, bool bCheckStatics = false);
+
+	TSharedPtr<FMaterialLayerStackFunctionsCache> GetSharedAvailableFunctionsCache();
+	void ResolveLayerInputs();
+	void CacheLayerInputs();
+#endif
+	inline TArray<FSoftObjectPath> GetAvailableLayerPaths() { return GetPathsFromAvailableFunctions(AvailableLayers); }
+	inline TArray<FSoftObjectPath> GetAvailableBlendPaths() { return GetPathsFromAvailableFunctions(AvailableBlends); }
+
+private:
+	TSharedPtr<FMaterialLayerStackFunctionsCache> SharedCache;
+	bool bAreAvailableLayersValid = false;
+
+	static TArray<FSoftObjectPath> GetPathsFromAvailableFunctions(TSet<TObjectPtr<UMaterialFunctionInterface>>& InFunctions);
+};
diff --git a/Engine/Source/Runtime/Engine/Public/Materials/MaterialExpressionMaterialAttributeLayers.h b/Engine/Source/Runtime/Engine/Public/Materials/MaterialExpressionMaterialAttributeLayers.h
--- a/Engine/Source/Runtime/Engine/Public/Materials/MaterialExpressionMaterialAttributeLayers.h
+++ b/Engine/Source/Runtime/Engine/Public/Materials/MaterialExpressionMaterialAttributeLayers.h
@@ -14,6 +14,14 @@ class UMaterialFunctionInterface;
 class UMaterialExpressionMaterialFunctionCall;
 struct FMaterialParameterInfo;
 
+/**
+* These defines are for tracking the required number of inputs / outputs for blends rather than the values being
+* hard coded throughout the layer functions and expressions code
+*/
+constexpr uint8 AcceptableNumLayerMAInputs = 1;
+constexpr uint8 AcceptableNumBlendMAInputs = 2;
+constexpr uint8 AcceptableNumLayerBlendMAOutputs = 1;
+
 UCLASS(hidecategories=Object, MinimalAPI)
 class UMaterialExpressionMaterialAttributeLayers : public UMaterialExpression
 {
@@ -88,12 +96,13 @@ class UMaterialExpressionMaterialAttributeLayers : public UMaterialExpression
 	//~ Begin UObject Interface
 
 #if WITH_EDITOR
-	ENGINE_API void RebuildLayerGraph(bool bReportErrors);
+	ENGINE_API virtual void RebuildLayerGraph(bool bReportErrors);
 	ENGINE_API void OverrideLayerGraph(const FMaterialLayersFunctions* OverrideLayers);
 #endif // WITH_EDITOR
 
 #if WITH_EDITOR
-	bool ValidateLayerConfiguration(FMaterialCompiler* Compiler, bool bReportErrors);
+	static void LogError(FMaterialCompiler* Compiler, const TCHAR* Format, ...);
+	virtual bool ValidateLayerConfiguration(FMaterialCompiler* Compiler, bool bReportErrors);
 #endif // WITH_EDITOR
 
 #if WITH_EDITORONLY_DATA
@@ -117,7 +126,7 @@ class UMaterialExpressionMaterialAttributeLayers : public UMaterialExpression
 #endif
 	//~ End UMaterialExpression Interface
 
-private:
+protected:
 	/** Internal pointer to parameter-driven layer graph */
 	const FMaterialLayersFunctions* ParamLayers = nullptr;
 };
diff --git a/Engine/Source/Runtime/Engine/Public/Materials/MaterialExpressionStaticBool.h b/Engine/Source/Runtime/Engine/Public/Materials/MaterialExpressionStaticBool.h
--- a/Engine/Source/Runtime/Engine/Public/Materials/MaterialExpressionStaticBool.h
+++ b/Engine/Source/Runtime/Engine/Public/Materials/MaterialExpressionStaticBool.h
@@ -25,7 +25,7 @@ class UMaterialExpressionStaticBool : public UMaterialExpression
 	virtual int32 Compile(class FMaterialCompiler* Compiler, int32 OutputIndex) override;
 	virtual void GetCaption(TArray<FString>& OutCaptions) const override;
 	virtual EMaterialValueType GetOutputValueType(int32 OutputIndex) override {return MCT_StaticBool;}
-
+	virtual bool IsStaticExpression() const override { return true; }
 #endif
 	//~ End UMaterialExpression Interface
 };
diff --git a/Engine/Source/Runtime/Engine/Public/Materials/MaterialExpressionStaticBoolParameter.h b/Engine/Source/Runtime/Engine/Public/Materials/MaterialExpressionStaticBoolParameter.h
--- a/Engine/Source/Runtime/Engine/Public/Materials/MaterialExpressionStaticBoolParameter.h
+++ b/Engine/Source/Runtime/Engine/Public/Materials/MaterialExpressionStaticBoolParameter.h
@@ -29,6 +29,7 @@ public:
 	virtual int32 CompilePreview(class FMaterialCompiler* Compiler, int32 OutputIndex) override;
 	virtual void GetCaption(TArray<FString>& OutCaptions) const override;
 	virtual EMaterialValueType GetOutputValueType(int32 OutputIndex) override {return DynamicBranch ? MCT_Bool : MCT_StaticBool;}
+	virtual bool IsStaticExpression() const override { return DynamicBranch == 0; }
 	virtual bool GetParameterValue(FMaterialParameterMetadata& OutMeta) const override
 	{
 		OutMeta.Value = (bool)DefaultValue;
diff --git a/Engine/Source/Runtime/Engine/Public/Materials/MaterialExpressionStaticComponentMaskParameter.h b/Engine/Source/Runtime/Engine/Public/Materials/MaterialExpressionStaticComponentMaskParameter.h
--- a/Engine/Source/Runtime/Engine/Public/Materials/MaterialExpressionStaticComponentMaskParameter.h
+++ b/Engine/Source/Runtime/Engine/Public/Materials/MaterialExpressionStaticComponentMaskParameter.h
@@ -39,6 +39,7 @@ public:
 	virtual void Build(MIR::FEmitter& Emitter) override;
 	virtual int32 Compile(class FMaterialCompiler* Compiler, int32 OutputIndex) override;
 	virtual void GetCaption(TArray<FString>& OutCaptions) const override;
+	virtual bool IsStaticExpression() const override { return true; }
 	virtual bool GetParameterValue(FMaterialParameterMetadata& OutMeta) const override
 	{
 		OutMeta.Value = FMaterialParameterValue(DefaultR, DefaultG, DefaultB, DefaultA);
diff --git a/Engine/Source/Runtime/Engine/Public/Materials/MaterialExpressionStaticSwitch.h b/Engine/Source/Runtime/Engine/Public/Materials/MaterialExpressionStaticSwitch.h
--- a/Engine/Source/Runtime/Engine/Public/Materials/MaterialExpressionStaticSwitch.h
+++ b/Engine/Source/Runtime/Engine/Public/Materials/MaterialExpressionStaticSwitch.h
@@ -39,7 +39,7 @@ public:
 	virtual bool IsResultMaterialAttributes(int32 OutputIndex) override;
 	virtual EMaterialValueType GetInputValueType(int32 InputIndex) override;
 	virtual EMaterialValueType GetOutputValueType(int32 OutputIndex) override {return MCT_Unknown;}
-
+	virtual bool IsStaticExpression() const override { return true; }
 
 	virtual bool IsResultSubstrateMaterial(int32 OutputIndex) override;
 	virtual void GatherSubstrateMaterialInfo(FSubstrateMaterialInfo& SubstrateMaterialInfo, int32 OutputIndex) override;
diff --git a/Engine/Source/Runtime/Engine/Public/Materials/MaterialExpressionStaticSwitchParameter.h b/Engine/Source/Runtime/Engine/Public/Materials/MaterialExpressionStaticSwitchParameter.h
--- a/Engine/Source/Runtime/Engine/Public/Materials/MaterialExpressionStaticSwitchParameter.h
+++ b/Engine/Source/Runtime/Engine/Public/Materials/MaterialExpressionStaticSwitchParameter.h
@@ -32,6 +32,7 @@ public:
 	virtual bool IsResultMaterialAttributes(int32 OutputIndex) override;
 	virtual EMaterialValueType GetInputValueType(int32 InputIndex) override {return MCT_Unknown;}
 	virtual EMaterialValueType GetOutputValueType(int32 OutputIndex) override {return MCT_Unknown;}
+	virtual bool IsStaticExpression() const override { return true; }
 
 	virtual bool IsResultSubstrateMaterial(int32 OutputIndex) override;
 	virtual void GatherSubstrateMaterialInfo(FSubstrateMaterialInfo& SubstrateMaterialInfo, int32 OutputIndex) override;
diff --git a/Engine/Source/Runtime/Engine/Public/Materials/MaterialFunctionMaterialLayerBlend.h b/Engine/Source/Runtime/Engine/Public/Materials/MaterialFunctionMaterialLayerBlend.h
--- a/Engine/Source/Runtime/Engine/Public/Materials/MaterialFunctionMaterialLayerBlend.h
+++ b/Engine/Source/Runtime/Engine/Public/Materials/MaterialFunctionMaterialLayerBlend.h
@@ -9,6 +9,10 @@
 #include "Materials/MaterialFunctionInstance.h"
 #include "MaterialFunctionMaterialLayerBlend.generated.h"
 
+//Legacy nodes rely on these hard coded names, this is mainly for fixup code.
+constexpr const TCHAR* TopMaterialBlendInputName = TEXT("Top Layer");
+constexpr const TCHAR* BottomMaterialBlendInputName = TEXT("Bottom Layer");
+
 /**
  * Specialized Material Function that acts as a blend
  */
@@ -16,6 +20,7 @@ UCLASS(hidecategories=object, MinimalAPI)
 class UMaterialFunctionMaterialLayerBlend : public UMaterialFunction
 {
 	GENERATED_BODY()
+	virtual void PostLoad() override;
 };
 
 /**
diff --git a/Engine/Source/Runtime/Engine/Public/Materials/MaterialLayersFunctions.h b/Engine/Source/Runtime/Engine/Public/Materials/MaterialLayersFunctions.h
--- a/Engine/Source/Runtime/Engine/Public/Materials/MaterialLayersFunctions.h
+++ b/Engine/Source/Runtime/Engine/Public/Materials/MaterialLayersFunctions.h
@@ -399,6 +399,12 @@ struct FMaterialLayersFunctionsRuntimeGraphCache
 #endif
 };
 
+struct FMaterialLayerStackFunctionsCache
+{
+	//Helper struct for managing available layers and blends across parent/instance materials utilising the LayerStack node.
+	TArray<FSoftObjectPath> AvailableLayerPaths;
+	TArray<FSoftObjectPath> AvailableBlendPaths;
+};
 
 USTRUCT()
 struct FMaterialLayersFunctionsRuntimeData
@@ -417,12 +423,17 @@ struct FMaterialLayersFunctionsRuntimeData
 
 	mutable TSharedPtr<FMaterialLayersFunctionsRuntimeGraphCache> RuntimeGraphCache;
 
+	// In multiple locations there are const references to this runtime data/it's children, 
+	// so this mutable entry is required to ensure we can keep available layers/blend up to date across all instances.
+	mutable TSharedPtr<FMaterialLayerStackFunctionsCache> LayerStackCache;
+
 	FMaterialLayersFunctionsRuntimeData() = default;
 	FMaterialLayersFunctionsRuntimeData(const FMaterialLayersFunctionsRuntimeData& Rhs)
 		: Layers(Rhs.Layers)
 		, Blends(Rhs.Blends)
 		, Tree(Rhs.Tree)
 		, RuntimeGraphCache(Rhs.RuntimeGraphCache)
+		, LayerStackCache(Rhs.LayerStackCache)
 	{}
 
 	FMaterialLayersFunctionsRuntimeData(const FMaterialLayersFunctions& Rhs) = delete;
@@ -433,6 +444,7 @@ struct FMaterialLayersFunctionsRuntimeData
 		Blends = Rhs.Blends;
 		Tree = Rhs.Tree;
 		RuntimeGraphCache = Rhs.RuntimeGraphCache;
+		LayerStackCache = Rhs.LayerStackCache;
 		return *this;
 	}
 
@@ -446,6 +458,7 @@ struct FMaterialLayersFunctionsRuntimeData
 		Blends.Empty();
 		Tree.Empty();
 		RuntimeGraphCache.Reset();
+		LayerStackCache.Reset();
 	}
 
 	FORCEINLINE bool operator==(const FMaterialLayersFunctionsRuntimeData& Other) const
@@ -456,6 +469,16 @@ struct FMaterialLayersFunctionsRuntimeData
 		{
 			return false;
 		}
+
+		if (LayerStackCache || Other.LayerStackCache)
+		{
+			if (!LayerStackCache || !Other.LayerStackCache 
+				|| LayerStackCache->AvailableLayerPaths != Other.LayerStackCache->AvailableLayerPaths
+				|| LayerStackCache->AvailableBlendPaths != Other.LayerStackCache->AvailableBlendPaths)
+			{
+				return false;
+			}
+		}
 		return true;
 	}
 
