[SceneGraph: AI Navigation] added missing filtering for colliders without queries or physics enabled

diff --git a/Engine/Source/Runtime/Engine/Classes/AI/Navigation/NavigationElement.h b/Engine/Source/Runtime/Engine/Classes/AI/Navigation/NavigationElement.h
--- a/Engine/Source/Runtime/Engine/Classes/AI/Navigation/NavigationElement.h
+++ b/Engine/Source/Runtime/Engine/Classes/AI/Navigation/NavigationElement.h
@@ -51,6 +51,12 @@ struct FNavigationElementHandle
 	/** @return true if the handle has been properly assigned. */
 	inline bool IsValid() const;
 
+	/** Invalidates the current handle so further calls to IsValid will return false. */
+	void Invalidate()
+	{
+		*this = Invalid;
+	}
+
 	/** Hash function to use TSet/TMap */
 	inline friend uint32 GetTypeHash(const FNavigationElementHandle& Handle);
 
