New Material Translator:  support for InverseLinearInterpolate, Normalize, Sign, Step.

#rnx

diff --git a/Engine/Source/Runtime/Engine/Private/Materials/MaterialExpressionsIR.cpp b/Engine/Source/Runtime/Engine/Private/Materials/MaterialExpressionsIR.cpp
--- a/Engine/Source/Runtime/Engine/Private/Materials/MaterialExpressionsIR.cpp
+++ b/Engine/Source/Runtime/Engine/Private/Materials/MaterialExpressionsIR.cpp
@@ -550,6 +550,19 @@ void UMaterialExpressionLength::Build(MIR::FEmitter& Em)
 	Em.Output(0, Em.Length(Em.Input(&Input)));
 }
 
+void UMaterialExpressionNormalize::Build(MIR::FEmitter& Em)
+{
+	FValueRef InputValue = Em.CastToFloatKind(Em.Input(&VectorInput));
+	if (InputValue->IsScalar())
+	{
+		Em.Output(0, Em.ConstantOne(MIR::ScalarKind_Float));
+	}
+	else
+	{
+		Em.Output(0, Em.Multiply(InputValue, Em.Rsqrt(Em.Dot(InputValue, InputValue))));
+	}
+}
+
 void UMaterialExpressionRound::Build(MIR::FEmitter& Em)
 {
 	Em.Output(0, Em.Round(Em.Input(&Input)));
@@ -678,6 +691,11 @@ void UMaterialExpressionSaturate::Build(MIR::FEmitter& Em)
 	Em.Output(0, Em.Saturate(Em.Input(&Input)));
 }
 
+void UMaterialExpressionSign::Build(MIR::FEmitter& Em)
+{
+	Em.Output(0, Em.Sign(Em.Input(&Input)));
+}
+
 void UMaterialExpressionSquareRoot::Build(MIR::FEmitter& Em)
 {
 	Em.Output(0, Em.Sqrt(Em.Input(&Input)));
@@ -773,6 +791,11 @@ void UMaterialExpressionMin::Build(MIR::FEmitter& Em)
 	BuildBinaryOperatorWithDefaults(Em, MIR::BO_Min, &A, ConstA, &B, ConstB);
 }
 
+void UMaterialExpressionStep::Build(MIR::FEmitter& Em)
+{
+	BuildBinaryOperatorWithDefaults(Em, MIR::BO_Step, &Y, ConstY, &X, ConstX);
+}
+
 void UMaterialExpressionDotProduct::Build(MIR::FEmitter& Em)
 {
 	Em.Output(0, Em.Dot(Em.Input(&A), Em.Input(&B)));
@@ -1269,6 +1292,19 @@ void UMaterialExpressionColorRamp::Build(MIR::FEmitter& Em)
 	Em.Output(0, FinalVector);
 }
 
+void UMaterialExpressionInverseLinearInterpolate::Build(MIR::FEmitter& Em)
+{
+	FValueRef ValueA = Em.InputDefaultFloat(&A, ConstA);
+	FValueRef ValueB = Em.InputDefaultFloat(&B, ConstB);
+	FValueRef ValueC = Em.InputDefaultFloat(&Value, ConstValue);
+	FValueRef Result = Em.Divide(Em.CastToFloatKind(Em.Subtract(ValueC, ValueA)), Em.CastToFloatKind(Em.Subtract(ValueB, ValueA)));
+	if (bClampResult)
+	{
+		Result = Em.Saturate(Result);
+	}
+	Em.Output(0, Result);
+}
+
 void UMaterialExpressionLinearInterpolate::Build(MIR::FEmitter& Em)
 {
 	BuildTernaryArithmeticOperator(Em, MIR::TO_Lerp, &A, ConstA, &B, ConstB, &Alpha, ConstAlpha);
diff --git a/Engine/Source/Runtime/Engine/Public/Materials/MaterialExpressionInverseLinearInterpolate.h b/Engine/Source/Runtime/Engine/Public/Materials/MaterialExpressionInverseLinearInterpolate.h
--- a/Engine/Source/Runtime/Engine/Public/Materials/MaterialExpressionInverseLinearInterpolate.h
+++ b/Engine/Source/Runtime/Engine/Public/Materials/MaterialExpressionInverseLinearInterpolate.h
@@ -43,6 +43,7 @@ public:
 
 	//~ Begin UMaterialExpression Interface
 #if WITH_EDITOR
+	virtual void Build(MIR::FEmitter& Emitter) override;
 	virtual int32 Compile(class FMaterialCompiler* Compiler, int32 OutputIndex) override;
 	virtual void GetCaption(TArray<FString>& OutCaptions) const override;
 	virtual FText GetKeywords() const override { return FText::FromString(TEXT("InvLerp")); }
diff --git a/Engine/Source/Runtime/Engine/Public/Materials/MaterialExpressionNormalize.h b/Engine/Source/Runtime/Engine/Public/Materials/MaterialExpressionNormalize.h
--- a/Engine/Source/Runtime/Engine/Public/Materials/MaterialExpressionNormalize.h
+++ b/Engine/Source/Runtime/Engine/Public/Materials/MaterialExpressionNormalize.h
@@ -22,6 +22,7 @@ public:
 
 	//~ Begin UMaterialExpression Interface
 #if WITH_EDITOR
+	virtual void Build(MIR::FEmitter& Emitter) override;
 	virtual int32 Compile(class FMaterialCompiler* Compiler, int32 OutputIndex) override;
 	virtual void GetCaption(TArray<FString>& OutCaptions) const override { OutCaptions.Add(TEXT("Normalize")); }
 	virtual FText GetKeywords() const override { return FText::FromString(TEXT("normal")); }
diff --git a/Engine/Source/Runtime/Engine/Public/Materials/MaterialExpressionSign.h b/Engine/Source/Runtime/Engine/Public/Materials/MaterialExpressionSign.h
--- a/Engine/Source/Runtime/Engine/Public/Materials/MaterialExpressionSign.h
+++ b/Engine/Source/Runtime/Engine/Public/Materials/MaterialExpressionSign.h
@@ -20,6 +20,7 @@ public:
 
 	//~ Begin UMaterialExpression Interface
 #if WITH_EDITOR
+	virtual void Build(MIR::FEmitter& Emitter) override;
 	virtual int32 Compile(class FMaterialCompiler* Compiler, int32 OutputIndex) override;
 	virtual void GetCaption(TArray<FString>& OutCaptions) const override;
 	virtual void GetExpressionToolTip(TArray<FString>& OutToolTip) override;
diff --git a/Engine/Source/Runtime/Engine/Public/Materials/MaterialExpressionStep.h b/Engine/Source/Runtime/Engine/Public/Materials/MaterialExpressionStep.h
--- a/Engine/Source/Runtime/Engine/Public/Materials/MaterialExpressionStep.h
+++ b/Engine/Source/Runtime/Engine/Public/Materials/MaterialExpressionStep.h
@@ -34,6 +34,7 @@ public:
 
 	//~ Begin UMaterialExpression Interface
 #if WITH_EDITOR
+	virtual void Build(MIR::FEmitter& Emitter) override;
 	virtual int32 Compile(class FMaterialCompiler* Compiler, int32 OutputIndex) override;
 	virtual void GetCaption(TArray<FString>& OutCaptions) const override;
 	virtual FText GetKeywords() const override { return FText::FromString(TEXT("Step")); }
diff --git a/Engine/Source/Runtime/Engine/Public/Materials/MaterialIREmitter.h b/Engine/Source/Runtime/Engine/Public/Materials/MaterialIREmitter.h
--- a/Engine/Source/Runtime/Engine/Public/Materials/MaterialIREmitter.h
+++ b/Engine/Source/Runtime/Engine/Public/Materials/MaterialIREmitter.h
@@ -535,6 +535,7 @@ public:
 	FValueRef Round(FValueRef A) { return Operator(UO_Round, A); }
 	FValueRef Rsqrt(FValueRef A) { return Operator(UO_Rsqrt, A); }
 	FValueRef Saturate(FValueRef A) { return Operator(UO_Saturate, A); }
+	FValueRef Sign(FValueRef A) { return Operator(UO_Sign, A); }
 	FValueRef Sin(FValueRef A) { return Operator(UO_Sin, A); }
 	FValueRef Sinh(FValueRef A) { return Operator(UO_Sinh, A); }
 	FValueRef Sqrt(FValueRef A) { return Operator(UO_Sqrt, A); }
@@ -562,6 +563,7 @@ public:
 	FValueRef Fmod(FValueRef A, FValueRef B) { return Operator(BO_Fmod, A, B); }
 	FValueRef Max(FValueRef A, FValueRef B) { return Operator(BO_Max, A, B); }
 	FValueRef Min(FValueRef A, FValueRef B) { return Operator(BO_Min, A, B); }
+	FValueRef Step(FValueRef A, FValueRef B) { return Operator(BO_Step, A, B); }
 	FValueRef Dot(FValueRef A, FValueRef B) { return Operator(BO_Dot, A, B); }
 	FValueRef Cross(FValueRef A, FValueRef B) { return Operator(BO_Cross, A, B); }
 	FValueRef Atan2(FValueRef A, FValueRef B) { return Operator(BO_ATan2, A, B); }
