Fixing a few TextureCollection issues
- We could have null textures if they aren't loaded
- Collections were not being referenced correctly inside of FMaterialCachedExpressionData

diff --git a/Engine/Source/Runtime/Engine/Private/Materials/Material.cpp b/Engine/Source/Runtime/Engine/Private/Materials/Material.cpp
--- a/Engine/Source/Runtime/Engine/Private/Materials/Material.cpp
+++ b/Engine/Source/Runtime/Engine/Private/Materials/Material.cpp
@@ -1251,7 +1251,10 @@ void UMaterial::GetUsedTextures(TArray<UTexture*>& OutTextures, EMaterialQuality
 				{
 					for (UTexture* Texture : TextureCollection->Textures)
 					{
-						OutTextures.AddUnique(Texture);
+						if (Texture)
+						{
+							OutTextures.AddUnique(Texture);
+						}
 					}
 				}
 			}
diff --git a/Engine/Source/Runtime/Engine/Private/Materials/MaterialCachedData.cpp b/Engine/Source/Runtime/Engine/Private/Materials/MaterialCachedData.cpp
--- a/Engine/Source/Runtime/Engine/Private/Materials/MaterialCachedData.cpp
+++ b/Engine/Source/Runtime/Engine/Private/Materials/MaterialCachedData.cpp
@@ -85,6 +85,7 @@ FMaterialCachedExpressionData::FMaterialCachedExpressionData()
 void FMaterialCachedExpressionData::AddReferencedObjects(FReferenceCollector& Collector)
 {
 	Collector.AddStableReferenceArray(&ReferencedTextures);
+	Collector.AddStableReferenceArray(&ReferencedTextureCollections);
 	Collector.AddStableReferenceArray(&GrassTypes);
 	Collector.AddStableReferenceArray(&MaterialLayers.Layers);
 	Collector.AddStableReferenceArray(&MaterialLayers.Blends);
diff --git a/Engine/Source/Runtime/Engine/Private/Materials/MaterialInstance.cpp b/Engine/Source/Runtime/Engine/Private/Materials/MaterialInstance.cpp
--- a/Engine/Source/Runtime/Engine/Private/Materials/MaterialInstance.cpp
+++ b/Engine/Source/Runtime/Engine/Private/Materials/MaterialInstance.cpp
@@ -1368,7 +1368,10 @@ void UMaterialInstance::GetTextureExpressionValues(const FMaterialResource* Mate
 			{
 				for (UTexture* Texture : TextureCollection->Textures)
 				{
-					OutTextures.AddUnique(Texture);
+					if (Texture)
+					{
+						OutTextures.AddUnique(Texture);
+					}
 				}
 			}
 		}
