[NavigationElement] minor code style fix

diff --git a/Engine/Source/Runtime/Engine/Classes/AI/Navigation/NavigationElement.h b/Engine/Source/Runtime/Engine/Classes/AI/Navigation/NavigationElement.h
--- a/Engine/Source/Runtime/Engine/Classes/AI/Navigation/NavigationElement.h
+++ b/Engine/Source/Runtime/Engine/Classes/AI/Navigation/NavigationElement.h
@@ -52,10 +52,7 @@ struct FNavigationElementHandle
 	inline bool IsValid() const;
 
 	/** Invalidates the current handle so further calls to IsValid will return false. */
-	void Invalidate()
-	{
-		*this = Invalid;
-	}
+	inline void Invalidate();
 
 	/** Hash function to use TSet/TMap */
 	inline friend uint32 GetTypeHash(const FNavigationElementHandle& Handle);
@@ -370,6 +367,11 @@ bool FNavigationElementHandle::IsValid() const
 	return *this != Invalid;
 }
 
+void FNavigationElementHandle::Invalidate()
+{
+	*this = Invalid;
+}
+
 
 //----------------------------------------------------------------------//
 // Inlines (FNavigationElement)
