New Material Translator:  Support for ACosFast, ASinFast, ATanFast, ATan2Fast, and ConstantBiasScale.

#rnx

diff --git a/Engine/Source/Runtime/Engine/Private/Materials/MaterialExpressionsIR.cpp b/Engine/Source/Runtime/Engine/Private/Materials/MaterialExpressionsIR.cpp
--- a/Engine/Source/Runtime/Engine/Private/Materials/MaterialExpressionsIR.cpp
+++ b/Engine/Source/Runtime/Engine/Private/Materials/MaterialExpressionsIR.cpp
@@ -348,6 +348,11 @@ void UMaterialExpressionConstant4Vector::Build(MIR::FEmitter& Em)
 	}
 }
 
+void UMaterialExpressionConstantBiasScale::Build(MIR::FEmitter& Em)
+{
+	Em.Output(0, Em.Multiply(Em.Add(Em.ConstantFloat(Bias), Em.Input(&Input)), Em.ConstantFloat(Scale)));
+}
+
 void UMaterialExpressionStaticBool::Build(MIR::FEmitter& Em)
 {
 	Em.Output(0, Em.ConstantBool(Value));
@@ -613,6 +618,21 @@ void UMaterialExpressionArctangent::Build(MIR::FEmitter& Em)
 	Em.Output(0, Em.ATan(Em.Input(&Input)));
 }
 
+void UMaterialExpressionArccosineFast::Build(MIR::FEmitter& Em)
+{
+	Em.Output(0, Em.Operator(MIR::UO_ACosFast, Em.Input(&Input)));
+}
+
+void UMaterialExpressionArcsineFast::Build(MIR::FEmitter& Em)
+{
+	Em.Output(0, Em.Operator(MIR::UO_ASinFast, Em.Input(&Input)));
+}
+
+void UMaterialExpressionArctangentFast::Build(MIR::FEmitter& Em)
+{
+	Em.Output(0, Em.Operator(MIR::UO_ATanFast, Em.Input(&Input)));
+}
+
 void UMaterialExpressionComponentMask::Build(MIR::FEmitter& Em)
 {
 	FValueRef Value = Em.Input(&Input);
@@ -811,6 +831,11 @@ void UMaterialExpressionArctangent2::Build(MIR::FEmitter& Em)
 	Em.Output(0, Em.Atan2(Em.Input(&Y), Em.Input(&X)));
 }
 
+void UMaterialExpressionArctangent2Fast::Build(MIR::FEmitter& Em)
+{
+	Em.Output(0, Em.Operator(MIR::BO_ATan2Fast, Em.Input(&Y), Em.Input(&X)));
+}
+
 void UMaterialExpressionEyeAdaptationInverse::Build(MIR::FEmitter& Em)
 {
 	check(ExternalCodeIdentifiers.Num() == 1);
diff --git a/Engine/Source/Runtime/Engine/Private/Materials/MaterialIR.cpp b/Engine/Source/Runtime/Engine/Private/Materials/MaterialIR.cpp
--- a/Engine/Source/Runtime/Engine/Private/Materials/MaterialIR.cpp
+++ b/Engine/Source/Runtime/Engine/Private/Materials/MaterialIR.cpp
@@ -701,10 +701,13 @@ const TCHAR* LexToString(EOperator Op)
 		/* Unary operators */
 		case UO_Abs: return TEXT("Abs");
 		case UO_ACos: return TEXT("ACos");
+		case UO_ACosFast: return TEXT("ACosFast");
 		case UO_ACosh: return TEXT("ACosh");
 		case UO_ASin: return TEXT("ASin");
+		case UO_ASinFast: return TEXT("ASinFast");
 		case UO_ASinh: return TEXT("ASinh");
 		case UO_ATan: return TEXT("ATan");
+		case UO_ATanFast: return TEXT("ATanFast");
 		case UO_ATanh: return TEXT("ATanh");
 		case UO_BitwiseNot: return TEXT("BitwiseNot");
 		case UO_Ceil: return TEXT("Ceil");
@@ -739,6 +742,7 @@ const TCHAR* LexToString(EOperator Op)
 		case BO_Add: return TEXT("Add");
 		case BO_And: return TEXT("And");
 		case BO_ATan2: return TEXT("ATan2");
+		case BO_ATan2Fast: return TEXT("ATan2Fast");
 		case BO_BitShiftLeft: return TEXT("BitShiftLeft");
 		case BO_BitShiftRight: return TEXT("BitShiftRight");
 		case BO_BitwiseAnd: return TEXT("BitwiseAnd");
diff --git a/Engine/Source/Runtime/Engine/Private/Materials/MaterialIREmitter.cpp b/Engine/Source/Runtime/Engine/Private/Materials/MaterialIREmitter.cpp
--- a/Engine/Source/Runtime/Engine/Private/Materials/MaterialIREmitter.cpp
+++ b/Engine/Source/Runtime/Engine/Private/Materials/MaterialIREmitter.cpp
@@ -1139,10 +1139,13 @@ static T FoldScalarOperator(FEmitter& Emitter, EOperator Operator, T A, T B, T C
 		switch (Operator)
 		{
 			case UO_ACos: return FGenericPlatformMath::Acos(A);
+			case UO_ACosFast: return FGenericPlatformMath::Acos(A);
 			case UO_ACosh: return MIR::ACosh(A);
 			case UO_ASin: return FGenericPlatformMath::Asin(A);
+			case UO_ASinFast: return FGenericPlatformMath::Asin(A);
 			case UO_ASinh: return MIR::ASinh(A);
 			case UO_ATan: return FGenericPlatformMath::Atan(A);
+			case UO_ATanFast: return FGenericPlatformMath::Atan(A);
 			case UO_ATanh: return MIR::ATanh(A);
 			case UO_Ceil:
 				if constexpr (std::is_same<T, float>::value)
@@ -1205,6 +1208,7 @@ static T FoldScalarOperator(FEmitter& Emitter, EOperator Operator, T A, T B, T C
 					return FGenericPlatformMath::TruncToDouble(A);
 				}
 			case BO_ATan2: return FGenericPlatformMath::Atan2(A, B);
+			case BO_ATan2Fast: return FGenericPlatformMath::Atan2(A, B);
 			case BO_Fmod: return FGenericPlatformMath::Fmod(A, B);
 			case BO_Pow: return FGenericPlatformMath::Pow(A, B);
 			case TO_Lerp: return FMath::Lerp<T>(A, B, C);
@@ -1566,10 +1570,13 @@ static const FOperatorSignature* GetOperatorSignature(EOperator Op)
 
 		S[UO_Abs] 						= UnaryFloatLWC;
 		S[UO_ACos] 						= UnaryFloatLWCDemote;
+		S[UO_ACosFast]					= UnaryFloat;
 		S[UO_ACosh] 					= { { OPTF_CheckIsArithmetic | OPTF_CastToAnyFloat | OPTF_CheckIsOneOrGreaterFloatConst }, ORT_FirstArgumentType };
 		S[UO_ASin] 						= UnaryFloatLWCDemote;
+		S[UO_ASinFast]					= UnaryFloat;
 		S[UO_ASinh] 					= UnaryFloat;
 		S[UO_ATan] 						= UnaryFloatLWCDemote;
+		S[UO_ATanFast]					= UnaryFloat;
 		S[UO_ATanh] 					= { { OPTF_CheckIsArithmetic | OPTF_CastToAnyFloat | OPTF_CheckIsBetweenMinusOneAndPlusOneFloatConst }, ORT_FirstArgumentType };
 		S[UO_Ceil] 						= UnaryFloatLWC;
 		S[UO_Cos] 						= UnaryFloatLWCDemote;
@@ -1616,6 +1623,7 @@ static const FOperatorSignature* GetOperatorSignature(EOperator Op)
 		S[BO_BitShiftRight]				= BinaryInteger;
 
 		S[BO_ATan2]						= BinaryFloat;
+		S[BO_ATan2Fast]					= BinaryFloat;
 		S[BO_Cross] 					= { { OPTF_CheckIsArithmetic | OPTF_CheckIsVector3, OPTF_CastToFirstArgumentType }, ORT_FirstArgumentType };
 		S[BO_Distance]					= { { OPTF_CastToCommonFloatType | OPTF_LWCInput, OPTF_CastToCommonFloatType | OPTF_LWCInput }, ORT_FirstArgumentTypeToScalarLWC };
 		S[BO_Dot] 						= { { OPTF_CheckIsArithmetic | OPTF_CheckIsNotMatrix | OPTF_LWCInput, OPTF_CastToFirstArgumentType | OPTF_LWCInput }, ORT_FirstArgumentTypeToScalarLWC };
@@ -2848,6 +2856,7 @@ static FValue* DifferentiateOperator(FEmitter& E, FOperator* Op, EDerivativeAxis
 
 		// d/dx arccos(f(x)) = -1 / sqrt(1 - f(x)^2) * f'(x)
 		case UO_ACos:
+		case UO_ACosFast:
 			return E.Negate(E.Divide(dF, E.Sqrt(E.Subtract(One(), E.Multiply(F, F)))));
 			
 		// d/dx acosh(f(x)) = 1 / sqrt(f(x)^2 - 1) * f'(x)
@@ -2856,6 +2865,7 @@ static FValue* DifferentiateOperator(FEmitter& E, FOperator* Op, EDerivativeAxis
 			
 		// d/dx arcsin(f(x)) = 1 / sqrt(1 - f(x)^2) * f'(x)
 		case UO_ASin:
+		case UO_ASinFast:
 			return E.Divide(dF, E.Sqrt(E.Subtract(One(), E.Multiply(F, F))));
 		
 		// d/dx asinh(f(x)) = 1 / sqrt(f(x)^2 + 1) * f'(x)
@@ -2864,6 +2874,7 @@ static FValue* DifferentiateOperator(FEmitter& E, FOperator* Op, EDerivativeAxis
 
 		// d/dx arctan(f(x)) = 1 / (1 + f(x)^2) * f'(x)
 		case UO_ATan:
+		case UO_ATanFast:
 			return E.Divide(dF, E.Add(One(), E.Multiply(F, F)));
 
     	// d/dx atanh(f(x)) = f'(x) / (1 - f(x)^2)
@@ -2997,6 +3008,7 @@ static FValue* DifferentiateOperator(FEmitter& E, FOperator* Op, EDerivativeAxis
 
 		// d/dx atan2(f(x), g(x)) = g(x) / (f(x)^2 + g(x)^2) * f'(x)  -  f(x) / (f(x)^2 + g(x)^2) * g'(x)
 		case BO_ATan2:
+		case BO_ATan2Fast:
 		{
 			FValue* Magnitude = E.Divide(One(), E.Add(E.Multiply(F, F), E.Multiply(G, G)));		// 1 / (f(x)^2 + g(x)^2)
 			return E.Subtract(E.Multiply(E.Multiply(G, Magnitude), dF), E.Multiply(E.Multiply(F, Magnitude), dG));
diff --git a/Engine/Source/Runtime/Engine/Private/Materials/MaterialIRToHLSLTranslator.cpp b/Engine/Source/Runtime/Engine/Private/Materials/MaterialIRToHLSLTranslator.cpp
--- a/Engine/Source/Runtime/Engine/Private/Materials/MaterialIRToHLSLTranslator.cpp
+++ b/Engine/Source/Runtime/Engine/Private/Materials/MaterialIRToHLSLTranslator.cpp
@@ -1157,10 +1157,13 @@ struct FTranslator : FMaterialIRToHLSLTranslation
 				{
 					case MIR::UO_Abs: OpString = TEXT("abs"); break;
 					case MIR::UO_ACos: OpString = TEXT("acos"); break;
+					case MIR::UO_ACosFast: OpString = TEXT("acosFast"); break;
 					case MIR::UO_ACosh: OpString = TEXT("acosh"); break;
 					case MIR::UO_ASin: OpString = TEXT("asin"); break;
+					case MIR::UO_ASinFast: OpString = TEXT("asinFast"); break;
 					case MIR::UO_ASinh: OpString = TEXT("asinh"); break;
 					case MIR::UO_ATan: OpString = TEXT("atan"); break;
+					case MIR::UO_ATanFast: OpString = TEXT("atanFast"); break;
 					case MIR::UO_ATanh: OpString = TEXT("atanh"); break;
 					case MIR::UO_Ceil: OpString = TEXT("ceil"); break;
 					case MIR::UO_Cos: OpString = TEXT("cos"); break;
@@ -1190,6 +1193,7 @@ struct FTranslator : FMaterialIRToHLSLTranslation
 
 					case MIR::BO_And: OpString = TEXT("and"); break;
 					case MIR::BO_ATan2: OpString = TEXT("atan2"); break;
+					case MIR::BO_ATan2Fast: OpString = TEXT("atan2Fast"); break;
 					case MIR::BO_Cross: OpString = TEXT("cross"); break;
 					case MIR::BO_Distance: OpString = TEXT("distance"); break;
 					case MIR::BO_Dot: OpString = TEXT("dot"); break;
diff --git a/Engine/Source/Runtime/Engine/Public/Materials/MaterialExpressionArccosineFast.h b/Engine/Source/Runtime/Engine/Public/Materials/MaterialExpressionArccosineFast.h
--- a/Engine/Source/Runtime/Engine/Public/Materials/MaterialExpressionArccosineFast.h
+++ b/Engine/Source/Runtime/Engine/Public/Materials/MaterialExpressionArccosineFast.h
@@ -20,6 +20,7 @@ public:
 
 	//~ Begin UMaterialExpression Interface
 #if WITH_EDITOR
+	virtual void Build(MIR::FEmitter& Emitter) override;
 	virtual int32 Compile(class FMaterialCompiler* Compiler, int32 OutputIndex) override;
 	virtual void GetCaption(TArray<FString>& OutCaptions) const override;
 	virtual void GetExpressionToolTip(TArray<FString>& OutToolTip) override;
diff --git a/Engine/Source/Runtime/Engine/Public/Materials/MaterialExpressionArcsineFast.h b/Engine/Source/Runtime/Engine/Public/Materials/MaterialExpressionArcsineFast.h
--- a/Engine/Source/Runtime/Engine/Public/Materials/MaterialExpressionArcsineFast.h
+++ b/Engine/Source/Runtime/Engine/Public/Materials/MaterialExpressionArcsineFast.h
@@ -20,6 +20,7 @@ public:
 
 	//~ Begin UMaterialExpression Interface
 #if WITH_EDITOR
+	virtual void Build(MIR::FEmitter& Emitter) override;
 	virtual int32 Compile(class FMaterialCompiler* Compiler, int32 OutputIndex) override;
 	virtual void GetCaption(TArray<FString>& OutCaptions) const override;
 	virtual void GetExpressionToolTip(TArray<FString>& OutToolTip) override;
diff --git a/Engine/Source/Runtime/Engine/Public/Materials/MaterialExpressionArctangent2Fast.h b/Engine/Source/Runtime/Engine/Public/Materials/MaterialExpressionArctangent2Fast.h
--- a/Engine/Source/Runtime/Engine/Public/Materials/MaterialExpressionArctangent2Fast.h
+++ b/Engine/Source/Runtime/Engine/Public/Materials/MaterialExpressionArctangent2Fast.h
@@ -23,6 +23,7 @@ public:
 
 	//~ Begin UMaterialExpression Interface
 #if WITH_EDITOR
+	virtual void Build(MIR::FEmitter& Emitter) override;
 	virtual int32 Compile(class FMaterialCompiler* Compiler, int32 OutputIndex) override;
 	virtual void GetCaption(TArray<FString>& OutCaptions) const override;
 	virtual void GetExpressionToolTip(TArray<FString>& OutToolTip) override;
diff --git a/Engine/Source/Runtime/Engine/Public/Materials/MaterialExpressionArctangentFast.h b/Engine/Source/Runtime/Engine/Public/Materials/MaterialExpressionArctangentFast.h
--- a/Engine/Source/Runtime/Engine/Public/Materials/MaterialExpressionArctangentFast.h
+++ b/Engine/Source/Runtime/Engine/Public/Materials/MaterialExpressionArctangentFast.h
@@ -20,6 +20,7 @@ public:
 
 	//~ Begin UMaterialExpression Interface
 #if WITH_EDITOR
+	virtual void Build(MIR::FEmitter& Emitter) override;
 	virtual int32 Compile(class FMaterialCompiler* Compiler, int32 OutputIndex) override;
 	virtual void GetCaption(TArray<FString>& OutCaptions) const override;
 	virtual void GetExpressionToolTip(TArray<FString>& OutToolTip) override;
diff --git a/Engine/Source/Runtime/Engine/Public/Materials/MaterialExpressionConstantBiasScale.h b/Engine/Source/Runtime/Engine/Public/Materials/MaterialExpressionConstantBiasScale.h
--- a/Engine/Source/Runtime/Engine/Public/Materials/MaterialExpressionConstantBiasScale.h
+++ b/Engine/Source/Runtime/Engine/Public/Materials/MaterialExpressionConstantBiasScale.h
@@ -27,6 +27,7 @@ public:
 
 	//~ Begin UMaterialExpression Interface
 #if WITH_EDITOR
+	virtual void Build(MIR::FEmitter& Emitter) override;
 	virtual int32 Compile(class FMaterialCompiler* Compiler, int32 OutputIndex) override;
 	virtual void GetCaption(TArray<FString>& OutCaptions) const override;
 #endif
diff --git a/Engine/Source/Runtime/Engine/Public/Materials/MaterialExpressionOperator.h b/Engine/Source/Runtime/Engine/Public/Materials/MaterialExpressionOperator.h
--- a/Engine/Source/Runtime/Engine/Public/Materials/MaterialExpressionOperator.h
+++ b/Engine/Source/Runtime/Engine/Public/Materials/MaterialExpressionOperator.h
@@ -18,10 +18,13 @@ enum class EMaterialExpressionOperatorKind
 	Not,
 	Abs,
 	ACos,
+	ACosFast,
 	ACosh,
 	ASin,
+	ASinFast,
 	ASinh,
 	ATan,
+	ATanFast,
 	ATanh,
 	Ceil,
 	Cos,
@@ -67,7 +70,8 @@ enum class EMaterialExpressionOperatorKind
 	BitwiseOr,
 	BitShiftLeft,
 	BitShiftRight,
-	Atan2,
+	ATan2,
+	ATan2Fast,
 	Cross,
 	Distance,
 	Dot,
diff --git a/Engine/Source/Runtime/Engine/Public/Materials/MaterialIR.h b/Engine/Source/Runtime/Engine/Public/Materials/MaterialIR.h
--- a/Engine/Source/Runtime/Engine/Public/Materials/MaterialIR.h
+++ b/Engine/Source/Runtime/Engine/Public/Materials/MaterialIR.h
@@ -624,10 +624,13 @@ enum EOperator
 	// Unary intrinsics
 	UO_Abs,
 	UO_ACos,
+	UO_ACosFast,
 	UO_ACosh,
 	UO_ASin,
+	UO_ASinFast,
 	UO_ASinh,
 	UO_ATan,
+	UO_ATanFast,
 	UO_ATanh,
 	UO_Ceil,
 	UO_Cos,
@@ -682,6 +685,7 @@ enum EOperator
 	
 	// Binary intrinsics
 	BO_ATan2,
+	BO_ATan2Fast,
 	BO_Cross,
 	BO_Distance,
 	BO_Dot,
