[CompositeDataTable]
- Fix regression caused by missing check for nullptr

diff --git a/Engine/Source/Runtime/Engine/Private/CompositeDataTable.cpp b/Engine/Source/Runtime/Engine/Private/CompositeDataTable.cpp
--- a/Engine/Source/Runtime/Engine/Private/CompositeDataTable.cpp
+++ b/Engine/Source/Runtime/Engine/Private/CompositeDataTable.cpp
@@ -223,7 +223,10 @@ void UCompositeDataTable::Serialize(FArchive& Ar)
 	{
 		for (UDataTable* ParentTable : ParentTables)
 		{
-			ParentTable->ConditionalPreload();
+			if (ParentTable)
+			{
+				ParentTable->ConditionalPreload();
+			}
 		}
 
 		OnParentTablesUpdated();
