Fix an assert that can trigger when creating a virtual texture editor only resource.
The size of editor only resource needs clamping, and that code was removed by mistake.
The issue can arise because thin textures are scaled to tile size when converted to virtual texture.

diff --git a/Engine/Source/Runtime/Engine/Private/Texture2D.cpp b/Engine/Source/Runtime/Engine/Private/Texture2D.cpp
--- a/Engine/Source/Runtime/Engine/Private/Texture2D.cpp
+++ b/Engine/Source/Runtime/Engine/Private/Texture2D.cpp
@@ -1682,13 +1682,14 @@ void FVirtualTexture2DResource::InitializeEditorResources(FRHICommandListBase& R
 	if (GIsEditor && !IsRunningCommandlet())
 	{
 		// Choose a mip level for the thumbnail texture to ensure proper size
-		const uint32 TargetSize = FMath::Clamp(128u, GetTileSize(), 1024u);
+		const uint32 MaxTargetSize = FMath::Min(1024u, GetMax2DTextureDimension());
+		const uint32 TargetSize = FMath::Clamp(128u, GetTileSize(), MaxTargetSize);
 		const uint32 MaxMipLevel = VTData->GetNumMips() - 1u;
 		
 		uint32 MipLevel = 0;
 		uint32 MipWidth = GetSizeX();
 		uint32 MipHeight = GetSizeY();
-		while (MipWidth > TargetSize && MipHeight > TargetSize && MipLevel < MaxMipLevel)
+		while (((MipWidth > TargetSize && MipHeight > TargetSize) || MipWidth > MaxTargetSize || MipHeight > MaxTargetSize) && MipLevel < MaxMipLevel)
 		{
 			++MipLevel;
 			MipWidth = FMath::DivideAndRoundUp(MipWidth, 2u);
