Add MatchesSearchQuery overload to allow an option for additional data to be retrieved from MaterialExpressions during searches.
Also extends embedded MFs and their expressions in function calls to be included in the search.

diff --git a/Engine/Source/Editor/MaterialEditor/Private/FindInMaterial.cpp b/Engine/Source/Editor/MaterialEditor/Private/FindInMaterial.cpp
--- a/Engine/Source/Editor/MaterialEditor/Private/FindInMaterial.cpp
+++ b/Engine/Source/Editor/MaterialEditor/Private/FindInMaterial.cpp
@@ -319,7 +319,7 @@ void SFindInMaterial::MatchTokensInGraph(const UEdGraph* Graph, const TArray<FSt
 			{
 				for (int32 Index = 0; Index < Tokens.Num(); ++Index)
 				{
-					if (!MatNode->MaterialExpression->MatchesSearchQuery(*Tokens[Index]))
+					if (!MatNode->MaterialExpression->DeepSearchQuery(*Tokens[Index], NodeResult))
 					{
 						bMatchesAllTokens = false;
 						break;
@@ -394,7 +394,7 @@ TSharedRef<ITableRow> SFindInMaterial::OnGenerateRow(FSearchResult InItem, const
 				.Padding(2, 0)
 				[
 					SNew(STextBlock)
-					.Text(FText::FromString(InItem->Value))
+					.Text(InItem->bEmbeddedViaNode ? FText::Format(LOCTEXT("MaterialResultSearchValueEmbeddedFmt", "Is found inside of {0}"), FText::FromString(InItem->Value)) : FText::FromString(InItem->Value))
 					.HighlightText(HighlightText)
 					.ToolTipText(FText::Format(LOCTEXT("MaterialResultSearchToolTipFmt", "{0} : {1}"), InItem->GetCategory(), FText::FromString(InItem->Value)))
 				]
diff --git a/Engine/Source/Editor/MaterialEditor/Public/FindInMaterial.h b/Engine/Source/Editor/MaterialEditor/Public/FindInMaterial.h
--- a/Engine/Source/Editor/MaterialEditor/Public/FindInMaterial.h
+++ b/Engine/Source/Editor/MaterialEditor/Public/FindInMaterial.h
@@ -68,6 +68,9 @@ public:
 
 	/** Display text for comment information */
 	FString CommentText;
+
+	/** Boolean for marking that the result was found embedded in another node e.g. cascading MFs*/
+	bool bEmbeddedViaNode = false;
 };
 
 
diff --git a/Engine/Source/Runtime/Engine/Private/Materials/MaterialExpressions.cpp b/Engine/Source/Runtime/Engine/Private/Materials/MaterialExpressions.cpp
--- a/Engine/Source/Runtime/Engine/Private/Materials/MaterialExpressions.cpp
+++ b/Engine/Source/Runtime/Engine/Private/Materials/MaterialExpressions.cpp
@@ -313,6 +313,7 @@
 #include "MaterialGraph/MaterialGraphNode_Composite.h"
 #include "MaterialEditor/MaterialNodes/SGraphNodeMaterialBase.h"
 #include "MaterialEditor/MaterialNodes/SGraphNodeMaterialConvert.h"
+#include "Editor/MaterialEditor/Public/FindInMaterial.h"
 #include "Framework/Notifications/NotificationManager.h"
 #include "Widgets/Notifications/SNotificationList.h"
 #include "MaterialGraph/MaterialGraphSchema.h"
@@ -1999,6 +2000,12 @@ bool UMaterialExpression::MatchesSearchQuery( const TCHAR* SearchQuery )
 	return Desc.Contains(SearchQuery);
 }
 
+bool UMaterialExpression::DeepSearchQuery(const TCHAR* SearchQuery, TSharedPtr<FFindInMaterialResult> SearchResult)
+{
+	//Unlike the matches call above, we always want to check the expression names in a deep search.
+	return GetName().Contains(SearchQuery) || MatchesSearchQuery(SearchQuery);
+}
+
 bool UMaterialExpression::IsExpressionConnected( FExpressionInput* Input, int32 OutputIndex )
 {
 	return Input->OutputIndex == OutputIndex && Input->Expression == this;
@@ -17795,6 +17802,28 @@ bool UMaterialExpressionMaterialFunctionCall::MatchesSearchQuery( const TCHAR* S
 	return Super::MatchesSearchQuery(SearchQuery);
 }
 
+bool UMaterialExpressionMaterialFunctionCall::DeepSearchQuery(const TCHAR* SearchQuery, TSharedPtr<FFindInMaterialResult> SearchResult)
+{
+	//First check to see if this expression matches, if not then inspect the subfunction(s)
+	if (Super::DeepSearchQuery(SearchQuery, SearchResult))
+	{
+		return true;
+	}
+	else if (SearchResult.IsValid())
+	{
+		for (UMaterialExpression* Expression : MaterialFunction->GetExpressions())
+		{
+			SearchResult->bEmbeddedViaNode = Expression->DeepSearchQuery(SearchQuery, SearchResult);
+			if (SearchResult->bEmbeddedViaNode)
+			{			
+				return true;
+			}
+		}
+	}
+	return false;
+}
+
+
 bool UMaterialExpressionMaterialFunctionCall::IsResultMaterialAttributes(int32 OutputIndex) 
 {
 	if( OutputIndex >= 0 && OutputIndex < FunctionOutputs.Num() && FunctionOutputs[OutputIndex].ExpressionOutput)
diff --git a/Engine/Source/Runtime/Engine/Public/Materials/MaterialExpression.h b/Engine/Source/Runtime/Engine/Public/Materials/MaterialExpression.h
--- a/Engine/Source/Runtime/Engine/Public/Materials/MaterialExpression.h
+++ b/Engine/Source/Runtime/Engine/Public/Materials/MaterialExpression.h
@@ -27,6 +27,8 @@ class UMaterialExpression;
 class UMaterialExpressionComment;
 class FMaterialExpressionKey;
 
+class FFindInMaterialResult;
+
 enum EEdGraphPinDirection : int;
 enum class EMaterialParameterType : uint8;
 
@@ -429,6 +431,15 @@ class UMaterialExpression : public UObject
 	 */
 	ENGINE_API virtual bool MatchesSearchQuery( const TCHAR* SearchQuery );
 
+	/**
+	 * MatchesSearchQuery: Check this expression and, if applicable, it's subobjects,to see if it matches the search query.
+	 * Also provides an FFindInMaterialResult object to return additional data for the result. 
+	 * @param SearchQuery - User's search query (never blank)
+	 * @param SearchResult - Result data
+	 * @return true if the expression matches the search query
+	 */
+	ENGINE_API virtual bool DeepSearchQuery(const TCHAR* SearchQuery, TSharedPtr<FFindInMaterialResult> SearchResult);
+
 	/**
 	 * Copy the SrcExpressions into the specified material, preserving internal references.
 	 * New material expressions are created within the specified material.
diff --git a/Engine/Source/Runtime/Engine/Public/Materials/MaterialExpressionMaterialFunctionCall.h b/Engine/Source/Runtime/Engine/Public/Materials/MaterialExpressionMaterialFunctionCall.h
--- a/Engine/Source/Runtime/Engine/Public/Materials/MaterialExpressionMaterialFunctionCall.h
+++ b/Engine/Source/Runtime/Engine/Public/Materials/MaterialExpressionMaterialFunctionCall.h
@@ -130,6 +130,7 @@ class UMaterialExpressionMaterialFunctionCall : public UMaterialExpression
 	ENGINE_API virtual void GetConnectorToolTip(int32 InputIndex, int32 OutputIndex, TArray<FString>& OutToolTip) override;
 	ENGINE_API virtual void GetExpressionToolTip(TArray<FString>& OutToolTip) override;
 	ENGINE_API virtual bool MatchesSearchQuery(const TCHAR* SearchQuery) override;
+	ENGINE_API virtual bool DeepSearchQuery(const TCHAR* SearchQuery, TSharedPtr<FFindInMaterialResult> SearchResult) override;
 	ENGINE_API virtual bool IsResultMaterialAttributes(int32 OutputIndex) override;
 	ENGINE_API virtual bool IsResultSubstrateMaterial(int32 OutputIndex) override;
 	ENGINE_API virtual void GatherSubstrateMaterialInfo(FSubstrateMaterialInfo& SubstrateMaterialInfo, int32 OutputIndex) override;
