World Partition - HLOD - SpatialHash - Make sure the FHLODCreationParams::bIsStandalone field is properly initialized (to false)

diff --git a/Engine/Source/Runtime/Engine/Private/WorldPartition/RuntimeSpatialHash/RuntimeSpatialHashHLOD.cpp b/Engine/Source/Runtime/Engine/Private/WorldPartition/RuntimeSpatialHash/RuntimeSpatialHashHLOD.cpp
--- a/Engine/Source/Runtime/Engine/Private/WorldPartition/RuntimeSpatialHash/RuntimeSpatialHashHLOD.cpp
+++ b/Engine/Source/Runtime/Engine/Private/WorldPartition/RuntimeSpatialHash/RuntimeSpatialHashHLOD.cpp
@@ -219,6 +219,7 @@ static TArray<FGuid> GenerateHLODActorsForGrid(UWorldPartition* WorldPartition,
 				CreationParams.MinVisibleDistance = RuntimeGrid.LoadingRange;
 				CreationParams.ContentBundleGuid = GridCellDataChunk.GetContentBundleID();
 				CreationParams.DataLayerInstances = GridCellDataChunk.GetDataLayers();
+				CreationParams.bIsStandalone = false;
 
 				TArray<AWorldPartitionHLOD*> CellHLODActors = WPHLODUtilities->CreateHLODActors(Context, CreationParams, ActorInstances);
 				if (!CellHLODActors.IsEmpty())
