[Engine] Fix lifetime of inlined AssetLoadState in AssetManager using UE_LAUNDER

#rnx

diff --git a/Engine/Source/Runtime/Engine/Private/AssetManager.cpp b/Engine/Source/Runtime/Engine/Private/AssetManager.cpp
--- a/Engine/Source/Runtime/Engine/Private/AssetManager.cpp
+++ b/Engine/Source/Runtime/Engine/Private/AssetManager.cpp
@@ -31,6 +31,7 @@
 #include "Misc/ScopedSlowTask.h"
 #include "Misc/PackageName.h"
 #include "Misc/PathViews.h"
+#include "Misc/Launder.h"
 #include "MoviePlayerProxy.h"
 #include "Modules/ModuleManager.h"
 #include "Stats/StatsMisc.h"
@@ -106,7 +107,7 @@ private:
 	void DestroyStoredItems()
 	{
 		check(Storage);
-		DestructItem(reinterpret_cast<TSharedPtr<FStreamableHandle>*>(Storage));
+		DestructItem(GetHandleInternal());
 		DestructItems(GetBundleNamesInternal(), NumBundleNames);
 		if (FAssetManagerDownloadParams* DownloadParams = GetDownloadParamsInternal())
 		{
@@ -115,13 +116,13 @@ private:
 	}
 
 	/** The handle to the streamable state for this asset, this keeps the objects in memory. If handle is invalid, not in memory at all */
-	TSharedPtr<FStreamableHandle>& GetHandleInternal()
+	TSharedPtr<FStreamableHandle>* GetHandleInternal()
 	{
 		check(Storage);
-		return *reinterpret_cast<TSharedPtr<FStreamableHandle>*>(Storage);
+		return UE_LAUNDER(reinterpret_cast<TSharedPtr<FStreamableHandle>*>(Storage));
 	}
 
-	const TSharedPtr<FStreamableHandle>& GetHandleInternal() const
+	const TSharedPtr<FStreamableHandle>* GetHandleInternal() const
 	{
 		return const_cast<FPrimaryAssetLoadState*>(this)->GetHandleInternal();
 	}
@@ -130,7 +131,7 @@ private:
 	FName* GetBundleNamesInternal()
 	{
 		check(Storage);
-		return reinterpret_cast<FName*>(Storage + GetBundleNamesOffset());
+		return UE_LAUNDER(reinterpret_cast<FName*>(Storage + GetBundleNamesOffset()));
 	}
 
 	const FName* GetBundleNamesInternal() const
@@ -144,7 +145,7 @@ private:
 		if (bWithDownloadParams)
 		{
 			check(Storage);
-			return reinterpret_cast<FAssetManagerDownloadParams*>(Storage + GetDownloadParamsOffset(NumBundleNames));
+			return UE_LAUNDER(reinterpret_cast<FAssetManagerDownloadParams*>(Storage + GetDownloadParamsOffset(NumBundleNames)));
 		}
 
 		return nullptr;
@@ -243,7 +244,7 @@ public:
 		}
 		else
 		{
-			GetHandleInternal() = MoveTemp(InHandle);
+			*GetHandleInternal() = MoveTemp(InHandle);
 
 			CopyAssignItems(GetBundleNamesInternal(), InBundleNames.GetData(), InNumBundleNames);
 
@@ -259,7 +260,7 @@ public:
 		*this = MoveTemp(InOther);
 		if (Storage)
 		{
-			GetHandleInternal() = MoveTemp(InHandle);
+			*GetHandleInternal() = MoveTemp(InHandle);
 		}
 	}
 
@@ -267,7 +268,7 @@ public:
 	{
 		if (Storage)
 		{
-			return GetHandleInternal();
+			return *GetHandleInternal();
 		}
 		return InvalidHandle;
 	}
