Add back in accidentally removed nullptr check to the UMaterialExpressionFunctionCall::DeepSearchQuery.

#rnx

diff --git a/Engine/Source/Runtime/Engine/Private/Materials/MaterialExpressions.cpp b/Engine/Source/Runtime/Engine/Private/Materials/MaterialExpressions.cpp
--- a/Engine/Source/Runtime/Engine/Private/Materials/MaterialExpressions.cpp
+++ b/Engine/Source/Runtime/Engine/Private/Materials/MaterialExpressions.cpp
@@ -17809,7 +17809,7 @@ bool UMaterialExpressionMaterialFunctionCall::DeepSearchQuery(const TCHAR* Searc
 	{
 		return true;
 	}
-	else if (SearchResult.IsValid())
+	else if (SearchResult.IsValid() && MaterialFunction)
 	{
 		for (UMaterialExpression* Expression : MaterialFunction->GetExpressions())
 		{
