[MIR] Fix sampler type output in new material translator.

LowerSamplerType() was invoked twice on standard textures, causing unpacking of normal maps to be called twice.

#rnx

diff --git a/Engine/Source/Runtime/Engine/Private/Materials/MaterialIRToHLSLTranslator.cpp b/Engine/Source/Runtime/Engine/Private/Materials/MaterialIRToHLSLTranslator.cpp
--- a/Engine/Source/Runtime/Engine/Private/Materials/MaterialIRToHLSLTranslator.cpp
+++ b/Engine/Source/Runtime/Engine/Private/Materials/MaterialIRToHLSLTranslator.cpp
@@ -1357,13 +1357,6 @@ struct FTranslator : FMaterialIRToHLSLTranslation
 
 	void LowerStandardTextureRead(const MIR::FTextureRead* TextureRead)
 	{
-		bool bSamplerNeedsBrackets = false;
-		LowerSamplerType(TextureRead->SamplerType, bSamplerNeedsBrackets);
-		if (bSamplerNeedsBrackets)
-		{
-			Printer << TEXT("(");
-		}
-		
 		switch (TextureRead->TextureObject->Type->AsObject()->ObjectKind)
 		{
 			case MIR::ObjectKind_Texture2D: Printer << TEXT("Texture2D"); break;
@@ -1397,11 +1390,6 @@ struct FTranslator : FMaterialIRToHLSLTranslation
 		}
 
 		Printer << EndArgs;
-
-		if (bSamplerNeedsBrackets)
-		{
-			Printer << TEXT(")");
-		}
 	}
 
 	void LowerVirtualTextureRead(const MIR::FTextureRead* TextureRead)
