[NewMaterialTranslator] - Implement UMaterialExpressionModulo.

- The operator already existed just needed to expose a helper on the emitter and implement the build function
- CPU side constant folding uses truncated division when implementing modulo which matches SPIRV, DXIL and HLSL behavior for `%`
- x-y*trunc(x/y) is the equation used and the sign of the result matches the sign of the dividend (operand 1)
- Modulo should work on floating point values in HLSL so change the IR so it can accept floats.
- Fix typo.

#rnx

diff --git a/Engine/Source/Runtime/Engine/Private/Materials/MaterialExpressionsIR.cpp b/Engine/Source/Runtime/Engine/Private/Materials/MaterialExpressionsIR.cpp
--- a/Engine/Source/Runtime/Engine/Private/Materials/MaterialExpressionsIR.cpp
+++ b/Engine/Source/Runtime/Engine/Private/Materials/MaterialExpressionsIR.cpp
@@ -811,6 +811,11 @@ void UMaterialExpressionMin::Build(MIR::FEmitter& Em)
 	BuildBinaryOperatorWithDefaults(Em, MIR::BO_Min, &A, ConstA, &B, ConstB);
 }
 
+void UMaterialExpressionModulo::Build(MIR::FEmitter& Em)
+{
+	Em.Output(0, Em.Operator(MIR::BO_Modulo, Em.Input(&A), Em.Input(&B)));
+}
+
 void UMaterialExpressionStep::Build(MIR::FEmitter& Em)
 {
 	BuildBinaryOperatorWithDefaults(Em, MIR::BO_Step, &Y, ConstY, &X, ConstX);
diff --git a/Engine/Source/Runtime/Engine/Private/Materials/MaterialIR.cpp b/Engine/Source/Runtime/Engine/Private/Materials/MaterialIR.cpp
--- a/Engine/Source/Runtime/Engine/Private/Materials/MaterialIR.cpp
+++ b/Engine/Source/Runtime/Engine/Private/Materials/MaterialIR.cpp
@@ -759,7 +759,7 @@ const TCHAR* LexToString(EOperator Op)
 		case BO_LessThanOrEquals: return TEXT("LessThanOrEquals");
 		case BO_Max: return TEXT("Max");
 		case BO_Min: return TEXT("Min");
-		case BO_Modulo: return TEXT("Module");
+		case BO_Modulo: return TEXT("Modulo");
 		case BO_Multiply: return TEXT("Multiply");
 		case BO_NotEquals: return TEXT("NotEquals");
 		case BO_Or: return TEXT("Or");
diff --git a/Engine/Source/Runtime/Engine/Private/Materials/MaterialIREmitter.cpp b/Engine/Source/Runtime/Engine/Private/Materials/MaterialIREmitter.cpp
--- a/Engine/Source/Runtime/Engine/Private/Materials/MaterialIREmitter.cpp
+++ b/Engine/Source/Runtime/Engine/Private/Materials/MaterialIREmitter.cpp
@@ -1210,6 +1210,8 @@ static T FoldScalarOperator(FEmitter& Emitter, EOperator Operator, T A, T B, T C
 			case BO_ATan2: return FGenericPlatformMath::Atan2(A, B);
 			case BO_ATan2Fast: return FGenericPlatformMath::Atan2(A, B);
 			case BO_Fmod: return FGenericPlatformMath::Fmod(A, B);
+			// truncated division (A-B*(trunc(A/B))) where the result takes on the sign of operand 1 the dividend
+			case BO_Modulo: return A - B * FGenericPlatformMath::TruncToFloat(A / B);
 			case BO_Pow: return FGenericPlatformMath::Pow(A, B);
 			case TO_Lerp: return FMath::Lerp<T>(A, B, C);
 			case TO_Smoothstep: return FMath::SmoothStep<T>(A, B, C);
@@ -1616,7 +1618,7 @@ static const FOperatorSignature* GetOperatorSignature(EOperator Op)
 		S[BO_Subtract] 					= BinaryArithmeticLWC;
 		S[BO_Multiply] 					= BinaryArithmeticLWC;
 		S[BO_Divide] 					= BinaryArithmeticLWC;
-		S[BO_Modulo] 					= BinaryInteger;
+		S[BO_Modulo] 					= BinaryFloat;
 		S[BO_BitwiseAnd]				= BinaryInteger;
 		S[BO_BitwiseOr]					= BinaryInteger;
 		S[BO_BitShiftLeft]				= BinaryInteger;
diff --git a/Engine/Source/Runtime/Engine/Public/Materials/MaterialExpressionModulo.h b/Engine/Source/Runtime/Engine/Public/Materials/MaterialExpressionModulo.h
--- a/Engine/Source/Runtime/Engine/Public/Materials/MaterialExpressionModulo.h
+++ b/Engine/Source/Runtime/Engine/Public/Materials/MaterialExpressionModulo.h
@@ -20,6 +20,7 @@ public:
 
 	//~ Begin UMaterialExpression Interface
 #if WITH_EDITOR
+	virtual void Build(MIR::FEmitter& Emitter) override;
 	virtual int32 Compile(class FMaterialCompiler* Compiler, int32 OutputIndex) override;
 	virtual EMaterialValueType GetInputValueType(int32 InputIndex) override;
 	virtual EMaterialValueType GetOutputValueType(int32 OutputIndex) override;
diff --git a/Engine/Source/Runtime/Engine/Public/Materials/MaterialIREmitter.h b/Engine/Source/Runtime/Engine/Public/Materials/MaterialIREmitter.h
--- a/Engine/Source/Runtime/Engine/Public/Materials/MaterialIREmitter.h
+++ b/Engine/Source/Runtime/Engine/Public/Materials/MaterialIREmitter.h
@@ -562,6 +562,7 @@ public:
 	FValueRef BitShiftRight(FValueRef A, FValueRef B) { return Operator(BO_BitShiftRight, A, B); }
 	FValueRef Fmod(FValueRef A, FValueRef B) { return Operator(BO_Fmod, A, B); }
 	FValueRef Max(FValueRef A, FValueRef B) { return Operator(BO_Max, A, B); }
+	FValueRef Modulo(FValueRef A, FValueRef B) { return Operator(BO_Modulo, A, B); }
 	FValueRef Min(FValueRef A, FValueRef B) { return Operator(BO_Min, A, B); }
 	FValueRef Step(FValueRef A, FValueRef B) { return Operator(BO_Step, A, B); }
 	FValueRef Dot(FValueRef A, FValueRef B) { return Operator(BO_Dot, A, B); }
