Game Instance ptr is reset when unregistered.
[RN] GameInstance pointer for Async Actions is correctly reset after it is un-registered.

diff --git a/Engine/Source/Runtime/Engine/Private/Kismet/BlueprintAsyncActionBase.cpp b/Engine/Source/Runtime/Engine/Private/Kismet/BlueprintAsyncActionBase.cpp
--- a/Engine/Source/Runtime/Engine/Private/Kismet/BlueprintAsyncActionBase.cpp
+++ b/Engine/Source/Runtime/Engine/Private/Kismet/BlueprintAsyncActionBase.cpp
@@ -56,6 +56,7 @@ void UBlueprintAsyncActionBase::SetReadyToDestroy()
 	if (OldGameInstance)
 	{
 		OldGameInstance->UnregisterReferencedObject(this);
+		RegisteredWithGameInstance.Reset();
 	}
 }
 
