New Material Translator:  More operator simplifications for trivial cases.  X-X=0, X/X=1, clamp(*,X,X)=X, lerp(X,X,*)=X, select(*,X,X)=X

#rnx

diff --git a/Engine/Source/Runtime/Engine/Private/Materials/MaterialIREmitter.cpp b/Engine/Source/Runtime/Engine/Private/Materials/MaterialIREmitter.cpp
--- a/Engine/Source/Runtime/Engine/Private/Materials/MaterialIREmitter.cpp
+++ b/Engine/Source/Runtime/Engine/Private/Materials/MaterialIREmitter.cpp
@@ -1311,6 +1311,10 @@ static FValueRef TrySimplifyOperator(FEmitter& Emitter, EOperator& Op, FValueRef
 			{
 				return Emitter.Negate(A);
 			}
+			else if (A->Equals(B))
+			{
+				return Emitter.ConstantZero(A->Type->AsPrimitive()->ScalarKind);
+			}
 			break;
 
 		case BO_Multiply:
@@ -1329,6 +1333,10 @@ static FValueRef TrySimplifyOperator(FEmitter& Emitter, EOperator& Op, FValueRef
 			{
 				return A;
 			}
+			else if (A->Equals(B))
+			{
+				return Emitter.ConstantOne(A->Type->AsPrimitive()->ScalarKind);
+			}
 			break;
 
 		case BO_Modulo:
@@ -1393,6 +1401,10 @@ static FValueRef TrySimplifyOperator(FEmitter& Emitter, EOperator& Op, FValueRef
 				B = {};
 				C = {};
 			}
+			else if (B->Equals(C))
+			{
+				return B;
+			}
 			break;
 
 		case TO_Lerp:
@@ -1404,6 +1416,10 @@ static FValueRef TrySimplifyOperator(FEmitter& Emitter, EOperator& Op, FValueRef
 			{
 				return B;
 			}
+			else if (A->Equals(B))
+			{
+				return A;
+			}
 			break;
 
 		case TO_Select:
@@ -1415,6 +1431,10 @@ static FValueRef TrySimplifyOperator(FEmitter& Emitter, EOperator& Op, FValueRef
 			{
 				return C;
 			}
+			else if (B->Equals(C))
+			{
+				return B;
+			}
 			break;
 
 		default:
