[NewMaterialTranslator] - Modulo should work on all BinaryArithmetic types (float, int, LWC, but not bools)

#rnx

diff --git a/Engine/Source/Runtime/Engine/Private/Materials/MaterialIREmitter.cpp b/Engine/Source/Runtime/Engine/Private/Materials/MaterialIREmitter.cpp
--- a/Engine/Source/Runtime/Engine/Private/Materials/MaterialIREmitter.cpp
+++ b/Engine/Source/Runtime/Engine/Private/Materials/MaterialIREmitter.cpp
@@ -1638,7 +1638,7 @@ static const FOperatorSignature* GetOperatorSignature(EOperator Op)
 		S[BO_Subtract] 					= BinaryArithmeticLWC;
 		S[BO_Multiply] 					= BinaryArithmeticLWC;
 		S[BO_Divide] 					= BinaryArithmeticLWC;
-		S[BO_Modulo] 					= BinaryFloat;
+		S[BO_Modulo] 					= BinaryArithmetic;
 		S[BO_BitwiseAnd]				= BinaryInteger;
 		S[BO_BitwiseOr]					= BinaryInteger;
 		S[BO_BitShiftLeft]				= BinaryInteger;
