- Move SoundConcurrency serialization to custom serializer
- Re-enact PerPlatform Sound Concurrency MaxCount
- Expose field in editor
- Provide Serialize upgrade logic

#rnx

diff --git a/Engine/Source/Runtime/Engine/Classes/Sound/SoundConcurrency.h b/Engine/Source/Runtime/Engine/Classes/Sound/SoundConcurrency.h
--- a/Engine/Source/Runtime/Engine/Classes/Sound/SoundConcurrency.h
+++ b/Engine/Source/Runtime/Engine/Classes/Sound/SoundConcurrency.h
@@ -2,7 +2,6 @@
 
 #pragma once
 
-#include "CoreMinimal.h"
 #include "UObject/ObjectMacros.h"
 #include "UObject/Object.h"
 #include "UObject/PerPlatformProperties.h"
@@ -76,12 +75,9 @@ struct FSoundConcurrencySettings
 {
 	GENERATED_BODY()
 
-	/** The max number of allowable concurrent active voices for voices playing in this concurrency group. */
-	UPROPERTY(EditAnywhere, BlueprintReadWrite, Category = Concurrency, meta = (UIMin = "1", ClampMin = "1"))
-	int32 MaxCount = 16;
-
-	UPROPERTY()
-	FPerPlatformInt PlatformMaxCount;
+	/** The max number of allowable concurrent active voices for voices playing in this concurrency group (Deprecated in favor of PlatformMaxCount, but maintained to avoid breaking Blueprint). */
+	UPROPERTY(BlueprintReadWrite, meta = (DisplayName = "Max Count (Deprecated)", Category = Concurrency))
+	int32 MaxCount;
 
 	/* Whether or not to limit the concurrency to per sound owner (i.e. the actor that plays the sound). If the sound doesn't have an owner, it falls back to global concurrency. */
 	UPROPERTY(EditAnywhere, BlueprintReadWrite, Category = Concurrency)
@@ -102,6 +98,9 @@ struct FSoundConcurrencySettings
 	float RetriggerTime = 0.0f;
 
 private:
+	UPROPERTY(EditAnywhere, Category = Concurrency, meta = (DisplayName = "Max Count", UIMin = "1", ClampMin = "1"))
+	FPerPlatformInt PlatformMaxCount;
+
 	/**
 	 * Ducking factor to apply per older voice instance (generation), which compounds based on scaling mode
 	 * and (optionally) revives them as they stop according to the provided attack/release times.
@@ -136,6 +135,8 @@ public:
 	UPROPERTY(EditAnywhere, BlueprintReadWrite, Category = "Voice Stealing", meta = (UIMin = "0.0", ClampMin = "0.0", UIMax = "1.0", ClampMax="1000000.0"))
 	float VoiceStealReleaseTime = 0.0f;
 
+	FSoundConcurrencySettings();
+
 	/**
 	  * Returns the MaxCount as a per-platform integer value.
 	  */
@@ -147,7 +148,10 @@ public:
 	  */
 	ENGINE_API bool SetMaxCount(int32 InMaxCount, FName PlatformName = { });
 
-	FSoundConcurrencySettings();
+	/**
+	  * Sets the MaxCount as a per-platform integer value.
+	  */
+	ENGINE_API void SetMaxCount(FPerPlatformInt InMaxCount);
 
 	/**
 	 * Retrieves the volume scale
@@ -160,7 +164,7 @@ public:
 	  */
 	ENGINE_API bool IsEvictionSupported() const;
 
-	ENGINE_API void PostSerialize(const FArchive& Ar);
+	ENGINE_API bool Serialize(FArchive& Ar);
 };
 
 template<>
@@ -169,7 +173,7 @@ struct TStructOpsTypeTraits<FSoundConcurrencySettings>
 {
 	enum
 	{
-		WithPostSerialize = true
+		WithSerializer = true
 	};
 };
 
diff --git a/Engine/Source/Runtime/Engine/Private/SoundConcurrency.cpp b/Engine/Source/Runtime/Engine/Private/SoundConcurrency.cpp
--- a/Engine/Source/Runtime/Engine/Private/SoundConcurrency.cpp
+++ b/Engine/Source/Runtime/Engine/Private/SoundConcurrency.cpp
@@ -3,6 +3,7 @@
 #include "Sound/SoundConcurrency.h"
 
 #include "AudioDevice.h"
+#include "Serialization/CustomVersion.h"
 
 #include UE_INLINE_GENERATED_CPP_BY_NAME(SoundConcurrency)
 
@@ -23,9 +24,29 @@ FAutoConsoleVariableRef CVarConcurrencyMinDistanceDelta(
 	TEXT("Minimum distance delta a sound needs to be to determine if it's further away from another sound.\n"),
 	ECVF_Default);
 
+struct FSoundConcurrencyCustomVersion
+{
+	FSoundConcurrencyCustomVersion() = delete;
+
+	enum Type
+	{
+		BeforeCustomVersionWasAdded = 0,
+		PlatformMaxCount,
+		VersionPlusOne,
+		LatestVersion = VersionPlusOne - 1,
+	};
+
+	const static FGuid GUID;
+};
+
+const FGuid FSoundConcurrencyCustomVersion::GUID(0xC7E7583A, 0x153B4256, 0x69717191, 0xB6B066BF);
+FCustomVersionRegistration GRegisterSoundConcurrencyCustomVersion(
+	FSoundConcurrencyCustomVersion::GUID, FSoundConcurrencyCustomVersion::LatestVersion, TEXT("SoundConcurrencyCustomVersion"));
 
-namespace SoundConcurrency
+namespace SoundConcurrencyPrivate
 {
+	static const int32 DefaultConcurrencyCount = 16;
+
 	void SetSoundDataTarget(const FActiveSound& ActiveSound, FConcurrencySoundData& SoundData, float InTargetVolume, float InLerpTime)
 	{
 #if UE_BUILD_SHIPPING
@@ -48,7 +69,7 @@ namespace SoundConcurrency
 		}
 #endif // UE_BUILD_SHIPPING
 	}
-} // namespace SoundConcurrency
+} // namespace SoundConcurrencyPrivate
 
 
 USoundConcurrency::USoundConcurrency(const FObjectInitializer& ObjectInitializer)
@@ -58,33 +79,26 @@ USoundConcurrency::USoundConcurrency(const FObjectInitializer& ObjectInitializer
 
 FSoundConcurrencySettings::FSoundConcurrencySettings()
 PRAGMA_DISABLE_DEPRECATION_WARNINGS
-	: MaxCount(0)
+	: MaxCount(SoundConcurrencyPrivate::DefaultConcurrencyCount)
 PRAGMA_ENABLE_DEPRECATION_WARNINGS
 {
-	PlatformMaxCount = FPerPlatformInt(16);
+	PlatformMaxCount = FPerPlatformInt(SoundConcurrencyPrivate::DefaultConcurrencyCount);
 }
 
-
 int32 FSoundConcurrencySettings::GetMaxCount() const
 {
-PRAGMA_DISABLE_DEPRECATION_WARNINGS
-	if (MaxCount > 0)
+	if (const int32 Count = PlatformMaxCount.GetValue(); Count > 0)
 	{
-		return MaxCount;
+		return Count;
 	}
-PRAGMA_ENABLE_DEPRECATION_WARNINGS
 
-	return PlatformMaxCount.GetValue();
+	return SoundConcurrencyPrivate::DefaultConcurrencyCount;
 }
 
 bool FSoundConcurrencySettings::SetMaxCount(int32 InMaxCount, FName PlatformName)
 {
 	if (InMaxCount > 0)
 	{
-PRAGMA_DISABLE_DEPRECATION_WARNINGS
-		MaxCount = 0;
-PRAGMA_ENABLE_DEPRECATION_WARNINGS
-
 		if (PlatformName.IsNone())
 		{
 			PlatformMaxCount.Default = InMaxCount;
@@ -102,9 +116,14 @@ PRAGMA_ENABLE_DEPRECATION_WARNINGS
 	return false;
 }
 
+void FSoundConcurrencySettings::SetMaxCount(FPerPlatformInt InMaxCount)
+{
+	PlatformMaxCount = MoveTemp(InMaxCount);
+}
+
 float FSoundConcurrencySettings::GetVolumeScale() const
 {
-	return FMath::Clamp(VolumeScale, 0.0f, 1.0f);;
+	return FMath::Clamp(VolumeScale, 0.0f, 1.0f);
 }
 
 bool FSoundConcurrencySettings::IsEvictionSupported() const
@@ -131,8 +150,50 @@ bool FSoundConcurrencySettings::IsEvictionSupported() const
 	}
 }
 
-void FSoundConcurrencySettings::PostSerialize(const FArchive& Ar)
+bool FSoundConcurrencySettings::Serialize(FArchive& Ar)
 {
+	Ar.UsingCustomVersion(FSoundConcurrencyCustomVersion::GUID);
+
+	auto MigratePlatformMaxCount = [&]()
+	{
+		if (Ar.CustomVer(FSoundConcurrencyCustomVersion::GUID) < FSoundConcurrencyCustomVersion::PlatformMaxCount)
+		{
+			if (MaxCount != SoundConcurrencyPrivate::DefaultConcurrencyCount)
+			{
+				PlatformMaxCount.Default = MaxCount;
+				MaxCount = 0;
+			}
+
+			return true;
+		}
+
+		return false;
+	};
+
+	if (Ar.IsSaving())
+	{
+		const bool bMigrated = MigratePlatformMaxCount();
+#if WITH_EDITORONLY_DATA
+		if (bMigrated)
+		{
+			if (Ar.IsCooking())
+			{
+				UE_LOG(LogAudioConcurrency, Display, TEXT("%s - Cook Migrated MaxCount to PlatformMaxCount"), *Ar.GetArchiveName());
+			}
+		}
+#endif // WITH_EDITORONLY_DATA
+	}
+
+	UScriptStruct* const Struct = FSoundConcurrencySettings::StaticStruct();
+	check(Struct);
+	Struct->SerializeTaggedProperties(Ar, (uint8*)this, Struct, nullptr);
+
+	if (Ar.IsLoading())
+	{
+		MigratePlatformMaxCount();
+	}
+
+	return true;
 }
 
 FConcurrencyHandle::FConcurrencyHandle(const FSoundConcurrencySettings& InSettings)
@@ -277,7 +338,7 @@ void FConcurrencyGroup::AddActiveSound(FActiveSound& ActiveSound)
 	if (Settings.VolumeScaleMode == EConcurrencyVolumeScaleMode::Default)
 	{
 		SoundData.Generation = ActiveSounds.Num();
-		SoundConcurrency::SetSoundDataTarget(ActiveSound, SoundData, 1.0f, 0.0f);
+		SoundConcurrencyPrivate::SetSoundDataTarget(ActiveSound, SoundData, 1.0f, 0.0f);
 	}
 
 	ActiveSounds.Add(&ActiveSound);
@@ -314,7 +375,7 @@ void FConcurrencyGroup::RemoveActiveSound(FActiveSound& ActiveSound)
 		{
 			const float GenerationDelta = static_cast<float>(ActiveSounds.Num() - OtherSoundData.Generation) - 1.0f;
 			const float NewTargetVolume = FMath::Clamp(FMath::Pow(Settings.GetVolumeScale(), GenerationDelta), 0.0f, 1.0f);
-			SoundConcurrency::SetSoundDataTarget(*OtherSound, OtherSoundData, NewTargetVolume, Settings.VolumeScaleReleaseTime);
+			SoundConcurrencyPrivate::SetSoundDataTarget(*OtherSound, OtherSoundData, NewTargetVolume, Settings.VolumeScaleReleaseTime);
 		}
 	}
 }
@@ -393,7 +454,7 @@ void FConcurrencyGroup::UpdateGeneration(FActiveSound* NewActiveSound)
 
 			if (!FMath::IsNearlyEqual(AttackTime, SoundData->GetLerpTime()) || !FMath::IsNearlyEqual(SoundData->GetTargetVolume(), NewTargetVolume, ConcurrencyMinVolumeScaleCVar))
 			{
-				SoundConcurrency::SetSoundDataTarget(*ActiveSound, *SoundData, NewTargetVolume, AttackTime);
+				SoundConcurrencyPrivate::SetSoundDataTarget(*ActiveSound, *SoundData, NewTargetVolume, AttackTime);
 			}
 		}
 	}
@@ -911,7 +972,7 @@ FActiveSound* FSoundConcurrencyManager::CreateAndEvictActiveSounds(const FActive
 						continue;
 					}
 
-					SoundConcurrency::SetSoundDataTarget(*OtherSound, OtherSoundData, NewTargetVolume, Settings.VolumeScaleAttackTime);
+					SoundConcurrencyPrivate::SetSoundDataTarget(*OtherSound, OtherSoundData, NewTargetVolume, Settings.VolumeScaleAttackTime);
 				}
 			}
 		}
