FMaterialShaderMapLayout: Fix ShaderMapHash to be deterministic; it was previously hashing the DebugName pointer which varies in each process.
#rnx
#incrementalcook

diff --git a/Engine/Source/Runtime/Engine/Private/Materials/MaterialShader.cpp b/Engine/Source/Runtime/Engine/Private/Materials/MaterialShader.cpp
--- a/Engine/Source/Runtime/Engine/Private/Materials/MaterialShader.cpp
+++ b/Engine/Source/Runtime/Engine/Private/Materials/MaterialShader.cpp
@@ -3064,8 +3064,8 @@ private:
 				{
 					Layout.Shaders.Add(FShaderLayoutEntry(ShaderType, PermutationId));
 
-					const FHashedName& TypeName = ShaderType->GetHashedName();
-					Hasher.Update((uint8*)&TypeName, sizeof(TypeName));
+					const uint64 TypeNameHash = ShaderType->GetHashedName().GetHash();
+					Hasher.Update((uint8*)&TypeNameHash, sizeof(TypeNameHash));
 					Hasher.Update((uint8*)&PermutationId, sizeof(PermutationId));
 				}
 			}
@@ -3080,8 +3080,8 @@ private:
 				{
 					Layout.ShaderPipelines.Add(ShaderPipelineType);
 
-					const FHashedName& TypeName = ShaderPipelineType->GetHashedName();
-					Hasher.Update((uint8*)&TypeName, sizeof(TypeName));
+					const uint64 TypeNameHash = ShaderPipelineType->GetHashedName().GetHash();
+					Hasher.Update((uint8*)&TypeNameHash, sizeof(TypeNameHash));
 				}
 			}
 		}
@@ -3109,8 +3109,8 @@ private:
 						}
 						MeshLayout->Shaders.Add(FShaderLayoutEntry(ShaderType, PermutationId));
 
-						const FHashedName& TypeName = ShaderType->GetHashedName();
-						Hasher.Update((uint8*)&TypeName, sizeof(TypeName));
+						const uint64 TypeNameHash = ShaderType->GetHashedName().GetHash();
+						Hasher.Update((uint8*)&TypeNameHash, sizeof(TypeNameHash));
 						Hasher.Update((uint8*)&PermutationId, sizeof(PermutationId));
 					}
 				}
@@ -3134,8 +3134,8 @@ private:
 						}
 						MeshLayout->ShaderPipelines.Add(ShaderPipelineType);
 
-						const FHashedName& TypeName = ShaderPipelineType->GetHashedName();
-						Hasher.Update((uint8*)&TypeName, sizeof(TypeName));
+						const uint64 TypeNameHash = ShaderPipelineType->GetHashedName().GetHash();
+						Hasher.Update((uint8*)&TypeNameHash, sizeof(TypeNameHash));
 					}
 				}
 			}
