Scalar parameter can set an enumeration object for its material instance UI.
Also scalar parameter can set an enumeration index for its material instance UI where the index is into an array of enumeration objects on the instance.

diff --git a/Engine/Source/Editor/MaterialEditor/Private/MaterialEditorDetailCustomization.cpp b/Engine/Source/Editor/MaterialEditor/Private/MaterialEditorDetailCustomization.cpp
--- a/Engine/Source/Editor/MaterialEditor/Private/MaterialEditorDetailCustomization.cpp
+++ b/Engine/Source/Editor/MaterialEditor/Private/MaterialEditorDetailCustomization.cpp
@@ -116,7 +116,27 @@ void FMaterialExpressionParameterDetails::CustomizeDetails( IDetailLayoutBuilder
 		{
 			// Store these for OnSliderMinMaxEdited
 			ScalarParameterObjects.Emplace(WeakObjectPtr);
-			DefaultValueHandles.Emplace(DetailLayout.GetProperty("DefaultValue", UMaterialExpressionScalarParameter::StaticClass()));
+
+			TSharedPtr<IPropertyHandle> DefaultValueHandle = DetailLayout.GetProperty("DefaultValue", UMaterialExpressionScalarParameter::StaticClass());
+			DefaultValueHandles.Emplace(DefaultValueHandle);
+
+			if (DefaultValueHandle.IsValid() && DefaultValueHandle->IsValidHandle())
+			{
+				if (!ScalarParameter->bUseCustomPrimitiveData && ScalarParameter->ControlType == EMaterialScalarParameterControlType::Enumeration && !ScalarParameter->Enumeration.IsNull())
+				{
+					// Display enumeration widget for DefaultValue
+					IDetailPropertyRow& PropertyRow = *DetailLayout.EditDefaultProperty(DefaultValueHandle);
+					PropertyRow.CustomWidget()
+					.NameContent()
+					[
+						DefaultValueHandle->CreatePropertyNameWidget()
+					]
+					.ValueContent()
+					[
+						FMaterialPropertyHelpers::CreateScalarEnumerationParameterValueWidget(ScalarParameter->Enumeration, DefaultValueHandle)
+					];
+				}
+			}
 
 			TSharedPtr<IPropertyHandle> SliderMinHandle = DetailLayout.GetProperty("SliderMin", UMaterialExpressionScalarParameter::StaticClass());
 
@@ -137,14 +157,25 @@ void FMaterialExpressionParameterDetails::CustomizeDetails( IDetailLayoutBuilder
 
 			OnSliderMinMaxEdited();
 
-			TSharedPtr<IPropertyHandle> PropertyHandle = DetailLayout.GetProperty("bUseCustomPrimitiveData", UMaterialExpressionScalarParameter::StaticClass());
-			if (PropertyHandle.IsValid() && PropertyHandle->IsValidHandle())
+			// Some properties require a layout rebuild on edit.
+			auto ForceRefresh = FSimpleDelegate::CreateLambda([&DetailLayout]() { DetailLayout.ForceRefreshDetails(); });
+
+			TSharedPtr<IPropertyHandle> ControlTypeHandle = DetailLayout.GetProperty("ControlType", UMaterialExpressionScalarParameter::StaticClass());
+			if (ControlTypeHandle.IsValid() && ControlTypeHandle->IsValidHandle())
 			{
-				// Rebuild the layout when the bUseCustomPrimitiveData property changes
-				PropertyHandle->SetOnPropertyValueChanged(FSimpleDelegate::CreateLambda([&DetailLayout]()
-				{
-					DetailLayout.ForceRefreshDetails();
-				}));
+				ControlTypeHandle->SetOnPropertyValueChanged(ForceRefresh);
+			}
+
+			TSharedPtr<IPropertyHandle> EnumerationHandle = DetailLayout.GetProperty("Enumeration", UMaterialExpressionScalarParameter::StaticClass());
+			if (EnumerationHandle.IsValid() && EnumerationHandle->IsValidHandle())
+			{
+				EnumerationHandle->SetOnPropertyValueChanged(ForceRefresh);
+			}
+
+			TSharedPtr<IPropertyHandle> UseCustomPrimitiveDataPropertyHandle = DetailLayout.GetProperty("bUseCustomPrimitiveData", UMaterialExpressionScalarParameter::StaticClass());
+			if (UseCustomPrimitiveDataPropertyHandle.IsValid() && UseCustomPrimitiveDataPropertyHandle->IsValidHandle())
+			{
+				UseCustomPrimitiveDataPropertyHandle->SetOnPropertyValueChanged(ForceRefresh);
 			}
 
 			if (ScalarParameter->IsUsedAsAtlasPosition())
@@ -182,20 +213,19 @@ void FMaterialExpressionParameterDetails::CustomizeDetails( IDetailLayoutBuilder
 						TSharedPtr<IPropertyHandle> AtlasHandle = DetailLayout.GetProperty("Atlas", UMaterialExpressionCurveAtlasRowParameter::StaticClass());
 						if (AtlasHandle.IsValid() && AtlasHandle->IsValidHandle())
 						{
-							// Rebuild the layout when the bUseCustomPrimitiveData property changes
-							AtlasHandle->SetOnPropertyValueChanged(FSimpleDelegate::CreateLambda([&DetailLayout]()
-							{
-								DetailLayout.ForceRefreshDetails();
-							}));
+							AtlasHandle->SetOnPropertyValueChanged(ForceRefresh);
 						}
 
 				}
+
+				DetailLayout.HideCategory(NAME_UI);
+				DetailLayout.HideCategory(TEXT("CustomPrimitiveData"));
 			}
 
 			if (ScalarParameter->bUseCustomPrimitiveData)
 			{
-				DetailLayout.HideCategory(TEXT("MaterialExpressionScalarParameter"));
 				DetailLayout.HideCategory(MaterialExpressionCategory);
+				DetailLayout.HideCategory(NAME_UI);
 			}
 		}
 
@@ -477,6 +507,7 @@ void FMaterialExpressionParameterDetails::CustomizeDetails( IDetailLayoutBuilder
 				[
 					SAssignNew(NewEditBox, SEditableText)
 					.Text(this, &FMaterialExpressionParameterDetails::OnGetText)
+					.Font(IPropertyTypeCustomizationUtils::GetRegularFont())
 					.OnTextCommitted(this, &FMaterialExpressionParameterDetails::OnTextCommitted)
 				]
 				.MenuContent()
diff --git a/Engine/Source/Editor/MaterialEditor/Private/MaterialEditorInstanceDetailCustomization.cpp b/Engine/Source/Editor/MaterialEditor/Private/MaterialEditorInstanceDetailCustomization.cpp
--- a/Engine/Source/Editor/MaterialEditor/Private/MaterialEditorInstanceDetailCustomization.cpp
+++ b/Engine/Source/Editor/MaterialEditor/Private/MaterialEditorInstanceDetailCustomization.cpp
@@ -363,6 +363,19 @@ void FMaterialInstanceParameterDetails::CustomizeDetails(IDetailLayoutBuilder& D
 				// Append the nanite material override.
 				MaterialPropertyOverrideGroup.AddPropertyRow(DetailLayout.GetProperty("NaniteOverrideMaterial"));
 			}
+
+			{
+				// Refresh UI when enumeration objects change.
+				IDetailPropertyRow& PropertyRow = DefaultCategory.AddProperty("EnumerationObjects");
+				TSharedPtr<IPropertyHandle> PropertyHandle = PropertyRow.GetPropertyHandle();
+				if (PropertyHandle->IsValidHandle())
+				{
+					// Needs to happen on property (add to array) and child (array entries).
+					auto ForceRefresh = FSimpleDelegate::CreateLambda([&DetailLayout]() { DetailLayout.ForceRefreshDetails(); });
+					PropertyHandle->SetOnPropertyValueChanged(ForceRefresh);
+					PropertyHandle->SetOnChildPropertyValueChanged(ForceRefresh);
+				}
+			}
 		}
 
 		// Add the preview mesh property directly from the material instance 
@@ -764,6 +777,33 @@ void FMaterialInstanceParameterDetails::CreateParameterValueWidget(UDEditorParam
 				}
 			}
 		}
+
+		// Scalar parameters can present an enumeration combo box.
+		if (UDEditorScalarParameterValue* ScalarParam = Cast<UDEditorScalarParameterValue>(Parameter))
+		{
+			TSoftObjectPtr<UObject> Enumeration;
+			if (!ScalarParam->Enumeration.IsNull())
+			{
+				Enumeration = ScalarParam->Enumeration;
+			}
+			else if (ScalarParam->EnumerationIndex >= 0)
+			{
+				Enumeration = MaterialEditorInstance->GetEnumerationObject(ScalarParam->EnumerationIndex);
+			}
+
+			if (!Enumeration.IsNull())
+			{
+				PropertyRow.CustomWidget()
+				.NameContent()
+				[
+					ParameterValueProperty->CreatePropertyNameWidget()
+				]
+				.ValueContent()
+				[
+					FMaterialPropertyHelpers::CreateScalarEnumerationParameterValueWidget(Enumeration, ParameterValueProperty)
+				];
+			}
+		}
 	}
 }
 
diff --git a/Engine/Source/Editor/MaterialEditor/Private/MaterialPropertyHelpers.cpp b/Engine/Source/Editor/MaterialEditor/Private/MaterialPropertyHelpers.cpp
--- a/Engine/Source/Editor/MaterialEditor/Private/MaterialPropertyHelpers.cpp
+++ b/Engine/Source/Editor/MaterialEditor/Private/MaterialPropertyHelpers.cpp
@@ -11,6 +11,7 @@
 #include "Widgets/Text/STextBlock.h"
 #include "Widgets/SToolTip.h"
 #include "Styling/AppStyle.h"
+#include "Materials/MaterialEnumeration.h"
 #include "Materials/MaterialInterface.h"
 #include "MaterialEditor/DEditorFontParameterValue.h"
 #include "MaterialEditor/DEditorMaterialLayersParameterValue.h"
@@ -1007,6 +1008,82 @@ void FMaterialPropertyHelpers::ResetCurveToDefault(UDEditorParameterValue* Param
 	}
 }
 
+TSharedRef<SWidget> FMaterialPropertyHelpers::CreateScalarEnumerationParameterValueWidget(TSoftObjectPtr<UObject> const& SoftEnumeration, TSharedPtr<IPropertyHandle> ParameterValueProperty)
+{
+	// Load the enumeration object before use.
+	// We support two object types: UEnum and IMaterialEnumerationProvider.
+	UObject* Enumeration = SoftEnumeration.LoadSynchronous();
+	TWeakObjectPtr<UObject> WeakEnumeration(Enumeration);
+
+	FOnGetPropertyComboBoxStrings GetStrings = FOnGetPropertyComboBoxStrings::CreateLambda([WeakEnumeration](TArray<TSharedPtr<FString>>& OutStrings, TArray<TSharedPtr<SToolTip>>&, TArray<bool>&)
+	{
+		UObject* Enumeration = WeakEnumeration.Get();
+		if (UEnum* Enum = Cast<UEnum>(Enumeration))
+		{
+			for (int32 EnumIndex = 0; EnumIndex < Enum->NumEnums(); ++EnumIndex)
+			{
+				OutStrings.Emplace(MakeShared<FString>(Enum->GetDisplayNameTextByIndex(EnumIndex).ToString()));
+			}
+		}
+		else if (IMaterialEnumerationProvider* EnumProvider = Cast<IMaterialEnumerationProvider>(Enumeration))
+		{
+			EnumProvider->ForEachEntry([&OutStrings](FName InEntryName, int32 InEntryValue)
+			{
+				OutStrings.Emplace(MakeShared<FString>(InEntryName.ToString()));
+			});
+		}
+	});
+
+	FOnGetPropertyComboBoxValue GetValue = FOnGetPropertyComboBoxValue::CreateLambda([WeakEnumeration, ParameterValueProperty]()
+	{
+		FString EntryName = TEXT("None");
+
+		float Value = 0.f;
+		ParameterValueProperty->GetValue(Value);
+
+		UObject* Enumeration = WeakEnumeration.Get();
+		if (UEnum* Enum = Cast<UEnum>(Enumeration))
+		{
+			EntryName = Enum->GetDisplayNameTextByValue((int64)Value).ToString();
+		}
+		else if (IMaterialEnumerationProvider* EnumProvider = Cast<IMaterialEnumerationProvider>(Enumeration))
+		{
+			EnumProvider->ForEachEntry([Value, &EntryName](FName InEntryName, int32 InEntryValue)
+			{
+				if (InEntryValue == (int32)Value)
+				{
+					EntryName = InEntryName.ToString();
+				}
+			});
+		}
+		
+		return EntryName;
+	});
+
+	FOnPropertyComboBoxValueSelected SetValue = FOnPropertyComboBoxValueSelected::CreateLambda([WeakEnumeration, ParameterValueProperty](const FString& StringValue)
+	{
+		UObject* Enumeration = WeakEnumeration.Get();
+		if (UEnum* Enum = Cast<UEnum>(Enumeration))
+		{
+			for (int32 EnumIndex = 0; EnumIndex < Enum->NumEnums(); ++EnumIndex)
+			{
+				if (Enum->GetDisplayNameTextByIndex(EnumIndex).ToString() == StringValue)
+				{
+					const float Value = Enum->GetValueByIndex(EnumIndex);
+					ParameterValueProperty->SetValue(Value);
+					break;
+				}
+			}
+		}
+		else if (IMaterialEnumerationProvider* EnumProvider = Cast<IMaterialEnumerationProvider>(Enumeration))
+		{
+			const float Value = EnumProvider->GetValueOrDefault(*StringValue);
+			ParameterValueProperty->SetValue(Value);
+		}
+	});
+
+	return PropertyCustomizationHelpers::MakePropertyComboBox(ParameterValueProperty, GetStrings, GetValue, SetValue);
+}
 
 #undef LOCTEXT_NAMESPACE
 
diff --git a/Engine/Source/Editor/MaterialEditor/Public/MaterialPropertyHelpers.h b/Engine/Source/Editor/MaterialEditor/Public/MaterialPropertyHelpers.h
--- a/Engine/Source/Editor/MaterialEditor/Public/MaterialPropertyHelpers.h
+++ b/Engine/Source/Editor/MaterialEditor/Public/MaterialPropertyHelpers.h
@@ -223,6 +223,8 @@ public:
 	static UE_API void ResetCurveToDefault(UDEditorParameterValue* Parameter, UMaterialEditorInstanceConstant* MaterialEditorInstance);
 	static void ResetCurveToDefault(TObjectPtr<UDEditorParameterValue> Parameter, UMaterialEditorInstanceConstant* MaterialEditorInstance) { ResetCurveToDefault(Parameter.Get(), MaterialEditorInstance); }
 
+	static  TSharedRef<SWidget> CreateScalarEnumerationParameterValueWidget(TSoftObjectPtr<UObject> const& SoftEnumeration, TSharedPtr<IPropertyHandle> ParameterValueProperty);
+
 	static UE_API FText LayerID;
 	static UE_API FText BlendID;
 	static UE_API FName LayerParamName;
diff --git a/Engine/Source/Editor/UnrealEd/Classes/MaterialEditor/DEditorScalarParameterValue.h b/Engine/Source/Editor/UnrealEd/Classes/MaterialEditor/DEditorScalarParameterValue.h
--- a/Engine/Source/Editor/UnrealEd/Classes/MaterialEditor/DEditorScalarParameterValue.h
+++ b/Engine/Source/Editor/UnrealEd/Classes/MaterialEditor/DEditorScalarParameterValue.h
@@ -32,14 +32,16 @@ public:
 	UPROPERTY(EditAnywhere, Category=DEditorScalarParameterValue)
 	float ParameterValue;
 
-	float SliderMin=0.0f;
-	float SliderMax = 0.0f;
-
 	UPROPERTY(Transient)
 	FScalarParameterAtlasData AtlasData;
 
 	UPROPERTY(Transient)
-	bool bUseCustomPrimitiveData;
+	TSoftObjectPtr<UObject> Enumeration;
+
+	float SliderMin = 0.0f;
+	float SliderMax = 0.0f;
+	int32 EnumerationIndex = INDEX_NONE;
+	bool bUseCustomPrimitiveData = false;
 
 	virtual FName GetDefaultGroupName() const override { return TEXT("Scalar Parameter Values"); }
 
diff --git a/Engine/Source/Editor/UnrealEd/Classes/MaterialEditor/MaterialEditorInstanceConstant.h b/Engine/Source/Editor/UnrealEd/Classes/MaterialEditor/MaterialEditorInstanceConstant.h
--- a/Engine/Source/Editor/UnrealEd/Classes/MaterialEditor/MaterialEditorInstanceConstant.h
+++ b/Engine/Source/Editor/UnrealEd/Classes/MaterialEditor/MaterialEditorInstanceConstant.h
@@ -373,6 +373,15 @@ private:
 	UPROPERTY(EditAnywhere, Category = MaterialEditorInstanceConstant, meta = (editcondition = "bNaniteOverride"))
 	TObjectPtr<UMaterialInterface> NaniteOverrideMaterial;
 
+	/** Array of enumeration objects for use by scalar parameter enumeration indices. */
+	UPROPERTY(EditAnywhere, Category = MaterialEditorInstanceConstant, meta = (DisplayThumbnail = "false", AllowedClasses = "/Script/CoreUObject.Enum, /Script/Engine.MaterialEnumerationProvider"))
+	TArray<TSoftObjectPtr<UObject>> EnumerationObjects;
+
+#if WITH_EDITOR
+	/** Get enumeration object by index. If nothing is found at the index we also look at parent chain until we find something. Returns a null pointer if nothing is found. */
+	UNREALED_API TSoftObjectPtr<UObject> GetEnumerationObject(int32 Index) const;
+#endif
+
 	/** Overrides specific to Post Process domain materials. */
 	UPROPERTY(EditAnywhere, Category = PostProcessOverrides)
 	FMaterialEditorPostProcessOverrides PostProcessOverrides;
diff --git a/Engine/Source/Editor/UnrealEd/Private/EditorParameters.cpp b/Engine/Source/Editor/UnrealEd/Private/EditorParameters.cpp
--- a/Engine/Source/Editor/UnrealEd/Private/EditorParameters.cpp
+++ b/Engine/Source/Editor/UnrealEd/Private/EditorParameters.cpp
@@ -121,6 +121,8 @@ static UDEditorParameterValue* CreateParameter_Scalar(UObject* Owner, const FMat
 		Parameter->ParameterValue = Meta.Value.AsScalar();
 		Parameter->SliderMin = Meta.ScalarMin;
 		Parameter->SliderMax = Meta.ScalarMax;
+		Parameter->Enumeration = Meta.ScalarEnumeration;
+		Parameter->EnumerationIndex = Meta.ScalarEnumerationIndex;
 		Parameter->AtlasData.bIsUsedAsAtlasPosition = Meta.bUsedAsAtlasPosition;
 		Parameter->AtlasData.Atlas = Meta.ScalarAtlas;
 		Parameter->AtlasData.Curve = Meta.ScalarCurve;
diff --git a/Engine/Source/Editor/UnrealEd/Private/PreviewMaterial.cpp b/Engine/Source/Editor/UnrealEd/Private/PreviewMaterial.cpp
--- a/Engine/Source/Editor/UnrealEd/Private/PreviewMaterial.cpp
+++ b/Engine/Source/Editor/UnrealEd/Private/PreviewMaterial.cpp
@@ -893,6 +893,19 @@ void UMaterialEditorInstanceConstant::ClearInvalidParameterOverrides()
 	}
 }
 
+TSoftObjectPtr<UObject> UMaterialEditorInstanceConstant::GetEnumerationObject(int32 Index) const
+{
+	if (EnumerationObjects.IsValidIndex(Index) && !EnumerationObjects[Index].IsNull())
+	{
+		return EnumerationObjects[Index];
+	}
+	if (UMaterialInstance* ParentInstance = Cast<UMaterialInstance>(Parent))
+	{
+		return ParentInstance->GetEnumerationObject(Index);
+	}
+	return TSoftObjectPtr<UObject>();
+}
+
 #endif
 
 void UMaterialEditorInstanceConstant::CopyToSourceInstance(const bool bForceStaticPermutationUpdate)
@@ -939,6 +952,7 @@ void UMaterialEditorInstanceConstant::CopyToSourceInstance(const bool bForceStat
 		SourceInstance->PhysMaterial = PhysMaterial;
 		SourceInstance->NaniteOverrideMaterial.bEnableOverride = bNaniteOverride;
 		SourceInstance->NaniteOverrideMaterial.OverrideMaterialEditor = NaniteOverrideMaterial;
+		SourceInstance->EnumerationObjects = EnumerationObjects;
 
 		// Copy the Lightmass settings...
 		SourceInstance->SetOverrideCastShadowAsMasked(LightmassSettings.CastShadowAsMasked.bOverride);
@@ -1045,6 +1059,7 @@ void UMaterialEditorInstanceConstant::SetSourceInstance(UMaterialInstanceConstan
 	PhysMaterial = SourceInstance->PhysMaterial;
 	bNaniteOverride = SourceInstance->NaniteOverrideMaterial.bEnableOverride;
 	NaniteOverrideMaterial = SourceInstance->NaniteOverrideMaterial.OverrideMaterialEditor;
+	EnumerationObjects = SourceInstance->EnumerationObjects;
 
 	CopyBasePropertiesFromParent();
 
diff --git a/Engine/Source/Runtime/Engine/Private/Materials/MaterialCachedData.cpp b/Engine/Source/Runtime/Engine/Private/Materials/MaterialCachedData.cpp
--- a/Engine/Source/Runtime/Engine/Private/Materials/MaterialCachedData.cpp
+++ b/Engine/Source/Runtime/Engine/Private/Materials/MaterialCachedData.cpp
@@ -206,6 +206,8 @@ bool FMaterialCachedExpressionData::AddParameter(const FMaterialParameterInfo& P
 		case EMaterialParameterType::Scalar:
 			ScalarValues.Insert(ParameterMeta.Value.AsScalar(), Index);
 			EditorOnlyData->ScalarMinMaxValues.Insert(FVector2D(ParameterMeta.ScalarMin, ParameterMeta.ScalarMax), Index);
+			EditorOnlyData->ScalarEnumerationValues.Insert(ParameterMeta.ScalarEnumeration, Index);
+			EditorOnlyData->ScalarEnumerationIndexValues.Insert(ParameterMeta.ScalarEnumerationIndex, Index);
 			ScalarPrimitiveDataIndexValues.Insert(ParameterMeta.PrimitiveDataIndex, Index);
 			if (ParameterMeta.bUsedAsAtlasPosition)
 			{
@@ -879,6 +881,8 @@ void FMaterialCachedExpressionData::GetParameterValueByIndex(EMaterialParameterT
 		{
 			OutResult.ScalarMin = EditorOnlyData->ScalarMinMaxValues[ParameterIndex].X;
 			OutResult.ScalarMax = EditorOnlyData->ScalarMinMaxValues[ParameterIndex].Y;
+			OutResult.ScalarEnumeration = EditorOnlyData->ScalarEnumerationValues[ParameterIndex];
+			OutResult.ScalarEnumerationIndex = EditorOnlyData->ScalarEnumerationIndexValues[ParameterIndex];
 			{
 				const TSoftObjectPtr<UCurveLinearColor>& Curve = EditorOnlyData->ScalarCurveValues[ParameterIndex];
 				const TSoftObjectPtr<UCurveLinearColorAtlas>& Atlas = EditorOnlyData->ScalarCurveAtlasValues[ParameterIndex];
diff --git a/Engine/Source/Runtime/Engine/Private/Materials/MaterialExpressions.cpp b/Engine/Source/Runtime/Engine/Private/Materials/MaterialExpressions.cpp
--- a/Engine/Source/Runtime/Engine/Private/Materials/MaterialExpressions.cpp
+++ b/Engine/Source/Runtime/Engine/Private/Materials/MaterialExpressions.cpp
@@ -8692,6 +8692,56 @@ void UMaterialExpressionScalarParameter::GetCaption(TArray<FString>& OutCaptions
 	 OutCaptions.Add(FString::Printf(TEXT("'%s'"), *ParameterName.ToString())); 
 }
 
+bool UMaterialExpressionScalarParameter::GetParameterValue(FMaterialParameterMetadata& OutMeta) const
+{
+	OutMeta.Value = DefaultValue;
+	OutMeta.PrimitiveDataIndex = INDEX_NONE;
+	OutMeta.ScalarMin = 0.f;
+	OutMeta.ScalarMin = 0.f;
+	OutMeta.ScalarEnumeration = nullptr;
+	OutMeta.ScalarEnumerationIndex = INDEX_NONE;
+
+	if (bUseCustomPrimitiveData)
+	{
+		OutMeta.PrimitiveDataIndex = PrimitiveDataIndex;
+	}
+	else
+	{
+		if (ControlType == EMaterialScalarParameterControlType::Numeric && SliderMax > SliderMin)
+		{
+			OutMeta.ScalarMin = SliderMin;
+			OutMeta.ScalarMax = SliderMax;
+		}
+		else if (ControlType == EMaterialScalarParameterControlType::Enumeration)
+		{
+			OutMeta.ScalarEnumeration = Enumeration;
+		}
+		else if (ControlType == EMaterialScalarParameterControlType::EnumerationIndex)
+		{
+			OutMeta.ScalarEnumerationIndex = EnumerationIndex;
+		}
+	}
+
+	return Super::GetParameterValue(OutMeta);
+}
+
+bool UMaterialExpressionScalarParameter::SetParameterValue(const FName& Name, const FMaterialParameterMetadata& Meta, EMaterialExpressionSetParameterValueFlags Flags)
+{
+	if (Meta.Value.Type == EMaterialParameterType::Scalar)
+	{
+		if (SetParameterValue(Name, Meta.Value.AsScalar(), Flags))
+		{
+			if (EnumHasAnyFlags(Flags, EMaterialExpressionSetParameterValueFlags::AssignGroupAndSortPriority))
+			{
+				Group = Meta.Group;
+				SortPriority = Meta.SortPriority;
+			}
+			return true;
+		}
+	}
+	return false;
+}
+
 bool UMaterialExpressionScalarParameter::SetParameterValue(FName InParameterName, float InValue, EMaterialExpressionSetParameterValueFlags Flags)
 {
 	if (InParameterName == ParameterName)
diff --git a/Engine/Source/Runtime/Engine/Private/Materials/MaterialInstance.cpp b/Engine/Source/Runtime/Engine/Private/Materials/MaterialInstance.cpp
--- a/Engine/Source/Runtime/Engine/Private/Materials/MaterialInstance.cpp
+++ b/Engine/Source/Runtime/Engine/Private/Materials/MaterialInstance.cpp
@@ -4586,6 +4586,19 @@ void UMaterialInstance::UpdateStaticPermutation(const FStaticParameterSet& NewPa
 	}
 }
 
+TSoftObjectPtr<UObject> UMaterialInstance::GetEnumerationObject(int32 Index) const
+{
+	if (EnumerationObjects.IsValidIndex(Index) && !EnumerationObjects[Index].IsNull())
+	{
+		return EnumerationObjects[Index];
+	}
+	if (UMaterialInstance* ParentInstance = Cast<UMaterialInstance>(Parent))
+	{
+		return ParentInstance->GetEnumerationObject(Index);
+	}
+	return TSoftObjectPtr<UObject>();
+}
+
 void UMaterialInstance::GetReferencedTexturesAndOverrides(TSet<const UTexture*>& InOutTextures) const
 {
 	for (UObject* UsedObject : GetCachedExpressionData().ReferencedTextures)
diff --git a/Engine/Source/Runtime/Engine/Public/MaterialCachedData.h b/Engine/Source/Runtime/Engine/Public/MaterialCachedData.h
--- a/Engine/Source/Runtime/Engine/Public/MaterialCachedData.h
+++ b/Engine/Source/Runtime/Engine/Public/MaterialCachedData.h
@@ -153,6 +153,12 @@ struct FMaterialCachedExpressionEditorOnlyData
 	UPROPERTY()
 	TArray<FVector2D> ScalarMinMaxValues;
 
+	UPROPERTY()
+	TArray<TSoftObjectPtr<UObject>> ScalarEnumerationValues;
+	
+	UPROPERTY()
+	TArray<int32> ScalarEnumerationIndexValues;
+
 	UPROPERTY()
 	TArray<TSoftObjectPtr<UCurveLinearColor>> ScalarCurveValues;
 
diff --git a/Engine/Source/Runtime/Engine/Public/Materials/MaterialEnumeration.h b/Engine/Source/Runtime/Engine/Public/Materials/MaterialEnumeration.h
--- a/Engine/Source/Runtime/Engine/Public/Materials/MaterialEnumeration.h
+++ b/Engine/Source/Runtime/Engine/Public/Materials/MaterialEnumeration.h
@@ -20,7 +20,7 @@ public:
 	virtual void ForEachEntry(TFunctionRef<void (FName Name, int32 Value)> Iterator) const = 0;
 
 	// Helper to get the value for a name, or return default if not found.
-	int32 GetValueOrDefault(FName EntryName, int32 DefaultValue = 0) const;
+	ENGINE_API int32 GetValueOrDefault(FName EntryName, int32 DefaultValue = 0) const;
 };
 
 UINTERFACE(MinimalAPI)
diff --git a/Engine/Source/Runtime/Engine/Public/Materials/MaterialExpressionScalarParameter.h b/Engine/Source/Runtime/Engine/Public/Materials/MaterialExpressionScalarParameter.h
--- a/Engine/Source/Runtime/Engine/Public/Materials/MaterialExpressionScalarParameter.h
+++ b/Engine/Source/Runtime/Engine/Public/Materials/MaterialExpressionScalarParameter.h
@@ -3,86 +3,79 @@
 
 #pragma once
 
-#include "CoreMinimal.h"
-#include "UObject/ObjectMacros.h"
 #include "Materials/MaterialExpressionParameter.h"
 #include "MaterialExpressionScalarParameter.generated.h"
 
-struct FPropertyChangedEvent;
+class IMaterialEnumerationProvider;
+
+UENUM()
+enum class EMaterialScalarParameterControlType
+{
+	/** Use the default numeric entry. */
+	Numeric,
+	/** Use a fixed enumeration combo box. */
+	Enumeration,
+	/** Use an enumeration combo box where the enumeration type is set by the material instance. */
+	EnumerationIndex,
+};
 
 UCLASS(collapsecategories, hidecategories=Object, MinimalAPI)
 class UMaterialExpressionScalarParameter : public UMaterialExpressionParameter
 {
 	GENERATED_UCLASS_BODY()
 
-	UPROPERTY(EditAnywhere, Category=MaterialExpressionScalarParameter, Meta = (ShowAsInputPin = "Primary"))
+	/** The default parameter value. */
+	UPROPERTY(EditAnywhere, Category=MaterialExpressionScalarParameter, meta = (ShowAsInputPin = "Primary"))
 	float DefaultValue;
 
-	UPROPERTY(EditAnywhere, Category=CustomPrimitiveData)
+	/** UI control type used for setting values in the material instance editor. */
+	UPROPERTY(EditAnywhere, Category = UI, meta = (EditCondition = "!bUseCustomPrimitiveData"))
+	EMaterialScalarParameterControlType ControlType = EMaterialScalarParameterControlType::Numeric;
+
+	/** Sets the lower bound for this parameter in the material instance editor. */
+	UPROPERTY(EditAnywhere, Category = UI, meta = (EditCondition = "!bUseCustomPrimitiveData && ControlType == EMaterialScalarParameterControlType::Numeric"))
+	float SliderMin = 0.f;
+
+	/**
+	 * Sets the upper bound for this parameter in the material instance editor.
+	 * The bounds will be disabled if SliderMax <= SliderMin.
+	 */
+	UPROPERTY(EditAnywhere, Category = UI, meta = (EditCondition = "ControlType == EMaterialScalarParameterControlType::Numeric"))
+	float SliderMax = 0.f;
+
+	/** Enumeration object to use in the material instance editor. */
+	UPROPERTY(EditAnywhere, Category = UI, meta = (EditCondition = "ControlType == EMaterialScalarParameterControlType::Enumeration", DisplayThumbnail = "false", AllowedClasses = "/Script/CoreUObject.Enum, /Script/Engine.MaterialEnumerationProvider"))
+	TSoftObjectPtr<UObject> Enumeration;
+
+	/** Index of the enumeration to use from the material instance Enumeration Objects. */
+	UPROPERTY(EditAnywhere, Category = UI, meta = (EditCondition = "ControlType == EMaterialScalarParameterControlType::EnumerationIndex", ClampMin = "0"))
+	uint8 EnumerationIndex = 0;
+
+	/** True if this parameter value provided by custom data on the primitive. */
+	UPROPERTY(EditAnywhere, Category = CustomPrimitiveData)
 	bool bUseCustomPrimitiveData = false;
 
-	UPROPERTY(EditAnywhere, Category=CustomPrimitiveData, meta=(ClampMin="0"))
+	/** The slot index for custom primitive data. */
+	UPROPERTY(EditAnywhere, Category = CustomPrimitiveData, meta = (ClampMin = "0", EditCondition = "bUseCustomPrimitiveData"))
 	uint8 PrimitiveDataIndex = 0;
-	/** 
-	 * Sets the lower bound for the slider on this parameter in the material instance editor. 
-	 */
-	UPROPERTY(EditAnywhere, Category=MaterialExpressionScalarParameter)
-	float SliderMin;
 
-	/** 
-	 * Sets the upper bound for the slider on this parameter in the material instance editor. 
-	 * The slider will be disabled if SliderMax <= SliderMin.
-	 */
-	UPROPERTY(EditAnywhere, Category=MaterialExpressionScalarParameter)
-	float SliderMax;
+#if WITH_EDITOR
+	//~ Begin UObject Interface
+	virtual void PostEditChangeProperty(FPropertyChangedEvent& PropertyChangedEvent) override;
+	//~ End UObject Interface
 
 	//~ Begin UMaterialExpression Interface
-#if WITH_EDITOR
 	virtual int32 Compile(class FMaterialCompiler* Compiler, int32 OutputIndex) override;
 	virtual void GetCaption(TArray<FString>& OutCaptions) const override;
-	virtual bool GetParameterValue(FMaterialParameterMetadata& OutMeta) const override
-	{
-		OutMeta.Value = DefaultValue;
-		
-		if (bUseCustomPrimitiveData)
-		{
-			OutMeta.PrimitiveDataIndex = PrimitiveDataIndex;
-		}
-
-		OutMeta.ScalarMin = SliderMin;
-		OutMeta.ScalarMax = SliderMax;
-		return Super::GetParameterValue(OutMeta);
-	}
-	virtual bool SetParameterValue(const FName& Name, const FMaterialParameterMetadata& Meta, EMaterialExpressionSetParameterValueFlags Flags) override
-	{
-		if (Meta.Value.Type == EMaterialParameterType::Scalar)
-		{
-			if (SetParameterValue(Name, Meta.Value.AsScalar(), Flags))
-			{
-				if (EnumHasAnyFlags(Flags, EMaterialExpressionSetParameterValueFlags::AssignGroupAndSortPriority))
-				{
-					Group = Meta.Group;
-					SortPriority = Meta.SortPriority;
-				}
-				return true;
-			}
-		}
-		return false;
-	}
-#endif
+	virtual bool GetParameterValue(FMaterialParameterMetadata& OutMeta) const override;
+	virtual bool SetParameterValue(const FName& Name, const FMaterialParameterMetadata& Meta, EMaterialExpressionSetParameterValueFlags Flags) override;
+	virtual bool HasClassAndNameCollision(UMaterialExpression* OtherExpression) const override;
+	virtual void ValidateParameterName(const bool bAllowDuplicateName) override;
 	//~ End UMaterialExpression Interface
 
-#if WITH_EDITOR
 	bool SetParameterValue(FName InParameterName, float InValue, EMaterialExpressionSetParameterValueFlags Flags = EMaterialExpressionSetParameterValueFlags::None);
-
-	virtual void PostEditChangeProperty(FPropertyChangedEvent& PropertyChangedEvent) override;
-	virtual void ValidateParameterName(const bool bAllowDuplicateName) override;
-	virtual bool HasClassAndNameCollision(UMaterialExpression* OtherExpression) const override;
 #endif
-
+	
+	/** Used to override behavior in derived curve atlas parameter types. */
 	virtual bool IsUsedAsAtlasPosition() const { return false; }
 };
-
-
-
-
diff --git a/Engine/Source/Runtime/Engine/Public/Materials/MaterialInstance.h b/Engine/Source/Runtime/Engine/Public/Materials/MaterialInstance.h
--- a/Engine/Source/Runtime/Engine/Public/Materials/MaterialInstance.h
+++ b/Engine/Source/Runtime/Engine/Public/Materials/MaterialInstance.h
@@ -840,6 +840,10 @@ public:
 #if WITH_EDITORONLY_DATA
 	UPROPERTY()
 	bool bOverrideBaseProperties_DEPRECATED;
+
+	/** Array of enumeration objects for use by scalar parameter enumeration indices. */ 
+	UPROPERTY()
+	TArray<TSoftObjectPtr<UObject>> EnumerationObjects;
 #endif
 
 	UPROPERTY(EditAnywhere, Category=MaterialInstance)
@@ -1134,6 +1138,11 @@ public:
 	ENGINE_API virtual void GetDependentFunctions(TArray<class UMaterialFunctionInterface*>& DependentFunctions) const override;
 #endif // WITH_EDITORONLY_DATA
 
+#if WITH_EDITOR
+	/** Get enumeration object by index. If nothing is found at the index we also look at parent chain until we find something. Returns a null pointer if nothing is found. */
+	ENGINE_API TSoftObjectPtr<UObject> GetEnumerationObject(int32 Index) const;
+#endif
+
 #if WITH_EDITOR
 	/** Add to the set any texture referenced by expressions, including nested functions, as well as any overrides from parameters. */
 	ENGINE_API virtual void GetReferencedTexturesAndOverrides(TSet<const UTexture*>& InOutTextures) const;
diff --git a/Engine/Source/Runtime/Engine/Public/Materials/MaterialParameters.h b/Engine/Source/Runtime/Engine/Public/Materials/MaterialParameters.h
--- a/Engine/Source/Runtime/Engine/Public/Materials/MaterialParameters.h
+++ b/Engine/Source/Runtime/Engine/Public/Materials/MaterialParameters.h
@@ -476,6 +476,12 @@ struct FMaterialParameterMetadata
 	float ScalarMin = 0.0f;
 	float ScalarMax = 0.0f;
 
+	/** Enumeration object for scalar values to use in the UI */
+	TSoftObjectPtr<UObject> ScalarEnumeration;
+
+	/** Index into material instance enumeration objects for scalar values to use in the UI */
+	int32 ScalarEnumerationIndex = INDEX_NONE;
+
 	/** Used for sorting parameter within the group, in the UI */
 	int32 SortPriority = 0;
 
