Change category of VirtualTexturePrefetchMips field in texture UI.
Means that a couple of related virtual texture options are closer/easier to find.
#rnx

diff --git a/Engine/Source/Runtime/Engine/Classes/Engine/Texture.h b/Engine/Source/Runtime/Engine/Classes/Engine/Texture.h
--- a/Engine/Source/Runtime/Engine/Classes/Engine/Texture.h
+++ b/Engine/Source/Runtime/Engine/Classes/Engine/Texture.h
@@ -1483,7 +1483,7 @@ PRAGMA_ENABLE_DEPRECATION_WARNINGS
 	 * In most cases this should be 0 because the virtual texture GPU feedback is sufficient and will minimize memory usage.
 	 * But CPU prefetching can be used in some cases to avoid late appearance of mips.
 	 */
-	UPROPERTY(EditAnywhere, Category = LevelOfDetail, AdvancedDisplay, meta=(EditCondition = "VirtualTextureStreaming"))
+	UPROPERTY(EditAnywhere, Category = Texture, AdvancedDisplay, meta=(EditCondition = "VirtualTextureStreaming"))
 	uint8 VirtualTexturePrefetchMips = 0;
 
 	/** If the platform supports it, tile the texture when cooking, or keep it linear and tile it when it's actually submitted to the GPU. */
