[NewMaterialTranslator] - Fix a bug where the value plugged into a StaticSwitch expression always evaluated to `true`.

- Em.InputDefaultBool(&Value, DefaultValue) would return a FValueRef
- The FValueRef was passed into `Em.ConstantBool(bool value)` and it happly cast the `FValueRef` to `bool` since `operator bool() const { return Value != nullptr; }` is implemented.
- In other words if a value ref had a valid Value it would convert to `true`.
- Instead we need to get the actually boolean value from the FValueRef, so call ToConstantBool instead

#rnx

diff --git a/Engine/Source/Runtime/Engine/Private/Materials/MaterialExpressionsIR.cpp b/Engine/Source/Runtime/Engine/Private/Materials/MaterialExpressionsIR.cpp
--- a/Engine/Source/Runtime/Engine/Private/Materials/MaterialExpressionsIR.cpp
+++ b/Engine/Source/Runtime/Engine/Private/Materials/MaterialExpressionsIR.cpp
@@ -441,7 +441,7 @@ void UMaterialExpressionStaticComponentMaskParameter::Build(MIR::FEmitter& Em)
 
 void UMaterialExpressionStaticSwitch::Build(MIR::FEmitter& Em)
 {
-	bool bCondition = Em.ConstantBool(Em.InputDefaultBool(&Value, DefaultValue));
+	bool bCondition = Em.ToConstantBool(Em.InputDefaultBool(&Value, DefaultValue));
 	UE_MIR_CHECKPOINT(Em); // Make sure that evaluating bCondition didn't raise an error
 
 	Em.Output(0, Em.Input(bCondition ? &A : &B));
