Fix for static analysis error in NetDriver stats code.

diff --git a/Engine/Source/Runtime/Engine/Private/NetDriver.cpp b/Engine/Source/Runtime/Engine/Private/NetDriver.cpp
--- a/Engine/Source/Runtime/Engine/Private/NetDriver.cpp
+++ b/Engine/Source/Runtime/Engine/Private/NetDriver.cpp
@@ -8552,7 +8552,11 @@ void UNetDriver::UpdateNetworkStats()
 						UNetConnection* Connection = ClientConnections[i];
 
 						NumAllConnections++;
-						NumAllConnections += Connection->Children.Num();
+
+						if (Connection)
+						{
+							NumAllConnections += Connection->Children.Num();
+						}
 
 						if (Connection != nullptr)
 						{
