Asking a world for its instance of a null parameter collection yields a null collection instance.

diff --git a/Engine/Source/Runtime/Engine/Private/World.cpp b/Engine/Source/Runtime/Engine/Private/World.cpp
--- a/Engine/Source/Runtime/Engine/Private/World.cpp
+++ b/Engine/Source/Runtime/Engine/Private/World.cpp
@@ -1819,6 +1819,8 @@ void UWorld::SetupParameterCollectionInstances()
 
 void UWorld::AddParameterCollectionInstance(UMaterialParameterCollection* Collection, bool bUpdateScene)
 {
+	check(Collection);
+
 	QUICK_SCOPE_CYCLE_COUNTER(STAT_World_AddParameterCollectionInstance);
 
 	int32 ExistingIndex = INDEX_NONE;
@@ -1838,6 +1840,11 @@ void UWorld::AddParameterCollectionInstance(UMaterialParameterCollection* Collec
 
 UMaterialParameterCollectionInstance* UWorld::GetParameterCollectionInstance(const UMaterialParameterCollection* Collection) const
 {
+	if (!Collection)
+	{
+		return nullptr;
+	}
+
 	for (int32 InstanceIndex = 0; InstanceIndex < ParameterCollectionInstances.Num(); InstanceIndex++)
 	{
 		if (ParameterCollectionInstances[InstanceIndex]->GetCollection() == Collection)
