Add forceDisableDebugViewModes command line option.
When set to true it disables debug view shaders completely.
This is useful for a specific commandlet case where ODSC is forced on and debug shaders are still requested but they are not actually needed.
In that case enabling the command line option can save 20% of shader compilation time.
#rnx

diff --git a/Engine/Source/Runtime/Engine/Private/DebugViewModeHelpers.cpp b/Engine/Source/Runtime/Engine/Private/DebugViewModeHelpers.cpp
--- a/Engine/Source/Runtime/Engine/Private/DebugViewModeHelpers.cpp
+++ b/Engine/Source/Runtime/Engine/Private/DebugViewModeHelpers.cpp
@@ -67,8 +67,22 @@ bool IsDebugViewShaderModeODSCOnly()
 	return false;
 }
 
+bool SupportDebugViewModes()
+{
+	// Allow command line disable of debug view modes.
+	// This is intended only for a very specific commandlet case where we can demand debug view shaders through ODSC even when they are not actually needed.
+	// Setting this command line avoids _any_ potential overhead from debug shader compilation.
+	static const bool bForceDisableDebugViewModes = FParse::Param(FCommandLine::Get(), TEXT("forceDisableDebugViewModes"));
+	return !bForceDisableDebugViewModes;
+}
+
 bool SupportDebugViewShaderMode(EDebugViewShaderMode ShaderMode, EShaderPlatform Platform)
 {
+	if (!SupportDebugViewModes())
+	{
+		return false;
+	}
+
 	// Quad modes don't compile/work on Metal or Playstation. 
 	const bool bSupportQuadComplexityModes =
 		!FDataDrivenShaderPlatformInfo::GetIsLanguageSony(Platform) && 
diff --git a/Engine/Source/Runtime/Engine/Public/DebugViewModeHelpers.h b/Engine/Source/Runtime/Engine/Public/DebugViewModeHelpers.h
--- a/Engine/Source/Runtime/Engine/Public/DebugViewModeHelpers.h
+++ b/Engine/Source/Runtime/Engine/Public/DebugViewModeHelpers.h
@@ -42,14 +42,17 @@ ENGINE_API const TCHAR* DebugViewShaderModeToString(EDebugViewShaderMode InShade
 #if WITH_DEBUG_VIEW_MODES
 /** Returns true if debug view shaders are only ODSC compiled (in editor or cooked builds). Returns false if they are can also be compiled for editor only shader maps as well as ODSC. */
 ENGINE_API bool IsDebugViewShaderModeODSCOnly();
+/** Returns true if any debug view rendering is supported. Can be used as input for shader cooking. */
+ENGINE_API bool SupportDebugViewModes();
 /** Returns true if a specific debug view supported for a platform. Can be used as input for shader cooking. */
 ENGINE_API bool SupportDebugViewShaderMode(EDebugViewShaderMode ShaderMode, EShaderPlatform Platform);
 /** Returns true if a specific debug view can be displayed at runtime. */
 ENGINE_API bool AllowDebugViewShaderMode(EDebugViewShaderMode ShaderMode, EShaderPlatform Platform, ERHIFeatureLevel::Type FeatureLevel);
 #else
-FORCEINLINE bool IsDebugViewShaderModeODSCOnly() { return false; }
-FORCEINLINE bool SupportDebugViewShaderMode(EDebugViewShaderMode ShaderMode, EShaderPlatform Platform) { return false; }
-FORCEINLINE bool AllowDebugViewShaderMode(EDebugViewShaderMode ShaderMode, EShaderPlatform Platform, ERHIFeatureLevel::Type FeatureLevel) { return false; }
+inline bool IsDebugViewShaderModeODSCOnly() { return false; }
+inline bool SupportDebugViewModes() { return false; }
+inline bool SupportDebugViewShaderMode(EDebugViewShaderMode ShaderMode, EShaderPlatform Platform) { return false; }
+inline bool AllowDebugViewShaderMode(EDebugViewShaderMode ShaderMode, EShaderPlatform Platform, ERHIFeatureLevel::Type FeatureLevel) { return false; }
 #endif
 
 ENGINE_API int32 GetNumActorsInWorld(UWorld* InWorld);
diff --git a/Engine/Source/Runtime/Renderer/Private/DebugViewModeRendering.cpp b/Engine/Source/Runtime/Renderer/Private/DebugViewModeRendering.cpp
--- a/Engine/Source/Runtime/Renderer/Private/DebugViewModeRendering.cpp
+++ b/Engine/Source/Runtime/Renderer/Private/DebugViewModeRendering.cpp
@@ -153,6 +153,11 @@ TRDGUniformBufferRef<FDebugViewModePassUniformParameters> CreateDebugViewModePas
 
 static bool ShouldCompileDebugViewModeShaders(const FMeshMaterialShaderPermutationParameters& Parameters)
 {
+	if (!SupportDebugViewModes())
+	{
+		return false;
+	}
+
 	EShaderPermutationFlags Flags = IsDebugViewShaderModeODSCOnly() ? EShaderPermutationFlags::IsODSCOnly : EShaderPermutationFlags::HasEditorOnlyData | EShaderPermutationFlags::IsODSCOnly;
 	return EnumHasAnyFlags(Parameters.Flags, Flags);
 }
