- Surface option to dynamically scale MaxCount (default true for back compat) or by platform on SoundConcurrency, circumventing need to add custom serialization & deprecation of old field and enable end-user to be more explicit about whether or not behavior should remain static (and reliably queriable at edittime)

#rnx

diff --git a/Engine/Source/Editor/AudioEditor/Private/Factories/SoundConcurrencyFactory.cpp b/Engine/Source/Editor/AudioEditor/Private/Factories/SoundConcurrencyFactory.cpp
--- a/Engine/Source/Editor/AudioEditor/Private/Factories/SoundConcurrencyFactory.cpp
+++ b/Engine/Source/Editor/AudioEditor/Private/Factories/SoundConcurrencyFactory.cpp
@@ -20,6 +20,9 @@ USoundConcurrencyFactory::USoundConcurrencyFactory(const FObjectInitializer& Obj
 
 UObject* USoundConcurrencyFactory::FactoryCreateNew(UClass* Class, UObject* InParent, FName Name, EObjectFlags Flags, UObject* Context, FFeedbackContext* Warn)
 {
-	return NewObject<USoundConcurrency>(InParent, Name, Flags);
-}
+	USoundConcurrency* NewConcurrency = NewObject<USoundConcurrency>(InParent, Name, Flags);
 
+	// bEnableMaxCountPlatformScaling is set to false for back compat, but most use cases will call for this being true.
+	NewConcurrency->Concurrency.SetEnableMaxCountPlatformScaling(true);
+	return NewConcurrency;
+}
diff --git a/Engine/Source/Runtime/Engine/Classes/Sound/SoundConcurrency.h b/Engine/Source/Runtime/Engine/Classes/Sound/SoundConcurrency.h
--- a/Engine/Source/Runtime/Engine/Classes/Sound/SoundConcurrency.h
+++ b/Engine/Source/Runtime/Engine/Classes/Sound/SoundConcurrency.h
@@ -2,7 +2,6 @@
 
 #pragma once
 
-#include "CoreMinimal.h"
 #include "UObject/ObjectMacros.h"
 #include "UObject/Object.h"
 #include "UObject/PerPlatformProperties.h"
@@ -76,12 +75,18 @@ struct FSoundConcurrencySettings
 {
 	GENERATED_BODY()
 
-	/** The max number of allowable concurrent active voices for voices playing in this concurrency group. */
-	UPROPERTY(EditAnywhere, BlueprintReadWrite, Category = Concurrency, meta = (UIMin = "1", ClampMin = "1"))
-	int32 MaxCount = 16;
-
-	UPROPERTY()
-	FPerPlatformInt PlatformMaxCount;
+	/**
+	 * The max number of allowable concurrent active voices for voices playing in this concurrency group.
+	 * Can be mutated at runtime via Blueprint or code (see "Enable MaxCount Platform Scaling" for disablement
+	 * of runtime manipulation, which in turn allows for platform scaling of the given value).
+	 */
+	UPROPERTY(EditAnywhere, BlueprintReadWrite, Category = Concurrency, meta = (
+		DisplayName = "Max Count (Dynamic)",
+		DisplayAfter = "bEnableMaxCountPlatformScaling",
+		EditCondition = "!bEnableMaxCountPlatformScaling",
+		EditConditionHides,
+		UIMin = "1", ClampMin = "1"))
+	int32 MaxCount;
 
 	/* Whether or not to limit the concurrency to per sound owner (i.e. the actor that plays the sound). If the sound doesn't have an owner, it falls back to global concurrency. */
 	UPROPERTY(EditAnywhere, BlueprintReadWrite, Category = Concurrency)
@@ -90,9 +95,15 @@ struct FSoundConcurrencySettings
 	/**
 	 * Whether or not volume scaling can recover volume ducking behavior when concurrency group sounds stop (default scale mode only).
 	 */
-	UPROPERTY(EditAnywhere, BlueprintReadWrite, Category = "Volume Scaling", meta = (DisplayName = "Can Recover", EditCondition = "VolumeScaleMode == EConcurrencyVolumeScaleMode::Default"))
+	UPROPERTY(EditAnywhere, BlueprintReadWrite, Category = "Volume Scaling", meta = (DisplayName = "Can Recover", DisplayAfter = "VolumeScaleMode", EditCondition = "VolumeScaleMode == EConcurrencyVolumeScaleMode::Default", EditConditionHides))
 	uint8 bVolumeScaleCanRelease : 1 = 0;
 
+private:
+	/* If true, MaxCount supports platform scaling, but cannot be dynamically changed at runtime (ex. from Blueprint or Gameplay Code).  If false, MaxCount is dynamically assignable at runtime, but is not platform scalable. */
+	UPROPERTY(EditAnywhere, BlueprintReadOnly, Category = Concurrency, meta = (AllowPrivateAccess))
+	uint8 bEnableMaxCountPlatformScaling : 1 = 0;
+
+public:
 	/** Which concurrency resolution policy to use if max voice count is reached. */
 	UPROPERTY(EditAnywhere, BlueprintReadWrite, Category = Concurrency)
 	TEnumAsByte<EMaxConcurrentResolutionRule::Type> ResolutionRule = EMaxConcurrentResolutionRule::StopFarthestThenOldest;
@@ -102,6 +113,20 @@ struct FSoundConcurrencySettings
 	float RetriggerTime = 0.0f;
 
 private:
+	/**
+	 * The max number of allowable concurrent active voices for voices playing in this concurrency group.
+	 * Scalable per platform or platform group. Cannot be mutated at runtime (Disable "Enable MaxCount Platform
+	 * Scaling" for enablement of MaxCount runtime manipulation).
+	 */
+	UPROPERTY(EditAnywhere, BlueprintReadOnly, Category = Concurrency, meta = (
+		DisplayName = "Max Count (Platform Scaled)",
+		DisplayAfter = "bEnableMaxCountPlatformScaling",
+		EditCondition = "bEnableMaxCountPlatformScaling",
+		EditConditionHides,
+		AllowPrivateAccess,
+		UIMin = "1", ClampMin = "1"))
+	FPerPlatformInt PlatformMaxCount;
+
 	/**
 	 * Ducking factor to apply per older voice instance (generation), which compounds based on scaling mode
 	 * and (optionally) revives them as they stop according to the provided attack/release times.
@@ -125,9 +150,17 @@ public:
 	float VolumeScaleAttackTime = 0.01f;
 
 	/**
-	 * Time taken to recover volume scalar duck.
+	 * Time taken to recover volume scalar duck (default scale mode only).
 	 */
-	UPROPERTY(EditAnywhere, BlueprintReadWrite, Category = "Volume Scaling", meta = (DisplayName = "Recover Time", EditCondition = "bVolumeScaleCanRelease && VolumeScaleMode == EConcurrencyVolumeScaleMode::Default", UIMin = "0.0", ClampMin = "0.0", UIMax = "10.0", ClampMax="1000000.0"))
+	UPROPERTY(EditAnywhere, BlueprintReadWrite, Category = "Volume Scaling", meta = (
+		DisplayName = "Recover Time",
+		DisplayAfter = "bVolumeScaleCanRelease",
+		EditCondition = "bVolumeScaleCanRelease && VolumeScaleMode == EConcurrencyVolumeScaleMode::Default",
+		EditConditionHides,
+		UIMin = "0.0",
+		ClampMin = "0.0",
+		UIMax = "10.0",
+		ClampMax = "1000000.0"))
 	float VolumeScaleReleaseTime = 0.5f;
 
 	/**
@@ -136,18 +169,35 @@ public:
 	UPROPERTY(EditAnywhere, BlueprintReadWrite, Category = "Voice Stealing", meta = (UIMin = "0.0", ClampMin = "0.0", UIMax = "1.0", ClampMax="1000000.0"))
 	float VoiceStealReleaseTime = 0.0f;
 
+	FSoundConcurrencySettings();
+
 	/**
 	  * Returns the MaxCount as a per-platform integer value.
 	  */
 	ENGINE_API int32 GetMaxCount() const;
 
 	/**
-	  * Sets the MaxCount as a per-platform integer value.
-	  * (Optional) PlatformName - If set, applies the given max count for the platform of the given name (editor only). Otherwise, applies value to default.
+	  * Applies the given integer value as the group's MaxCount if passed valid MaxCount (greater than 0) and MaxCount platform scaling is enabled.
+	  * Returns true if successfully set, false if not.
 	  */
-	ENGINE_API bool SetMaxCount(int32 InMaxCount, FName PlatformName = { });
+	ENGINE_API bool SetMaxCount(int32 InMaxCount);
 
-	FSoundConcurrencySettings();
+#if WITH_EDITOR
+	/**
+	 * Enables or disables MaxCount platform scaling.
+	 */
+	ENGINE_API void SetEnableMaxCountPlatformScaling(bool bEnabled);
+
+	/**
+	 * Returns whether or not platform scaling is enabled for MaxCount.
+	 */
+	ENGINE_API bool IsMaxCountPlatformScalingEnabled() const;
+
+	/**
+	  * If MaxCount platform scaling is enabled, applying the given per-platform integer value as the MaxCount.
+	  */
+	ENGINE_API bool SetPlatformMaxCount(FPerPlatformInt InMaxCount);
+#endif // WITH_EDITOR
 
 	/**
 	 * Retrieves the volume scale
@@ -159,18 +209,6 @@ public:
 	  * prior to being removed from a concurrency group.
 	  */
 	ENGINE_API bool IsEvictionSupported() const;
-
-	ENGINE_API void PostSerialize(const FArchive& Ar);
-};
-
-template<>
-struct TStructOpsTypeTraits<FSoundConcurrencySettings>
-	: public TStructOpsTypeTraitsBase2<FSoundConcurrencySettings>
-{
-	enum
-	{
-		WithPostSerialize = true
-	};
 };
 
 
@@ -182,6 +220,13 @@ class USoundConcurrency : public UObject
 public:
 	UPROPERTY(EditAnywhere, BlueprintReadOnly, Category = Settings, meta = (ShowOnlyInnerProperties))
 	FSoundConcurrencySettings Concurrency;
+
+	/**
+	  * Applies the given integer value as the group's MaxCount if passed valid MaxCount (greater than 0) and MaxCount platform scaling is enabled.
+	  * Returns true if successfully set, false if not.
+	  */
+	UFUNCTION(BlueprintCallable, Category = Concurrency)
+	UPARAM(DisplayName = "Success") bool SetMaxCount(UPARAM(DisplayName = "MaxCount") int32 InMaxCount);
 };
 
 /** How the concurrency request is handled by the concurrency manager */
diff --git a/Engine/Source/Runtime/Engine/Private/SoundConcurrency.cpp b/Engine/Source/Runtime/Engine/Private/SoundConcurrency.cpp
--- a/Engine/Source/Runtime/Engine/Private/SoundConcurrency.cpp
+++ b/Engine/Source/Runtime/Engine/Private/SoundConcurrency.cpp
@@ -2,6 +2,7 @@
 
 #include "Sound/SoundConcurrency.h"
 
+#include "Algo/AllOf.h"
 #include "AudioDevice.h"
 
 #include UE_INLINE_GENERATED_CPP_BY_NAME(SoundConcurrency)
@@ -24,8 +25,10 @@ FAutoConsoleVariableRef CVarConcurrencyMinDistanceDelta(
 	ECVF_Default);
 
 
-namespace SoundConcurrency
+namespace SoundConcurrencyPrivate
 {
+	static const int32 DefaultMaxCount = 16;
+
 	void SetSoundDataTarget(const FActiveSound& ActiveSound, FConcurrencySoundData& SoundData, float InTargetVolume, float InLerpTime)
 	{
 #if UE_BUILD_SHIPPING
@@ -48,7 +51,7 @@ namespace SoundConcurrency
 		}
 #endif // UE_BUILD_SHIPPING
 	}
-} // namespace SoundConcurrency
+} // namespace SoundConcurrencyPrivate
 
 
 USoundConcurrency::USoundConcurrency(const FObjectInitializer& ObjectInitializer)
@@ -56,55 +59,81 @@ USoundConcurrency::USoundConcurrency(const FObjectInitializer& ObjectInitializer
 {
 }
 
-FSoundConcurrencySettings::FSoundConcurrencySettings()
-PRAGMA_DISABLE_DEPRECATION_WARNINGS
-	: MaxCount(0)
-PRAGMA_ENABLE_DEPRECATION_WARNINGS
+bool USoundConcurrency::SetMaxCount(int32 InMaxCount)
 {
-	PlatformMaxCount = FPerPlatformInt(16);
+	return Concurrency.SetMaxCount(InMaxCount);
 }
 
+FSoundConcurrencySettings::FSoundConcurrencySettings()
+	: MaxCount(SoundConcurrencyPrivate::DefaultMaxCount)
+	, PlatformMaxCount(SoundConcurrencyPrivate::DefaultMaxCount)
+{
+}
 
 int32 FSoundConcurrencySettings::GetMaxCount() const
 {
-PRAGMA_DISABLE_DEPRECATION_WARNINGS
-	if (MaxCount > 0)
+	if (bEnableMaxCountPlatformScaling)
+	{
+		if (const int32 Count = PlatformMaxCount.GetValue(); Count > 0)
+		{
+			return Count;
+		}
+	}
+	else
 	{
-		return MaxCount;
+		if (MaxCount > 0)
+		{
+			return MaxCount;
+		}
 	}
-PRAGMA_ENABLE_DEPRECATION_WARNINGS
 
-	return PlatformMaxCount.GetValue();
+	return SoundConcurrencyPrivate::DefaultMaxCount;
 }
 
-bool FSoundConcurrencySettings::SetMaxCount(int32 InMaxCount, FName PlatformName)
+bool FSoundConcurrencySettings::SetMaxCount(int32 InMaxCount)
 {
-	if (InMaxCount > 0)
+	if (InMaxCount <= 0 || bEnableMaxCountPlatformScaling)
 	{
-PRAGMA_DISABLE_DEPRECATION_WARNINGS
-		MaxCount = 0;
-PRAGMA_ENABLE_DEPRECATION_WARNINGS
+		return false;
+	}
+
+	MaxCount = InMaxCount;
+	return true;
+}
+
+#if WITH_EDITOR
+void FSoundConcurrencySettings::SetEnableMaxCountPlatformScaling(bool bEnabled)
+{
+	bEnableMaxCountPlatformScaling = bEnabled ? 1 : 0;
+}
+
+bool FSoundConcurrencySettings::IsMaxCountPlatformScalingEnabled() const
+{
+	return bEnableMaxCountPlatformScaling > 0;
+}
+
+bool FSoundConcurrencySettings::SetPlatformMaxCount(FPerPlatformInt InMaxCount)
+{
+	bool bIsMaxCountValid = InMaxCount.Default > 0;
 
-		if (PlatformName.IsNone())
-		{
-			PlatformMaxCount.Default = InMaxCount;
-			return true;
-		}
 #if WITH_EDITORONLY_DATA
-		else
-		{
-			PlatformMaxCount.PerPlatform.FindOrAdd(PlatformName) = InMaxCount;
-			return true;
-		}
+	auto IsValidMaxCount = [](const TPair<FName, int32>& Pair) { return Pair.Value > 0; };
+	bIsMaxCountValid &= Algo::AllOf(InMaxCount.PerPlatform, IsValidMaxCount);
 #endif // WITH_EDITORONLY_DATA
+
+	if (bIsMaxCountValid)
+	{
+		PlatformMaxCount = MoveTemp(InMaxCount);
+		return true;
 	}
 
 	return false;
 }
+#endif // WITH_EDITOR
 
 float FSoundConcurrencySettings::GetVolumeScale() const
 {
-	return FMath::Clamp(VolumeScale, 0.0f, 1.0f);;
+	return FMath::Clamp(VolumeScale, 0.0f, 1.0f);
 }
 
 bool FSoundConcurrencySettings::IsEvictionSupported() const
@@ -131,10 +160,6 @@ bool FSoundConcurrencySettings::IsEvictionSupported() const
 	}
 }
 
-void FSoundConcurrencySettings::PostSerialize(const FArchive& Ar)
-{
-}
-
 FConcurrencyHandle::FConcurrencyHandle(const FSoundConcurrencySettings& InSettings)
 	: Settings(InSettings)
 	, ObjectID(0)
@@ -277,7 +302,7 @@ void FConcurrencyGroup::AddActiveSound(FActiveSound& ActiveSound)
 	if (Settings.VolumeScaleMode == EConcurrencyVolumeScaleMode::Default)
 	{
 		SoundData.Generation = ActiveSounds.Num();
-		SoundConcurrency::SetSoundDataTarget(ActiveSound, SoundData, 1.0f, 0.0f);
+		SoundConcurrencyPrivate::SetSoundDataTarget(ActiveSound, SoundData, 1.0f, 0.0f);
 	}
 
 	ActiveSounds.Add(&ActiveSound);
@@ -314,7 +339,7 @@ void FConcurrencyGroup::RemoveActiveSound(FActiveSound& ActiveSound)
 		{
 			const float GenerationDelta = static_cast<float>(ActiveSounds.Num() - OtherSoundData.Generation) - 1.0f;
 			const float NewTargetVolume = FMath::Clamp(FMath::Pow(Settings.GetVolumeScale(), GenerationDelta), 0.0f, 1.0f);
-			SoundConcurrency::SetSoundDataTarget(*OtherSound, OtherSoundData, NewTargetVolume, Settings.VolumeScaleReleaseTime);
+			SoundConcurrencyPrivate::SetSoundDataTarget(*OtherSound, OtherSoundData, NewTargetVolume, Settings.VolumeScaleReleaseTime);
 		}
 	}
 }
@@ -393,7 +418,7 @@ void FConcurrencyGroup::UpdateGeneration(FActiveSound* NewActiveSound)
 
 			if (!FMath::IsNearlyEqual(AttackTime, SoundData->GetLerpTime()) || !FMath::IsNearlyEqual(SoundData->GetTargetVolume(), NewTargetVolume, ConcurrencyMinVolumeScaleCVar))
 			{
-				SoundConcurrency::SetSoundDataTarget(*ActiveSound, *SoundData, NewTargetVolume, AttackTime);
+				SoundConcurrencyPrivate::SetSoundDataTarget(*ActiveSound, *SoundData, NewTargetVolume, AttackTime);
 			}
 		}
 	}
@@ -911,7 +936,7 @@ FActiveSound* FSoundConcurrencyManager::CreateAndEvictActiveSounds(const FActive
 						continue;
 					}
 
-					SoundConcurrency::SetSoundDataTarget(*OtherSound, OtherSoundData, NewTargetVolume, Settings.VolumeScaleAttackTime);
+					SoundConcurrencyPrivate::SetSoundDataTarget(*OtherSound, OtherSoundData, NewTargetVolume, Settings.VolumeScaleAttackTime);
 				}
 			}
 		}
