[NewMaterialTranslator] - Fix crash when a Material Function uses a static bool as a function input.

#rnx

diff --git a/Engine/Source/Runtime/Engine/Private/Materials/MaterialIRTypes.cpp b/Engine/Source/Runtime/Engine/Private/Materials/MaterialIRTypes.cpp
--- a/Engine/Source/Runtime/Engine/Private/Materials/MaterialIRTypes.cpp
+++ b/Engine/Source/Runtime/Engine/Private/Materials/MaterialIRTypes.cpp
@@ -80,6 +80,9 @@ const FType* FType::FromMaterialValueType(EMaterialValueType Type)
 		case MCT_Texture2D:
 			return FObjectType::GetTexture2D();
 
+		case MCT_StaticBool:
+			return FPrimitiveType::GetVector(ScalarKind_Bool, 1);
+
 		case MCT_UInt1:
 			return FPrimitiveType::GetVector(ScalarKind_Int, 1);
 		case MCT_UInt2:
