AssetManager,AssetRegistry: Report an additional set of GetPackageManagers from the AssetManager: GetPackageCookRuleManagers. The existing GetPackageManagers references are a superset of these new GetPackageCookRuleManager, and are used for assigning a chunk to each package. But as of this change, only the managers returned from GetPackageCookRuleManagers are used to assign cookrules to each package (and therefore to decide whether the package gets cooked). The two sets are created by doing two graph searches in UpdateManagementDatabase. During the graph search for all references, we traverse package dependencies that are marked either as Game or Build. During the graph search for CookRule references, we only traverse Game dependencies. This allows types with platform-specific dependencies, such as FMaterialOverrideNanite, to declare their references as EditorOnly Build references and have the references receive the Chunk assignment from the FMaterialOverrideNanite's chunk assignments, but not be added to the cook until and unless the FMaterialOverrideNanite serialization code decides to reference them when it is saved in the cooked package for the TargetPlatform.
#rnx

diff --git a/Engine/Source/Runtime/Engine/Classes/Engine/AssetManager.h b/Engine/Source/Runtime/Engine/Classes/Engine/AssetManager.h
--- a/Engine/Source/Runtime/Engine/Classes/Engine/AssetManager.h
+++ b/Engine/Source/Runtime/Engine/Classes/Engine/AssetManager.h
@@ -667,12 +667,30 @@ public:
 	/** Returns list of asset packages managed by primary asset */
 	ENGINE_API virtual bool GetManagedPackageList(FPrimaryAssetId PrimaryAssetId, TArray<FName>& AssetPackageList) const;
 
-	/** Returns list of PrimaryAssetIds that manage a package. Will optionally recurse up the management chain */
-	ENGINE_API virtual bool GetPackageManagers(FName PackageName, bool bRecurseToParents, TSet<FPrimaryAssetId>& ManagerSet) const;
-	/** Returns PrimaryAssetIds that manage a package, with property describing the reference (direct or indirect). */
+	/**
+	 * Returns PrimaryAssetIds that manage whether a package is included in the cook, with property describing the
+	 * reference (direct or indirect).
+	 */
+	ENGINE_API virtual bool GetPackageCookRuleManagers(FName PackageName, bool bRecurseToParents,
+		TMap<FPrimaryAssetId, UE::AssetRegistry::EDependencyProperty>& Managers) const;
+
+	/**
+	 * Returns PrimaryAssetIds that manage a package in any way, either chunk assignment or inclusion in the cook.
+	 * Chunk assignments are a superset of inclusion in the cook assignments, so this is a supserset of
+	 * GetPackageCookRuleManagers.
+	 * Includes property describing the reference (direct or indirect, and cookrule or not).
+	 */
 	ENGINE_API virtual bool GetPackageManagers(FName PackageName, bool bRecurseToParents,
 		TMap<FPrimaryAssetId, UE::AssetRegistry::EDependencyProperty>& Managers) const;
 
+	/** Returns just the keys of the TMap returned from GetPackageManagers that takes a TMap&. */
+	ENGINE_API virtual bool GetPackageManagers(FName PackageName, bool bRecurseToParents, TSet<FPrimaryAssetId>& ManagerSet) const;
+
+protected:
+	ENGINE_API bool GetPackageManagersInternal(FName PackageName,
+		bool bUseCookRuleReferences, bool bRecurseToParents,
+		TMap<FPrimaryAssetId, UE::AssetRegistry::EDependencyProperty>& Managers) const;
+public:
 
 	// GENERAL ASSET UTILITY FUNCTIONS
 
@@ -1105,8 +1123,10 @@ protected:
 	/** Overridden asset management data for specific types */
 	TMap<FPrimaryAssetId, FPrimaryAssetRulesExplicitOverride> AssetRuleOverrides;
 
-	/** Map from PrimaryAssetId to list of PrimaryAssetIds that are the parent of this one, for determining chunking/cooking */
+	/** Map from PrimaryAssetId to list of PrimaryAssetIds that are the parent of this one, for determining chunking */
 	TMap<FPrimaryAssetId, TArray<FPrimaryAssetId>> ManagementParentMap;
+	/** Map from PrimaryAssetId to list of PrimaryAssetIds that are the parent of this one, for determining what gets cooked */
+	TMap<FPrimaryAssetId, TArray<FPrimaryAssetId>> CookRuleManagementParentMap;
 
 	/** Cached map of asset bundles, global and per primary asset */
 	TMap<FPrimaryAssetId, TSharedPtr<FAssetBundleData, ESPMode::ThreadSafe>> CachedAssetBundles;
diff --git a/Engine/Source/Runtime/Engine/Private/AssetManager.cpp b/Engine/Source/Runtime/Engine/Private/AssetManager.cpp
--- a/Engine/Source/Runtime/Engine/Private/AssetManager.cpp
+++ b/Engine/Source/Runtime/Engine/Private/AssetManager.cpp
@@ -4009,7 +4009,15 @@ bool UAssetManager::GetManagedPackageList(FPrimaryAssetId PrimaryAssetId, TArray
 	return bFoundAny;
 }
 
-bool UAssetManager::GetPackageManagers(FName PackageName, bool bRecurseToParents, TSet<FPrimaryAssetId>& ManagerSet) const
+bool UAssetManager::GetPackageCookRuleManagers(FName PackageName, bool bRecurseToParents,
+	TMap<FPrimaryAssetId, UE::AssetRegistry::EDependencyProperty>& Managers) const
+{
+	return GetPackageManagersInternal(PackageName, true /* bUseCookRuleReferences */,
+		bRecurseToParents, Managers);
+}
+
+bool UAssetManager::GetPackageManagers(FName PackageName, bool bRecurseToParents,
+	TSet<FPrimaryAssetId>& ManagerSet) const
 {
 	TMap<FPrimaryAssetId, UE::AssetRegistry::EDependencyProperty> Managers;
 	bool bFoundAny = GetPackageManagers(PackageName, bRecurseToParents, Managers);
@@ -4022,6 +4030,14 @@ bool UAssetManager::GetPackageManagers(FName PackageName, bool bRecurseToParents
 
 bool UAssetManager::GetPackageManagers(FName PackageName, bool bRecurseToParents,
 	TMap<FPrimaryAssetId, UE::AssetRegistry::EDependencyProperty>& Managers) const
+{
+	return GetPackageManagersInternal(PackageName, false /* bUseCookRuleReferences */,
+		bRecurseToParents, Managers);
+}
+
+bool UAssetManager::GetPackageManagersInternal(FName PackageName,
+	bool bUseCookRuleReferences, bool bRecurseToParents,
+	TMap<FPrimaryAssetId, UE::AssetRegistry::EDependencyProperty>& Managers) const
 {
 	IAssetRegistry& AssetRegistry = GetAssetRegistry();
 	auto UnionManageProperties = [](UE::AssetRegistry::EDependencyProperty A, UE::AssetRegistry::EDependencyProperty B)
@@ -4029,11 +4045,21 @@ bool UAssetManager::GetPackageManagers(FName PackageName, bool bRecurseToParents
 		return A | B;
 	};
 
+	// We have two ManagementParentMaps that are calculated using the two kinds of dependencies we track:
+	// CookRule or Chunking. Pick the corresponding map based on bUseCookRuleReferencess.
+	const TMap<FPrimaryAssetId, TArray<FPrimaryAssetId>>& UsedParentMap = bUseCookRuleReferences ?
+		CookRuleManagementParentMap : ManagementParentMap;
+	UE::AssetRegistry::EDependencyProperty RequiredProperties = bUseCookRuleReferences ?
+		UE::AssetRegistry::EDependencyProperty::CookRule : UE::AssetRegistry::EDependencyProperty::None;
+
 	bool bFoundAny = false;
 	TArray<FAssetDependency> ReferencingPrimaryAssets;
 	ReferencingPrimaryAssets.Reserve(128);
 
-	AssetRegistry.GetReferencers(PackageName, ReferencingPrimaryAssets, UE::AssetRegistry::EDependencyCategory::Manage);
+	UE::AssetRegistry::FDependencyQuery RequiredPropertiesQuery;
+	RequiredPropertiesQuery.Required = RequiredProperties;
+	AssetRegistry.GetReferencers(PackageName, ReferencingPrimaryAssets,
+		UE::AssetRegistry::EDependencyCategory::Manage, RequiredPropertiesQuery);
 
 	for (int32 IdentifierIndex = 0; IdentifierIndex < ReferencingPrimaryAssets.Num(); IdentifierIndex++)
 	{
@@ -4044,15 +4070,17 @@ bool UAssetManager::GetPackageManagers(FName PackageName, bool bRecurseToParents
 			continue;
 		}
 		bFoundAny = true;
-		UE::AssetRegistry::EDependencyProperty& ExistingProperties = Managers.FindOrAdd(PrimaryAssetId, UE::AssetRegistry::EDependencyProperty::None);
+		UE::AssetRegistry::EDependencyProperty& ExistingProperties = Managers.FindOrAdd(
+			PrimaryAssetId, UE::AssetRegistry::EDependencyProperty::None);
 		ExistingProperties = UnionManageProperties(ExistingProperties, AssetDependency.Properties);
 
 		if (bRecurseToParents)
 		{
-			const TArray<FPrimaryAssetId>* ManagementParents = ManagementParentMap.Find(PrimaryAssetId);
+			const TArray<FPrimaryAssetId>* ManagementParents = UsedParentMap.Find(PrimaryAssetId);
 			if (ManagementParents)
 			{
-				// AssetDependency can be invalidated because we are modifying the array it is in. Copy its data now before we modify the array.
+				// AssetDependency can be invalidated because we are modifying the array it is in.
+				// Copy its data now before we modify the array.
 				UE::AssetRegistry::EDependencyCategory IndirectCategory = AssetDependency.Category;
 				UE::AssetRegistry::EDependencyProperty IndirectProperties = AssetDependency.Properties;
 				EnumRemoveFlags(IndirectProperties, UE::AssetRegistry::EDependencyProperty::Direct);
@@ -4359,6 +4387,7 @@ void UAssetManager::UpdateManagementDatabase(EUpdateManagementDatabaseFlags InFl
 	}
 	LLM_SCOPE_BYTAG(AssetManager);
 
+	CookRuleManagementParentMap.Reset();
 	ManagementParentMap.Reset();
 
 	// Make sure the asset labels are up to date 
@@ -4511,7 +4540,9 @@ void UAssetManager::UpdateManagementDatabase(EUpdateManagementDatabaseFlags InFl
 			FPrimaryAssetId PrimaryAssetId(TypePair.Key, NamePair.Key);
 			const FSoftObjectPath& AssetRef = NameData.GetAssetPtr().ToSoftObjectPath();
 
-			TSet<FPrimaryAssetId> Managers;
+			// ChunkManagers are a superset of CookRuleManagers; we rely on that to avoid the need
+			// for querying CookRuleManagers when Managers is empty.
+			TMap<FPrimaryAssetId, UE::AssetRegistry::EDependencyProperty> Managers;
 
 			if (AssetRef.IsValid())
 			{
@@ -4520,26 +4551,33 @@ void UAssetManager::UpdateManagementDatabase(EUpdateManagementDatabaseFlags InFl
 				if (GetPackageManagers(PackageName, false, Managers) && Managers.Num() > 1)
 				{
 					// Find all managers that aren't this specific asset
-					for (const FPrimaryAssetId& Manager : Managers)
+					for (const TPair<FPrimaryAssetId, UE::AssetRegistry::EDependencyProperty>& ManagerPair : Managers)
 					{
-						if (Manager != PrimaryAssetId)
+						if (ManagerPair.Key != PrimaryAssetId)
 						{
 							// Update the cached version and the version in registry
-							ManagementParentMap.FindOrAdd(PrimaryAssetId).AddUnique(Manager);
+							if (EnumHasAnyFlags(ManagerPair.Value, UE::AssetRegistry::EDependencyProperty::CookRule))
+							{
+								CookRuleManagementParentMap.FindOrAdd(PrimaryAssetId).AddUnique(ManagerPair.Key);
+							}
+							ManagementParentMap.FindOrAdd(PrimaryAssetId).AddUnique(ManagerPair.Key);
 
-							PrimaryAssetIdManagementMap.Add(Manager, PrimaryAssetId);
+							PrimaryAssetIdManagementMap.Add(ManagerPair.Key, PrimaryAssetId);
 						}
 					}
 				}
 			}
 			else
 			{
-				Managers.Add(PrimaryAssetId);
+				Managers.Add(PrimaryAssetId, UE::AssetRegistry::EDependencyProperty::ManageMask);
 			}
 
-			// Compute chunk assignment and store those as manager references
+			// Compute chunk assignments for the PrimaryAsset and record them as manage references from a
+			// PrimaryAssetId we create for each chunk
 			ChunkList.Reset();
-			GetPrimaryAssetSetChunkIds(Managers, nullptr, ExistingChunkList, ChunkList);
+			TSet<FPrimaryAssetId> ManagersSet;
+			Managers.GetKeys(ManagersSet);
+			GetPrimaryAssetSetChunkIds(ManagersSet, nullptr, ExistingChunkList, ChunkList);
 
 			for (int32 ChunkId : ChunkList)
 			{
@@ -4799,7 +4837,7 @@ EPrimaryAssetCookRule UAssetManager::GetPackageCookRule(FName PackageName) const
 	TRACE_CPUPROFILER_EVENT_SCOPE(UAssetManager::GetPackageCookRule);
 
 	TMap<FPrimaryAssetId, UE::AssetRegistry::EDependencyProperty> Managers;
-	GetPackageManagers(PackageName, true, Managers);
+	GetPackageCookRuleManagers(PackageName, true, Managers);
 
 	TOptional<TTuple<FPrimaryAssetId, FPrimaryAssetId>> ConflictIds;
 	EPrimaryAssetCookRule CookRule = CalculateCookRuleUnion(Managers, &ConflictIds);
@@ -4994,12 +5032,12 @@ bool UAssetManager::VerifyCanCookPackage(UE::Cook::ICookInfo* CookInfo, FName Pa
 
 	if (!bRetVal && bLogError)
 	{
-		TSet<FPrimaryAssetId> Managers;
-		GetPackageManagers(PackageName, true, Managers);
+		TMap<FPrimaryAssetId, UE::AssetRegistry::EDependencyProperty> Managers;
+		GetPackageCookRuleManagers(PackageName, true, Managers);
 		UE_LOG(LogAssetManager, Display, TEXT("Listing Managers... (Count:%d)"), Managers.Num());
-		for (const FPrimaryAssetId& PrimaryAssetId : Managers)
+		for (const TPair<FPrimaryAssetId, UE::AssetRegistry::EDependencyProperty>& PrimaryAssetPair : Managers)
 		{
-			UE_LOG(LogAssetManager, Display, TEXT("  %s"), *PrimaryAssetId.ToString());
+			UE_LOG(LogAssetManager, Display, TEXT("  %s"), *PrimaryAssetPair.Key.ToString());
 		}
 
 		TArray<FName> PackageReferencers;
@@ -5035,11 +5073,13 @@ bool UAssetManager::GetPackageChunkIds(FName PackageName, const ITargetPlatform*
 	}
 
 	// Add all chunk ids from the asset rules of managers. By default priority will not override other chunks
-	TSet<FPrimaryAssetId> Managers;
+	TMap<FPrimaryAssetId, UE::AssetRegistry::EDependencyProperty> Managers;
 	Managers.Reserve(128);
 
 	GetPackageManagers(PackageName, true, Managers);
-	return GetPrimaryAssetSetChunkIds(Managers, TargetPlatform, ExistingChunkList, OutChunkList);
+	TSet<FPrimaryAssetId> ManagersSet;
+	Managers.GetKeys(ManagersSet);
+	return GetPrimaryAssetSetChunkIds(ManagersSet, TargetPlatform, ExistingChunkList, OutChunkList);
 }
 
 void UAssetManager::UpdateCachedChunkMapAfterCook(const TSet<FName>& AllCookedPackages, const TSet<FName>& StartupPackages)
@@ -5156,7 +5196,9 @@ TMap<int32, FAssetManagerChunkInfo> UAssetManager::BuildChunkMap(const TSet<FNam
 			if (AssetRef.IsValid())
 			{
 				FName PackageName = FName(*AssetRef.GetLongPackageName());
-				GetPackageManagers(PackageName, false, Managers);
+				TMap<FPrimaryAssetId, UE::AssetRegistry::EDependencyProperty> ManagersMap;
+				GetPackageManagers(PackageName, false, ManagersMap);
+				ManagersMap.GetKeys(Managers);
 			}
 			else
 			{
@@ -5214,6 +5256,7 @@ void UAssetManager::ReinitializeFromConfig()
 		}
 	}
 	check(AssetPathMap.IsEmpty()); // Should have been emptied by the ResetAssets calls
+	CookRuleManagementParentMap.Reset();
 	ManagementParentMap.Reset();
 	CachedAssetBundles.Reset();
 	AlreadyScannedDirectories.Reset();
