PCGSplineStruct: Fix Create Spline not working

FSpline was defaulting to a 2 point spline. We were also unconditionally applying FSpline in ApplyTo which is wrong, the source data should be determined by ShouldUseSplineCurves(). These were the 2 bugs.

diff --git a/Engine/Plugins/PCG/Source/PCG/Private/Data/PCGSplineStruct.cpp b/Engine/Plugins/PCG/Source/PCG/Private/Data/PCGSplineStruct.cpp
--- a/Engine/Plugins/PCG/Source/PCG/Private/Data/PCGSplineStruct.cpp
+++ b/Engine/Plugins/PCG/Source/PCG/Private/Data/PCGSplineStruct.cpp
@@ -190,7 +190,16 @@ void FPCGSplineStruct::ApplyTo(USplineComponent* InSplineComponent) const
 	InSplineComponent->DefaultUpVector = DefaultUpVector;
 	InSplineComponent->ReparamStepsPerSegment = ReparamStepsPerSegment;
 
-	InSplineComponent->SetSpline(Spline);
+	// Not using a ternary so we are not doing any implicit conversions (SetSpline is overloaded)
+	if (ShouldUseSplineCurves())
+	{
+		InSplineComponent->SetSpline(WarninglessSplineCurves());
+	}
+	else
+	{
+		InSplineComponent->SetSpline(Spline);
+	}
+	
 	InSplineComponent->bStationaryEndpoints = false;
 	// TODO: metadata? might not be needed
 	InSplineComponent->SetClosedLoop(bClosedLoop);
@@ -246,14 +255,10 @@ FPCGSplineStruct& FPCGSplineStruct::operator=(FPCGSplineStruct&& Other)
 
 void FPCGSplineStruct::AddPoint(const FSplinePoint& InSplinePoint, bool bUpdateSpline)
 {
-	const TArray<FInterpCurvePoint<FVector>>& SourcePoints = ShouldUseSplineCurves()
-		? WarninglessSplineCurves().Position.Points
-		: Spline.GetSplinePointsPosition().Points;
-	
-	const int32 Index = PCGSplineStruct::UpperBound(SourcePoints, InSplinePoint.InputKey);
-
 	Spline.AddPoint(InSplinePoint);
-
+ 
+	const int32 Index = PCGSplineStruct::UpperBound(WarninglessSplineCurves().Position.Points, InSplinePoint.InputKey);
+	
 	WarninglessSplineCurves().Position.Points.Insert(FInterpCurvePoint<FVector>(
 		InSplinePoint.InputKey,
 		InSplinePoint.Position,
diff --git a/Engine/Source/Runtime/Engine/Classes/Curves/Spline.h b/Engine/Source/Runtime/Engine/Classes/Curves/Spline.h
--- a/Engine/Source/Runtime/Engine/Classes/Curves/Spline.h
+++ b/Engine/Source/Runtime/Engine/Classes/Curves/Spline.h
@@ -129,9 +129,6 @@ struct FSpline
 	
 	/** Reset the scale attribute channel to default values. */
 	ENGINE_API void ResetScale();
-	
-	/** Reset the spline to the default spline (2 points). */
-	ENGINE_API void ResetToDefault();
 
 	struct FUpdateSplineParams
 	{
diff --git a/Engine/Source/Runtime/Engine/Private/Components/SplineComponent.cpp b/Engine/Source/Runtime/Engine/Private/Components/SplineComponent.cpp
--- a/Engine/Source/Runtime/Engine/Private/Components/SplineComponent.cpp
+++ b/Engine/Source/Runtime/Engine/Private/Components/SplineComponent.cpp
@@ -142,19 +142,25 @@ bool USplineComponent::CanResetToDefault() const
 
 void USplineComponent::SetDefaultSpline()
 {
-	Spline.ResetToDefault();
+	const FVector StartPoint(0.0, 0.0, 0.0);
+	const FVector EndPoint(100.0, 0.0, 0.0);
+	constexpr float StartParam = 0.f;
+	constexpr float EndParam = 1.f;
 
+	Spline.Reset();
 	WarninglessSplineCurves().Position.Points.Reset(10);
     WarninglessSplineCurves().Rotation.Points.Reset(10);
     WarninglessSplineCurves().Scale.Points.Reset(10);
 
-    WarninglessSplineCurves().Position.Points.Emplace(0.0f, FVector(0, 0, 0), FVector::ZeroVector, FVector::ZeroVector, CIM_CurveAuto);
-    WarninglessSplineCurves().Rotation.Points.Emplace(0.0f, FQuat::Identity, FQuat::Identity, FQuat::Identity, CIM_CurveAuto);
-    WarninglessSplineCurves().Scale.Points.Emplace(0.0f, FVector(1.0f), FVector::ZeroVector, FVector::ZeroVector, CIM_CurveAuto);
+	Spline.AddPoint(FSplinePoint(StartParam, StartPoint));
+    WarninglessSplineCurves().Position.Points.Emplace(StartParam, StartPoint, FVector::ZeroVector, FVector::ZeroVector, CIM_CurveAuto);
+    WarninglessSplineCurves().Rotation.Points.Emplace(StartParam, FQuat::Identity, FQuat::Identity, FQuat::Identity, CIM_CurveAuto);
+    WarninglessSplineCurves().Scale.Points.Emplace(StartParam, FVector(1.0f), FVector::ZeroVector, FVector::ZeroVector, CIM_CurveAuto);
 
-    WarninglessSplineCurves().Position.Points.Emplace(1.0f, FVector(100, 0, 0), FVector::ZeroVector, FVector::ZeroVector, CIM_CurveAuto);
-    WarninglessSplineCurves().Rotation.Points.Emplace(1.0f, FQuat::Identity, FQuat::Identity, FQuat::Identity, CIM_CurveAuto);
-    WarninglessSplineCurves().Scale.Points.Emplace(1.0f, FVector(1.0f), FVector::ZeroVector, FVector::ZeroVector, CIM_CurveAuto);
+	Spline.AddPoint(FSplinePoint(EndParam, EndPoint));
+    WarninglessSplineCurves().Position.Points.Emplace(EndParam, EndPoint, FVector::ZeroVector, FVector::ZeroVector, CIM_CurveAuto);
+    WarninglessSplineCurves().Rotation.Points.Emplace(EndParam, FQuat::Identity, FQuat::Identity, FQuat::Identity, CIM_CurveAuto);
+    WarninglessSplineCurves().Scale.Points.Emplace(EndParam, FVector(1.0f), FVector::ZeroVector, FVector::ZeroVector, CIM_CurveAuto);
 }
 
 void USplineComponent::GetLifetimeReplicatedProps(TArray<FLifetimeProperty>& OutLifetimeProps) const
diff --git a/Engine/Source/Runtime/Engine/Private/Curves/Spline.cpp b/Engine/Source/Runtime/Engine/Private/Curves/Spline.cpp
--- a/Engine/Source/Runtime/Engine/Private/Curves/Spline.cpp
+++ b/Engine/Source/Runtime/Engine/Private/Curves/Spline.cpp
@@ -162,7 +162,6 @@ public:
     virtual void SetClosedLoop(bool bInClosedLoop) override;
 
 	void Reset();
-	void ResetToDefault();
 	
 	void ResetRotation();
 	void ResetScale();
@@ -669,7 +668,7 @@ class FLegacySpline
 	
 public:
 
-	ENGINE_API FLegacySpline();
+	ENGINE_API FLegacySpline() = default;
 	ENGINE_API FLegacySpline(const FNewSpline& Other);
 	ENGINE_API FLegacySpline(const FSplineCurves& Other);
 	
@@ -757,9 +756,6 @@ public:
 	/** Reset the scale attribute channel to default values. */
 	void ResetScale();
 	
-	/** Reset the spline to the default spline (2 points). */
-	void ResetToDefault();
-	
 	/** Update the spline's internal data according to the passed-in params. */
 	void UpdateSpline(const FSpline::FUpdateSplineParams& InParams);
 };
@@ -1551,16 +1547,6 @@ void FSpline::ResetScale()
 	case 2: NewData->ResetScale(); break;
 	}
 }
-	
-void FSpline::ResetToDefault()
-{
-	switch (CurrentImplementation)
-	{
-	default: break;
-	case 1: LegacyData->ResetToDefault(); break;
-	case 2: NewData->ResetToDefault(); break;
-	}
-}
 
 void FSpline::UpdateSpline(const FUpdateSplineParams& InParams)
 {
@@ -1629,8 +1615,6 @@ FNewSpline::FNewSpline()
 	CreateAttributeChannel<TTangentBezierSpline<FVector>>(ScaleAttrName);
 	CreateAttributeChannel<TTangentBezierSpline<float>>(ParamToLengthAttrName);
 	CreateAttributeChannel<TTangentBezierSpline<float>>(NormalizedLengthToParamAttrName);
-    
-    ResetToDefault();
 	
 	ValidateRotScale();
 }
@@ -1848,17 +1832,6 @@ void FNewSpline::Reset()
 	ScaleCurve.Reset();
 }
 
-void FNewSpline::ResetToDefault()
-{
-    GetSpline().SetControlPoints( { FTangentBezierControlPoint(FVector(0.0, 0.0, 0.0)),
-    					FTangentBezierControlPoint(FVector(100.0, 0.0, 0.0)) } );
-
-	ResetRotation();
-	ResetScale();
-	
-    UpdateSpline();
-}
-
 void FNewSpline::ResetRotation()
 {
     ClearAttributeChannel(RotationAttrName);
@@ -2503,11 +2476,6 @@ void FNewSpline::ValidateRotScale() const
 
 /** FLegacySpline Implementation */
 
-FLegacySpline::FLegacySpline()
-{
-	ResetToDefault();
-}
-
 FLegacySpline::FLegacySpline(const FNewSpline& Other)
 	: FLegacySpline()
 {
@@ -2936,21 +2904,6 @@ void FLegacySpline::ResetScale()
 	}
 }
 
-void FLegacySpline::ResetToDefault()
-{
-	PositionCurve.Points.Reset(10);
-	RotationCurve.Points.Reset(10);
-	ScaleCurve.Points.Reset(10);
-
-	PositionCurve.Points.Emplace(0.0f, FVector(0, 0, 0), FVector::ZeroVector, FVector::ZeroVector, CIM_CurveAuto);
-	RotationCurve.Points.Emplace(0.0f, FQuat::Identity, FQuat::Identity, FQuat::Identity, CIM_CurveAuto);
-	ScaleCurve.Points.Emplace(0.0f, FVector(1.0f), FVector::ZeroVector, FVector::ZeroVector, CIM_CurveAuto);
-
-	PositionCurve.Points.Emplace(1.0f, FVector(100, 0, 0), FVector::ZeroVector, FVector::ZeroVector, CIM_CurveAuto);
-	RotationCurve.Points.Emplace(1.0f, FQuat::Identity, FQuat::Identity, FQuat::Identity, CIM_CurveAuto);
-	ScaleCurve.Points.Emplace(1.0f, FVector(1.0f), FVector::ZeroVector, FVector::ZeroVector, CIM_CurveAuto);
-}
-
 void FLegacySpline::UpdateSpline(const FSpline::FUpdateSplineParams& InParams)
 {
 	const int32 NumPoints = PositionCurve.Points.Num();
