[NewMaterialTranslator] - Minor cleanup.

- Fixup comment to make it clearer what the bool operation does on FValueRef.
- When validating FValueRef in `ValidateOperatorAndGetResultType` call `IsValid` instead of bool operator.

#rnx

diff --git a/Engine/Source/Runtime/Engine/Private/Materials/MaterialIREmitter.cpp b/Engine/Source/Runtime/Engine/Private/Materials/MaterialIREmitter.cpp
--- a/Engine/Source/Runtime/Engine/Private/Materials/MaterialIREmitter.cpp
+++ b/Engine/Source/Runtime/Engine/Private/Materials/MaterialIREmitter.cpp
@@ -1723,7 +1723,7 @@ static const FPrimitiveType* ValidateOperatorAndGetResultType(FEmitter& Emitter,
 			// Cast this argument to the first argument primitive type
 			ArgumentPrimitiveType = FirstArgumentPrimitiveType;
 			Arguments[i] = Emitter.Cast(Arguments[i], ArgumentPrimitiveType);
-			bValid &= Arguments[i];
+			bValid &= Arguments[i].IsValid();
 		}
 		else if (ArgumentPrimitiveType->ScalarKind == EScalarKind::ScalarKind_LWC && !(Filter & OPTF_LWCInput))
 		{
@@ -1737,7 +1737,7 @@ static const FPrimitiveType* ValidateOperatorAndGetResultType(FEmitter& Emitter,
 			{
 				ArgumentPrimitiveType = ArgumentPrimitiveType->WithScalarKind(MIR::ScalarKind_Float);
 				Arguments[i] = Emitter.Cast(Arguments[i], ArgumentPrimitiveType);
-				bValid &= Arguments[i];
+				bValid &= Arguments[i].IsValid();
 			}
 		}
 		
diff --git a/Engine/Source/Runtime/Engine/Public/Materials/MaterialIREmitter.h b/Engine/Source/Runtime/Engine/Public/Materials/MaterialIREmitter.h
--- a/Engine/Source/Runtime/Engine/Public/Materials/MaterialIREmitter.h
+++ b/Engine/Source/Runtime/Engine/Public/Materials/MaterialIREmitter.h
@@ -81,7 +81,7 @@ struct FValueRef
 	// Constructs from a value pointer and its corresponding input source.
 	FValueRef(FValue* Value, const FExpressionInput* Input) : Value{ Value }, Input{ Input } {}
 
-	// Returns whether this reference is valid, that is it references to a non-null value.
+	// Returns whether this references a non-null value.
 	operator bool() const { return Value != nullptr; }
 
 	// Returns the referenced value, which must be non-null.
