Skeletal Mesh: Compress cooked CPU-side morph target storage

- Combined MorphMesh.cpp and MorphTools.cpp into a single MorphTarget.cpp file.

[RN] CPU morphs are now stored compressed in the same way that GPU morphs are. This will reduce memory usage on mobile platforms when using morph targets.

diff --git a/Engine/Plugins/Animation/ACLPlugin/Source/ACLPlugin/Private/AnimCurveCompressionCodec_ACL.cpp b/Engine/Plugins/Animation/ACLPlugin/Source/ACLPlugin/Private/AnimCurveCompressionCodec_ACL.cpp
--- a/Engine/Plugins/Animation/ACLPlugin/Source/ACLPlugin/Private/AnimCurveCompressionCodec_ACL.cpp
+++ b/Engine/Plugins/Animation/ACLPlugin/Source/ACLPlugin/Private/AnimCurveCompressionCodec_ACL.cpp
@@ -89,12 +89,9 @@ static TArray<float> GetMorphTargetMaxPositionDeltas(const FCompressibleAnimData
 		if (Target != nullptr)
 		{
 			// This curve drives a morph target, find the largest displacement it can have
-			const int32 LODIndex = 0;
-			int32 NumDeltas = 0;
-			const FMorphTargetDelta* Deltas = Target->GetMorphTargetDelta(LODIndex, NumDeltas);
-			for (int32 DeltaIndex = 0; DeltaIndex < NumDeltas; ++DeltaIndex)
+			constexpr int32 LODIndex = 0;
+			for (const FMorphTargetDelta& Delta: Target->GetMorphTargetDeltas(LODIndex))
 			{
-				const FMorphTargetDelta& Delta = Deltas[DeltaIndex];
 				MaxDeltaPosition = FMath::Max(MaxDeltaPosition, Delta.PositionDelta.Size());
 			}
 		}
diff --git a/Engine/Plugins/Animation/MLDeformer/MLDeformerFramework/Source/MLDeformerFrameworkEditor/Private/MLDeformerMorphModelEditorModel.cpp b/Engine/Plugins/Animation/MLDeformer/MLDeformerFramework/Source/MLDeformerFrameworkEditor/Private/MLDeformerMorphModelEditorModel.cpp
--- a/Engine/Plugins/Animation/MLDeformer/MLDeformerFramework/Source/MLDeformerFrameworkEditor/Private/MLDeformerMorphModelEditorModel.cpp
+++ b/Engine/Plugins/Animation/MLDeformer/MLDeformerFramework/Source/MLDeformerFrameworkEditor/Private/MLDeformerMorphModelEditorModel.cpp
@@ -219,15 +219,13 @@ namespace UE::MLDeformer
 			const float MaxWeight = !MinMaxMorphWeights.IsEmpty() ? FMath::Max(FMath::Abs(MinMaxMorphWeights[MorphIndex].Min), FMath::Abs(MinMaxMorphWeights[MorphIndex].Min)) : 1.0f;
 
 			// Get the array of deltas.
-			int32 NumDeltas = 0;
-			const FMorphTargetDelta* Deltas = MorphTarget->GetMorphTargetDelta(LOD, NumDeltas);
+			TConstArrayView<FMorphTargetDelta> Deltas = MorphTarget->GetMorphTargetDeltas(LOD);
 
 			// Build the array of position delta lengths.
-			DeltaLengths.Reset();
-			DeltaLengths.SetNumUninitialized(NumDeltas);
-			for (int32 DeltaIndex = 0; DeltaIndex < NumDeltas; ++DeltaIndex)
+			DeltaLengths.Reset(Deltas.Num());
+			for (const FMorphTargetDelta& VertexDelta: Deltas)
 			{
-				DeltaLengths[DeltaIndex] = Deltas[DeltaIndex].PositionDelta.Length() * MaxWeight;
+				DeltaLengths.Add(VertexDelta.PositionDelta.Length() * MaxWeight);
 			}
 
 			// Now calculate the standard deviation of those lengths.
diff --git a/Engine/Plugins/Runtime/USDCore/Source/USDUtilities/Private/USDSkeletalDataConversion.cpp b/Engine/Plugins/Runtime/USDCore/Source/USDUtilities/Private/USDSkeletalDataConversion.cpp
--- a/Engine/Plugins/Runtime/USDCore/Source/USDUtilities/Private/USDSkeletalDataConversion.cpp
+++ b/Engine/Plugins/Runtime/USDCore/Source/USDUtilities/Private/USDSkeletalDataConversion.cpp
@@ -1029,14 +1029,13 @@ namespace UnrealToUsdImpl
 
 	// Converts UE morph target deltas from DeltaArray into offsets, pointIndices and normalOffsets attributes of BlendShape
 	bool ConvertMorphTargetDeltas(
-		const FMorphTargetDelta* DeltaArray,
-		int32 NumDeltas,
+		TConstArrayView<FMorphTargetDelta> DeltaArray,
 		TArray<int32> SourceToPackedVertexIndex,
 		pxr::UsdSkelBlendShape& BlendShape,
 		pxr::UsdTimeCode TimeCode
 	)
 	{
-		if (!DeltaArray || NumDeltas == 0 || !BlendShape)
+		if (DeltaArray.IsEmpty() || !BlendShape)
 		{
 			return false;
 		}
@@ -1048,14 +1047,12 @@ namespace UnrealToUsdImpl
 		pxr::VtArray<int> PointIndices;
 		pxr::VtArray<pxr::GfVec3f> Normals;
 
-		Offsets.reserve(NumDeltas);
-		PointIndices.reserve(NumDeltas);
-		Normals.reserve(NumDeltas);
+		Offsets.reserve(DeltaArray.Num());
+		PointIndices.reserve(DeltaArray.Num());
+		Normals.reserve(DeltaArray.Num());
 
-		for (int32 DeltaIndex = 0; DeltaIndex < NumDeltas; ++DeltaIndex)
+		for (const FMorphTargetDelta& Delta: DeltaArray)
 		{
-			const FMorphTargetDelta& Delta = DeltaArray[DeltaIndex];
-
 			int32 PackedIndex = SourceToPackedVertexIndex[Delta.SourceIdx];
 			if (PackedIndex == INDEX_NONE)
 			{
@@ -3823,9 +3820,8 @@ bool UnrealToUsd::ConvertSkeletalMesh(
 				continue;
 			}
 
-			int32 NumDeltas = 0;
-			const FMorphTargetDelta* DeltaArray = MorphTarget->GetMorphTargetDelta(LODIndex, NumDeltas);
-			if (!DeltaArray || NumDeltas == 0)
+			TConstArrayView<FMorphTargetDelta> DeltaArray = MorphTarget->GetMorphTargetDeltas(LODIndex);
+			if (DeltaArray.IsEmpty())
 			{
 				continue;
 			}
@@ -3840,7 +3836,6 @@ bool UnrealToUsd::ConvertSkeletalMesh(
 
 			bool bCreatedBlendShape = UnrealToUsdImpl::ConvertMorphTargetDeltas(
 				DeltaArray,
-				NumDeltas,
 				SourceToPackedVertexIndex,
 				BlendShape,
 				TimeCode
diff --git a/Engine/Source/Developer/SkeletalMeshUtilitiesCommon/Private/LODUtilities.cpp b/Engine/Source/Developer/SkeletalMeshUtilitiesCommon/Private/LODUtilities.cpp
--- a/Engine/Source/Developer/SkeletalMeshUtilitiesCommon/Private/LODUtilities.cpp
+++ b/Engine/Source/Developer/SkeletalMeshUtilitiesCommon/Private/LODUtilities.cpp
@@ -1723,21 +1723,12 @@ void FLODUtilities::SimplifySkeletalMeshLOD(
 			SkeletalMeshResource->InlineReductionCacheDatas[DesiredLOD].SetCacheGeometryInfo(SrcModel);
 
 			InlineReductionDataParameter.InlineOriginalSrcMorphTargetData.Empty(SkeletalMesh->GetMorphTargets().Num());
-			for (UMorphTarget* MorphTarget : SkeletalMesh->GetMorphTargets())
+			for (const UMorphTarget* MorphTarget : SkeletalMesh->GetMorphTargets())
 			{
-				if (!MorphTarget->HasDataForLOD(DesiredLOD))
-				{
-					continue;
-				}
+				// Iterate each original morph target source index to fill the NewMorphTargetDeltas array with the TargetMatchData.
+				// If there's no morph data for the LOD, the returned array view is empty, and so nothing happens.
 				TArray<FMorphTargetDelta>& MorphDeltasArray = InlineReductionDataParameter.InlineOriginalSrcMorphTargetData.FindOrAdd(MorphTarget->GetFullName());
-				const FMorphTargetLODModel& BaseMorphModel = MorphTarget->GetMorphLODModels()[DesiredLOD];
-				//Iterate each original morph target source index to fill the NewMorphTargetDeltas array with the TargetMatchData.
-				int32 NumDeltas = 0;
-				const FMorphTargetDelta* BaseDeltaArray = MorphTarget->GetMorphTargetDelta(DesiredLOD, NumDeltas);
-				for (int32 DeltaIndex = 0; DeltaIndex < NumDeltas; DeltaIndex++)
-				{
-					MorphDeltasArray.Add(BaseDeltaArray[DeltaIndex]);
-				}
+				MorphDeltasArray.Append(MorphTarget->GetMorphTargetDeltas(DesiredLOD));
 			}
 
 			// Copy the original SkeletalMesh LODModel
diff --git a/Engine/Source/Editor/UnrealEd/Private/Fbx/FbxSkeletalMeshExport.cpp b/Engine/Source/Editor/UnrealEd/Private/Fbx/FbxSkeletalMeshExport.cpp
--- a/Engine/Source/Editor/UnrealEd/Private/Fbx/FbxSkeletalMeshExport.cpp
+++ b/Engine/Source/Editor/UnrealEd/Private/Fbx/FbxSkeletalMeshExport.cpp
@@ -284,16 +284,13 @@ FbxNode* FFbxExporter::CreateMesh(const USkeletalMesh* SkelMesh, const TCHAR* Me
 					// Replicate the base mesh in the shape control points to set up the data.
 					for (int32 VertIndex = 0; VertIndex < VertexCount; ++VertIndex)
 					{
-						FVector Position = (FVector)Vertices[VertIndex].Position;
+						FVector Position = static_cast<FVector>(Vertices[VertIndex].Position);
 						ShapeControlPoints[VertIndex] = Converter.ConvertToFbxPos(Position);
 					}
 				
-					int32 NumberOfDeltas = 0;
-					const FMorphTargetDelta* MorphTargetDeltas = MorphTarget->GetMorphTargetDelta(LODIndex, NumberOfDeltas);
-					for (int32 MorphTargetDeltaIndex = 0; MorphTargetDeltaIndex < NumberOfDeltas; ++MorphTargetDeltaIndex)
+					for (const FMorphTargetDelta& CurrentDelta: MorphTarget->GetMorphTargetDeltas(LODIndex))
 					{
 						// Apply the morph target deltas to the control points.
-						const FMorphTargetDelta& CurrentDelta = MorphTargetDeltas[MorphTargetDeltaIndex];
 						uint32 RemappedSourceIndex = CurrentDelta.SourceIdx;
 
 						if (VertexIndexOffsetPairArray.Num() > 1)
diff --git a/Engine/Source/Runtime/Core/Private/UObject/FortniteMainBranchObjectVersion.cpp b/Engine/Source/Runtime/Core/Private/UObject/FortniteMainBranchObjectVersion.cpp
--- a/Engine/Source/Runtime/Core/Private/UObject/FortniteMainBranchObjectVersion.cpp
+++ b/Engine/Source/Runtime/Core/Private/UObject/FortniteMainBranchObjectVersion.cpp
@@ -12,7 +12,7 @@ TMap<FGuid, FGuid> FFortniteMainBranchObjectVersion::GetSystemGuids()
 	SystemGuids.Add(DevGuids.NANITE_DERIVEDDATA_VER, FGuid("FB5EAC6147EF4591B0D5EB7A85062AC9"));
 	SystemGuids.Add(DevGuids.NIAGARASHADERMAP_DERIVEDDATA_VER, FGuid("6360A977062842A29ED30E3A7ACB0E64"));
 	SystemGuids.Add(DevGuids.Niagara_LatestScriptCompileVersion, FGuid("97F0FD7149D90D42B65D6557AF477D1A"));
-	SystemGuids.Add(DevGuids.SkeletalMeshDerivedDataVersion, FGuid("716424BCEC314E969ACFE36D2C2ECE49"));
+	SystemGuids.Add(DevGuids.SkeletalMeshDerivedDataVersion, FGuid("07AE97536567481DBE673C51EBEC96EF"));
 	SystemGuids.Add(DevGuids.STATICMESH_DERIVEDDATA_VER, FGuid("3ABF32149F9D4D838750368AB95573BE"));
 	SystemGuids.Add(DevGuids.MaterialTranslationDDCVersion, FGuid("E541A459F3374E628240BBB309910753"));
 	return SystemGuids;
diff --git a/Engine/Source/Runtime/Core/Public/UObject/FortniteMainBranchObjectVersions.inl b/Engine/Source/Runtime/Core/Public/UObject/FortniteMainBranchObjectVersions.inl
--- a/Engine/Source/Runtime/Core/Public/UObject/FortniteMainBranchObjectVersions.inl
+++ b/Engine/Source/Runtime/Core/Public/UObject/FortniteMainBranchObjectVersions.inl
@@ -643,5 +643,8 @@ FORTNITE_MAIN_VERSION(PerParticleFlagToAllowPartialIslandSleepInConnectedIsland,
 // Material Function Blend Deserialize Top/Bottom input nodes with clearer enum marker.
 FORTNITE_MAIN_VERSION(MaterialFunctionBlendTopBottomInputEnum, 210)
 
+// Cooked CPU-side morph target points are now stored internally in the same compressed format as the GPU morph data. 
+FORTNITE_MAIN_VERSION(MorphTargetCookedCPUDataCompressed, 211)
+
 // -----<new versions can be added above this line>-------------------------------------------------
 #undef FORTNITE_MAIN_VERSION
\ No newline at end of file
diff --git a/Engine/Source/Runtime/Engine/Classes/Animation/MorphTarget.h b/Engine/Source/Runtime/Engine/Classes/Animation/MorphTarget.h
--- a/Engine/Source/Runtime/Engine/Classes/Animation/MorphTarget.h
+++ b/Engine/Source/Runtime/Engine/Classes/Animation/MorphTarget.h
@@ -2,10 +2,9 @@
 
 #pragma once
 
-#include "CoreMinimal.h"
-#include "UObject/ObjectMacros.h"
 #include "UObject/Object.h"
 #include "PackedNormal.h"
+#include "Rendering/MorphTargetVertexCodec.h"
 #include "Serialization/MemoryArchive.h"
 
 #include "MorphTarget.generated.h"
@@ -17,13 +16,13 @@ class UStaticMesh;
 struct FMorphTargetDelta
 {
 	/** change in position */
-	FVector3f			PositionDelta;
+	FVector3f			PositionDelta = FVector3f::ZeroVector;
 
 	/** Tangent basis normal */
-	FVector3f			TangentZDelta;
+	FVector3f			TangentZDelta = FVector3f::ZeroVector;
 
 	/** index of source vertex to apply deltas to */
-	uint32			SourceIdx;
+	uint32			SourceIdx = std::numeric_limits<uint32>::max();
 
 	/** pipe operator */
 	friend FArchive& operator<<(FArchive& Ar, FMorphTargetDelta& V)
@@ -43,6 +42,121 @@ struct FMorphTargetDelta
 	}
 };
 
+#if !WITH_EDITOR
+/** Compressed morph vertex deltas, used for cooked serialization. Only available in game. This matched the
+ *  GPU compressed delta storage and has the same quantization behavior. The tolerance stored is the same
+ *  value that was used when the deltas were compressed and should not be changed, or the decompression will
+ *  result in erroneous values.
+ */
+struct FMorphTargetCompressedLODModel
+{
+	TArray<UE::MorphTargetVertexCodec::FDeltaBatchHeader> PackedDeltaHeaders;
+	TArray<uint32> PackedDeltaData;
+	float PositionPrecision = 0.0f;
+	float TangentPrecision = 0.0f;
+};
+#endif
+
+
+class FMorphTargetDeltaIterator
+{
+public:
+	FMorphTargetDeltaIterator() = default;
+	FMorphTargetDeltaIterator(const FMorphTargetDeltaIterator&) = default;
+	FMorphTargetDeltaIterator(FMorphTargetDeltaIterator&&) = default;
+	FMorphTargetDeltaIterator& operator=(const FMorphTargetDeltaIterator&) = default;
+
+	// Support the minimal functionality required for a ranged-for.
+	[[nodiscard]] const FMorphTargetDelta& operator*() const 
+	{
+		return CurrentDelta;
+	}
+	
+	[[nodiscard]] const FMorphTargetDelta* operator->() const
+	{
+		return &CurrentDelta;
+	}
+	
+	FMorphTargetDeltaIterator& operator++()
+	{
+		Advance();
+		return *this;
+	}
+
+	[[nodiscard]] bool AtEnd() const
+	{
+		return Token == InvalidToken;
+	}
+	
+private:
+	static constexpr uint64 InvalidToken = std::numeric_limits<uint64>::max(); 
+	friend class UMorphTarget;
+
+	FMorphTargetDeltaIterator(
+		TConstArrayView<FMorphTargetDelta> InRawDeltas
+		) :
+		Token(0),
+		RawDeltas(InRawDeltas)
+	{
+		Advance();
+	}
+#if !WITH_EDITOR
+	FMorphTargetDeltaIterator(
+		const FMorphTargetCompressedLODModel& InCompressedDeltas 
+		) :
+		Token(0),
+		CompressedDelta(&InCompressedDeltas)
+	{
+		if (!InCompressedDeltas.PackedDeltaHeaders.IsEmpty() && !InCompressedDeltas.PackedDeltaData.IsEmpty())
+		{
+			Advance();
+		}
+		else
+		{
+			Clear();
+		}
+	}
+#endif
+	
+	void Advance()
+	{
+		if (!RawDeltas.IsEmpty())
+		{
+			if (Token < RawDeltas.Num())
+			{
+				CurrentDelta = RawDeltas[static_cast<int32>(Token++)];
+			}
+			else
+			{
+				Clear();
+			}
+		}
+#if !WITH_EDITOR
+		if (!IterativeDecode(Token, CompressedDelta->PackedDeltaHeaders, CompressedDelta->PackedDeltaData, CompressedDelta->PositionPrecision, CompressedDelta->TangentPrecision, CurrentDelta))
+		{
+			// If the decode failed, we're either at the end, or the data is invalid. Stop any further decoding.
+			Clear();
+		}
+#endif
+	}
+
+	void Clear()
+	{
+		Token = InvalidToken;
+		RawDeltas = {};
+#if !WITH_EDITOR
+		CompressedDelta = nullptr;
+#endif
+	}
+
+	uint64 Token = InvalidToken;
+	TConstArrayView<FMorphTargetDelta> RawDeltas;
+#if !WITH_EDITOR
+	const FMorphTargetCompressedLODModel* CompressedDelta = nullptr;
+#endif
+	FMorphTargetDelta CurrentDelta;
+};
+
 
 /**
 * Mesh data for a single LOD model of a morph target
@@ -53,25 +167,23 @@ struct FMorphTargetLODModel
 	TArray<FMorphTargetDelta> Vertices;
 
 	/** Number of elements in Vertices array. This property is set at runtime and is not serialized. */
-	int32 NumVertices;
+	int32 NumVertices = 0;
 
 	/** number of original verts in the base mesh */
-	int32 NumBaseMeshVerts;
+	int32 NumBaseMeshVerts = 0;
 	
 	/** list of sections this morph is used */
 	TArray<int32> SectionIndices;
 
 	/** Is this LOD generated by reduction setting */
-	bool bGeneratedByEngine;
+	bool bGeneratedByEngine = false;
 
-	/* The source filename use to import this morph target. If source is empty this morph target was import with the LOD geometry. */
+	/** The source filename use to import this morph target. If source is empty this morph target was import with the LOD geometry.
+	    Only defined in editor, since we don't need this information to spill into a cooked build. 
+	    */
 	FString SourceFilename;
-
-	FMorphTargetLODModel()
-		: NumVertices(0)
-		, NumBaseMeshVerts(0)
-		, bGeneratedByEngine(false)
-	{ }
+	
+	FMorphTargetLODModel() = default;
 
 	/** pipe operator */
 	friend FArchive& operator<<(FArchive& Ar, FMorphTargetLODModel& M);
@@ -85,22 +197,11 @@ struct FMorphTargetLODModel
 		// since engine cleared it, we mark as engine generated
 		// this makes it clear to clear up later
 		bGeneratedByEngine = true;
-
 		SourceFilename.Empty();
 	}
-
-	void DiscardVertexData()
-	{
-		// Vertices array can be discarded after render data is initialized and if CPU data is no longer required,
-		// but cache the array size so that UMorphTarget::HasValidData() can still query if morph target is valid. 
-		if (Vertices.Num())
-		{
-			NumVertices = Vertices.Num();
-			Vertices.Empty();
-		}
-	}
 };
 
+
 #if WITH_EDITOR
 /**
 * Data to cache serialization results for async asset building
@@ -144,31 +245,53 @@ public:
 protected:
 	/** morph mesh vertex data for each LOD */
 	TArray<FMorphTargetLODModel>	MorphLODModels;
-
+#if !WITH_EDITOR
+	TArray<FMorphTargetCompressedLODModel> CompressedLODModels;
+#endif
+	
 public:
 
+#if WITH_EDITOR
 	/** Get Morphtarget Delta array for the given input Index */
+	UE_DEPRECATED(5.7, "Use the TConstArrayView version of GetMorphTargetDelta instead.")
 	ENGINE_API virtual const FMorphTargetDelta* GetMorphTargetDelta(int32 LODIndex, int32& OutNumDeltas) const;
+
+	/** Get the morph target vertex delta array for a given LOD. Only available in editor targets, since cooked builds
+	 *  store the deltas internally in a compressed format. For non-editor builds, use GetDeltaIteratorForLOD */
+	ENGINE_API virtual TConstArrayView<FMorphTargetDelta> GetMorphTargetDeltas(int32 LODIndex) const;
+	
+#endif
+
+	/** Returns an iterator that runs over vertices, either the raw delta form or the cooked compressed form. Use this
+	 *  function at runtime to avoid the overhead of keeping the uncompressed raw deltas in memory. */
+	ENGINE_API virtual FMorphTargetDeltaIterator GetDeltaIteratorForLOD(int32 LODIndex) const;
+
+	/** Returns the number of morph deltas for the given LOD. If 0, there is no morph for this LOD.
+	 *  At runtime, this value may be non-zero, even if GetDeltaIteratorForLOD returns no morph deltas, since
+	 *  the morphs might be tuned to only compress the deltas for the GPU, leaving no CPU deltas in memory.
+	 *  In editor, this value will always match the number of entries returned by GetMorphTargetDeltas.
+	 */
+	ENGINE_API virtual int32 GetNumDeltasForLOD(int32 LODIndex) const;
+	
+	/** Returns true if the given LOD has morph data */
 	ENGINE_API virtual bool HasDataForLOD(int32 LODIndex) const;
+	
 	/** return true if this morphtarget contains data for section within LOD */
 	ENGINE_API virtual bool HasDataForSection(int32 LODIndex, int32 SectionIndex) const;
+	
 	/** return true if this morphtarget contains valid vertices */
 	ENGINE_API virtual bool HasValidData() const;
 	ENGINE_API virtual void EmptyMorphLODModels();
 
 	/** return true if this morphtarget was custom imported from a file */
 	ENGINE_API virtual bool IsCustomImported(int32 LODIndex) const;
-	ENGINE_API virtual FString GetCustomImportedSourceFilename(int32 LODIndex) const;
+	ENGINE_API virtual const FString& GetCustomImportedSourceFilename(int32 LODIndex) const;
 	ENGINE_API virtual void SetCustomImportedSourceFilename(int32 LODIndex, const FString& InSourceFilename);
 
 	/** return true if this morphtarget was generated by an engine tool */
 	ENGINE_API virtual bool IsGeneratedByEngine(int32 LODIndex) const;
 	ENGINE_API virtual void SetGeneratedByEngine(int32 LODIndex, bool bInGeneratedByEngine);
 
-	/** Discard CPU Buffers after render resources have been created. */
-	UE_DEPRECATED(5.0, "No longer in use, will be deleted. Whether to discard vertex data is now determined during cooking instead of loading.")
-	ENGINE_API virtual void DiscardVertexData();
-
 	/** Return true if this morph target uses engine built-in compression */
 	virtual bool UsesBuiltinMorphTargetCompression() const { return true; }
 
@@ -181,8 +304,6 @@ public:
 	ENGINE_API virtual TUniquePtr<FFinishBuildMorphTargetData> CreateFinishBuildMorphTargetData() const;
 #endif // WITH_EDITOR
 
-public:
-
 	//~ UObject interface
 
 	ENGINE_API virtual void Serialize(FArchive& Ar) override;
diff --git a/Engine/Source/Runtime/Engine/Private/MorphMesh.cpp b/Engine/Source/Runtime/Engine/Private/MorphMesh.cpp
deleted file mode 100644
diff --git a/Engine/Source/Runtime/Engine/Private/MorphTarget.cpp b/Engine/Source/Runtime/Engine/Private/MorphTarget.cpp
new file mode 100644
--- /dev/null
+++ b/Engine/Source/Runtime/Engine/Private/MorphTarget.cpp
@@ -0,0 +1,709 @@
+// Copyright Epic Games, Inc. All Rights Reserved.
+
+#include "Animation/MorphTarget.h"
+
+#include "Engine/SkeletalMesh.h"
+#include "EngineUtils.h"
+#include "Algo/AnyOf.h"
+#include "Interfaces/ITargetPlatform.h"
+#include "Rendering/SkeletalMeshModel.h"
+#include "Serialization/MemoryArchive.h"
+#include "UObject/UE5SpecialProjectStreamObjectVersion.h"
+
+//////////////////////////////////////////////////////////////////////////
+
+FArchive& operator<<(FArchive& Ar, FMorphTargetLODModel& M)
+{
+	Ar.UsingCustomVersion(FEditorObjectVersion::GUID);
+	Ar.UsingCustomVersion(FFortniteMainBranchObjectVersion::GUID);
+	Ar.UsingCustomVersion(FUE5SpecialProjectStreamObjectVersion::GUID);
+
+	if (!Ar.IsObjectReferenceCollector())
+	{
+		if (Ar.IsLoading() && Ar.CustomVer(FEditorObjectVersion::GUID) < FEditorObjectVersion::AddedMorphTargetSectionIndices)
+		{
+			Ar << M.Vertices << M.NumBaseMeshVerts;
+			M.bGeneratedByEngine = false;
+		}
+		else if (Ar.IsLoading() && Ar.CustomVer(FFortniteMainBranchObjectVersion::GUID) < FFortniteMainBranchObjectVersion::SaveGeneratedMorphTargetByEngine)
+		{
+			Ar << M.Vertices << M.NumBaseMeshVerts << M.SectionIndices;
+			M.bGeneratedByEngine = false;
+		}
+		else
+		{
+			bool bVerticesAreStrippedForCookedBuilds = false;
+
+			if (Ar.IsPersistent() && (Ar.CustomVer(FUE5SpecialProjectStreamObjectVersion::GUID) >= FUE5SpecialProjectStreamObjectVersion::StripMorphTargetSourceDataForCookedBuilds))
+			{
+				// Strip raw vertex data for cooked builds, but only if we're cooking to disk. 
+				bVerticesAreStrippedForCookedBuilds = Ar.IsCooking();
+				Ar << bVerticesAreStrippedForCookedBuilds;
+			}
+
+			if (bVerticesAreStrippedForCookedBuilds)
+			{
+				M.NumVertices = M.Vertices.Num();
+				Ar << M.NumVertices;
+			}
+			else
+			{
+				Ar << M.Vertices;
+
+				if (Ar.IsLoading())
+				{
+					M.NumVertices = M.Vertices.Num();
+				}
+			}
+
+			Ar << M.NumBaseMeshVerts << M.SectionIndices << M.bGeneratedByEngine;
+		}
+
+		if (Ar.IsLoading() && Ar.CustomVer(FFortniteMainBranchObjectVersion::GUID) < FFortniteMainBranchObjectVersion::MorphTargetCustomImport)
+		{
+			M.SourceFilename.Empty();
+		}
+		else
+		{
+			//Do not cook the source filename, we don't need it at runtime
+			if (Ar.IsCooking() && Ar.IsSaving())
+			{
+				FString EmptySourceFilename;
+				Ar << EmptySourceFilename;
+			}
+			else
+			{
+				Ar << M.SourceFilename;
+			}
+		}
+
+		/* Fix the data for the generated by engine flag */
+		if(Ar.IsLoading() && M.SourceFilename.StartsWith(FMorphTargetImportedSourceFileInfo::GeneratedByEnginePrefix))
+		{
+			M.bGeneratedByEngine = true;
+			M.SourceFilename.Empty();
+		}
+	}
+
+	return Ar;
+}
+
+
+//////////////////////////////////////////////////////////////////////////
+
+
+UMorphTarget::UMorphTarget(const FObjectInitializer& ObjectInitializer)
+	: Super(ObjectInitializer)
+{
+}
+
+
+void UMorphTarget::Serialize( FArchive& Ar )
+{
+	LLM_SCOPE(ELLMTag::Animation);
+	
+	Super::Serialize( Ar );
+	Ar.UsingCustomVersion(FEditorObjectVersion::GUID);
+	Ar.UsingCustomVersion(FFortniteMainBranchObjectVersion::GUID);
+
+	FStripDataFlags StripFlags( Ar );
+	if( StripFlags.IsAudioVisualDataStripped() )
+	{
+		return;
+	}
+
+	bool bCooked = Ar.IsCooking() && Ar.IsPersistent();
+	if (Ar.IsLoading() && Ar.CustomVer(FFortniteMainBranchObjectVersion::GUID) < FFortniteMainBranchObjectVersion::MorphTargetCookedCPUDataCompressed)
+	{
+		bCooked = false;
+	}
+	else
+	{
+		Ar << bCooked;
+	}
+	
+	// We always store the morph models (save for the raw vertex data when cooking).  
+	Ar << MorphLODModels;
+	
+	if (bCooked)
+	{
+		// Only platforms that support ES31 rendering need compressed vertex data on the CPU side. Non-ES31 platforms operate on compressed GPU
+		// data which is cooked out elsewhere.
+		bool bStoreCompressedVertices = false;
+		
+#if WITH_EDITOR
+		bStoreCompressedVertices = Ar.CookingTarget()->SupportsFeature(ETargetPlatformFeatures::MobileRendering);
+#endif
+
+		Ar << bStoreCompressedVertices;
+		
+		if (bStoreCompressedVertices)
+		{
+#if WITH_EDITOR
+			if (Ar.IsCooking() && Ar.IsPersistent())
+			{
+				int32 NumLODs = MorphLODModels.Num();
+
+				static const FSkeletalMeshLODInfo DefaultLODInfo{};
+				
+				float TangentPrecision = UE::MorphTargetVertexCodec::ComputeTangentPrecision();
+
+				// BaseSkelMesh is technically redundant, since it should just be the outer. To play it safe, we always use it first
+				// and just default to the outer when it is null (which can happen).
+				USkeletalMesh* SkeletalMesh = BaseSkelMesh;
+				if (!SkeletalMesh)
+				{
+					SkeletalMesh = Cast<USkeletalMesh>(GetOuter());
+				}
+
+				Ar << NumLODs;
+
+				for (int32 LODIndex = 0; LODIndex < NumLODs; LODIndex++)
+				{
+					const FMorphTargetLODModel& LODModel = MorphLODModels[LODIndex];
+					
+					TArray<UE::MorphTargetVertexCodec::FDeltaBatchHeader> PackedDeltaHeaders;
+					TArray<uint32> PackedDeltaData;
+					float PositionPrecision = UE::MorphTargetVertexCodec::ComputePositionPrecision(
+						SkeletalMesh ? SkeletalMesh->GetLODInfo(LODIndex)->MorphTargetPositionErrorTolerance : DefaultLODInfo.MorphTargetPositionErrorTolerance);
+
+					UE::MorphTargetVertexCodec::Encode(LODModel.Vertices, nullptr,
+						PositionPrecision, TangentPrecision, PackedDeltaHeaders, PackedDeltaData);
+
+					Ar << PackedDeltaHeaders;
+					Ar << PackedDeltaData;
+					Ar << PositionPrecision;
+					Ar << TangentPrecision;
+				}
+			}
+			else
+#endif
+			{
+				int32 NumLODs = 0;
+				Ar << NumLODs;
+
+#if !WITH_EDITOR
+				CompressedLODModels.Reset(NumLODs);
+#endif
+			
+				for (int32 LODIndex = 0; LODIndex < NumLODs; LODIndex++)
+				{
+					TArray<UE::MorphTargetVertexCodec::FDeltaBatchHeader> PackedDeltaHeaders;
+					TArray<uint32> PackedDeltaData;
+					float PositionPrecision = 0.0f;
+					float TangentPrecision = 0.0f;
+
+					Ar << PackedDeltaHeaders;
+					Ar << PackedDeltaData;
+					Ar << PositionPrecision;
+					Ar << TangentPrecision;
+
+#if WITH_EDITOR
+					FMorphTargetLODModel& MorphModel = MorphLODModels[LODIndex];
+					// Unpack the compressed data into the raw vertex data arrays.
+					UE::MorphTargetVertexCodec::Decode(PackedDeltaHeaders, PackedDeltaData, PositionPrecision, TangentPrecision, MorphModel.Vertices);
+#else
+					FMorphTargetCompressedLODModel& CompressedModel = CompressedLODModels.AddDefaulted_GetRef();
+					CompressedModel.PackedDeltaHeaders = MoveTemp(PackedDeltaHeaders);
+					CompressedModel.PackedDeltaData = MoveTemp(PackedDeltaData);
+					CompressedModel.PositionPrecision = PositionPrecision;
+					CompressedModel.TangentPrecision = TangentPrecision;
+#endif
+				}
+			}
+		}
+	}
+}
+
+#if WITH_EDITORONLY_DATA
+void UMorphTarget::DeclareCustomVersions(FArchive& Ar, const UClass* SpecificSubclass)
+{
+	Super::DeclareCustomVersions(Ar, SpecificSubclass);
+	FMorphTargetLODModel MorphLODModel;
+	Ar << MorphLODModel;
+}
+#endif
+
+namespace
+{
+	void SerializeMorphLODModels(FMemoryArchive& Ar, TArray<FMorphTargetLODModel>& MorphLODModels)
+	{
+		int32 MorphLODModelNumber = 0;
+		if (Ar.IsLoading())
+		{
+			Ar << MorphLODModelNumber;
+			MorphLODModels.Empty(MorphLODModelNumber);
+			MorphLODModels.AddDefaulted(MorphLODModelNumber);
+		}
+		else
+		{
+			MorphLODModelNumber = MorphLODModels.Num();
+			Ar << MorphLODModelNumber;
+		}
+
+		for (int32 MorphIndex = 0; MorphIndex < MorphLODModelNumber; ++MorphIndex)
+		{
+			Ar << MorphLODModels[MorphIndex];
+		}
+	}
+}
+
+void UMorphTarget::SerializeMemoryArchive(FMemoryArchive & Ar)
+{
+	FName MorphTargetName = GetFName();
+	Ar << MorphTargetName;
+	SerializeMorphLODModels(Ar, MorphLODModels);
+}
+
+#if WITH_EDITOR
+void FFinishBuildMorphTargetData::LoadFromMemoryArchive(FMemoryArchive & Ar)
+{
+	check(Ar.IsLoading());
+	
+	if (!ensureMsgf(!bApplyMorphTargetsData, TEXT("Error in FFinishBuildMorphTargetData::LoadFromMemoryArchive. The compilation context morph targets data was already set.")))
+	{
+		MorphLODModelsPerTargetName.Empty();
+		bApplyMorphTargetsData = false;
+	}
+	bApplyMorphTargetsData = true;
+	
+	int32 MorphTargetNumber = 0;
+	Ar << MorphTargetNumber;
+	MorphLODModelsPerTargetName.Reserve(MorphTargetNumber);
+	for (int32 MorphTargetIndex = 0; MorphTargetIndex < MorphTargetNumber; ++MorphTargetIndex)
+	{
+		FName MorphTargetName = NAME_None;
+		Ar << MorphTargetName;
+		TArray<FMorphTargetLODModel>&MorphLODModels = MorphLODModelsPerTargetName.FindOrAdd(MorphTargetName);
+		SerializeMorphLODModels(Ar, MorphLODModels);
+	}
+}
+#endif
+
+void UMorphTarget::PostLoad()
+{
+	Super::PostLoad();
+	
+#if WITH_EDITOR
+	if (GetLinkerCustomVersion(FEditorObjectVersion::GUID) < FEditorObjectVersion::AddedMorphTargetSectionIndices &&
+		BaseSkelMesh)
+	{
+		const int32 MaxLOD = FMath::Min(BaseSkelMesh->GetImportedModel()->LODModels.Num(), MorphLODModels.Num());
+		for (int32 LODIndex = 0; LODIndex < MaxLOD; ++LODIndex)
+		{
+			FMorphTargetLODModel& MorphLODModel = MorphLODModels[LODIndex];
+			MorphLODModel.SectionIndices.Empty();
+			const FSkeletalMeshLODModel& LODModel = BaseSkelMesh->GetImportedModel()->LODModels[LODIndex];
+			TArray<int32> BaseIndexes;
+			TArray<int32> LastIndexes;
+			for (int32 SectionIdx = 0; SectionIdx < LODModel.Sections.Num(); ++SectionIdx)
+			{
+				const int32 BaseVertexBufferIndex = LODModel.Sections[SectionIdx].GetVertexBufferIndex();
+				BaseIndexes.Add(BaseVertexBufferIndex);
+				LastIndexes.Add(BaseVertexBufferIndex + LODModel.Sections[SectionIdx].GetNumVertices());
+			}
+			// brute force
+			for (int32 VertIndex = 0; VertIndex < MorphLODModel.Vertices.Num() && MorphLODModel.SectionIndices.Num() < BaseIndexes.Num(); ++VertIndex)
+			{
+				int32 SourceVertexIdx = MorphLODModel.Vertices[VertIndex].SourceIdx;
+				for (int32 SectionIdx = 0; SectionIdx < BaseIndexes.Num(); ++SectionIdx)
+				{
+					if (!MorphLODModel.SectionIndices.Contains(SectionIdx))
+					{
+						if (BaseIndexes[SectionIdx] <= SourceVertexIdx && SourceVertexIdx < LastIndexes[SectionIdx])
+						{
+							MorphLODModel.SectionIndices.AddUnique(SectionIdx);
+							break;
+						}
+					}
+				}
+			}
+		}
+	}
+#endif //#if WITH_EDITOR
+}
+
+#if WITH_EDITOR
+const FMorphTargetDelta* UMorphTarget::GetMorphTargetDelta(int32 LODIndex, int32& OutNumDeltas) const
+{
+	if(GetMorphLODModels().IsValidIndex(LODIndex))
+	{
+		// Calling GetMorphLODModels to potentially get from subclasses
+		const FMorphTargetLODModel& MorphModel = GetMorphLODModels()[LODIndex];
+		OutNumDeltas = MorphModel.Vertices.Num();
+		return MorphModel.Vertices.GetData();
+	}
+
+	OutNumDeltas = 0;
+	return nullptr;
+}
+
+TConstArrayView<FMorphTargetDelta> UMorphTarget::GetMorphTargetDeltas(int32 LODIndex) const
+{
+	if (GetMorphLODModels().IsValidIndex(LODIndex))
+	{
+		return GetMorphLODModels()[LODIndex].Vertices;
+	}
+
+	return {};
+}
+#endif
+
+
+FMorphTargetDeltaIterator UMorphTarget::GetDeltaIteratorForLOD(int32 LODIndex) const
+{
+	if (GetMorphLODModels().IsValidIndex(LODIndex))
+	{
+#if !WITH_EDITOR
+		if (CompressedLODModels.IsValidIndex(LODIndex))
+		{
+			return FMorphTargetDeltaIterator(CompressedLODModels[LODIndex]);
+		}
+#endif
+		return FMorphTargetDeltaIterator(GetMorphLODModels()[LODIndex].Vertices);
+	}
+
+	return {};
+}
+
+int32 UMorphTarget::GetNumDeltasForLOD(int32 LODIndex) const
+{
+	if (GetMorphLODModels().IsValidIndex(LODIndex))
+	{
+#if WITH_EDITOR
+		return GetMorphLODModels()[LODIndex].Vertices.Num();
+#else
+		return GetMorphLODModels()[LODIndex].NumVertices;
+#endif
+	}
+
+	return 0;
+}
+
+bool UMorphTarget::HasDataForLOD(int32 LODIndex) const
+{
+	return GetNumDeltasForLOD(LODIndex) != 0;
+}
+
+bool UMorphTarget::HasValidData() const
+{
+	for (const FMorphTargetLODModel& Model : MorphLODModels)
+	{
+		if (Model.Vertices.Num() > 0)
+		{
+			return true;
+		}
+#if !WITH_EDITOR
+		// In cooked builds, Model.Vertices is stripped but Model.NumVertices is valid
+		else if (Model.NumVertices > 0)
+		{
+			return true;
+		}
+#endif
+	}
+
+	return false;
+}
+
+bool UMorphTarget::HasDataForSection(int32 LODIndex, int32 SectionIndex) const
+{
+	return HasDataForLOD(LODIndex) && MorphLODModels[LODIndex].SectionIndices.Contains(SectionIndex);
+}
+
+void UMorphTarget::EmptyMorphLODModels()
+{
+	MorphLODModels.Empty();
+}
+
+bool UMorphTarget::IsCustomImported(int32 LODIndex) const
+{
+	if (LODIndex < MorphLODModels.Num())
+	{
+		// Calling GetMorphLODModels to potentially get from subclasses
+		const FMorphTargetLODModel& MorphModel = GetMorphLODModels()[LODIndex];
+		return !MorphModel.SourceFilename.IsEmpty();
+	}
+	return false;
+}
+
+const FString& UMorphTarget::GetCustomImportedSourceFilename(int32 LODIndex) const
+{
+	if (LODIndex < MorphLODModels.Num())
+	{
+		// Calling GetMorphLODModels to potentially get from subclasses
+		const FMorphTargetLODModel& MorphModel = GetMorphLODModels()[LODIndex];
+		return MorphModel.SourceFilename;
+	}
+	static const FString EmptyString;
+	return EmptyString;
+}
+
+void UMorphTarget::SetCustomImportedSourceFilename(int32 LODIndex, const FString& InSourceFilename)
+{
+	if (LODIndex < MorphLODModels.Num())
+	{
+		// Calling GetMorphLODModels to potentially get from subclasses
+		FMorphTargetLODModel& MorphModel = GetMorphLODModels()[LODIndex];
+		MorphModel.SourceFilename = InSourceFilename;
+		MorphModel.bGeneratedByEngine = false;
+	}
+}
+
+bool UMorphTarget::IsGeneratedByEngine(int32 LODIndex) const
+{
+	if (LODIndex < MorphLODModels.Num())
+	{
+		// Calling GetMorphLODModels to potentially get from subclasses
+		const FMorphTargetLODModel& MorphModel = GetMorphLODModels()[LODIndex];
+		if(MorphModel.SourceFilename.IsEmpty())
+		{
+			return MorphModel.bGeneratedByEngine;
+		}
+	}
+	return false;
+}
+
+void UMorphTarget::SetGeneratedByEngine(int32 LODIndex, bool bInGeneratedByEngine)
+{
+	if (LODIndex < MorphLODModels.Num())
+	{
+		// Calling GetMorphLODModels to potentially get from subclasses
+		FMorphTargetLODModel& MorphModel = GetMorphLODModels()[LODIndex];
+		MorphModel.SourceFilename.Empty();
+		MorphModel.bGeneratedByEngine = bInGeneratedByEngine;
+	}
+}
+
+#if WITH_EDITOR
+
+void UMorphTarget::PopulateDeltas(
+	const TArray<FMorphTargetDelta>& Deltas,
+	const int32 LODIndex,
+	const TArray<FSkelMeshSection>& Sections,
+	const bool bCompareNormal,
+	const bool bGeneratedByReductionSetting,
+	const float InPositionThreshold
+	)
+{
+	// create the LOD entry if it doesn't already exist
+	if (LODIndex >= MorphLODModels.Num())
+	{
+		MorphLODModels.AddDefaulted(LODIndex - MorphLODModels.Num() + 1);
+	}
+
+	// morph mesh data to modify
+	FMorphTargetLODModel& MorphModel = MorphLODModels[LODIndex];
+	// copy the wedge point indices
+	// for now just keep every thing 
+
+	// set the original number of vertices
+	MorphModel.NumBaseMeshVerts = Deltas.Num();
+
+	// empty morph mesh vertices first
+	MorphModel.Vertices.Empty(Deltas.Num());
+
+	// mark if generated by reduction setting, so that we can remove them later if we want to
+	// we don't want to delete if it has been imported
+	MorphModel.bGeneratedByEngine = bGeneratedByReductionSetting;
+
+	// Still keep this (could remove in long term due to incoming data)
+	const float PositionThresholdSquared = FMath::Square(InPositionThreshold);
+	for (const FMorphTargetDelta& Delta : Deltas)
+	{
+		if (Delta.PositionDelta.SizeSquared() > PositionThresholdSquared ||
+			( bCompareNormal && Delta.TangentZDelta.SizeSquared() > 0.01f))
+		{
+			MorphModel.Vertices.Add(Delta);
+			for (int32 SectionIdx = 0; SectionIdx < Sections.Num(); ++SectionIdx)
+			{
+				if (MorphModel.SectionIndices.Contains(SectionIdx))
+				{
+					continue;
+				}
+				const uint32 BaseVertexBufferIndex = static_cast<uint32>(Sections[SectionIdx].GetVertexBufferIndex());
+				const uint32 LastVertexBufferIndex = BaseVertexBufferIndex + Sections[SectionIdx].GetNumVertices();
+				if (BaseVertexBufferIndex <= Delta.SourceIdx && Delta.SourceIdx < LastVertexBufferIndex)
+				{
+					MorphModel.SectionIndices.AddUnique(SectionIdx);
+					break;
+				}
+			}
+		}
+	}
+
+	// sort the array of vertices for this morph target based on the base mesh indices
+	// that each vertex is associated with. This allows us to sequentially traverse the list
+	// when applying the morph blends to each vertex.
+	MorphModel.Vertices.Sort(
+		[](const FMorphTargetDelta& A, const FMorphTargetDelta& B)
+		{
+			return A.SourceIdx < B.SourceIdx;
+		}
+	);
+
+	// remove array slack
+	MorphModel.Vertices.Shrink();
+	MorphModel.NumVertices = MorphModel.Vertices.Num();
+}
+
+void UMorphTarget::RemoveEmptyMorphTargets()
+{
+	// do not remove LOD 0
+	for (int32 CurIndex = MorphLODModels.Num() - 1; CurIndex > 0; --CurIndex)
+	{
+		//  if it's empty, no reason to keep
+		if (MorphLODModels[CurIndex].Vertices.IsEmpty())
+		{
+			MorphLODModels.RemoveAt(CurIndex);
+		}
+		else
+		{
+			// we leave empty ones between. We can't remove between ones. 
+			// Once we found valid one, just get out
+			break;
+		}
+	}
+}
+
+TUniquePtr<FFinishBuildMorphTargetData> UMorphTarget::CreateFinishBuildMorphTargetData() const
+{
+	return MakeUnique<FFinishBuildMorphTargetData>();
+}
+
+void FFinishBuildMorphTargetData::ApplyEditorData(USkeletalMesh * SkeletalMesh, bool bIsSerializeSaving) const
+{
+	//List of potential issue with this function
+	// -This function change the array USkeletalMesh::MorphTargets, if a find reference call happen in same time for this asset(on the game thread)
+	//  The find reference can crash because it is accessing this array to find UObject reference.
+	//
+	// -If this function is call in the pre garbage collector event (we do finish compilation in this event), the call to FindObjectSafe will acquire a GC lock
+	//  which will cause a dead lock since the garbage collector already have the gc lock.
+	//
+	// TODO: remove the UMorphTarget object and replace the data in the import data of the skeletal mesh.
+	if (!IsInGameThread())
+	{
+		UE_ASSET_LOG(LogSkeletalMesh, Display, SkeletalMesh, TEXT("Calling function FFinishBuildMorphTargetData::ApplyEditorData outside of the game thread is not safe and can deadlock or crash."));
+	}
+	//Return if we do not need to apply data
+	if (!bApplyMorphTargetsData)
+	{
+		return;
+	}
+	
+	if (SkeletalMesh->HasAnyFlags(RF_BeginDestroyed | RF_FinishDestroyed))
+	{
+		//No need to apply the morph targets if the asset is being deleted.
+		//Also acquiring the GCScopeGuard when we are in the destruction process will create a deadlock, if there is code
+		//in the BeginDestroy function of the skeletalmesh that touch any property that is lock by the async build.
+		return;
+	}
+	//GC should not be active during the build since we force finish the skinned asset compilation during the pre garbage delegate
+	check(!IsGarbageCollectingAndLockingUObjectHashTables());
+
+	FSkeletalMeshModel * SkelMeshModel = SkeletalMesh->GetImportedModel();
+	check(SkelMeshModel);
+	
+	TMap<FName, UMorphTarget*> ExistingMorphTargets;
+	for (UMorphTarget* MorphTarget : SkeletalMesh->GetMorphTargets())
+	{
+		ExistingMorphTargets.Add(MorphTarget->GetFName(), MorphTarget);
+	}
+	
+	TArray<UMorphTarget*> ToDeleteMorphTargets;
+	ToDeleteMorphTargets.Append(SkeletalMesh->GetMorphTargets());
+	SkeletalMesh->GetMorphTargets().Empty();
+	//Rebuild the MorphTarget object
+	for (const TPair<FName, TArray<FMorphTargetLODModel>>& TargetNameAndMorphLODModels : MorphLODModelsPerTargetName)
+	{
+		FName MorphTargetName = TargetNameAndMorphLODModels.Key;
+		const TArray<FMorphTargetLODModel>& MorphTargetLODModels = TargetNameAndMorphLODModels.Value;
+		int32 MorphLODModelNumber = MorphTargetLODModels.Num();
+
+		UMorphTarget * MorphTarget = ExistingMorphTargets.FindRef(MorphTargetName);
+		if (!MorphTarget)
+		{
+			if (!Algo::AnyOf(MorphTargetLODModels, [](const FMorphTargetLODModel& Model) { return !Model.Vertices.IsEmpty();}))
+			{
+				//Skip this empty morphtarget
+				continue;
+			}
+
+			//Reuse the morph target if it already exist and was not garbage collect (it can happen if we play with the morph target threshold build options)
+			MorphTarget = FindObjectSafe<UMorphTarget>(SkeletalMesh, *MorphTargetName.ToString(), EFindObjectFlags::ExactClass);
+			if (MorphTarget && MorphTarget->HasAnyFlags(RF_BeginDestroyed | RF_FinishDestroyed))
+			{
+				MorphTarget = nullptr;
+			}
+
+			if (!MorphTarget)
+			{
+				//When we save the cook result we should never have to build a new morph target
+				//When saving cook build, we call GetPlatformSkeletalMeshRenderData in USkeletalMesh::BeginCacheForCookedPlatformData
+				//which happen before the serialization of that cook skeletalmesh
+				if (!bIsSerializeSaving)
+				{
+					//Avoid recycling morphtarget with NewObject it cannot be done asynchronously
+					//Find the UMorphTarget and simply clear the data if it exist.
+					TArray<UObject*> SubObjects;
+					GetObjectsWithOuter(SkeletalMesh, SubObjects, true);
+					for (UObject* SubObject : SubObjects)
+					{
+						if (SubObject->GetFName() == MorphTargetName)
+						{
+							if (UMorphTarget* SubMorphTarget = Cast<UMorphTarget>(SubObject))
+							{
+								MorphTarget = SubMorphTarget;
+								MorphTarget->EmptyMorphLODModels();
+								MorphTarget->ClearGarbage();
+								break;
+							}
+						}
+					}
+					//Create a new morph target, if the object do not exist (creating a new uobject is ok to do asynchronously)
+					if (!MorphTarget)
+					{
+						MorphTarget = NewObject<UMorphTarget>(SkeletalMesh, MorphTargetName);
+					}
+					check(MorphTarget);
+				}
+				else
+				{
+					UE_ASSET_LOG(LogSkeletalMesh, Error, SkeletalMesh, TEXT("Cannot cache a skeletal mesh during a serialize if some morph targets need to be created. The solution is to Pre cache the skeletal mesh before the serialization so no morph target get created."));
+					continue;
+				}
+			}
+		}
+		else
+		{
+			ToDeleteMorphTargets.Remove(MorphTarget);
+		}
+		MorphTarget->EmptyMorphLODModels();
+		SkeletalMesh->GetMorphTargets().Add(MorphTarget);
+		
+		MorphTarget->GetMorphLODModels().AddDefaulted(MorphLODModelNumber);
+		for (int32 MorphDataIndex = 0; MorphDataIndex < MorphLODModelNumber; ++MorphDataIndex)
+		{
+			MorphTarget->GetMorphLODModels()[MorphDataIndex] = MorphTargetLODModels[MorphDataIndex];
+		}
+	}
+	//Rebuild the mapping and rehook the curve data
+	SkeletalMesh->InitMorphTargets();
+	
+	//Clear any async flags after the morphtargets have been set to the skeletalmesh
+	for (UMorphTarget* MorphTarget : SkeletalMesh->GetMorphTargets())
+	{
+		constexpr EInternalObjectFlags AsyncFlags = EInternalObjectFlags::Async;
+		MorphTarget->ClearInternalFlags(AsyncFlags);
+	}
+
+	for (UMorphTarget* ToDeleteMorphTarget : ToDeleteMorphTargets)
+	{
+		ToDeleteMorphTarget->BaseSkelMesh = nullptr;
+		ToDeleteMorphTarget->EmptyMorphLODModels();
+		ToDeleteMorphTarget->MarkAsGarbage();
+	}
+}
+
+#endif // WITH_EDITOR
+
+
diff --git a/Engine/Source/Runtime/Engine/Private/MorphTools.cpp b/Engine/Source/Runtime/Engine/Private/MorphTools.cpp
deleted file mode 100644
diff --git a/Engine/Source/Runtime/Engine/Private/Rendering/MorphTargetVertexCodec.cpp b/Engine/Source/Runtime/Engine/Private/Rendering/MorphTargetVertexCodec.cpp
new file mode 100644
--- /dev/null
+++ b/Engine/Source/Runtime/Engine/Private/Rendering/MorphTargetVertexCodec.cpp
@@ -0,0 +1,397 @@
+﻿// Copyright Epic Games, Inc. All Rights Reserved.
+
+#include "Rendering/MorphTargetVertexCodec.h"
+
+#include "Animation/MorphTarget.h"
+#include "Containers/Array.h"
+#include "Math/Vector.h"
+
+#include "Engine/EngineConsoleCommandExecutor.h"
+
+namespace UE::MorphTargetVertexCodec
+{
+
+static constexpr uint32 BatchSizeBits = 6;
+static constexpr uint32 BatchSize = 1 << BatchSizeBits;
+
+class FDwordBitWriter
+{
+public:
+	FDwordBitWriter(TArray<uint32>& Buffer) :
+		Buffer(Buffer),
+		PendingBits(0ull),
+		NumPendingBits(0)
+	{
+	}
+
+	void PutBits(uint32 Bits, uint32 NumBits)
+	{
+		check(static_cast<uint64>(Bits) < (1ull << NumBits));
+		PendingBits |= static_cast<uint64>(Bits) << NumPendingBits;
+		NumPendingBits += NumBits;
+
+		while (NumPendingBits >= 32)
+		{
+			Buffer.Add(static_cast<uint32>(PendingBits));
+			PendingBits >>= 32;
+			NumPendingBits -= 32;
+		}
+	}
+
+	void Flush()
+	{
+		if (NumPendingBits > 0)
+			Buffer.Add(static_cast<uint32>(PendingBits));
+		PendingBits = 0;
+		NumPendingBits = 0;
+	}
+
+private:
+	TArray<uint32>& Buffer;
+	uint64 			PendingBits;
+	int32 			NumPendingBits;
+};
+
+class FDwordBitReader
+{
+public:
+	FDwordBitReader(TConstArrayView<uint32> InBuffer, uint32 InOffset = 0) :
+		Buffer(InBuffer),
+		Offset(InOffset)
+	{
+	}
+
+	uint32 GetBits(uint32 NumBits)
+	{
+		check(NumBits <= 32);
+		check(Offset + NumBits <= (Buffer.Num() * 32u));
+		
+		if (NumBits == 0)
+		{
+			return 0;
+		}
+		
+		const uint32 BaseIndex = Offset >> 5u;
+		const uint32 BitOffset = Offset & 31u;
+
+		// Advance
+		Offset += NumBits;
+
+		if ((BitOffset + NumBits) > 32)
+		{
+			const uint32 BitMaskLow = ((1 << (32 - BitOffset)) - 1);
+			const uint32 BitMaskHigh = (1 << (NumBits + BitOffset - 32)) - 1;
+			const uint32 BitOffsetLow = BitOffset;
+			const uint32 BitOffsetHigh = 32 - BitOffset;
+
+			const uint32 Low = ((Buffer[BaseIndex + 0] >> BitOffsetLow) & BitMaskLow);
+			const uint32 High = ((Buffer[BaseIndex + 1] & BitMaskHigh) << BitOffsetHigh);
+
+			return Low | High;
+		}
+		else
+		{
+			const uint32 BitMask = (1 << NumBits) - 1;
+			return (Buffer[BaseIndex] >> BitOffset) & BitMask;
+		}
+	}
+
+	uint32 GetOffset() const
+	{
+		return Offset;
+	}
+
+private:
+	TConstArrayView<uint32> Buffer;
+	uint32 Offset = 0;
+};
+
+
+void Encode(
+	TConstArrayView<FMorphTargetDelta> InMorphDeltas,
+	const TBitArray<>* InVertexNeedsTangents,
+	const float InPositionPrecision,
+	const float InTangentZPrecision,
+	TArray<FDeltaBatchHeader>& OutBatchHeaders,
+	TArray<uint32>& OutCompressedVertices
+	)
+{
+	// Simple Morph compression 0.1
+	// Instead of storing vertex deltas individually they are organized into batches of 64.
+	// Each batch has a header that describes how many bits are allocated to each of the vertex components.
+	// Batches also store an explicit offset to its associated data. This makes it trivial to decode batches
+	// in parallel, and because deltas are fixed-width inside a batch, deltas can also be decoded in parallel.
+	// The result is a semi-adaptive encoding that functions as a crude substitute for entropy coding, that is
+	// fast to decode on parallel hardware.
+
+	// Quantization still happens globally to avoid issues with cracks at duplicate vertices.
+	// The quantization is artist controlled on a per LOD basis. Higher error tolerance results in smaller deltas
+	// and a smaller compressed size.
+
+	const float RcpPositionPrecision = 1.0f / InPositionPrecision;
+	const float RcpTangentZPrecision = 1.0f / InTangentZPrecision;
+
+	constexpr uint32 IndexMaxBits = 31u;
+	constexpr uint32 PositionMaxBits = 28u;				// Probably more than we need, but let's just allow it to go this high to be safe for now.
+	// For larger deltas this can even be more precision than what was in the float input data!
+	// Maybe consider float-like or exponential encoding of large values?
+	constexpr float PositionMinValue = -134217728.0f;	// -2^(MaxBits-1)
+	constexpr float PositionMaxValue = 134217720.0f;	// Largest float smaller than 2^(MaxBits-1)-1
+	// Using 134217727.0f would NOT work as it would be rounded up to 134217728.0f, which is
+	// outside the range.
+
+	constexpr uint32 TangentZMaxBits = 16u;
+	constexpr float TangentZMinValue = -32768.0f;		// -2^(MaxBits-1)
+	constexpr float TangentZMaxValue = 32767.0f;		//  2^(MaxBits-1)-1
+
+	struct FQuantizedDelta
+	{
+		FIntVector	Position;
+		FIntVector	TangentZ;
+		uint32		Index;
+	};
+	TArray<FQuantizedDelta> QuantizedDeltas;
+	QuantizedDeltas.Reserve(InMorphDeltas.Num());
+
+	bool bVertexIndicesSorted = true;
+
+	int32 PrevVertexIndex = -1;
+	for (int32 DeltaIndex = 0; DeltaIndex < InMorphDeltas.Num(); DeltaIndex++)
+	{
+		const FMorphTargetDelta& MorphDelta = InMorphDeltas[DeltaIndex];
+		const FVector3f TangentZDelta = !InVertexNeedsTangents || (InVertexNeedsTangents->IsValidIndex(MorphDelta.SourceIdx) && (*InVertexNeedsTangents)[MorphDelta.SourceIdx]) ? MorphDelta.TangentZDelta : FVector3f::ZeroVector;
+
+		// Check if input is sorted. It usually is, but it might not be.
+		if (static_cast<int32>(MorphDelta.SourceIdx) < PrevVertexIndex)
+			bVertexIndicesSorted = false;
+		PrevVertexIndex = static_cast<int32>(MorphDelta.SourceIdx);
+
+		// Quantize delta
+		FQuantizedDelta QuantizedDelta;
+		const FVector3f& PositionDelta = MorphDelta.PositionDelta;
+		QuantizedDelta.Position.X = FMath::RoundToInt(FMath::Clamp(PositionDelta.X * RcpPositionPrecision, PositionMinValue, PositionMaxValue));
+		QuantizedDelta.Position.Y = FMath::RoundToInt(FMath::Clamp(PositionDelta.Y * RcpPositionPrecision, PositionMinValue, PositionMaxValue));
+		QuantizedDelta.Position.Z = FMath::RoundToInt(FMath::Clamp(PositionDelta.Z * RcpPositionPrecision, PositionMinValue, PositionMaxValue));
+		QuantizedDelta.TangentZ.X = FMath::RoundToInt(FMath::Clamp(TangentZDelta.X * RcpTangentZPrecision, TangentZMinValue, TangentZMaxValue));
+		QuantizedDelta.TangentZ.Y = FMath::RoundToInt(FMath::Clamp(TangentZDelta.Y * RcpTangentZPrecision, TangentZMinValue, TangentZMaxValue));
+		QuantizedDelta.TangentZ.Z = FMath::RoundToInt(FMath::Clamp(TangentZDelta.Z * RcpTangentZPrecision, TangentZMinValue, TangentZMaxValue));
+		QuantizedDelta.Index = MorphDelta.SourceIdx;
+
+		if (QuantizedDelta.Position != FIntVector::ZeroValue || QuantizedDelta.TangentZ != FIntVector::ZeroValue)
+		{
+			// Only add delta if it is non-zero
+			QuantizedDeltas.Add(QuantizedDelta);
+		}
+	}
+
+	// Sort deltas if the source wasn't already sorted
+	if (!bVertexIndicesSorted)
+	{
+		Algo::Sort(QuantizedDeltas, [](const FQuantizedDelta& A, const FQuantizedDelta& B) { return A.Index < B.Index; });
+	}
+
+	// Encode batch deltas
+	const uint32 MorphNumBatches = (QuantizedDeltas.Num() + BatchSize - 1u) / BatchSize;
+	for (uint32 BatchIndex = 0; BatchIndex < MorphNumBatches; BatchIndex++)
+	{
+		const uint32 BatchFirstElementIndex = BatchIndex * BatchSize;
+		const uint32 NumElements = FMath::Min(BatchSize, QuantizedDeltas.Num() - BatchFirstElementIndex);
+
+		// Calculate batch min/max bounds
+		uint32 IndexMin = MAX_uint32;
+		uint32 IndexMax = MIN_uint32;
+		FIntVector PositionMin = FIntVector(MAX_int32);
+		FIntVector PositionMax = FIntVector(MIN_int32);
+		FIntVector TangentZMin = FIntVector(MAX_int32);
+		FIntVector TangentZMax = FIntVector(MIN_int32);
+
+		for (uint32 LocalElementIndex = 0; LocalElementIndex < NumElements; LocalElementIndex++)
+		{
+			const FQuantizedDelta& Delta = QuantizedDeltas[BatchFirstElementIndex + LocalElementIndex];
+
+			// Trick: Deltas are sorted by index, so the index increase by at least one per delta.
+			//		  Naively this would mean that a batch always spans at least 64 index values and
+			//		  indices would have to use at least 6 bits per index.
+			//		  If instead of storing the raw index, we store the index relative to its position in the batch,
+			//		  then the spanned range becomes 63 smaller.
+			//		  For a consecutive range this even gets us down to 0 bits per index!
+			check(Delta.Index >= LocalElementIndex);
+			const uint32 AdjustedIndex = Delta.Index - LocalElementIndex;
+			IndexMin = FMath::Min(IndexMin, AdjustedIndex);
+			IndexMax = FMath::Max(IndexMax, AdjustedIndex);
+
+			PositionMin.X = FMath::Min(PositionMin.X, Delta.Position.X);
+			PositionMin.Y = FMath::Min(PositionMin.Y, Delta.Position.Y);
+			PositionMin.Z = FMath::Min(PositionMin.Z, Delta.Position.Z);
+
+			PositionMax.X = FMath::Max(PositionMax.X, Delta.Position.X);
+			PositionMax.Y = FMath::Max(PositionMax.Y, Delta.Position.Y);
+			PositionMax.Z = FMath::Max(PositionMax.Z, Delta.Position.Z);
+
+			TangentZMin.X = FMath::Min(TangentZMin.X, Delta.TangentZ.X);
+			TangentZMin.Y = FMath::Min(TangentZMin.Y, Delta.TangentZ.Y);
+			TangentZMin.Z = FMath::Min(TangentZMin.Z, Delta.TangentZ.Z);
+
+			TangentZMax.X = FMath::Max(TangentZMax.X, Delta.TangentZ.X);
+			TangentZMax.Y = FMath::Max(TangentZMax.Y, Delta.TangentZ.Y);
+			TangentZMax.Z = FMath::Max(TangentZMax.Z, Delta.TangentZ.Z);
+		}
+
+		const uint32 IndexDelta = IndexMax - IndexMin;
+		const FIntVector PositionDelta = PositionMax - PositionMin;
+		const FIntVector TangentZDelta = TangentZMax - TangentZMin;
+		const bool bBatchHasTangents = TangentZMin != FIntVector::ZeroValue || TangentZMax != FIntVector::ZeroValue;
+
+		FDeltaBatchHeader& BatchHeader = OutBatchHeaders.AddDefaulted_GetRef();
+		BatchHeader.DataOffset = OutCompressedVertices.Num() * sizeof(uint32);
+		BatchHeader.bTangents = bBatchHasTangents;
+		BatchHeader.NumElements = NumElements;
+		BatchHeader.IndexBits = FMath::CeilLogTwo(IndexDelta + 1);
+		BatchHeader.PositionBits.X = FMath::CeilLogTwo(static_cast<uint32>(PositionDelta.X) + 1);
+		BatchHeader.PositionBits.Y = FMath::CeilLogTwo(static_cast<uint32>(PositionDelta.Y) + 1);
+		BatchHeader.PositionBits.Z = FMath::CeilLogTwo(static_cast<uint32>(PositionDelta.Z) + 1);
+		BatchHeader.TangentZBits.X = FMath::CeilLogTwo(static_cast<uint32>(TangentZDelta.X) + 1);
+		BatchHeader.TangentZBits.Y = FMath::CeilLogTwo(static_cast<uint32>(TangentZDelta.Y) + 1);
+		BatchHeader.TangentZBits.Z = FMath::CeilLogTwo(static_cast<uint32>(TangentZDelta.Z) + 1);
+		check(BatchHeader.IndexBits <= IndexMaxBits);
+		check(BatchHeader.PositionBits.X <= PositionMaxBits);
+		check(BatchHeader.PositionBits.Y <= PositionMaxBits);
+		check(BatchHeader.PositionBits.Z <= PositionMaxBits);
+		check(BatchHeader.TangentZBits.X <= TangentZMaxBits);
+		check(BatchHeader.TangentZBits.Y <= TangentZMaxBits);
+		check(BatchHeader.TangentZBits.Z <= TangentZMaxBits);
+		BatchHeader.IndexMin = IndexMin;
+		BatchHeader.PositionMin = PositionMin;
+		BatchHeader.TangentZMin = TangentZMin;
+
+		// Write quantized bits
+		FDwordBitWriter BitWriter(OutCompressedVertices);
+		for (uint32 LocalElementIndex = 0; LocalElementIndex < NumElements; LocalElementIndex++)
+		{
+			const FQuantizedDelta& Delta = QuantizedDeltas[BatchFirstElementIndex + LocalElementIndex];
+			const uint32 AdjustedIndex = Delta.Index - LocalElementIndex;
+			BitWriter.PutBits(AdjustedIndex - IndexMin, BatchHeader.IndexBits);
+			BitWriter.PutBits(static_cast<uint32>(Delta.Position.X - PositionMin.X), BatchHeader.PositionBits.X);
+			BitWriter.PutBits(static_cast<uint32>(Delta.Position.Y - PositionMin.Y), BatchHeader.PositionBits.Y);
+			BitWriter.PutBits(static_cast<uint32>(Delta.Position.Z - PositionMin.Z), BatchHeader.PositionBits.Z);
+			if (bBatchHasTangents)
+			{
+				BitWriter.PutBits(static_cast<uint32>(Delta.TangentZ.X - TangentZMin.X), BatchHeader.TangentZBits.X);
+				BitWriter.PutBits(static_cast<uint32>(Delta.TangentZ.Y - TangentZMin.Y), BatchHeader.TangentZBits.Y);
+				BitWriter.PutBits(static_cast<uint32>(Delta.TangentZ.Z - TangentZMin.Z), BatchHeader.TangentZBits.Z);
+			}
+		}
+		BitWriter.Flush();
+	}
+}
+
+
+static FMorphTargetDelta DecodeMorphTargetDelta(
+	const FDeltaBatchHeader& InHeader,
+	FDwordBitReader& InReader,
+	uint32 InLocalIndex, 
+	const float InPositionPrecision,
+	const float InTangentZPrecision
+	)
+{
+	FMorphTargetDelta Result;
+	
+	Result.SourceIdx = InReader.GetBits(InHeader.IndexBits) + InHeader.IndexMin + InLocalIndex;
+	
+	Result.PositionDelta.X = static_cast<float>(static_cast<int32>(InReader.GetBits(InHeader.PositionBits.X)) + InHeader.PositionMin.X) * InPositionPrecision;
+	Result.PositionDelta.Y = static_cast<float>(static_cast<int32>(InReader.GetBits(InHeader.PositionBits.Y)) + InHeader.PositionMin.Y) * InPositionPrecision;
+	Result.PositionDelta.Z = static_cast<float>(static_cast<int32>(InReader.GetBits(InHeader.PositionBits.Z)) + InHeader.PositionMin.Z) * InPositionPrecision;
+	
+	if (InHeader.bTangents)
+	{
+		Result.TangentZDelta.X = static_cast<float>(static_cast<int32>(InReader.GetBits(InHeader.TangentZBits.X)) + InHeader.TangentZMin.X) * InTangentZPrecision;
+		Result.TangentZDelta.Y = static_cast<float>(static_cast<int32>(InReader.GetBits(InHeader.TangentZBits.Y)) + InHeader.TangentZMin.Y) * InTangentZPrecision;
+		Result.TangentZDelta.Z = static_cast<float>(static_cast<int32>(InReader.GetBits(InHeader.TangentZBits.Z)) + InHeader.TangentZMin.Z) * InTangentZPrecision;
+	}
+	
+	return Result;
+}
+
+
+void Decode(
+	TConstArrayView<FDeltaBatchHeader> InBatchHeaders,
+	TConstArrayView<uint32> InCompressedVertices,
+	const float InPositionPrecision,
+	const float InTangentZPrecision, 
+	TArray<FMorphTargetDelta>& OutMorphDeltas
+	)
+{
+	// Tally up how many deltas we'll need.
+	int32 NumDeltas = 0;
+	for (const FDeltaBatchHeader& Header: InBatchHeaders)
+	{
+		NumDeltas += static_cast<int32>(Header.NumElements);
+	}
+
+	OutMorphDeltas.Reset(NumDeltas);
+
+	FDwordBitReader Reader(InCompressedVertices);
+	
+	for (const FDeltaBatchHeader& Header: InBatchHeaders)
+	{
+		for (uint32 ElementIndex = 0; ElementIndex < Header.NumElements; ElementIndex++)
+		{
+			OutMorphDeltas.Add(DecodeMorphTargetDelta(Header, Reader, ElementIndex, InPositionPrecision, InTangentZPrecision));
+		}
+	}
+}
+
+
+bool IterativeDecode(
+	uint64& InOutNextItemToken,
+	TConstArrayView<FDeltaBatchHeader> InBatchHeaders,
+	TConstArrayView<uint32> InCompressedVertices,
+	const float InPositionPrecision,
+	const float InTangentZPrecision,
+	FMorphTargetDelta& OutMorphTargetDelta
+	)
+{
+	// Make sure this is not the same as uint64::max() since the iterator in MorphTarget.h uses that as an
+	// iterator invalidation.
+	static constexpr uint64 EndOfStreamToken = 0xEFFFFFFFFFFFFFFFull;
+	
+	// We've reached the end.
+	if (InOutNextItemToken == EndOfStreamToken)
+	{
+		return false;
+	}
+
+	// We assume the token is intact.
+	uint32 HeaderIndex = (static_cast<uint32>(InOutNextItemToken) >> BatchSizeBits) & ((32 - BatchSize) - 1);
+	uint32 ElementIndex = InOutNextItemToken & (BatchSize - 1);
+	uint32 DataOffset =  InOutNextItemToken >> 32;
+
+	check(HeaderIndex < static_cast<uint32>(InBatchHeaders.Num()));
+	
+	const FDeltaBatchHeader& BatchHeader = InBatchHeaders[HeaderIndex];
+	check(ElementIndex < BatchHeader.NumElements);
+
+	FDwordBitReader Reader(InCompressedVertices, DataOffset);
+	OutMorphTargetDelta = DecodeMorphTargetDelta(BatchHeader, Reader, ElementIndex, InPositionPrecision, InTangentZPrecision);
+	
+	DataOffset = Reader.GetOffset();
+	ElementIndex++;
+	if (ElementIndex == BatchHeader.NumElements)
+	{
+		ElementIndex = 0;
+		HeaderIndex++;
+	}
+
+	if (HeaderIndex == InBatchHeaders.Num())
+	{
+		InOutNextItemToken = EndOfStreamToken;
+	}
+	else
+	{
+		InOutNextItemToken = (static_cast<uint64>(DataOffset) << 32ull) | (HeaderIndex << BatchSizeBits) | ElementIndex;
+	}
+	
+	return true;	
+}
+
+} // namespace UE::MorphTargetVertexCodec
diff --git a/Engine/Source/Runtime/Engine/Private/Rendering/MorphTargetVertexInfoBuffers.cpp b/Engine/Source/Runtime/Engine/Private/Rendering/MorphTargetVertexInfoBuffers.cpp
--- a/Engine/Source/Runtime/Engine/Private/Rendering/MorphTargetVertexInfoBuffers.cpp
+++ b/Engine/Source/Runtime/Engine/Private/Rendering/MorphTargetVertexInfoBuffers.cpp
@@ -55,12 +55,6 @@ uint32 FMorphTargetVertexInfoBuffers::GetMaximumThreadGroupSize()
 	return uint32(FMath::Min<uint64>(MaximumThreadGroupSize, UINT32_MAX));
 }
 
-const float FMorphTargetVertexInfoBuffers::CalculatePositionPrecision(float TargetPositionErrorTolerance)
-{
-	const float UnrealUnitPerMeter = 100.0f;
-	return TargetPositionErrorTolerance * 2.0f * 1e-6f * UnrealUnitPerMeter;	// * 2.0 because correct rounding guarantees error is at most half of the cell size.
-}
-
 void FMorphTargetVertexInfoBuffers::ResetCPUData()
 {
 	MorphData.Empty();
@@ -223,57 +217,12 @@ void FMorphTargetVertexInfoBuffers::InitMorphResourcesStreaming(const TArray<FSk
 
 void FMorphTargetVertexInfoBuffers::Compress(const TArray<FSkelMeshRenderSection>& RenderSections, const TArray<const FMorphTargetLODModel*>& MorphTargets, const TBitArray<>& UsesBuiltinMorphTargetCompression, int32 NumVertices, float TargetPositionErrorTolerance)
 {
-	// Simple Morph compression 0.1
-	// Instead of storing vertex deltas individually they are organized into batches of 64.
-	// Each batch has a header that describes how many bits are allocated to each of the vertex components.
-	// Batches also store an explicit offset to its associated data. This makes it trivial to decode batches
-	// in parallel, and because deltas are fixed-width inside a batch, deltas can also be decoded in parallel.
-	// The result is a semi-adaptive encoding that functions as a crude substitute for entropy coding, that is
-	// fast to decode on parallel hardware.
-
-	// Quantization still happens globally to avoid issues with cracks at duplicate vertices.
-	// The quantization is artist controlled on a per LOD basis. Higher error tolerance results in smaller deltas
-	// and a smaller compressed size.
-
-	PositionPrecision = CalculatePositionPrecision(TargetPositionErrorTolerance);
-	const float RcpPositionPrecision = 1.0f / PositionPrecision;
-
-	TangentZPrecision = 1.0f / 2048.0f;				// Object scale irrelevant here. Let's assume ~12bits per component is plenty.
-	const float RcpTangentZPrecision = 1.0f / TangentZPrecision;
-
-	const uint32 BatchSize = 64;
-	const uint32 NumBatchHeaderDwords = 10u;
-
-	const uint32 IndexMaxBits = 31u;
-
-	const uint32 PositionMaxBits = 28u;				// Probably more than we need, but let's just allow it to go this high to be safe for now.
-													// For larger deltas this can even be more precision than what was in the float input data!
-													// Maybe consider float-like or exponential encoding of large values?
-	const float PositionMinValue = -134217728.0f;	// -2^(MaxBits-1)
-	const float PositionMaxValue = 134217720.0f;	// Largest float smaller than 2^(MaxBits-1)-1
-													// Using 134217727.0f would NOT work as it would be rounded up to 134217728.0f, which is
-													// outside the range.
-
-	const uint32 TangentZMaxBits = 16u;
-	const float TangentZMinValue = -32768.0f;		// -2^(MaxBits-1)
-	const float TangentZMaxValue = 32767.0f;		//  2^(MaxBits-1)-1
-
-	struct FBatchHeader
-	{
-		uint32		DataOffset;
-		uint32		NumElements;
-		bool		bTangents;
-
-		uint32		IndexBits;
-		FIntVector	PositionBits;
-		FIntVector	TangentZBits;
-
-		uint32		IndexMin;
-		FIntVector	PositionMin;
-		FIntVector	TangentZMin;
-	};
-
-	// uint32 StartTime = FPlatformTime::Cycles();
+	constexpr uint32 NumBatchHeaderDwords = 10u;
+	
+	using namespace UE::MorphTargetVertexCodec;
+	
+	PositionPrecision = ComputePositionPrecision(TargetPositionErrorTolerance);
+	TangentZPrecision = ComputeTangentPrecision();
 
 	MorphData.Empty();
 	NumTotalBatches = 0;
@@ -300,7 +249,7 @@ void FMorphTargetVertexInfoBuffers::Compress(const TArray<FSkelMeshRenderSection
 	}
 
 	// Populate the arrays to be filled in later in the render thread
-	TArray<FBatchHeader> BatchHeaders;
+	TArray<FDeltaBatchHeader> BatchHeaders;
 	TArray<uint32> BitstreamData;
 	for (int32 AnimIdx = 0; AnimIdx < MorphTargets.Num(); ++AnimIdx)
 	{
@@ -328,18 +277,6 @@ void FMorphTargetVertexInfoBuffers::Compress(const TArray<FSkelMeshRenderSection
 		{
 			const FMorphTargetDelta* MorphDeltas = MorphModel->Vertices.GetData();
 
-			struct FQuantizedDelta
-			{
-				FIntVector	Position;
-				FIntVector	TangentZ;
-				uint32		Index;
-			};
-			TArray<FQuantizedDelta> QuantizedDeltas;
-			QuantizedDeltas.Reserve(NumSrcDeltas);
-
-			bool bVertexIndicesSorted = true;
-
-			int32 PrevVertexIndex = -1;
 			for (int32 DeltaIndex = 0; DeltaIndex < NumSrcDeltas; DeltaIndex++)
 			{
 				const auto& MorphDelta = MorphDeltas[DeltaIndex];
@@ -356,143 +293,23 @@ void FMorphTargetVertexInfoBuffers::Compress(const TArray<FSkelMeshRenderSection
 				MinimumValues[1] = FMath::Min(MinimumValues[1], MorphDelta.PositionDelta.Y);
 				MinimumValues[2] = FMath::Min(MinimumValues[2], MorphDelta.PositionDelta.Z);
 				MinimumValues[3] = FMath::Min(MinimumValues[3], FMath::Min(TangentZDelta.X, FMath::Min(TangentZDelta.Y, TangentZDelta.Z)));
-
-				// Check if input is sorted. It usually is, but it might not be.
-				if ((int32)MorphDelta.SourceIdx < PrevVertexIndex)
-					bVertexIndicesSorted = false;
-				PrevVertexIndex = (int32)MorphDelta.SourceIdx;
-
-				// Quantize delta
-				FQuantizedDelta QuantizedDelta;
-				const FVector3f& PositionDelta = MorphDelta.PositionDelta;
-				QuantizedDelta.Position.X = FMath::RoundToInt(FMath::Clamp(PositionDelta.X * RcpPositionPrecision, PositionMinValue, PositionMaxValue));
-				QuantizedDelta.Position.Y = FMath::RoundToInt(FMath::Clamp(PositionDelta.Y * RcpPositionPrecision, PositionMinValue, PositionMaxValue));
-				QuantizedDelta.Position.Z = FMath::RoundToInt(FMath::Clamp(PositionDelta.Z * RcpPositionPrecision, PositionMinValue, PositionMaxValue));
-				QuantizedDelta.TangentZ.X = FMath::RoundToInt(FMath::Clamp(TangentZDelta.X * RcpTangentZPrecision, TangentZMinValue, TangentZMaxValue));
-				QuantizedDelta.TangentZ.Y = FMath::RoundToInt(FMath::Clamp(TangentZDelta.Y * RcpTangentZPrecision, TangentZMinValue, TangentZMaxValue));
-				QuantizedDelta.TangentZ.Z = FMath::RoundToInt(FMath::Clamp(TangentZDelta.Z * RcpTangentZPrecision, TangentZMinValue, TangentZMaxValue));
-				QuantizedDelta.Index = MorphDelta.SourceIdx;
-
-				if (QuantizedDelta.Position != FIntVector::ZeroValue || QuantizedDelta.TangentZ != FIntVector::ZeroValue)
-				{
-					// Only add delta if it is non-zero
-					QuantizedDeltas.Add(QuantizedDelta);
-				}
 			}
 
-			// Sort deltas if the source wasn't already sorted
-			if (!bVertexIndicesSorted)
-			{
-				Algo::Sort(QuantizedDeltas, [](const FQuantizedDelta& A, const FQuantizedDelta& B) { return A.Index < B.Index; });
-			}
+			// Encode the actual morph vertex info into the quantized bitstream.
+			Encode(MorphModel->Vertices, &VertexNeedsTangents, PositionPrecision, TangentZPrecision, BatchHeaders, BitstreamData);
 
-			// Encode batch deltas
-			const uint32 MorphNumBatches = (QuantizedDeltas.Num() + BatchSize - 1u) / BatchSize;
-			for (uint32 BatchIndex = 0; BatchIndex < MorphNumBatches; BatchIndex++)
-			{
-				const uint32 BatchFirstElementIndex = BatchIndex * BatchSize;
-				const uint32 NumElements = FMath::Min(BatchSize, QuantizedDeltas.Num() - BatchFirstElementIndex);
-
-				// Calculate batch min/max bounds
-				uint32 IndexMin = MAX_uint32;
-				uint32 IndexMax = MIN_uint32;
-				FIntVector PositionMin = FIntVector(MAX_int32);
-				FIntVector PositionMax = FIntVector(MIN_int32);
-				FIntVector TangentZMin = FIntVector(MAX_int32);
-				FIntVector TangentZMax = FIntVector(MIN_int32);
-
-				for (uint32 LocalElementIndex = 0; LocalElementIndex < NumElements; LocalElementIndex++)
-				{
-					const FQuantizedDelta& Delta = QuantizedDeltas[BatchFirstElementIndex + LocalElementIndex];
-
-					// Trick: Deltas are sorted by index, so the index increase by at least one per delta.
-					//		  Naively this would mean that a batch always spans at least 64 index values and
-					//		  indices would have to use at least 6 bits per index.
-					//		  If instead of storing the raw index, we store the index relative to its position in the batch,
-					//		  then the spanned range becomes 63 smaller.
-					//		  For a consecutive range this even gets us down to 0 bits per index!
-					check(Delta.Index >= LocalElementIndex);
-					const uint32 AdjustedIndex = Delta.Index - LocalElementIndex;
-					IndexMin = FMath::Min(IndexMin, AdjustedIndex);
-					IndexMax = FMath::Max(IndexMax, AdjustedIndex);
-
-					PositionMin.X = FMath::Min(PositionMin.X, Delta.Position.X);
-					PositionMin.Y = FMath::Min(PositionMin.Y, Delta.Position.Y);
-					PositionMin.Z = FMath::Min(PositionMin.Z, Delta.Position.Z);
-
-					PositionMax.X = FMath::Max(PositionMax.X, Delta.Position.X);
-					PositionMax.Y = FMath::Max(PositionMax.Y, Delta.Position.Y);
-					PositionMax.Z = FMath::Max(PositionMax.Z, Delta.Position.Z);
-
-					TangentZMin.X = FMath::Min(TangentZMin.X, Delta.TangentZ.X);
-					TangentZMin.Y = FMath::Min(TangentZMin.Y, Delta.TangentZ.Y);
-					TangentZMin.Z = FMath::Min(TangentZMin.Z, Delta.TangentZ.Z);
-
-					TangentZMax.X = FMath::Max(TangentZMax.X, Delta.TangentZ.X);
-					TangentZMax.Y = FMath::Max(TangentZMax.Y, Delta.TangentZ.Y);
-					TangentZMax.Z = FMath::Max(TangentZMax.Z, Delta.TangentZ.Z);
-				}
-
-				const uint32 IndexDelta = IndexMax - IndexMin;
-				const FIntVector PositionDelta = PositionMax - PositionMin;
-				const FIntVector TangentZDelta = TangentZMax - TangentZMin;
-				const bool bBatchHasTangents = TangentZMin != FIntVector::ZeroValue || TangentZMax != FIntVector::ZeroValue;
-
-				FBatchHeader BatchHeader;
-				BatchHeader.DataOffset = BitstreamData.Num() * sizeof(uint32);
-				BatchHeader.bTangents = bBatchHasTangents;
-				BatchHeader.NumElements = NumElements;
-				BatchHeader.IndexBits = FMath::CeilLogTwo(IndexDelta + 1);
-				BatchHeader.PositionBits.X = FMath::CeilLogTwo(uint32(PositionDelta.X) + 1);
-				BatchHeader.PositionBits.Y = FMath::CeilLogTwo(uint32(PositionDelta.Y) + 1);
-				BatchHeader.PositionBits.Z = FMath::CeilLogTwo(uint32(PositionDelta.Z) + 1);
-				BatchHeader.TangentZBits.X = FMath::CeilLogTwo(uint32(TangentZDelta.X) + 1);
-				BatchHeader.TangentZBits.Y = FMath::CeilLogTwo(uint32(TangentZDelta.Y) + 1);
-				BatchHeader.TangentZBits.Z = FMath::CeilLogTwo(uint32(TangentZDelta.Z) + 1);
-				check(BatchHeader.IndexBits <= IndexMaxBits);
-				check(BatchHeader.PositionBits.X <= PositionMaxBits);
-				check(BatchHeader.PositionBits.Y <= PositionMaxBits);
-				check(BatchHeader.PositionBits.Z <= PositionMaxBits);
-				check(BatchHeader.TangentZBits.X <= TangentZMaxBits);
-				check(BatchHeader.TangentZBits.Y <= TangentZMaxBits);
-				check(BatchHeader.TangentZBits.Z <= TangentZMaxBits);
-				BatchHeader.IndexMin = IndexMin;
-				BatchHeader.PositionMin = PositionMin;
-				BatchHeader.TangentZMin = TangentZMin;
-
-				// Write quantized bits
-				FDwordBitWriter BitWriter(BitstreamData);
-				for (uint32 LocalElementIndex = 0; LocalElementIndex < NumElements; LocalElementIndex++)
-				{
-					const FQuantizedDelta& Delta = QuantizedDeltas[BatchFirstElementIndex + LocalElementIndex];
-					const uint32 AdjustedIndex = Delta.Index - LocalElementIndex;
-					BitWriter.PutBits(AdjustedIndex - IndexMin, BatchHeader.IndexBits);
-					BitWriter.PutBits(uint32(Delta.Position.X - PositionMin.X), BatchHeader.PositionBits.X);
-					BitWriter.PutBits(uint32(Delta.Position.Y - PositionMin.Y), BatchHeader.PositionBits.Y);
-					BitWriter.PutBits(uint32(Delta.Position.Z - PositionMin.Z), BatchHeader.PositionBits.Z);
-					if (bBatchHasTangents)
-					{
-						BitWriter.PutBits(uint32(Delta.TangentZ.X - TangentZMin.X), BatchHeader.TangentZBits.X);
-						BitWriter.PutBits(uint32(Delta.TangentZ.Y - TangentZMin.Y), BatchHeader.TangentZBits.Y);
-						BitWriter.PutBits(uint32(Delta.TangentZ.Z - TangentZMin.Z), BatchHeader.TangentZBits.Z);
-					}
-				}
-				BitWriter.Flush();
-
-				BatchHeaders.Add(BatchHeader);
-			}
-			NumTotalBatches += MorphNumBatches;
-		}
+			NumTotalBatches += BatchHeaders.Num();
 
-		const uint32 MorphNumBatches = NumTotalBatches - BatchStartOffset;
-		BatchStartOffsetPerMorph.Add(BatchStartOffset);
-		BatchesPerMorph.Add(MorphNumBatches);
-		MaximumValuePerMorph.Add(FVector4f(MaximumValues[0], MaximumValues[1], MaximumValues[2], MaximumValues[3]));
-		MinimumValuePerMorph.Add(FVector4f(MinimumValues[0], MinimumValues[1], MinimumValues[2], MinimumValues[3]));
+			const uint32 MorphNumBatches = NumTotalBatches - BatchStartOffset;
+			BatchStartOffsetPerMorph.Add(BatchStartOffset);
+			BatchesPerMorph.Add(MorphNumBatches);
+			MaximumValuePerMorph.Add(FVector4f(MaximumValues[0], MaximumValues[1], MaximumValues[2], MaximumValues[3]));
+			MinimumValuePerMorph.Add(FVector4f(MinimumValues[0], MinimumValues[1], MinimumValues[2], MinimumValues[3]));
+		}
 	}
 
 	// Write packed batch headers
-	for (const FBatchHeader& BatchHeader : BatchHeaders)
+	for (const FDeltaBatchHeader& BatchHeader : BatchHeaders)
 	{
 		const uint32 DataOffset = BatchHeader.DataOffset + BatchHeaders.Num() * NumBatchHeaderDwords * sizeof(uint32);
 
@@ -514,7 +331,7 @@ void FMorphTargetVertexInfoBuffers::Compress(const TArray<FSkelMeshRenderSection
 
 	// Append bitstream data
 	MorphData.Append(BitstreamData);
-	
+
 	if(MorphData.Num() > 0)
 	{
 		// Pad to make sure it is always safe to access the data with load4s.
diff --git a/Engine/Source/Runtime/Engine/Private/SkeletalMeshLODRenderData.cpp b/Engine/Source/Runtime/Engine/Private/SkeletalMeshLODRenderData.cpp
--- a/Engine/Source/Runtime/Engine/Private/SkeletalMeshLODRenderData.cpp
+++ b/Engine/Source/Runtime/Engine/Private/SkeletalMeshLODRenderData.cpp
@@ -16,7 +16,6 @@
 #include "Serialization/MemoryReader.h"
 #include "SkeletalMeshLegacyCustomVersions.h"
 #include "UObject/Package.h"
-#include "Stats/StatsTrace.h"
 #include "UObject/UE5MainStreamObjectVersion.h"
 #include "UObject/UE5SpecialProjectStreamObjectVersion.h"
 #include "GPUSkinCache.h"
@@ -870,7 +869,7 @@ void FSkeletalMeshLODRenderData::SerializeStreamedData(FArchive& Ar, USkinnedAss
 
 #if DO_CHECK
 				const FSkeletalMeshLODInfo* SkeletalMeshLODInfo = Owner->GetLODInfo(LODIdx);
-				float PositionPrecision = FMorphTargetVertexInfoBuffers::CalculatePositionPrecision(SkeletalMeshLODInfo->MorphTargetPositionErrorTolerance);
+				float PositionPrecision = UE::MorphTargetVertexCodec::ComputePositionPrecision(SkeletalMeshLODInfo->MorphTargetPositionErrorTolerance);
 				checkf(PositionPrecision == MorphTargetVertexInfoBuffers.GetPositionPrecision(), TEXT("Morph target render data position tollerance for the skeleton %s does not match the expected tolerance in the LOD info."), *Owner->GetName());
 #endif
 			}
diff --git a/Engine/Source/Runtime/Engine/Private/SkeletalRenderCPUSkin.cpp b/Engine/Source/Runtime/Engine/Private/SkeletalRenderCPUSkin.cpp
--- a/Engine/Source/Runtime/Engine/Private/SkeletalRenderCPUSkin.cpp
+++ b/Engine/Source/Runtime/Engine/Private/SkeletalRenderCPUSkin.cpp
@@ -561,12 +561,9 @@ struct FMorphTargetInfo
 	/** The index into the morph weight list */
 	int32						WeightIndex = INDEX_NONE;
 	
-	/** Index of next delta to try applying. This prevents us looking at every delta for every vertex. */
-	int32						NextDeltaIndex = INDEX_NONE;
-	/** Array of deltas to apply to mesh, sorted based on the index of the base mesh vert that they affect. */
-	const FMorphTargetDelta*	Deltas = nullptr;
-	/** How many deltas are in array */
-	int32						NumDeltas = 0;
+	/** An iterator into the list of morph target deltas, depending on build, this can either be the raw
+	 *  vertex deltas or the compressed ones, generated at cook time. */
+	FMorphTargetDeltaIterator	VertexDeltaIterator;
 };
 
 /**
@@ -594,8 +591,7 @@ static uint32 InitEvalInfos(const FMorphTargetWeightMap& InActiveMorphTargets, c
 		{
 			// start at the first vertex since they affect base mesh verts in ascending order
 			NewInfo.WeightIndex = WeightIndex;
-			NewInfo.NextDeltaIndex = 0;
-			NewInfo.Deltas = MorphTarget->GetMorphTargetDelta(LODIndex, NewInfo.NumDeltas);
+			NewInfo.VertexDeltaIterator = MorphTarget->GetDeltaIteratorForLOD(LODIndex);
 
 			NumValidMorphTargets++;
 		}
@@ -658,16 +654,11 @@ FORCEINLINE void UpdateMorphedVertex( VertexType& MorphedVertex, const VertexTyp
 	for( int32 MorphIdx=0; MorphIdx < EvalInfos.Num(); MorphIdx++ )
 	{
 		FMorphTargetInfo& Info = EvalInfos[MorphIdx];
-
-		// if the next delta to use matches the current vertex, apply it
-		if( Info.NextDeltaIndex != INDEX_NONE &&
-			Info.NextDeltaIndex < Info.NumDeltas &&
-			Info.Deltas[Info.NextDeltaIndex].SourceIdx == CurBaseVertIdx )
+		
+		if (!Info.VertexDeltaIterator.AtEnd() && Info.VertexDeltaIterator->SourceIdx == CurBaseVertIdx)
 		{
-			ApplyMorphBlend( MorphedVertex, Info.Deltas[Info.NextDeltaIndex], MorphWeights[Info.WeightIndex] );
-
-			// Update 'next delta to use'
-			Info.NextDeltaIndex += 1;
+			ApplyMorphBlend(MorphedVertex, *Info.VertexDeltaIterator, MorphWeights[MorphIdx]);
+			++Info.VertexDeltaIterator;
 		}
 	}
 
@@ -1145,13 +1136,11 @@ static void CalculateMorphTargetWeights(FFinalSkinVertex* DestVertex, FSkeletalM
 		ClearVert->TextureCoordinates[0].Y = 0.0f;
 	}
 
-	for (const UMorphTarget* Morphtarget : InMorphTargetOfInterest)
+	for (const UMorphTarget* MorphTarget : InMorphTargetOfInterest)
 	{
-		int32 NumDeltas;
-		const FMorphTargetDelta* MTLODVertices = Morphtarget->GetMorphTargetDelta(LODIndex, NumDeltas);
-		for (int32 MorphVertexIndex = 0; MorphVertexIndex < NumDeltas; ++MorphVertexIndex)
+		for (FMorphTargetDeltaIterator DeltaIt = MorphTarget->GetDeltaIteratorForLOD(LODIndex); !DeltaIt.AtEnd(); ++DeltaIt)
 		{
-			FFinalSkinVertex* SetVert = DestVertex + MTLODVertices[MorphVertexIndex].SourceIdx;
+			FFinalSkinVertex* SetVert = DestVertex + DeltaIt->SourceIdx;
 			SetVert->TextureCoordinates[0].X += 1.0f;
 			SetVert->TextureCoordinates[0].Y += 1.0f;
 		}
diff --git a/Engine/Source/Runtime/Engine/Private/SkeletalRenderGPUSkin.cpp b/Engine/Source/Runtime/Engine/Private/SkeletalRenderGPUSkin.cpp
--- a/Engine/Source/Runtime/Engine/Private/SkeletalRenderGPUSkin.cpp
+++ b/Engine/Source/Runtime/Engine/Private/SkeletalRenderGPUSkin.cpp
@@ -1811,28 +1811,20 @@ void FSkeletalMeshObjectGPUSkin::FSkeletalMeshObjectLOD::UpdateMorphVertexBuffer
 			const float MorphAbsWeight = FMath::Abs(MorphTargetWeight);
 			checkSlow(MorphAbsWeight >= MinMorphTargetBlendWeight && MorphAbsWeight <= MorphTargetMaxBlendWeight);
 
-			// Get deltas
-			int32 NumDeltas;
-			const FMorphTargetDelta* Deltas = MorphTarget->GetMorphTargetDelta(LODIndex, NumDeltas);
-
 			// iterate over the vertices that this lod model has changed
-			for (int32 MorphVertIdx = 0; MorphVertIdx < NumDeltas; MorphVertIdx++)
+			for (FMorphTargetDeltaIterator DeltaIt = MorphTarget->GetDeltaIteratorForLOD(LODIndex); !DeltaIt.AtEnd(); ++DeltaIt)
 			{
-				const FMorphTargetDelta& MorphVertex = Deltas[MorphVertIdx];
-
-				// @TODO FIXMELH : temp hack until we fix importing issue
-				if ((MorphVertex.SourceIdx < LodData.GetNumVertices()))
+				if (ensure(DeltaIt->SourceIdx < LodData.GetNumVertices()))
 				{
-					FMorphGPUSkinVertex& DestVertex = Buffer[MorphVertex.SourceIdx];
-
-					DestVertex.DeltaPosition += MorphVertex.PositionDelta * MorphTargetWeight;
+					FMorphGPUSkinVertex& DestVertex = Buffer[DeltaIt->SourceIdx];
+					DestVertex.DeltaPosition += DeltaIt->PositionDelta * MorphTargetWeight;
 
 					// todo: could be moved out of the inner loop to be more efficient
 					if (bBlendTangentsOnCPU)
 					{
-						DestVertex.DeltaTangentZ += MorphVertex.TangentZDelta * MorphTargetWeight;
+						DestVertex.DeltaTangentZ += DeltaIt->TangentZDelta * MorphTargetWeight;
 						// accumulate the weight so we can normalized it later
-						MorphAccumulatedWeightArray[MorphVertex.SourceIdx] += MorphAbsWeight;
+						MorphAccumulatedWeightArray[DeltaIt->SourceIdx] += MorphAbsWeight;
 					}
 				}
 			} // for all vertices
diff --git a/Engine/Source/Runtime/Engine/Public/Rendering/MorphTargetVertexCodec.h b/Engine/Source/Runtime/Engine/Public/Rendering/MorphTargetVertexCodec.h
new file mode 100644
--- /dev/null
+++ b/Engine/Source/Runtime/Engine/Public/Rendering/MorphTargetVertexCodec.h
@@ -0,0 +1,115 @@
+﻿// Copyright Epic Games, Inc. All Rights Reserved.
+
+#pragma once
+
+#include "Containers/BitArray.h"
+#include "Containers/ContainersFwd.h"
+#include "Math/IntVector.h"
+
+struct FMorphTargetDelta;
+class UMorphTarget;
+class FDefaultBitArrayAllocator;
+
+
+namespace UE::MorphTargetVertexCodec
+{
+
+struct FDeltaBatchHeader
+{
+	uint32		DataOffset;
+	uint32		NumElements;
+
+	bool		bTangents;
+	int8		IndexBits;
+	Math::TIntVector3<int8>	PositionBits;
+	Math::TIntVector3<int8>	TangentZBits;
+	
+	uint32		IndexMin;
+	FIntVector	PositionMin;
+	FIntVector	TangentZMin;
+
+	friend FArchive& operator<<(FArchive& Ar, FDeltaBatchHeader& InHeader)
+	{
+		Ar << InHeader.DataOffset;
+		Ar << InHeader.NumElements;
+		Ar << InHeader.bTangents;
+		Ar << InHeader.IndexBits;
+		Ar << InHeader.PositionBits;
+		Ar << InHeader.TangentZBits;
+		Ar << InHeader.IndexMin;
+		Ar << InHeader.PositionMin;
+		Ar << InHeader.TangentZMin;
+		
+		return Ar;
+	}
+};
+
+
+inline float ComputePositionPrecision(float InTargetPositionErrorTolerance)
+{
+	constexpr float UnrealUnitPerMeter = 100.0f;
+	return InTargetPositionErrorTolerance * 2.0f * 1e-6f * UnrealUnitPerMeter;	// * 2.0 because correct rounding guarantees error is at most half of the cell size.
+}
+
+inline float ComputeTangentPrecision()
+{
+	return 1.0f / 2048.0f;	// Object scale irrelevant here. Let's assume ~12bits per component is plenty.
+}
+
+/** Encodes the morph deltas, using the given precision, into a pair of arrays -- one to store the header
+  * data that describes a bit-packed batch of vertices and another that stores the actual bit packed vertices
+  * themselves.
+  * Note that due to the compressor throwing away deltas that are under the precision threshold, the number of 
+  * output deltas may not match the number of deltas that end up getting packed.
+  * @param InMorphDeltas The list of morph deltas that will be compressed.
+  * @param InVertexNeedsTangents An optional bit array of vertices that require tangents to be present. If nullptr,
+  *    then all vertices are assumed to require tangents. Used to ignore compressing tangents on sections that compute
+  *    them automatically at render time.
+  * @param InPositionPrecision The precision required for the compression the position data. The greater this value
+  *    the fewer bits are used for storing the resulting position values.
+  * @param InTangentZPrecision The precision required for the compression the tangent data. The greater this value
+  *    the fewer bits are used for storing the resulting tangent values.
+  * @param OutBatchHeaders The list of block headers for each compressed block of morph deltas. Multiple headers are
+  *    required since the final amount of bits required is dependent on the spread of position/tangent values within
+  *    each block.
+  * @param OutCompressedVertices The actual compressed data. 
+  */
+ENGINE_API void Encode(
+	TConstArrayView<FMorphTargetDelta> InMorphDeltas,
+	const TBitArray<>* InVertexNeedsTangents,
+	const float InPositionPrecision,
+	const float InTangentZPrecision,
+	TArray<FDeltaBatchHeader>& OutBatchHeaders,
+	TArray<uint32>& OutCompressedVertices
+	);
+
+ENGINE_API void Encode(
+	TConstArrayView<FMorphTargetDelta> InMorphDeltas,
+	const TBitArray<>* InVertexNeedsTangents,
+	const float InPositionPrecision,
+	const float InTangentZPrecision,
+	TArray<FDeltaBatchHeader>& OutBatchHeaders,
+	TArray<uint32>& OutCompressedVertices
+	);
+
+/** Decode all the vertices given by the two arrays. */ 
+ENGINE_API void Decode(
+	TConstArrayView<FDeltaBatchHeader> InBatchHeaders,
+	TConstArrayView<uint32> InCompressedVertices,
+	const float InPositionPrecision,
+	const float InTangentZPrecision,
+	TArray<FMorphTargetDelta>& OutMorphDeltas
+	);
+
+/** Iteratively decode a single morph target delta entry from the data stream. The InOutNextItemToken is an opaque value that should
+ *  be initialized to zero for the first entry. It will be adjusted to point to the next entry after a successful decode. */ 
+ENGINE_API bool IterativeDecode(
+	uint64& InOutNextItemToken,
+	TConstArrayView<FDeltaBatchHeader> InBatchHeaders,
+	TConstArrayView<uint32> InCompressedVertices,
+	const float InPositionPrecision,
+	const float InTangentZPrecision,
+	FMorphTargetDelta& OutMorphTargetDelta
+	);
+
+}
diff --git a/Engine/Source/Runtime/Engine/Public/Rendering/MorphTargetVertexInfoBuffers.h b/Engine/Source/Runtime/Engine/Public/Rendering/MorphTargetVertexInfoBuffers.h
--- a/Engine/Source/Runtime/Engine/Public/Rendering/MorphTargetVertexInfoBuffers.h
+++ b/Engine/Source/Runtime/Engine/Public/Rendering/MorphTargetVertexInfoBuffers.h
@@ -62,23 +62,21 @@ public:
 		return MinimumValuePerMorph[Index];
 	}
 
-	const float GetPositionPrecision() const
+	float GetPositionPrecision() const
 	{
 		return PositionPrecision;
 	}
 
+	float GetTangentZPrecision() const
+	{
+		return TangentZPrecision;
+	}
+
 	uint64 GetMorphDataSizeInBytes() const 
 	{ 
 		return MorphData.Num() * sizeof(uint32);
 	}
 
-	static ENGINE_API const float CalculatePositionPrecision(float TargetPositionErrorTolerance);
-
-	const float GetTangentZPrecision() const
-	{
-		return TangentZPrecision;
-	}
-
 	static ENGINE_API bool IsPlatformShaderSupported(EShaderPlatform ShaderPlatform);
 
 	FBufferRHIRef MorphDataBuffer;
