Skeletal Mesh: Fix MSVC SA issue in the morph decoder.

#rnx

diff --git a/Engine/Source/Runtime/Engine/Private/Rendering/MorphTargetVertexCodec.cpp b/Engine/Source/Runtime/Engine/Private/Rendering/MorphTargetVertexCodec.cpp
--- a/Engine/Source/Runtime/Engine/Private/Rendering/MorphTargetVertexCodec.cpp
+++ b/Engine/Source/Runtime/Engine/Private/Rendering/MorphTargetVertexCodec.cpp
@@ -388,7 +388,7 @@ bool IterativeDecode(
 	}
 	else
 	{
-		InOutNextItemToken = (static_cast<uint64>(DataOffset) << 32ull) | (HeaderIndex << BatchSizeBits) | ElementIndex;
+		InOutNextItemToken = (static_cast<uint64>(DataOffset) << 32) | (static_cast<uint64>(HeaderIndex) << BatchSizeBits) | ElementIndex;
 	}
 	
 	return true;	
