Fixup API macro usage

#rnx

diff --git a/Engine/Source/Runtime/Engine/Classes/GameFramework/PlayerController.h b/Engine/Source/Runtime/Engine/Classes/GameFramework/PlayerController.h
--- a/Engine/Source/Runtime/Engine/Classes/GameFramework/PlayerController.h
+++ b/Engine/Source/Runtime/Engine/Classes/GameFramework/PlayerController.h
@@ -2315,7 +2315,7 @@ public:
 	ENGINE_API FAsyncPhysicsTimestamp GetPhysicsTimestamp();
 
 	UE_DEPRECATED(5.7, "Deprecated, use the GetPhysicsTimestamp() without the DeltaSeconds parameter, it will return the upcoming frame as intended.")
-	ENGINE_API FAsyncPhysicsTimestamp GetPhysicsTimestamp(float DeltaSeconds) { return GetPhysicsTimestamp(); }
+	FAsyncPhysicsTimestamp GetPhysicsTimestamp(float DeltaSeconds) { return GetPhysicsTimestamp(); }
 
 	/** Get the physics frame number offset between Local and Server, recommended is to use GetPhysicsTimestamp() which takes this offset into account on the client for the ServerFrame value */
 	int32 GetNetworkPhysicsTickOffset() const { return NetworkPhysicsTickOffset; }
diff --git a/Engine/Source/Runtime/Experimental/Chaos/Public/PBDRigidsSolver.h b/Engine/Source/Runtime/Experimental/Chaos/Public/PBDRigidsSolver.h
--- a/Engine/Source/Runtime/Experimental/Chaos/Public/PBDRigidsSolver.h
+++ b/Engine/Source/Runtime/Experimental/Chaos/Public/PBDRigidsSolver.h
@@ -297,7 +297,7 @@ namespace Chaos
 		bool RewindUsesCollisionResimCache() const { return GetUseCollisionResimCache(); }
 
 		UE_DEPRECATED(5.7, "Deprecated, This functionality is now handled by FRewindData::ProcessDirtyPTParticles")
-		CHAOS_API void FinalizeRewindData(const TParticleView<FPBDRigidParticles>& DirtyParticles) { };
+		void FinalizeRewindData(const TParticleView<FPBDRigidParticles>& DirtyParticles) { };
 
 		FPerSolverFieldSystem& GetPerSolverField() { return *PerSolverField; }
 		const FPerSolverFieldSystem& GetPerSolverField() const { return *PerSolverField; }
diff --git a/Engine/Source/Runtime/Online/HTTP/Public/GenericPlatform/GenericPlatformHttp.h b/Engine/Source/Runtime/Online/HTTP/Public/GenericPlatform/GenericPlatformHttp.h
--- a/Engine/Source/Runtime/Online/HTTP/Public/GenericPlatform/GenericPlatformHttp.h
+++ b/Engine/Source/Runtime/Online/HTTP/Public/GenericPlatform/GenericPlatformHttp.h
@@ -81,7 +81,7 @@ public:
 	/**
 	 * Platform shutdown custom ssl if supported
 	 */
-	static HTTP_API void ShutdownPlatformCustomSsl() {}
+	static void ShutdownPlatformCustomSsl() {}
 
 	/**
 	 * Creates a new Http request instance for the current platform
