[FontFace]
- Fix thread-safety issue in serialize when called from async loading thread

diff --git a/Engine/Source/Runtime/Engine/Private/FontFace.cpp b/Engine/Source/Runtime/Engine/Private/FontFace.cpp
--- a/Engine/Source/Runtime/Engine/Private/FontFace.cpp
+++ b/Engine/Source/Runtime/Engine/Private/FontFace.cpp
@@ -97,7 +97,12 @@ void UFontFace::Serialize(FArchive& Ar)
 		}
 
 #if WITH_EDITORONLY_DATA
-		CacheSubFaces();
+		// This will be executed in Postload instead to avoid thread-safety issue
+		// when the serialization comes from the loader, which could run on any thread.
+		if (!HasAnyFlags(RF_NeedPostLoad))
+		{
+			CacheSubFaces();
+		}
 #endif // WITH_EDITORONLY_DATA
 	}
 	else
@@ -135,6 +140,10 @@ void UFontFace::PostLoad()
 {
 	Super::PostLoad();
 
+#if WITH_EDITORONLY_DATA
+	CacheSubFaces();
+#endif // WITH_EDITORONLY_DATA
+
 	UpdateDeviceRasterizationSettings();
 }
 
