Force Layer Stack input nodes to re-resolve changes to dependent material functions.

diff --git a/Engine/Source/Runtime/Engine/Private/Materials/MaterialCachedData.cpp b/Engine/Source/Runtime/Engine/Private/Materials/MaterialCachedData.cpp
--- a/Engine/Source/Runtime/Engine/Private/Materials/MaterialCachedData.cpp
+++ b/Engine/Source/Runtime/Engine/Private/Materials/MaterialCachedData.cpp
@@ -679,6 +679,7 @@ void FMaterialCachedExpressionData::UpdateForExpressions(const FMaterialCachedEx
 				UpdateForLayerFunctions(Context, Layers);
 				if (UMaterialExpressionLayerStack* LayerStackExpression = Cast<UMaterialExpressionLayerStack>(LayersExpression))
 				{
+					LayerStackExpression->ResolveLayerInputs();
 					Layers.LayerStackCache = LayerStackExpression->GetSharedAvailableFunctionsCache();
 				}
 
