Logging a player id on the game server and proxy when migrating actors between DSTM servers to help debug issues with 100s of players.

diff --git a/Engine/Plugins/Runtime/MultiServerReplication/Source/MultiServerReplication/Private/MultiServerProxy.cpp b/Engine/Plugins/Runtime/MultiServerReplication/Source/MultiServerReplication/Private/MultiServerProxy.cpp
--- a/Engine/Plugins/Runtime/MultiServerReplication/Source/MultiServerReplication/Private/MultiServerProxy.cpp
+++ b/Engine/Plugins/Runtime/MultiServerReplication/Source/MultiServerReplication/Private/MultiServerProxy.cpp
@@ -916,10 +916,11 @@ void UProxyNetDriver::DisableActorLogicAndGameCode()
 
 void UProxyNetDriver::GameServerAssignPlayerController(UNetConnection* ChildGameServerConnection, UNetConnection* NotUsedParentGameServerConnection, APlayerController* GameServerPlayerController)
 {
-	UE_LOG(LogNetProxy, Log, TEXT("[%d] Received a new player controller %s:%s for connection %s:%s."), 
+	UE_LOG(LogNetProxy, Log, TEXT("[%d] Received a new player controller %s:%s for connection %s:%s with player id %s."), 
 		   GameServerPlayerController->GetClientHandshakeId(),
 		   *GameServerPlayerController->GetName(), *GameServerPlayerController->GetClass()->GetName(),
-		   *ChildGameServerConnection->GetDriver()->GetName(), *ChildGameServerConnection->GetName());
+		   *ChildGameServerConnection->GetDriver()->GetName(), *ChildGameServerConnection->GetName(),
+		   *ChildGameServerConnection->PlayerId.ToDebugString());
 
 	const uint32 ClientHandshakeId = GameServerPlayerController->GetClientHandshakeId();
 	FMultiServerProxyInternalConnectionRoute* Route = InternalRoutes.Find(ClientHandshakeId);
diff --git a/Engine/Source/Runtime/Engine/Private/PlayerController.cpp b/Engine/Source/Runtime/Engine/Private/PlayerController.cpp
--- a/Engine/Source/Runtime/Engine/Private/PlayerController.cpp
+++ b/Engine/Source/Runtime/Engine/Private/PlayerController.cpp
@@ -5035,8 +5035,8 @@ void APlayerController::PostMigrate(const FUObjectMigrationContext& MigrationCon
 	{
 		EObjectMigrationRecvType MigrationRecvType = MigrationContext.GetObjectMigrationRecvType(this);
 
-		UE_LOG(LogPlayerController, Log, TEXT("[%s] APlayerController::PostMigrate [%s]: %s bHasOwnership=%d MigrationSide=Receive RecvType=%s"), 
-			   *UE::RemoteObject::GetGlobalServerId().ToString(), *FRemoteObjectId(this).ToString(), *GetName(), bHasOwnership, ToString(MigrationRecvType));
+		UE_LOG(LogPlayerController, Log, TEXT("[%s] APlayerController::PostMigrate [%s]: %s bHasOwnership=%d MigrationSide=Receive RecvType=%s PlayerId=%s"), 
+			   *UE::RemoteObject::GetGlobalServerId().ToString(), *FRemoteObjectId(this).ToString(), *GetName(), bHasOwnership, ToString(MigrationRecvType), *CachedConnectionPlayerId.ToDebugString());
 
 		ensure(NetConnection == nullptr);
 
@@ -5075,8 +5075,8 @@ void APlayerController::PostMigrate(const FUObjectMigrationContext& MigrationCon
 	{
 		EObjectMigrationSendType MigrationSendType = MigrationContext.GetObjectMigrationSendType(this);
 
-		UE_LOG(LogPlayerController, Log, TEXT("[%s] APlayerController::PostMigrate [%s]: %s bHasOwnership=%d MigrationSide=Send SendType=%s"), 
-			   *UE::RemoteObject::GetGlobalServerId().ToString(), *FRemoteObjectId(this).ToString(), *GetName(), bHasOwnership, ToString(MigrationSendType));
+		UE_LOG(LogPlayerController, Log, TEXT("[%s] APlayerController::PostMigrate [%s]: %s bHasOwnership=%d MigrationSide=Send SendType=%s PlayerId=%s"), 
+			   *UE::RemoteObject::GetGlobalServerId().ToString(), *FRemoteObjectId(this).ToString(), *GetName(), bHasOwnership, ToString(MigrationSendType), *CachedConnectionPlayerId.ToDebugString());
 
 		const bool bDetachControllerFromConnection = true;
 		const bool bDeregisterFromWorld = (MigrationSendType == EObjectMigrationSendType::ReassignOwnership);
