Fix for accidental double increment of layer caller index when building the graph in the LayerStack expression and inversed pairing of layer and blend nodes.

diff --git a/Engine/Source/Runtime/Engine/Private/Materials/MaterialExpressionLayerStack.cpp b/Engine/Source/Runtime/Engine/Private/Materials/MaterialExpressionLayerStack.cpp
--- a/Engine/Source/Runtime/Engine/Private/Materials/MaterialExpressionLayerStack.cpp
+++ b/Engine/Source/Runtime/Engine/Private/Materials/MaterialExpressionLayerStack.cpp
@@ -593,8 +593,8 @@ void UMaterialExpressionLayerStack::RebuildLayerGraph(bool bReportErrors)
 		if (NumActiveLayerCallers >= 2 && NumActiveBlendCallers == NumActiveLayerCallers - 1)
 		{
 			int32 CurrentLayerIndex = 0;
-			TObjectPtr<UMaterialExpressionMaterialFunctionCall> Bottom = LayerCallers[CurrentLayerIndex++];
-			TObjectPtr<UMaterialExpressionMaterialFunctionCall> Top = LayerCallers[CurrentLayerIndex++];
+			TObjectPtr<UMaterialExpressionMaterialFunctionCall> Bottom = LayerCallers[CurrentLayerIndex];
+			TObjectPtr<UMaterialExpressionMaterialFunctionCall> Top = LayerCallers[++CurrentLayerIndex];
 			for (int32 BlendIndex = 0; BlendIndex < NumActiveBlendCallers; BlendIndex++)
 			{
 				TObjectPtr<UMaterialExpressionMaterialFunctionCall> Blend = BlendCallers[BlendIndex];
@@ -630,8 +630,8 @@ void UMaterialExpressionLayerStack::RebuildLayerGraph(bool bReportErrors)
 
 				if (++CurrentLayerIndex < NumActiveLayerCallers)
 				{
-					Bottom = LayerCallers[CurrentLayerIndex];
-					Top = Blend;
+					Bottom = Blend;
+					Top = LayerCallers[CurrentLayerIndex];
 				}
 				else
 				{
