Fixing up some Logs for multi-server mode.

diff --git a/Engine/Source/Programs/AutomationTool/Scripts/LaunchMultiServer.Automation.cs b/Engine/Source/Programs/AutomationTool/Scripts/LaunchMultiServer.Automation.cs
--- a/Engine/Source/Programs/AutomationTool/Scripts/LaunchMultiServer.Automation.cs
+++ b/Engine/Source/Programs/AutomationTool/Scripts/LaunchMultiServer.Automation.cs
@@ -4,7 +4,6 @@ using System;
 using System.IO;
 using UnrealBuildBase;
 using UnrealBuildTool;
-using AutomationScripts;
 using EpicGames.Core;
 using System.Linq;
 using Microsoft.Extensions.Logging;
@@ -447,7 +446,7 @@ Example running 4 clients locally connecting to single proxy server w/ a non-loc
 			string[] LocalGameServerAddresses = new string[NumTotalServers];
 
 			// This is optional to the multi-server process, these are in addition to the MultiServerNumPeers argument
-			string ServerCommonArgs = $" -server {AdditionalServerArguments} -NODEBUGOUTPUT -SuppressConsoleOutputVerboseLogging";
+			string ServerCommonArgs = $" -server {AdditionalServerArguments} -NODEBUGOUTPUT";
 			ServerCommonArgs += NoTimeouts ? " -notimeouts" : "";
 			ServerCommonArgs += IsAttachingDebugger(ProcessLaunchType.Server) ? " -WaitForDebuggerNoBreak" : "";
 			ServerCommonArgs += $" -MultiServerNumServers={NumTotalServers}";
diff --git a/Engine/Source/Runtime/Engine/Private/PlayerController.cpp b/Engine/Source/Runtime/Engine/Private/PlayerController.cpp
--- a/Engine/Source/Runtime/Engine/Private/PlayerController.cpp
+++ b/Engine/Source/Runtime/Engine/Private/PlayerController.cpp
@@ -5035,7 +5035,7 @@ void APlayerController::PostMigrate(const FUObjectMigrationContext& MigrationCon
 	{
 		EObjectMigrationRecvType MigrationRecvType = MigrationContext.GetObjectMigrationRecvType(this);
 
-		UE_LOG(LogPlayerController, Log, TEXT("[%s] APlayerController::PostMigrate [%s]: %s bHasOwnership=%d MigrationSide=Receive RecvType=%s PlayerId=%s"), 
+		UE_LOG(LogPlayerController, Verbose, TEXT("[%s] APlayerController::PostMigrate [%s]: %s bHasOwnership=%d MigrationSide=Receive RecvType=%s PlayerId=%s"), 
 			   *UE::RemoteObject::GetGlobalServerId().ToString(), *FRemoteObjectId(this).ToString(), *GetName(), bHasOwnership, ToString(MigrationRecvType), *CachedConnectionPlayerId.ToDebugString());
 
 		ensure(NetConnection == nullptr);
@@ -5060,7 +5060,7 @@ void APlayerController::PostMigrate(const FUObjectMigrationContext& MigrationCon
 				ensure(PrevNetPlayerIndex == NetPlayerIndex);
 				ensure(PrevHandshakeTransactionId == GetClientHandshakeId());
 
-				UE_LOG(LogPlayerController, Log, TEXT("[%s] APlayerController::PostMigrate [%s]: Binding %s to connection %s"), 
+				UE_LOG(LogPlayerController, Verbose, TEXT("[%s] APlayerController::PostMigrate [%s]: Binding %s to connection %s"), 
 					   *UE::RemoteObject::GetGlobalServerId().ToString(), *FRemoteObjectId(this).ToString(), 
 					   *GetName(), *NetConnection->GetName());
 			}
@@ -5075,7 +5075,7 @@ void APlayerController::PostMigrate(const FUObjectMigrationContext& MigrationCon
 	{
 		EObjectMigrationSendType MigrationSendType = MigrationContext.GetObjectMigrationSendType(this);
 
-		UE_LOG(LogPlayerController, Log, TEXT("[%s] APlayerController::PostMigrate [%s]: %s bHasOwnership=%d MigrationSide=Send SendType=%s PlayerId=%s"), 
+		UE_LOG(LogPlayerController, Verbose, TEXT("[%s] APlayerController::PostMigrate [%s]: %s bHasOwnership=%d MigrationSide=Send SendType=%s PlayerId=%s"), 
 			   *UE::RemoteObject::GetGlobalServerId().ToString(), *FRemoteObjectId(this).ToString(), *GetName(), bHasOwnership, ToString(MigrationSendType), *CachedConnectionPlayerId.ToDebugString());
 
 		const bool bDetachControllerFromConnection = true;
@@ -5127,7 +5127,7 @@ void APlayerController::PostMigrate(const FUObjectMigrationContext& MigrationCon
 
 					PrevNetConnection->OwningActor = PrevNetConnection->PlayerController;
 					
-					UE_LOG(LogPlayerController, Log, TEXT("[%s] APlayerController::PostMigrate [%s]: Replacing %s with new player controller %s on connection %s"), 
+					UE_LOG(LogPlayerController, Verbose, TEXT("[%s] APlayerController::PostMigrate [%s]: Replacing %s with new player controller %s on connection %s"), 
 						   *UE::RemoteObject::GetGlobalServerId().ToString(), *FRemoteObjectId(this).ToString(), 
 						   *GetName(), *PrevNetConnection->PlayerController->GetName(), *PrevNetConnection->GetName());
 
@@ -5135,7 +5135,7 @@ void APlayerController::PostMigrate(const FUObjectMigrationContext& MigrationCon
 				}
 				else
 				{
-					UE_LOG(LogPlayerController, Log, TEXT("[%s] APlayerController::PostMigrate [%s]: Replacing %s with existing player controller %s on connection %s"),
+					UE_LOG(LogPlayerController, Verbose, TEXT("[%s] APlayerController::PostMigrate [%s]: Replacing %s with existing player controller %s on connection %s"),
 						   *UE::RemoteObject::GetGlobalServerId().ToString(), *FRemoteObjectId(this).ToString(), 
 						   *GetName(), *PrevNetConnection->PlayerController->GetName(), *PrevNetConnection->GetName());
 				}
@@ -5147,7 +5147,7 @@ void APlayerController::PostMigrate(const FUObjectMigrationContext& MigrationCon
 
 		if (bDeregisterFromWorld)
 		{
-			UE_LOG(LogPlayerController, Log, TEXT("[%s] APlayerController::PostMigrate [%s]: Removing %s from the world"), 
+			UE_LOG(LogPlayerController, Verbose, TEXT("[%s] APlayerController::PostMigrate [%s]: Removing %s from the world"), 
 				   *UE::RemoteObject::GetGlobalServerId().ToString(), *FRemoteObjectId(this).ToString(), *GetName());
 
 			GetWorld()->RemoveController(this);
