Skeletal Mesh: Fix GPU morph target generation that got broken during refactoring.

#rnx

diff --git a/Engine/Source/Runtime/Core/Private/UObject/FortniteMainBranchObjectVersion.cpp b/Engine/Source/Runtime/Core/Private/UObject/FortniteMainBranchObjectVersion.cpp
--- a/Engine/Source/Runtime/Core/Private/UObject/FortniteMainBranchObjectVersion.cpp
+++ b/Engine/Source/Runtime/Core/Private/UObject/FortniteMainBranchObjectVersion.cpp
@@ -12,7 +12,7 @@ TMap<FGuid, FGuid> FFortniteMainBranchObjectVersion::GetSystemGuids()
 	SystemGuids.Add(DevGuids.NANITE_DERIVEDDATA_VER, FGuid("FB5EAC6147EF4591B0D5EB7A85062AC9"));
 	SystemGuids.Add(DevGuids.NIAGARASHADERMAP_DERIVEDDATA_VER, FGuid("6360A977062842A29ED30E3A7ACB0E64"));
 	SystemGuids.Add(DevGuids.Niagara_LatestScriptCompileVersion, FGuid("97F0FD7149D90D42B65D6557AF477D1A"));
-	SystemGuids.Add(DevGuids.SkeletalMeshDerivedDataVersion, FGuid("07AE97536567481DBE673C51EBEC96EF"));
+	SystemGuids.Add(DevGuids.SkeletalMeshDerivedDataVersion, FGuid("C998F2F3B2884A02B7E290B8C25835F0"));
 	SystemGuids.Add(DevGuids.STATICMESH_DERIVEDDATA_VER, FGuid("3ABF32149F9D4D838750368AB95573BE"));
 	SystemGuids.Add(DevGuids.MaterialTranslationDDCVersion, FGuid("E541A459F3374E628240BBB309910753"));
 	return SystemGuids;
diff --git a/Engine/Source/Runtime/Engine/Private/Rendering/MorphTargetVertexInfoBuffers.cpp b/Engine/Source/Runtime/Engine/Private/Rendering/MorphTargetVertexInfoBuffers.cpp
--- a/Engine/Source/Runtime/Engine/Private/Rendering/MorphTargetVertexInfoBuffers.cpp
+++ b/Engine/Source/Runtime/Engine/Private/Rendering/MorphTargetVertexInfoBuffers.cpp
@@ -225,7 +225,6 @@ void FMorphTargetVertexInfoBuffers::Compress(const TArray<FSkelMeshRenderSection
 	TangentZPrecision = ComputeTangentPrecision();
 
 	MorphData.Empty();
-	NumTotalBatches = 0;
 
 	BatchStartOffsetPerMorph.Empty(MorphTargets.Num());
 	BatchesPerMorph.Empty(MorphTargets.Num());
@@ -253,7 +252,7 @@ void FMorphTargetVertexInfoBuffers::Compress(const TArray<FSkelMeshRenderSection
 	TArray<uint32> BitstreamData;
 	for (int32 AnimIdx = 0; AnimIdx < MorphTargets.Num(); ++AnimIdx)
 	{
-		uint32 BatchStartOffset = NumTotalBatches;
+		uint32 BatchStartOffset = BatchHeaders.Num();
 
 		float MaximumValues[4] = { -FLT_MAX, -FLT_MAX, -FLT_MAX, -FLT_MAX };
 		float MinimumValues[4] = { +FLT_MAX, +FLT_MAX, +FLT_MAX, +FLT_MAX };
@@ -297,17 +296,17 @@ void FMorphTargetVertexInfoBuffers::Compress(const TArray<FSkelMeshRenderSection
 
 			// Encode the actual morph vertex info into the quantized bitstream.
 			Encode(MorphModel->Vertices, &VertexNeedsTangents, PositionPrecision, TangentZPrecision, BatchHeaders, BitstreamData);
-
-			NumTotalBatches += BatchHeaders.Num();
-
-			const uint32 MorphNumBatches = NumTotalBatches - BatchStartOffset;
-			BatchStartOffsetPerMorph.Add(BatchStartOffset);
-			BatchesPerMorph.Add(MorphNumBatches);
-			MaximumValuePerMorph.Add(FVector4f(MaximumValues[0], MaximumValues[1], MaximumValues[2], MaximumValues[3]));
-			MinimumValuePerMorph.Add(FVector4f(MinimumValues[0], MinimumValues[1], MinimumValues[2], MinimumValues[3]));
 		}
+		
+		const uint32 MorphNumBatches = BatchHeaders.Num() - BatchStartOffset;
+		BatchStartOffsetPerMorph.Add(BatchStartOffset);
+		BatchesPerMorph.Add(MorphNumBatches);
+		MaximumValuePerMorph.Add(FVector4f(MaximumValues[0], MaximumValues[1], MaximumValues[2], MaximumValues[3]));
+		MinimumValuePerMorph.Add(FVector4f(MinimumValues[0], MinimumValues[1], MinimumValues[2], MinimumValues[3]));
 	}
 
+	NumTotalBatches = BatchHeaders.Num();
+
 	// Write packed batch headers
 	for (const FDeltaBatchHeader& BatchHeader : BatchHeaders)
 	{
