SceneCulling: Fixed broken occlusion culling caused by undefined behavior in hash calculation when building with Clang 19.

The problem manifested itself as many VSM rendering tests failing with missing shadows. It was caused by Clang 19 aggressively optimizing the hash calculation when we build the spatial hash grid, which led to incorrect results when running occlusion queries because hashes when building the structur did not match the ones we used in queries.

The exact cause is that the hashing function multiplies signed integer coordinates by large numbers, which is susceptible to signed overflows and constitutes undefined behavior. In most cases it might work fine because the compiler is not able to determine that overflows will happen at runtime. However, in SceneRendering we provide the compiler with a crucial hint because at some point in code we increment one of the components (W/Level) by 3. Since 3 * 'the large number' already overflows int32, in heavily inlined code the compiler is able to take this info and apply in the hash calculation to determine that the multiplication of the affected component (W/Level) is not needed as we overflow anyway.

Resolved by simply casting signed coordinates to the unsigned versions. 'Overflowing' unsigned integers is well defined and it resolves the issue. We have the exact same code copied in RenderingSpatialHash.h and I've applied the fix there as well for completeness, but it doesn't matter in practice because TLocation::GetTypeHash is not used in our codebase.

diff --git a/Engine/Source/Runtime/Engine/Public/Rendering/RenderingSpatialHash.h b/Engine/Source/Runtime/Engine/Public/Rendering/RenderingSpatialHash.h
--- a/Engine/Source/Runtime/Engine/Public/Rendering/RenderingSpatialHash.h
+++ b/Engine/Source/Runtime/Engine/Public/Rendering/RenderingSpatialHash.h
@@ -37,7 +37,9 @@ struct TLocation
 
 	inline friend uint32 GetTypeHash(const TLocation& CellLocation)
 	{
-		return uint32(CellLocation.Coord.X * 1150168907 + CellLocation.Coord.Y * 1235029793 + CellLocation.Coord.Z * 1282581571 + CellLocation.Level * 1264559321);
+		using UScalarType = std::make_unsigned_t<ScalarType>;
+
+		return uint32((UScalarType)CellLocation.Coord.X * 1150168907 + (UScalarType)CellLocation.Coord.Y * 1235029793 + (UScalarType)CellLocation.Coord.Z * 1282581571 + (UScalarType)CellLocation.Level * 1264559321);
 	}
 	
 	inline bool operator == (const TLocation& Other) const
diff --git a/Engine/Source/Runtime/Renderer/Private/SceneCulling/SceneCulling.h b/Engine/Source/Runtime/Renderer/Private/SceneCulling/SceneCulling.h
--- a/Engine/Source/Runtime/Renderer/Private/SceneCulling/SceneCulling.h
+++ b/Engine/Source/Runtime/Renderer/Private/SceneCulling/SceneCulling.h
@@ -128,8 +128,10 @@ public:
 
 		FORCEINLINE uint32 GetHash() const
 		{
+			using UIntType = std::make_unsigned_t<FIntVector4::IntType>;
+
 			// TODO: Vectorize? Maybe convert to float vector & use dot product? Maybe not? (mul is easy, dot maybe not?)
-			return uint32(Data.X * 1150168907 + Data.Y * 1235029793 + Data.Z * 1282581571 + Data.W * 1264559321);
+			return uint32((UIntType)Data.X * 1150168907 + (UIntType)Data.Y * 1235029793 + (UIntType)Data.Z * 1282581571 + (UIntType)Data.W * 1264559321);
 		}
 
 		FIntVector4 Data;
