Fixed SphereTraceComponent not returning the hit FaceIndex when tracing against complex shapes.

diff --git a/Engine/Source/Runtime/Engine/Private/PhysicsEngine/Experimental/PhysInterface_Chaos.cpp b/Engine/Source/Runtime/Engine/Private/PhysicsEngine/Experimental/PhysInterface_Chaos.cpp
--- a/Engine/Source/Runtime/Engine/Private/PhysicsEngine/Experimental/PhysInterface_Chaos.cpp
+++ b/Engine/Source/Runtime/Engine/Private/PhysicsEngine/Experimental/PhysInterface_Chaos.cpp
@@ -1096,7 +1096,7 @@ bool FPhysInterface_Chaos::Sweep_Geom(FHitResult& OutHit, const FBodyInstance* I
 							SetFlags(Hit, EHitFlags::Distance | EHitFlags::Normal | EHitFlags::Position | EHitFlags::FaceIndex);
 
 							FTransform StartTransform(InStart);
-							ConvertQueryImpactHit(OwnerComponentInst->GetWorld(), Hit, OutHit, DeltaMag, QueryFilter, InStart, InEnd, &QueryGeometry, StartTransform, false, false);
+							ConvertQueryImpactHit(OwnerComponentInst->GetWorld(), Hit, OutHit, DeltaMag, QueryFilter, InStart, InEnd, &QueryGeometry, StartTransform, true, false);
 							bSweepHit = true;
 						}
 					}
