World Partition: make sure to flush streaming data when saving a new map in the editor while in a PIE session.

diff --git a/Engine/Source/Runtime/Engine/Private/WorldPartition/RuntimeHashSet/WorldPartitionRuntimeHashSet.cpp b/Engine/Source/Runtime/Engine/Private/WorldPartition/RuntimeHashSet/WorldPartitionRuntimeHashSet.cpp
--- a/Engine/Source/Runtime/Engine/Private/WorldPartition/RuntimeHashSet/WorldPartitionRuntimeHashSet.cpp
+++ b/Engine/Source/Runtime/Engine/Private/WorldPartition/RuntimeHashSet/WorldPartitionRuntimeHashSet.cpp
@@ -1157,6 +1157,11 @@ void UWorldPartitionRuntimeHashSet::PostLoad()
 	{
 		FixupHLODSetup(RuntimePartition);
 	}
+
+	if (!GetWorld()->IsGameWorld())
+	{
+		FlushStreamingContent();
+	}
 }
 
 void UWorldPartitionRuntimeHashSet::FixupHLODSetup(FRuntimePartitionDesc& RuntimePartition)
