Fix memory tracking in FMemoryBoundQueue

#rnx

diff --git a/Engine/Source/Runtime/Engine/Private/AssetCompilingManager.cpp b/Engine/Source/Runtime/Engine/Private/AssetCompilingManager.cpp
--- a/Engine/Source/Runtime/Engine/Private/AssetCompilingManager.cpp
+++ b/Engine/Source/Runtime/Engine/Private/AssetCompilingManager.cpp
@@ -70,7 +70,7 @@ static TAutoConsoleVariable<int32> CVarAsyncAssetCompilationMaxMemoryUsage(
 
 TRACE_DECLARE_INT_COUNTER(AsyncCompilationMaxConcurrency, TEXT("AsyncCompilation/MaxConcurrency"));
 TRACE_DECLARE_INT_COUNTER(AsyncCompilationConcurrency, TEXT("AsyncCompilation/Concurrency"));
-TRACE_DECLARE_MEMORY_COUNTER(AsyncCompilationTotalMemoryLimit, TEXT("AsyncCompilation/TotalMemoryLimit"));
+TRACE_DECLARE_MEMORY_COUNTER(AsyncCompilationAvailableMemory, TEXT("AsyncCompilation/AvailableMemory"));
 TRACE_DECLARE_MEMORY_COUNTER(AsyncCompilationTotalEstimatedMemory, TEXT("AsyncCompilation/TotalEstimatedMemory"));
 
 #endif //#if WITH_EDITOR
@@ -137,13 +137,13 @@ public:
 		return FMath::Max(0ll, (int64)CVarAsyncAssetCompilationMemoryPerCore.GetValueOnAnyThread(false) * 1024 * 1024 * 1024);
 	}
 
-	int64 GetMemoryLimit() const
+	int64 GetMemoryAvailableForAssetCompilation() const
 	{
 		const FPlatformMemoryStats MemoryStats = FPlatformMemory::GetStats();
 		
 		// UsedPhysical = "working set"
 		// UsedVirtual  = "commit charge"
-		UE_LOG(LogAsyncCompilation, Verbose, TEXT("GetMemoryLimit UsedPhysical = %.3f MB (Peak %.3f MB) UsedVirtual = %.3f MB (Peak %.3f MB) AvailPhysical = %.3f MB AvailVirtual = %.3f MB "),
+		UE_LOG(LogAsyncCompilation, Verbose, TEXT("GetMemoryAvailableForAssetCompilation UsedPhysical = %.3f MB (Peak %.3f MB) UsedVirtual = %.3f MB (Peak %.3f MB) AvailPhysical = %.3f MB AvailVirtual = %.3f MB "),
 				MemoryStats.UsedPhysical/(1024*1024.f),
 				MemoryStats.PeakUsedPhysical/(1024*1024.f),
 				MemoryStats.UsedVirtual/(1024*1024.f),
@@ -154,18 +154,32 @@ public:
 		// Just make sure available physical fits in a int64, if it's bigger than that, we're not expecting to be memory limited anyway
 		// Also uses AvailableVirtual because the system might have plenty of physical memory but still be limited by virtual memory available in some cases.
 		//   (i.e. per-process quota, paging file size lower than actual memory available, etc.).
-		int64 AvailableMemory = (int64)FMath::Min3(MemoryStats.AvailablePhysical, MemoryStats.AvailableVirtual, (uint64)INT64_MAX);
+		const int64 SystemAvailableMemory = (int64)FMath::Min3(MemoryStats.AvailablePhysical, MemoryStats.AvailableVirtual, (uint64)INT64_MAX);
 
-		int64 HardMemoryLimit = GetHardMemoryLimit();
-		if (HardMemoryLimit > 0)
-		{
-			AvailableMemory = FMath::Min(HardMemoryLimit, AvailableMemory);
+		const int64 HardMemoryLimit = GetHardMemoryLimit();
+		const int64 LeaveMemoryHeadRoom = 64 * 1024 * 1024; // leave 64 MB head room
+
+		// No hard limit specified, use system memory
+		if (HardMemoryLimit == 0) {
+			return FMath::Max(0, SystemAvailableMemory - LeaveMemoryHeadRoom);
 		}
+
+		const int64 MemoryAvailableAgainstHardLimit = FMath::Max(0, HardMemoryLimit - TotalEstimatedMemory);
 		
-		const int64 LeaveMemoryHeadRoom = 64 * 1024 * 1024; // leave 64 MB head room
-		AvailableMemory = FMath::Max(0, AvailableMemory - LeaveMemoryHeadRoom);
+		if (SystemAvailableMemory < MemoryAvailableAgainstHardLimit) {
+			UE_LOG(LogAsyncCompilation, Warning, TEXT(
+				"Not enough system memory to provide asset compilation full hard limit of %.3f MB - SystemAvailableMemory = %.3f MB < MemoryAvailableAgainstHardLimit = %.3f MB"),
+				   HardMemoryLimit / (1024 * 1024.f),
+				   SystemAvailableMemory / (1024 * 1024.f),
+				   MemoryAvailableAgainstHardLimit / (1024 * 1024.f));
+		}
+
+		const int64 MemoryAvailableForAssetCompilation = FMath::Min(
+			FMath::Max(0, MemoryAvailableAgainstHardLimit - LeaveMemoryHeadRoom), 
+			FMath::Max(0, SystemAvailableMemory - LeaveMemoryHeadRoom)
+		);
 
-		return AvailableMemory;
+		return MemoryAvailableForAssetCompilation;
 	}
 
 	int64 GetRequiredMemory(const IQueuedWork* InWork) const
@@ -195,7 +209,7 @@ public:
 
 	void UpdateCounters()
 	{
-		TRACE_COUNTER_SET(AsyncCompilationTotalMemoryLimit, GetMemoryLimit());
+		TRACE_COUNTER_SET(AsyncCompilationAvailableMemory, GetMemoryAvailableForAssetCompilation());
 		TRACE_COUNTER_SET(AsyncCompilationTotalEstimatedMemory, TotalEstimatedMemory);
 		TRACE_COUNTER_SET(AsyncCompilationConcurrency, GetCurrentConcurrency());
 	}
@@ -261,14 +275,10 @@ public:
 			NewRequiredMemory += GetRequiredMemory(NextWork);
 		}
 
-		// TotalEstimatedMemory is the sum of all "RequiredMemory" for currently scheduled tasks
-		//	note TotalEstimatedMemory is atomic and can be changing as we act
-		int64 TotalRequiredMemory = TotalEstimatedMemory.load() + NewRequiredMemory;
-
 		int32 DynamicMaxConcurrency = GetAdjustedMaxConcurrency(Priority);
 		int32 Concurrency = FQueuedThreadPoolWrapper::GetCurrentConcurrency();
 		
-		int64 MemoryLimit = GetMemoryLimit();
+		int64 AvailableAssetCompilationMemory = GetMemoryAvailableForAssetCompilation();
 
 		// we check MemoryLimit which uses current system available memory against TotalRequiredMemory
 		//	which includes already running tasks
@@ -278,7 +288,7 @@ public:
 		//    because after it allocated there is 4 GB free but we also then subtract off the 4 GB the task requires
 		// this is conservative, so okay if we want to avoid OOM and are okay with less concurrency than possible
 
-		if ( TotalRequiredMemory > 0 && TotalRequiredMemory >= MemoryLimit )
+		if ( NewRequiredMemory > 0 && NewRequiredMemory >= AvailableAssetCompilationMemory )
 		{
 			if ( Concurrency == 0 )
 			{
@@ -304,11 +314,11 @@ public:
 						// and don't just exceed the currently available memory.
 						UE_LOGFMT_NSLOC(LogAsyncCompilation, Warning, "AsyncAssetCompilation", "HardMemoryLimitExceeded",
 							"BEWARE: AssetCompile memory estimate is greater than the hard memory limit, but we're running it [{TaskName}] anyway! "
-							"RequiredMemory = {TotalEstimatedMemory} MiB + {RequiredMemory} MiB, MemoryLimit = {MemoryLimit} MiB, HardMemoryLimit = {HardMemoryLimit} MiB",
+							"RequiredMemory = {TotalEstimatedMemory} MiB + {RequiredMemory} MiB, AvailableAssetCompilationMemory = {AvailableAssetCompilationMemory} MiB, HardMemoryLimit = {HardMemoryLimit} MiB",
 							("TaskName", DebugName),
 							("RequiredMemory", FString::SanitizeFloat(NewRequiredMemory / (1024 * 1024.f), 3)),
 							("TotalEstimatedMemory", FString::SanitizeFloat(TotalEstimatedMemory / (1024 * 1024.f), 3)),
-							("MemoryLimit", FString::SanitizeFloat(MemoryLimit / (1024 * 1024.f), 3)),
+							("AvailableAssetCompilationMemory", FString::SanitizeFloat(AvailableAssetCompilationMemory / (1024 * 1024.f), 3)),
 							("HardMemoryLimit", FString::SanitizeFloat(HardMemoryLimit / (1024 * 1024.f), 3))
 						);
 					}
@@ -316,11 +326,11 @@ public:
 					{
 						UE_LOGFMT_NSLOC(LogAsyncCompilation, Display, "AsyncAssetCompilation", "MemoryLimitExceeded",
 							"BEWARE: AssetCompile memory estimate is greater than available, but we're running it [{TaskName}] anyway! "
-							"RequiredMemory = {TotalEstimatedMemory} MiB + {RequiredMemory} MiB, MemoryLimit = {MemoryLimit} MiB ",
+							"RequiredMemory = {TotalEstimatedMemory} MiB + {RequiredMemory} MiB, MemoryLimit = {AvailableAssetCompilationMemory} MiB ",
 							("TaskName", DebugName),
 							("RequiredMemory", FString::SanitizeFloat(NewRequiredMemory / (1024 * 1024.f), 3)),
 							("TotalEstimatedMemory", FString::SanitizeFloat(TotalEstimatedMemory / (1024 * 1024.f), 3)),
-							("MemoryLimit", FString::SanitizeFloat(MemoryLimit / (1024 * 1024.f), 3))
+							("AvailableAssetCompilationMemory", FString::SanitizeFloat(AvailableAssetCompilationMemory / (1024 * 1024.f), 3))
 						);
 					}
 				}
@@ -332,11 +342,11 @@ public:
 			}
 			else
 			{
-				UE_LOG(LogAsyncCompilation, Verbose, TEXT("CONCURRENCY LIMITED to %d/%d; RequiredMemory = %.3f MB + %.3f MB MemoryLimit = %.3f MB "),
+				UE_LOG(LogAsyncCompilation, Warning, TEXT("CONCURRENCY LIMITED to %d/%d; RequiredMemory = %.3f MB + %.3f MB AvailableAssetCompilationMemory = %.3f MB "),
 					Concurrency,DynamicMaxConcurrency,
 					NewRequiredMemory/(1024*1024.f),
 					TotalEstimatedMemory/(1024*1024.f),
-					MemoryLimit/(1024*1024.f));
+					AvailableAssetCompilationMemory/(1024*1024.f));
 					
 				DynamicMaxConcurrency = Concurrency; // Limit concurrency to what we already allowed
 				// this will cause CanSchedule() to return false
