Removed weak ptr dereferences in collision conversion.
Added GetItemIndex OverlapHitResult so ItemIndex can be deprecated.

diff --git a/Engine/Plugins/Experimental/GameplayTargetingSystem/Source/GameplayTargetingSystem/Private/Tasks/TargetingSelectionTask_AOE.cpp b/Engine/Plugins/Experimental/GameplayTargetingSystem/Source/GameplayTargetingSystem/Private/Tasks/TargetingSelectionTask_AOE.cpp
--- a/Engine/Plugins/Experimental/GameplayTargetingSystem/Source/GameplayTargetingSystem/Private/Tasks/TargetingSelectionTask_AOE.cpp
+++ b/Engine/Plugins/Experimental/GameplayTargetingSystem/Source/GameplayTargetingSystem/Private/Tasks/TargetingSelectionTask_AOE.cpp
@@ -281,7 +281,7 @@ void UTargetingSelectionTask_AOE::ProcessOverlapResults(const FTargetingRequestH
 				ResultData->HitResult.Location = OverlapResult.GetActor()->GetActorLocation();
 				ResultData->HitResult.bBlockingHit = OverlapResult.bBlockingHit;
 				ResultData->HitResult.TraceStart = SourceLocation;
-				ResultData->HitResult.Item = OverlapResult.ItemIndex;
+				ResultData->HitResult.Item = OverlapResult.GetItemIndex();
 				ResultData->HitResult.Distance = FVector::Distance(OverlapResult.GetActor()->GetActorLocation(), SourceLocation);
 			}
 		}
diff --git a/Engine/Plugins/Experimental/Mover/Source/Mover/Private/MoveLibrary/MovementUtils.cpp b/Engine/Plugins/Experimental/Mover/Source/Mover/Private/MoveLibrary/MovementUtils.cpp
--- a/Engine/Plugins/Experimental/Mover/Source/Mover/Private/MoveLibrary/MovementUtils.cpp
+++ b/Engine/Plugins/Experimental/Mover/Source/Mover/Private/MoveLibrary/MovementUtils.cpp
@@ -939,7 +939,7 @@ bool UMovementUtils::TestEncroachmentWithAdjustment_Internal(const UWorld* World
 				{
 					NumBlockingHits++;
 					FCollisionShape const NonShrunkenCollisionShape = PrimComp->GetCollisionShape();
-					const FBodyInstance* OverlapBodyInstance = OverlapComponent->GetBodyInstance(NAME_None, true, Overlaps[HitIdx].ItemIndex);
+					const FBodyInstance* OverlapBodyInstance = OverlapComponent->GetBodyInstance(NAME_None, true, Overlaps[HitIdx].GetItemIndex());
 					bool bSuccess = OverlapBodyInstance && OverlapBodyInstance->OverlapTest(TestWorldTransform.GetLocation(), TestWorldTransform.GetRotation(), NonShrunkenCollisionShape, &MTDResult);
 					if (bSuccess)
 					{
diff --git a/Engine/Source/Runtime/Engine/Classes/Engine/OverlapResult.h b/Engine/Source/Runtime/Engine/Classes/Engine/OverlapResult.h
--- a/Engine/Source/Runtime/Engine/Classes/Engine/OverlapResult.h
+++ b/Engine/Source/Runtime/Engine/Classes/Engine/OverlapResult.h
@@ -25,6 +25,10 @@ struct FOverlapResult
 		For SkeletalMeshComponents this is the Body index or INDEX_NONE for single body */
 	int32 ItemIndex;
 
+	UE_INTERNAL int32 GetItemIndexInternal() const;
+	ENGINE_API int32 GetItemIndex() const;
+	ENGINE_API void SetItemIndex(const int32 InItemIndex);
+
 	/** Utility to return the Actor that owns the Component that was hit */
 	AActor* GetActor() const;
 
diff --git a/Engine/Source/Runtime/Engine/Private/Collision/CollisionConversions.cpp b/Engine/Source/Runtime/Engine/Private/Collision/CollisionConversions.cpp
--- a/Engine/Source/Runtime/Engine/Private/Collision/CollisionConversions.cpp
+++ b/Engine/Source/Runtime/Engine/Private/Collision/CollisionConversions.cpp
@@ -21,17 +21,27 @@
 #include "PhysicsProxy/SingleParticlePhysicsProxy.h"
 #include "PhysicsProxy/GeometryCollectionPhysicsProxy.h"
 
+int32 GetItemIndex(const FOverlapResult& InResult)
+{
+	// This causes different behavior when enabled as we don't deduplicate hit results when `OwnerComponent->bMultiBodyOverlap` is false.
+#if UE_WITH_REMOTE_OBJECT_HANDLE
+	return InResult.GetItemIndexInternal();
+#else
+	return InResult.GetItemIndex();
+#endif
+}
+
 // Used to place overlaps into a TMap when deduplicating them
 struct FOverlapKey
 {
 	const Chaos::FPhysicsObjectHandle PhysicsObject;
-	const UPrimitiveComponent* Component;
+	TWeakObjectPtr<class UPrimitiveComponent> Component;
 	const int32 ComponentIndex;
 
 	FOverlapKey(const FOverlapResult& InResult)
 		: PhysicsObject(InResult.PhysicsObject)
-		, Component(InResult.GetComponent())
-		, ComponentIndex(InResult.ItemIndex)
+		, Component(InResult.Component)
+		, ComponentIndex(GetItemIndex(InResult))
 	{
 	}
 
@@ -678,13 +688,23 @@ void ConvertQueryOverlap(const FPhysicsShape& Shape, const FPhysicsActor& Actor,
 	else if (const FBodyInstance* BodyInst = GetUserData(Actor))
 	{
 		BodyInst = FPhysicsInterface::ShapeToOriginalBodyInstance(BodyInst, &Shape);
-		if (const UPrimitiveComponent* OwnerComponent = BodyInst->OwnerComponent.Get())
+		if (BodyInst->OwnerComponent.IsValid())
 		{
 			// Create a handle that needs to be resolved later when accessed in game thread,
 			// since resolving FActorInstanceHandle may require to access related UObjects (e.g. owner actor), which is not safe outside game thread.
 			OutOverlap.OverlapObjectHandle = FActorInstanceHandle::MakeActorHandleToResolve(BodyInst->OwnerComponent, BodyInst->InstanceBodyIndex);
 			OutOverlap.Component = BodyInst->OwnerComponent; // Copying weak pointer is faster than assigning raw pointer.
-			OutOverlap.ItemIndex = OwnerComponent->bMultiBodyOverlap ? BodyInst->InstanceBodyIndex : INDEX_NONE;
+			bool bUseBodyIndex = false;
+#if UE_WITH_REMOTE_OBJECT_HANDLE
+			// To avoid migrations, we always set the body index and we defer bMultiBodyOverlap to the getter
+			bUseBodyIndex = true;
+#else
+			if (const UPrimitiveComponent* OwnerComponent = BodyInst->OwnerComponent.Get())
+			{
+				bUseBodyIndex = OwnerComponent->bMultiBodyOverlap;
+			}
+#endif
+			OutOverlap.SetItemIndex(bUseBodyIndex ? BodyInst->InstanceBodyIndex : INDEX_NONE);
 		}
 		OutOverlap.PhysicsObject = BodyInst->GetPhysicsActor() ? BodyInst->GetPhysicsActor()->GetPhysicsObject() : nullptr;
 
@@ -709,7 +729,7 @@ void ConvertQueryOverlap(const FPhysicsShape& Shape, const FPhysicsActor& Actor,
 			// since resolving FActorInstanceHandle may require to access related UObjects (e.g. owner actor), which is not safe outside game thread.
 			OutOverlap.OverlapObjectHandle = FActorInstanceHandle::MakeActorHandleToResolve(PossibleOwner, INDEX_NONE);
 			OutOverlap.Component = PossibleOwner;
-			OutOverlap.ItemIndex = INDEX_NONE;
+			OutOverlap.SetItemIndex(INDEX_NONE);
 
 			// If this is a geometry collection, we can do more to extract the properly ItemIndex.
 			if (const IPhysicsProxyBase* ActorProxy = Actor.GetProxy())
@@ -718,7 +738,7 @@ void ConvertQueryOverlap(const FPhysicsShape& Shape, const FPhysicsActor& Actor,
 				{
 					const FGeometryCollectionPhysicsProxy* ConcreteProxy = static_cast<const FGeometryCollectionPhysicsProxy*>(ActorProxy);
 					const FGeometryCollectionItemIndex ItemIndex = ConcreteProxy->GetItemIndexFromGTParticle_External(Actor.CastToRigidParticle());
-					OutOverlap.ItemIndex = ItemIndex.GetItemIndex();
+					OutOverlap.SetItemIndex(ItemIndex.GetItemIndex());
 				}
 			}
 
@@ -879,7 +899,7 @@ FHitResult ConvertOverlapToHitResult(const FOverlapResult& Overlap)
 {
 	FHitResult Hit;
 	Hit.bBlockingHit = Overlap.bBlockingHit;
-	Hit.Item = Overlap.ItemIndex;
+	Hit.Item = GetItemIndex(Overlap);
 	Hit.Component = Overlap.Component;
 	Hit.PhysicsObject = Overlap.PhysicsObject;
 	Hit.PhysicsObjectOwner = Overlap.PhysicsObjectOwner;
diff --git a/Engine/Source/Runtime/Engine/Private/Components/PrimitiveComponent.cpp b/Engine/Source/Runtime/Engine/Private/Components/PrimitiveComponent.cpp
--- a/Engine/Source/Runtime/Engine/Private/Components/PrimitiveComponent.cpp
+++ b/Engine/Source/Runtime/Engine/Private/Components/PrimitiveComponent.cpp
@@ -4354,7 +4354,7 @@ bool UPrimitiveComponent::UpdateOverlapsImpl(const TOverlapArrayView* NewPending
 							const bool bCheckOverlapFlags = false; // Already checked above
 							if (!ShouldIgnoreOverlapResult(MyWorld, MyActor, *this, Result.OverlapObjectHandle, *HitComp, bCheckOverlapFlags))
 							{
-								OverlapMultiResult.Emplace(HitComp, Result.ItemIndex);		// don't need to add unique unless the overlap check can return dupes
+								OverlapMultiResult.Emplace(HitComp, Result.GetItemIndex());		// don't need to add unique unless the overlap check can return dupes
 							}
 						}
 					}
diff --git a/Engine/Source/Runtime/Engine/Private/Engine/OverlapResult.cpp b/Engine/Source/Runtime/Engine/Private/Engine/OverlapResult.cpp
new file mode 100644
--- /dev/null
+++ b/Engine/Source/Runtime/Engine/Private/Engine/OverlapResult.cpp
@@ -0,0 +1,24 @@
+// Copyright Epic Games, Inc. All Rights Reserved.
+
+#include "Engine/OverlapResult.h"
+
+#include "Components/PrimitiveComponent.h"
+
+int32 FOverlapResult::GetItemIndexInternal() const
+{
+	return ItemIndex;
+}
+
+int32 FOverlapResult::GetItemIndex() const
+{
+	if (const UPrimitiveComponent* OwnerComponent = GetComponent())
+	{
+		return OwnerComponent->bMultiBodyOverlap ? ItemIndex : INDEX_NONE;
+	}
+	return INDEX_NONE;
+}
+
+void FOverlapResult::SetItemIndex(const int32 InItemIndex)
+{
+	ItemIndex = InItemIndex;
+}
diff --git a/Engine/Source/Runtime/Engine/Private/LevelActor.cpp b/Engine/Source/Runtime/Engine/Private/LevelActor.cpp
--- a/Engine/Source/Runtime/Engine/Private/LevelActor.cpp
+++ b/Engine/Source/Runtime/Engine/Private/LevelActor.cpp
@@ -1378,7 +1378,7 @@ static bool ComponentEncroachesBlockingGeometry_WithAdjustment(UWorld const* Wor
 				{
 					NumBlockingHits++;
 					FCollisionShape const NonShrunkenCollisionShape = PrimComp->GetCollisionShape();
-					const FBodyInstance* OverlapBodyInstance = OverlapComponent->GetBodyInstance(NAME_None, true, Overlaps[HitIdx].ItemIndex);
+					const FBodyInstance* OverlapBodyInstance = OverlapComponent->GetBodyInstance(NAME_None, true, Overlaps[HitIdx].GetItemIndex());
 					bool bSuccess = OverlapBodyInstance && OverlapBodyInstance->OverlapTest(TestWorldTransform.GetLocation(), TestWorldTransform.GetRotation(), NonShrunkenCollisionShape, &MTDResult);
 					if (bSuccess)
 					{
diff --git a/Engine/Source/Runtime/Engine/Private/PhysicsEngine/CollisionAnalyzerCapture.cpp b/Engine/Source/Runtime/Engine/Private/PhysicsEngine/CollisionAnalyzerCapture.cpp
--- a/Engine/Source/Runtime/Engine/Private/PhysicsEngine/CollisionAnalyzerCapture.cpp
+++ b/Engine/Source/Runtime/Engine/Private/PhysicsEngine/CollisionAnalyzerCapture.cpp
@@ -30,7 +30,7 @@ void CaptureOverlap(const UWorld* World, const FPhysicsGeometryCollection& PGeom
 		NewResult.bBlockingHit = OverlapResult.bBlockingHit;
 		NewResult.HitObjectHandle = OverlapResult.OverlapObjectHandle;
 		NewResult.Component = OverlapResult.Component;
-		NewResult.Item = OverlapResult.ItemIndex;
+		NewResult.Item = OverlapResult.GetItemIndex();
 		HitResults.Add(NewResult);
 	}
 
@@ -61,7 +61,7 @@ void CaptureOverlap(const UWorld* World, const FCollisionShape& PGeom, const FTr
 		NewResult.bBlockingHit = OverlapResult.bBlockingHit;
 		NewResult.HitObjectHandle = OverlapResult.OverlapObjectHandle;
 		NewResult.Component = OverlapResult.Component;
-		NewResult.Item = OverlapResult.ItemIndex;
+		NewResult.Item = OverlapResult.GetItemIndex();
 		HitResults.Add(NewResult);
 	}
 
