IncrementalCook: SkeletalMesh: Invalidate previously cooked SkeletalMesh packages when the SkeletalMesh DDC version changes.
#rnx
#incrementalcook

diff --git a/Engine/Source/Runtime/Engine/Classes/Engine/SkeletalMesh.h b/Engine/Source/Runtime/Engine/Classes/Engine/SkeletalMesh.h
--- a/Engine/Source/Runtime/Engine/Classes/Engine/SkeletalMesh.h
+++ b/Engine/Source/Runtime/Engine/Classes/Engine/SkeletalMesh.h
@@ -2497,7 +2497,8 @@ public:
 	ENGINE_API virtual void PostInitProperties() override;
 	ENGINE_API virtual void PostLoad() override;
 #if WITH_EDITORONLY_DATA
-	static ENGINE_API void DeclareConstructClasses(TArray<FTopLevelAssetPath>& OutConstructClasses, const UClass* SpecificSubclass);
+	ENGINE_API static void AppendToClassSchema(FAppendToClassSchemaContext& Context);
+	ENGINE_API static void DeclareConstructClasses(TArray<FTopLevelAssetPath>& OutConstructClasses, const UClass* SpecificSubclass);
 #endif
 
 	ENGINE_API virtual bool IsPostLoadThreadSafe() const override;
diff --git a/Engine/Source/Runtime/Engine/Private/SkeletalMesh.cpp b/Engine/Source/Runtime/Engine/Private/SkeletalMesh.cpp
--- a/Engine/Source/Runtime/Engine/Private/SkeletalMesh.cpp
+++ b/Engine/Source/Runtime/Engine/Private/SkeletalMesh.cpp
@@ -3798,6 +3798,13 @@ void USkeletalMesh::BeginPostLoadInternal(FSkinnedAssetPostLoadContext& Context)
 }
 
 #if WITH_EDITORONLY_DATA
+void USkeletalMesh::AppendToClassSchema(FAppendToClassSchemaContext& Context)
+{
+	Super::AppendToClassSchema(Context);
+	const FString& DerivedDataVersion = GetSkeletalMeshDerivedDataVersion();
+	Context.Update(*DerivedDataVersion, DerivedDataVersion.Len() * sizeof(**DerivedDataVersion));
+}
+
 void USkeletalMesh::DeclareConstructClasses(TArray<FTopLevelAssetPath>& OutConstructClasses, const UClass* SpecificSubclass)
 {
 	Super::DeclareConstructClasses(OutConstructClasses, SpecificSubclass);
