Moving GetRHIShaderBase calls to tasks to distribute the decompression overhead.

diff --git a/Engine/Source/Runtime/Engine/Private/ShaderCompiler/ShaderCompiler.cpp b/Engine/Source/Runtime/Engine/Private/ShaderCompiler/ShaderCompiler.cpp
--- a/Engine/Source/Runtime/Engine/Private/ShaderCompiler/ShaderCompiler.cpp
+++ b/Engine/Source/Runtime/Engine/Private/ShaderCompiler/ShaderCompiler.cpp
@@ -4914,21 +4914,51 @@ void PrecacheComputePipelineStatesForGlobalShaders(ERHIFeatureLevel::Type Featur
 		if (PipelineStateCache::IsPSOPrecachingEnabled())
 		{
 			UE_LOG(LogShaders, Display, TEXT("Precaching %d global compute shaders"), ComputeShadersToPrecache.Num());
-			ENQUEUE_RENDER_COMMAND(PrecachePSOsForGlobalShaders)(
-				[ComputeShadersToPrecache](FRHICommandListImmediate& RHICmdList)
+
+			if (GRHIGlobals.SupportsMultithreadedShaderCreation)
 			{
-				for (TShaderRef<FShader> GlobalShader : ComputeShadersToPrecache)
+				TArray<UE::Tasks::TTask<void>> Tasks;
+				for (TShaderRef<FShader>& GlobalShader : ComputeShadersToPrecache)
 				{
-					// PSO precache shaders are not required to all load correctly
-					bool bRequired = false;
-					const TCHAR* TypeName = GlobalShader.GetType()->GetName();
-					FRHIComputeShader* RHIComputeShader = static_cast<FRHIComputeShader*>(GlobalShader.GetRHIShaderBase(SF_Compute, bRequired));
-					if (RHIComputeShader)
-					{
-						PipelineStateCache::PrecacheComputePipelineState(RHIComputeShader, TypeName);
-					}
+					Tasks.Emplace(UE::Tasks::Launch(TEXT("PrecachePSOsForGlobalShaders"),
+						[&GlobalShader]()
+						{
+							// PSO precache shaders are not required to all load correctly
+							constexpr bool bRequired = false;
+
+							FRHIComputeShader* RHIComputeShader = static_cast<FRHIComputeShader*>(GlobalShader.GetRHIShaderBase(SF_Compute, bRequired));
+							if (RHIComputeShader)
+							{
+								const TCHAR* TypeName = GlobalShader.GetType()->GetName();
+								PipelineStateCache::PrecacheComputePipelineState(RHIComputeShader, TypeName);
+							}
+						}
+					));
 				}
-			});
+				ENQUEUE_RENDER_COMMAND(WaitOnPSOPrecacheTasks)(
+					[ComputeShadersToPrecache = MoveTemp(ComputeShadersToPrecache), Tasks = MoveTemp(Tasks)](FRHICommandListImmediate& RHICmdList)
+					{
+						UE::Tasks::Wait(Tasks);
+					});
+			}
+			else
+			{
+				ENQUEUE_RENDER_COMMAND(PrecachePSOsForGlobalShaders)(
+					[ComputeShadersToPrecache = MoveTemp(ComputeShadersToPrecache)](FRHICommandListImmediate& RHICmdList)
+					{
+						for (TShaderRef<FShader> GlobalShader : ComputeShadersToPrecache)
+						{
+							// PSO precache shaders are not required to all load correctly
+							bool bRequired = false;
+							FRHIComputeShader* RHIComputeShader = static_cast<FRHIComputeShader*>(GlobalShader.GetRHIShaderBase(SF_Compute, bRequired));
+							if (RHIComputeShader)
+							{
+								const TCHAR* TypeName = GlobalShader.GetType()->GetName();
+								PipelineStateCache::PrecacheComputePipelineState(RHIComputeShader, TypeName);
+							}
+						}
+					});
+			}
 		}
 		else if (IsPSOShaderPreloadingEnabled())
 		{
