New Material Translator:  Double constant and double vector parameter support.  AppendVector fix for appending float with LWC.  Improved tidiness when fetching uniform parameters: when a subset of a uniform parameter is fetched as a composite, it generates a swizzle, rather than fetching individual components -- works for both float and LWC parameters, although the code size reduction is much more significant for LWC.

#rnx

diff --git a/Engine/Source/Runtime/Engine/Private/Materials/MaterialExpressions.cpp b/Engine/Source/Runtime/Engine/Private/Materials/MaterialExpressions.cpp
--- a/Engine/Source/Runtime/Engine/Private/Materials/MaterialExpressions.cpp
+++ b/Engine/Source/Runtime/Engine/Private/Materials/MaterialExpressions.cpp
@@ -8450,11 +8450,11 @@ UMaterialExpressionDoubleVectorParameter::UMaterialExpressionDoubleVectorParamet
 {
 #if WITH_EDITORONLY_DATA
 	Outputs.Reset();
-	Outputs.Add(FExpressionOutput(TEXT(""), 1, 1, 1, 1, 0));
-	Outputs.Add(FExpressionOutput(TEXT(""), 1, 1, 0, 0, 0));
-	Outputs.Add(FExpressionOutput(TEXT(""), 1, 0, 1, 0, 0));
-	Outputs.Add(FExpressionOutput(TEXT(""), 1, 0, 0, 1, 0));
-	Outputs.Add(FExpressionOutput(TEXT(""), 1, 0, 0, 0, 1));
+	Outputs.Add(FExpressionOutput(TEXT("XYZ"), 1, 1, 1, 1, 0));
+	Outputs.Add(FExpressionOutput(TEXT("X"), 1, 1, 0, 0, 0));
+	Outputs.Add(FExpressionOutput(TEXT("Y"), 1, 0, 1, 0, 0));
+	Outputs.Add(FExpressionOutput(TEXT("Z"), 1, 0, 0, 1, 0));
+	Outputs.Add(FExpressionOutput(TEXT("W"), 1, 0, 0, 0, 1));
 #endif // WITH_EDITORONLY_DATA
 }
 
diff --git a/Engine/Source/Runtime/Engine/Private/Materials/MaterialExpressionsIR.cpp b/Engine/Source/Runtime/Engine/Private/Materials/MaterialExpressionsIR.cpp
--- a/Engine/Source/Runtime/Engine/Private/Materials/MaterialExpressionsIR.cpp
+++ b/Engine/Source/Runtime/Engine/Private/Materials/MaterialExpressionsIR.cpp
@@ -348,6 +348,12 @@ void UMaterialExpressionConstant4Vector::Build(MIR::FEmitter& Em)
 	}
 }
 
+void UMaterialExpressionGenericConstant::Build(MIR::FEmitter& Em)
+{
+	FValueRef Value = Em.ConstantFromShaderValue(GetConstantValue());
+	Em.Output(0, Value);
+}
+
 void UMaterialExpressionConstantBiasScale::Build(MIR::FEmitter& Em)
 {
 	Em.Output(0, Em.Multiply(Em.Add(Em.ConstantFloat(Bias), Em.Input(&Input)), Em.ConstantFloat(Scale)));
@@ -385,6 +391,16 @@ void UMaterialExpressionVectorParameter::Build(MIR::FEmitter& Em)
 	Em.Output(4, Em.Subscript(Value, 3));
 }
 
+void UMaterialExpressionDoubleVectorParameter::Build(MIR::FEmitter& Em)
+{
+	FValueRef Value = BuildMaterialExpressionParameter(Em, this);
+	Em.Output(0, Em.Cast(Value, MIR::FPrimitiveType::GetLWC3()));
+	Em.Output(1, Em.Subscript(Value, 0));
+	Em.Output(2, Em.Subscript(Value, 1));
+	Em.Output(3, Em.Subscript(Value, 2));
+	Em.Output(4, Em.Subscript(Value, 3));
+}
+
 void UMaterialExpressionChannelMaskParameter::Build(MIR::FEmitter& Em)
 {
 	FValueRef DotResult = Em.Dot(
@@ -494,19 +510,19 @@ void UMaterialExpressionAppendVector::Build(MIR::FEmitter& Em)
 	MIR::EScalarKind ResultKind = (AType->ScalarKind == MIR::ScalarKind_LWC || (BType && BType->ScalarKind == MIR::ScalarKind_LWC)) ? MIR::ScalarKind_LWC : MIR::ScalarKind_Float;
 	const MIR::FPrimitiveType* ResultType = MIR::FPrimitiveType::GetVector(ResultKind, Dimensions);
 
-	// Set up each output vector component 
+	// Set up each output vector component.  These need CastToScalarKind in case we are appending LWC and non-LWC.
 	FValueRef Components[4] = { nullptr, nullptr, nullptr, nullptr };
 	int ComponentIndex = 0;
 	for (int i = 0; i < AType->NumRows; ++i, ++ComponentIndex)
 	{
-		Components[ComponentIndex] = Em.Subscript(AVal, i);
+		Components[ComponentIndex] = Em.CastToScalarKind(Em.Subscript(AVal, i), ResultKind);
 	}
 
 	if (BType)
 	{
 		for (int i = 0; i < BType->NumRows; ++i, ++ComponentIndex)
 		{
-			Components[ComponentIndex] = Em.Subscript(BVal, i);
+			Components[ComponentIndex] = Em.CastToScalarKind(Em.Subscript(BVal, i), ResultKind);
 		}
 	}
 
diff --git a/Engine/Source/Runtime/Engine/Private/Materials/MaterialIR.cpp b/Engine/Source/Runtime/Engine/Private/Materials/MaterialIR.cpp
--- a/Engine/Source/Runtime/Engine/Private/Materials/MaterialIR.cpp
+++ b/Engine/Source/Runtime/Engine/Private/Materials/MaterialIR.cpp
@@ -3,6 +3,7 @@
 #include "Materials/MaterialIR.h"
 #include "Materials/MaterialIRTypes.h"
 #include "Materials/MaterialAggregate.h"
+#include "Misc/LargeWorldRenderPosition.h"
 #include "Engine/Texture.h"
 
 #if WITH_EDITOR
@@ -421,6 +422,21 @@ bool FConstant::IsFloat() const
 	return Type == FPrimitiveType::GetFloat();
 }
 
+bool FConstant::IsLWC() const
+{
+	return Type == FPrimitiveType::GetLWC();
+}
+
+bool FConstant::IsLWCRepresentableAsFloat() const
+{
+	if (Type != FPrimitiveType::GetLWC())
+	{
+		return false;
+	}
+	FLargeWorldRenderScalar ConstantLWC(LWCValue);
+	return ConstantLWC.GetTile() == 0.0f;
+}
+
 const TCHAR* LexToString(EExternalInput Input)
 {
 	switch (Input)
diff --git a/Engine/Source/Runtime/Engine/Private/Materials/MaterialIRDebug.cpp b/Engine/Source/Runtime/Engine/Private/Materials/MaterialIRDebug.cpp
--- a/Engine/Source/Runtime/Engine/Private/Materials/MaterialIRDebug.cpp
+++ b/Engine/Source/Runtime/Engine/Private/Materials/MaterialIRDebug.cpp
@@ -115,7 +115,8 @@ struct FDebugDumpIRUseGraphState
 			{
 				case ScalarKind_Bool:	Out.Append(Constant->Boolean ? TEXT("true") : TEXT("false")); break;
 				case ScalarKind_Int:	Out.Appendf(TEXT("%lld"), Constant->Integer); break;
-				case ScalarKind_Float: Out.Appendf(TEXT("%f"), Constant->Float); break;
+				case ScalarKind_Float:	Out.Appendf(TEXT("%f"), Constant->Float); break;
+				case ScalarKind_LWC:	Out.Appendf(TEXT("%f"), Constant->LWCValue); break;
 				default: UE_MIR_UNREACHABLE();
 			}
 		}
@@ -367,6 +368,10 @@ struct FDebugDumpIRState
 			{
 				Out.Appendf(TEXT("%.5ff"), Constant->Float);
 			}
+			else if (Constant->IsLWC())
+			{
+				Out.Appendf(TEXT("%.8f"), Constant->LWCValue);
+			}
 			else
 			{
 				UE_MIR_UNREACHABLE();
diff --git a/Engine/Source/Runtime/Engine/Private/Materials/MaterialIREmitter.cpp b/Engine/Source/Runtime/Engine/Private/Materials/MaterialIREmitter.cpp
--- a/Engine/Source/Runtime/Engine/Private/Materials/MaterialIREmitter.cpp
+++ b/Engine/Source/Runtime/Engine/Private/Materials/MaterialIREmitter.cpp
@@ -535,6 +535,11 @@ FValueRef FEmitter::ConstantFromShaderValue(const UE::Shader::FValue& InValue)
 		case UE::Shader::EValueType::Int2: return ConstantInt2({ InValue.Component[0].Int, InValue.Component[1].Int });
 		case UE::Shader::EValueType::Int3: return ConstantInt3({ InValue.Component[0].Int, InValue.Component[1].Int, InValue.Component[2].Int });
 		case UE::Shader::EValueType::Int4: return ConstantInt4({ InValue.Component[0].Int, InValue.Component[1].Int, InValue.Component[2].Int, InValue.Component[3].Int });
+
+		case UE::Shader::EValueType::Double1: return ConstantLWC(InValue.Component[0].Double);
+		case UE::Shader::EValueType::Double2: return Vector2(ConstantLWC(InValue.Component[0].Double), ConstantLWC(InValue.Component[1].Double));
+		case UE::Shader::EValueType::Double3: return Vector3(ConstantLWC(InValue.Component[0].Double), ConstantLWC(InValue.Component[1].Double), ConstantLWC(InValue.Component[2].Double));
+		case UE::Shader::EValueType::Double4: return Vector4(ConstantLWC(InValue.Component[0].Double), ConstantLWC(InValue.Component[1].Double), ConstantLWC(InValue.Component[2].Double), ConstantLWC(InValue.Component[3].Double));
 	}
 
 	UE_MIR_UNREACHABLE();
@@ -577,7 +582,7 @@ FValueRef FEmitter::ConstantZero(EScalarKind Kind)
 		case ScalarKind_Bool: return ConstantFalse();
 		case ScalarKind_Int: return ConstantInt(0);
 		case ScalarKind_Float: return ConstantFloat(0.0f);
-		case ScalarKind_LWC: return CastToScalarKind(ConstantFloat(0.0f), ScalarKind_LWC);
+		case ScalarKind_LWC: return ConstantLWC(0.0);
 		default: UE_MIR_UNREACHABLE();
 	}
 }
@@ -589,7 +594,7 @@ FValueRef FEmitter::ConstantOne(EScalarKind Kind)
 		case ScalarKind_Bool: return ConstantTrue();
 		case ScalarKind_Int: return ConstantInt(1);
 		case ScalarKind_Float: return ConstantFloat(1.0f);
-		case ScalarKind_LWC: return CastToScalarKind(ConstantFloat(1.0f), ScalarKind_LWC);
+		case ScalarKind_LWC: return ConstantLWC(1.0);
 		default: UE_MIR_UNREACHABLE();
 	}
 }
@@ -601,7 +606,7 @@ FValueRef FEmitter::ConstantScalar(EScalarKind Kind, TFloat FromFloat)
 		case ScalarKind_Bool: return ConstantBool(FromFloat != 0.0f);
 		case ScalarKind_Int: return ConstantInt(TInteger(FromFloat));
 		case ScalarKind_Float: return ConstantFloat(FromFloat);
-		case ScalarKind_LWC: return CastToScalarKind(ConstantFloat(FromFloat), ScalarKind_LWC);
+		case ScalarKind_LWC: return ConstantLWC(FromFloat);
 		default: UE_MIR_UNREACHABLE();
 	}
 }
@@ -707,6 +712,13 @@ FValueRef FEmitter::ConstantFloat4(UE::Math::TVector4<TFloat> InValue)
 	return Vector4(X, Y, Z, W);
 }
 
+FValueRef FEmitter::ConstantLWC(TDouble InX)
+{
+	FConstant Scalar = MakePrototype<FConstant>(FPrimitiveType::GetScalar(ScalarKind_LWC));
+	Scalar.LWCValue = InX;
+	return EmitPrototype(*this, Scalar);
+}
+
 /*--------------------- Other non-instruction values emission ---------------------*/
 
 FValueRef FEmitter::Poison()
@@ -803,6 +815,22 @@ FValueRef FEmitter::Parameter(FName Name, FMaterialParameterMetadata& Metadata,
 			return EmitPrototype(*this, Prototype);
 		}
 
+		case EMaterialParameterType::DoubleVector:
+		{
+			if (Metadata.PrimitiveDataIndex != INDEX_NONE)
+			{
+				MIR::FValue* X = CustomPrimitiveData(Metadata.PrimitiveDataIndex + 0);
+				MIR::FValue* Y = CustomPrimitiveData(Metadata.PrimitiveDataIndex + 1);
+				MIR::FValue* Z = CustomPrimitiveData(Metadata.PrimitiveDataIndex + 2);
+				MIR::FValue* W = CustomPrimitiveData(Metadata.PrimitiveDataIndex + 3);
+				return Vector4(X, Y, Z, W);
+			}
+
+			MIR::FUniformParameter Prototype = MakePrototype<MIR::FUniformParameter>(FPrimitiveType::GetLWC4());
+			Prototype.ParameterIdInModule = RegisterParameter(Module, Info, Metadata);
+			return EmitPrototype(*this, Prototype);
+		}
+
 		case EMaterialParameterType::Texture:
 		{
 			switch (Metadata.Value.Texture->GetMaterialType())
@@ -1485,6 +1513,10 @@ static FValueRef TryFoldOperatorScalar(FEmitter& Emitter, EOperator Op, FValueRe
 				Result = FoldComparisonOperatorScalar<TFloat>(Op, AConstant->Float, BConstant->Float);
 				break;
 
+			case ScalarKind_LWC:
+				Result = FoldComparisonOperatorScalar<TDouble>(Op, AConstant->LWCValue, BConstant->LWCValue);
+				break;
+
 			default:
 				UE_MIR_UNREACHABLE();
 		}
@@ -1512,6 +1544,12 @@ static FValueRef TryFoldOperatorScalar(FEmitter& Emitter, EOperator Op, FValueRe
 				return Emitter.ConstantFloat(Result);
 			}
 
+			case ScalarKind_LWC:
+			{
+				TDouble Result = FoldScalarOperator<TDouble>(Emitter, Op, AConstant->LWCValue, BConstant ? BConstant->LWCValue : 0, CConstant ? CConstant->LWCValue : 0);
+				return Emitter.ConstantLWC(Result);
+			}
+
 			default:
 				UE_MIR_UNREACHABLE();
 		}
@@ -1727,7 +1765,7 @@ static const FPrimitiveType* ValidateOperatorAndGetResultType(FEmitter& Emitter,
 		}
 		else if (ArgumentPrimitiveType->ScalarKind == EScalarKind::ScalarKind_LWC && !(Filter & OPTF_LWCInput))
 		{
-		// Handle automatic cast to float for operators that don't support LWC inputs
+			// Handle automatic cast to float for operators that don't support LWC inputs
 			ArgumentPrimitiveType = ArgumentPrimitiveType->WithScalarKind(EScalarKind::ScalarKind_Float);
 			Arguments[i] = Emitter.Cast(Arguments[i], ArgumentPrimitiveType);
 		}
@@ -2062,6 +2100,16 @@ static MIR::TFloat ConstantDotFloat(TArrayView<FConstant*> AComponents, TArrayVi
 	}
 	return Result;
 }
+static MIR::TDouble ConstantDotDouble(TArrayView<FConstant*> AComponents, TArrayView<FConstant*> BComponents, int32 NumComponents)
+{
+	double Result = 0.0;
+	for (int32 i = 0; i < NumComponents; ++i)
+	{
+		Result += AComponents[i]->LWCValue * BComponents[i]->LWCValue;
+	}
+	return Result;
+}
+
 
 /// Tries to fold the operator, that is to evaluate its result now at translation time if its arguments are constant.
 /// If the operator could not be folded in any way, it returns nullptr.
@@ -2086,6 +2134,11 @@ static FValue* TryFoldOperator(FEmitter& Emitter, EOperator Op, FValue* A, FValu
 				float Result = FMath::Sqrt(ConstantDotFloat(AComponents, AComponents, ANumComponents));
 				return Emitter.ConstantFloat(Result);
 			}
+			else if (ResultType->ScalarKind == ScalarKind_LWC)
+			{
+				double Result = FMath::Sqrt(ConstantDotDouble(AComponents, AComponents, ANumComponents));
+				return Emitter.ConstantLWC(Result);
+			}
 			else
 			{
 				UE_MIR_UNREACHABLE();
@@ -2103,6 +2156,11 @@ static FValue* TryFoldOperator(FEmitter& Emitter, EOperator Op, FValue* A, FValu
 					float Result = ConstantDotFloat(AComponents, BComponents, ANumComponents);
 					return Emitter.ConstantFloat(Result);
 				}
+				else if (ResultType->ScalarKind == ScalarKind_LWC)
+				{
+					double Result = ConstantDotDouble(AComponents, BComponents, ANumComponents);
+					return Emitter.ConstantLWC(Result);
+				}
 				else
 				{
 					UE_MIR_UNREACHABLE();
@@ -2370,13 +2428,23 @@ static FValue* CastConstant(FEmitter& Emitter, FConstant* Constant, EScalarKind
 	switch (ConstantScalarKind)
 	{
 		case ScalarKind_Bool:
+		{
+			switch (TargetKind)
+			{
+				case ScalarKind_Int: return Emitter.ConstantInt(Constant->Boolean ? 1 : 0);
+				case ScalarKind_Float: return Emitter.ConstantFloat(Constant->Boolean ? 1.0f : 0.0f);
+				case ScalarKind_LWC: return Emitter.ConstantLWC(Constant->Boolean ? 1.0 : 0.0);
+				default: UE_MIR_UNREACHABLE();
+			}
+		}
+
 		case ScalarKind_Int:
 		{
 			switch (TargetKind)
 			{
-				case ScalarKind_Bool: return {};
-				case ScalarKind_Int: return Emitter.ConstantInt(Constant->Integer);
+				case ScalarKind_Bool: return Emitter.ConstantBool(Constant->Integer != 0);
 				case ScalarKind_Float: return Emitter.ConstantFloat((TFloat)Constant->Integer);
+				case ScalarKind_LWC: return Emitter.ConstantLWC(Constant->Integer);
 				default: UE_MIR_UNREACHABLE();
 			}
 		}
@@ -2385,8 +2453,20 @@ static FValue* CastConstant(FEmitter& Emitter, FConstant* Constant, EScalarKind
 		{
 			switch (TargetKind)
 			{
-				case ScalarKind_Bool: return {};
+				case ScalarKind_Bool: return Emitter.ConstantBool(Constant->Float != 0.0f);
 				case ScalarKind_Int: return Emitter.ConstantInt((int32)Constant->Float);
+				case ScalarKind_LWC: return Emitter.ConstantLWC(Constant->Float);
+				default: UE_MIR_UNREACHABLE();
+			}
+		}
+
+		case ScalarKind_LWC:
+		{
+			switch (TargetKind)
+			{
+				case ScalarKind_Bool: return Emitter.ConstantBool(Constant->LWCValue != 0.0);
+				case ScalarKind_Int: return Emitter.ConstantInt((int32)Constant->LWCValue);
+				case ScalarKind_Float: return Emitter.ConstantFloat((float)Constant->LWCValue);
 				default: UE_MIR_UNREACHABLE();
 			}
 		}
@@ -2419,9 +2499,8 @@ static FValue* CastValueToPrimitiveType(FEmitter& Emitter, FValueRef Value, cons
 		{
 			return Value;
 		}
-		// If value is a constant, we can cast the constant now.  LWC is not a supported constant type at the moment, so in that case,
-		// we go through the dynamic cast code path below.
-		else if (FConstant* ConstantInitializer = TargetPrimitiveType->ScalarKind != EScalarKind::ScalarKind_LWC ? As<FConstant>(Value) : nullptr)
+		// If value is a constant, we can cast the constant now.
+		else if (FConstant* ConstantInitializer = As<FConstant>(Value))
 		{
 			return CastConstant(Emitter, ConstantInitializer, ValuePrimitiveType->ScalarKind, TargetPrimitiveType->ScalarKind);
 		}
@@ -2860,11 +2939,14 @@ static FValue* DifferentiateOperator(FEmitter& E, FOperator* Op, EDerivativeAxis
 	// Convenience local functions as multiplications and division operations are common in derivatives.
 	auto Zero = [&E, PrimType] () { return E.ConstantZero(PrimType->ScalarKind); };
 	auto One = [&E, PrimType] () { return E.ConstantOne(PrimType->ScalarKind); };
-	auto Constant = [&E, PrimType] (TFloat FromFloat) { return E.ConstantScalar(PrimType->ScalarKind, FromFloat); };
+	auto Constant = [&E, PrimType] (TDouble FromDouble)
+	{ 
+		return PrimType->ScalarKind == ScalarKind_LWC ? E.ConstantLWC(FromDouble) : E.ConstantScalar(PrimType->ScalarKind, (float)FromDouble);
+	};
 
 	// Some constants
-	constexpr TFloat Ln2 = 0.69314718055994530941723212145818;
-	constexpr TFloat Ln10 = 2.3025850929940456840179914546844;
+	constexpr TDouble Ln2 = 0.69314718055994530941723212145818;
+	constexpr TDouble Ln10 = 2.3025850929940456840179914546844;
 	
 	switch (Op->Op)
 	{
diff --git a/Engine/Source/Runtime/Engine/Private/Materials/MaterialIRToHLSLTranslator.cpp b/Engine/Source/Runtime/Engine/Private/Materials/MaterialIRToHLSLTranslator.cpp
--- a/Engine/Source/Runtime/Engine/Private/Materials/MaterialIRToHLSLTranslator.cpp
+++ b/Engine/Source/Runtime/Engine/Private/Materials/MaterialIRToHLSLTranslator.cpp
@@ -15,6 +15,7 @@
 #include "MaterialDomain.h"
 #include "DataDrivenShaderPlatformInfo.h"
 #include "Materials/MaterialExpressionVolumetricAdvancedMaterialOutput.h"
+#include "Misc/LargeWorldRenderPosition.h"
 #include "RenderUtils.h"
 #include "Engine/Texture.h"
 
@@ -34,6 +35,8 @@ enum EListSeparator { ListSeparator };
 
 #define TAB "    "
 
+static const char* GVector4SwizzleSubset[4] = { ".x", ".xy", ".xyz", "" };
+
 struct FHLSLPrinter
 {
 	FString Buffer;
@@ -466,7 +469,12 @@ struct FTranslator : FMaterialIRToHLSLTranslation
 				Printer << Constant->Float;
 				break;
 
-			case MIR::ScalarKind_LWC: UE_MIR_UNREACHABLE();
+			case MIR::ScalarKind_LWC:
+				{
+					FLargeWorldRenderScalar ConstantLWC(Constant->LWCValue);
+					Printer << "MakeLWCScalar(" << ConstantLWC.GetTile() << ", " << ConstantLWC.GetOffset() << ")";
+				}
+				break;
 		}
 	}
 
@@ -832,11 +840,6 @@ struct FTranslator : FMaterialIRToHLSLTranslation
 		check(PrimitiveType);
 		check(PrimitiveType->IsScalar() || PrimitiveType->IsVector()); // no matrices yet
 
-		if (PrimitiveType->ScalarKind == MIR::EScalarKind::ScalarKind_Int)
-		{
-			Printer << TEXT("asint(");
-		}
-
 		const FUniformExpressionSet& UniformExpressionSet = Module->GetCompilationOutput().UniformExpressionSet;
 		
 		// Get the global float4 component index (e.g. if this is 13, it refer to PreshaderBuffer[3].y)
@@ -848,26 +851,62 @@ struct FTranslator : FMaterialIRToHLSLTranslation
 		// Starting component of the float4 slot
 		uint32 BufferSlotOffset = GlobalComponentOffset % 4;
 
-		Printer << TEXT("Material.PreshaderBuffer[") << BufferSlotIndex << TEXT("]");
-
-		if (PrimitiveType->NumComponents() < 4)
+		if (PrimitiveType->ScalarKind == MIR::EScalarKind::ScalarKind_LWC)
+		{
+			// Currently preshaders store values as DF, but eventually it would be more efficient to store them as LWC to avoid conversion.
+			// LWC parameters are always four element vectors, so we assume the slot offset is zero.
+			check(BufferSlotOffset == 0);
+			Printer << "DFToWS(MakeDFVector4(Material.PreshaderBuffer[" << BufferSlotIndex << "], Material.PreshaderBuffer[" << BufferSlotIndex + 1 << "]))";
+		}
+		else
 		{
-			Printer << TEXT(".");
+			if (PrimitiveType->ScalarKind == MIR::EScalarKind::ScalarKind_Int)
+			{
+				Printer << TEXT("asint(");
+			}
+
+			Printer << TEXT("Material.PreshaderBuffer[") << BufferSlotIndex << TEXT("]");
 
-			static const TCHAR* Components = TEXT("xyzw");
-			for (int32 i = 0; i < PrimitiveType->NumComponents(); ++i)
+			if (PrimitiveType->NumComponents() < 4)
 			{
-				check(BufferSlotOffset + i < sizeof(Components));
-				Printer.Buffer.AppendChar(Components[BufferSlotOffset + i]);
+				Printer << TEXT(".");
+
+				static const TCHAR* Components = TEXT("xyzw");
+				for (int32 i = 0; i < PrimitiveType->NumComponents(); ++i)
+				{
+					check(BufferSlotOffset + i < sizeof(Components));
+					Printer.Buffer.AppendChar(Components[BufferSlotOffset + i]);
+				}
 			}
-		}
 
-		if (PrimitiveType->ScalarKind != MIR::EScalarKind::ScalarKind_Float)
-		{
-			Printer << ")"; // close the "asint(" bracket
+			if (PrimitiveType->ScalarKind == MIR::EScalarKind::ScalarKind_Int)
+			{
+				Printer << ")"; // close the "asint(" bracket
+			}
 		}
 	}
 
+	// Variation of LowerPrimitiveUniformParameter specific to LWC, to support fetching a subset of an LWC vector
+	void LowerPrimitiveUniformParameterLWC(const MIR::FUniformParameter* UniformParameter, int32 NumComponents)
+	{
+		check(NumComponents >= 2 && NumComponents <= 4);
+
+		const MIR::FPrimitiveType* PrimitiveType = UniformParameter->Type->AsPrimitive();
+		check(PrimitiveType->IsLWC());
+
+		const FUniformExpressionSet& UniformExpressionSet = Module->GetCompilationOutput().UniformExpressionSet;
+
+		// Get the global float4 component index (e.g. if this is 13, it refer to PreshaderBuffer[3].y)
+		uint32 GlobalComponentOffset = UniformExpressionSet.GetNumericParameterEvaluationOffset(UniformParameter->Analysis_UniformParameterIndex);
+
+		// Index of the float4 slot.  LWC parameters always have four elements, so there won't be an offset.
+		uint32 BufferSlotIndex = GlobalComponentOffset / 4;
+
+		Printer << "DFToWS(MakeDFVector" << NumComponents;
+		Printer << "(Material.PreshaderBuffer[" << BufferSlotIndex << "]" << GVector4SwizzleSubset[NumComponents - 1];
+		Printer << ", Material.PreshaderBuffer[" << BufferSlotIndex + 1 << "]" << GVector4SwizzleSubset[NumComponents - 1] << "))";
+	}
+
 	bool HasMatchingScalarComponentCastChain(MIR::FValue* FirstComponent, MIR::FValue* CurrentComponent, int32 Index, int32 NumComponents)
 	{
 		if (FirstComponent->Kind == MIR::VK_Subscript && CurrentComponent->Kind == MIR::VK_Subscript)
@@ -875,7 +914,21 @@ struct FTranslator : FMaterialIRToHLSLTranslation
 			if (((MIR::FSubscript*)FirstComponent)->Arg == ((MIR::FSubscript*)CurrentComponent)->Arg)
 			{
 				const MIR::FPrimitiveType* PrimitiveType = ((MIR::FSubscript*)FirstComponent)->Arg->Type->AsVector();
-				return PrimitiveType && PrimitiveType->NumRows == NumComponents && ((MIR::FSubscript*)FirstComponent)->Index == 0 && ((MIR::FSubscript*)CurrentComponent)->Index == Index;
+				if (PrimitiveType && ((MIR::FSubscript*)FirstComponent)->Index == 0 && ((MIR::FSubscript*)CurrentComponent)->Index == Index)
+				{
+					// Exact match?
+					if (PrimitiveType->NumRows == NumComponents)
+					{
+						return true;
+					}
+
+					// Subset of components?  See if we can generate a swizzled vector.
+					if (PrimitiveType->NumRows >= NumComponents)
+					{
+						// Allow swizzle for non-LWC values or uniform parameters (the latter having a special case for LWC)
+						return PrimitiveType->ScalarKind != MIR::ScalarKind_LWC || ((MIR::FSubscript*)FirstComponent)->Arg->Kind == MIR::VK_UniformParameter;
+					}
+				}
 			}
 		}
 		else if (FirstComponent->Kind == MIR::VK_Scalar && CurrentComponent->Kind == MIR::VK_Scalar)
@@ -962,11 +1015,46 @@ struct FTranslator : FMaterialIRToHLSLTranslation
 		{
 			// Finally reached the inner subscript, print its vector argument
 			check(FirstComponent->Kind == MIR::VK_Subscript);
-			Printer << LowerValue(((MIR::FSubscript*)FirstComponent)->Arg);
+			MIR::FValue* InnerArg = ((MIR::FSubscript*)FirstComponent)->Arg;
+			const MIR::FPrimitiveType* PrimitiveType = InnerArg->Type->AsVector();
+
+			if (PrimitiveType->IsLWC() && InnerArg->Kind == MIR::VK_UniformParameter)
+			{
+				// LWC types support special case swizzling logic for uniform parameters on initial fetch.  See HasMatchingScalarComponentCastChain above.
+				LowerPrimitiveUniformParameterLWC((const MIR::FUniformParameter*)InnerArg, NumComponents);
+			}
+			else
+			{
+				Printer << LowerValue(InnerArg);
+
+				if (PrimitiveType->NumRows > NumComponents)
+				{
+					Printer << GVector4SwizzleSubset[NumComponents - 1];
+				}
+			}
 		}
 		return NoOp;
 	}
 
+	// Check if we can tidy up casts from non-LWC scalar to LWC vectors (assuming we already know all components are the same and LWC)
+	bool IsScalarToLWCVectorCast(const MIR::FValue* FirstComponent)
+	{
+		if (FirstComponent->Kind == MIR::VK_Scalar)
+		{
+			// Check if this is a cast from an arithmetic non-LWC scalar type.  If so, we can cast the scalar to a float vector, and then to LWC.
+			const MIR::FScalar* Scalar = (const MIR::FScalar*)FirstComponent;
+			const MIR::EScalarKind ArgScalarKind = Scalar->Arg->Type->AsPrimitive()->ScalarKind;
+
+			return ArgScalarKind != MIR::ScalarKind_LWC && ArgScalarKind != MIR::ScalarKind_Bool;
+		}
+		else if (FirstComponent->Kind == MIR::VK_Constant)
+		{
+			// Check if this is a cast from a constant representable exactly as a non-LWC float
+			return ((const MIR::FConstant*)FirstComponent)->IsLWCRepresentableAsFloat();
+		}
+		return false;
+	}
+
 	void LowerComposite(const MIR::FComposite* Dimensional) 
 	{
 		TArrayView<MIR::FValue* const> Components = Dimensional->GetComponents();
@@ -979,27 +1067,87 @@ struct FTranslator : FMaterialIRToHLSLTranslation
 		// to call a function to convert types.
 		bool bSameComponents = true;
 
+		// Track if all components are part of a constant vector
+		bool bConstantVector = Components[0]->Kind == MIR::VK_Constant && ArithmeticType->IsVector();
+
 		// We can also generate tidier HLSL for cases where casts are done for whole vectors.  This is a frequent case for LWC, where casts happen
 		// in both directions (LWC to float and back), due to operations automatically casting their inputs or outputs.  For example, BO_Sub
 		// automatically downcasts LWC to float, but if you feed that into a BO_Add that has LWC for its other input, it will immediately cast it
 		// back to LWC -- two consecutive casts.  These then get expanded into individual scalar casts, which we would like to do as whole vector
 		// casts for readability.  We only need to consider this special case if the first component is a Scalar, and the type is a Vector.
-		bool bWholeVectorCast = Components[0]->Kind == MIR::VK_Scalar && ArithmeticType->IsVector();
+		bool bWholeVectorCast = ArithmeticType->IsVector() && (Components[0]->Kind == MIR::VK_Scalar || Components[0]->Kind == MIR::VK_Subscript);
 
-		for (int32 i = 1; bSameComponents && i < Components.Num(); ++i)
+		for (int32 i = 1; i < Components.Num(); ++i)
 		{
 			bSameComponents &= (Components[i] == Components[0]);
-			bWholeVectorCast &= HasMatchingScalarComponentCastChain(Components[0], Components[i], i, Components.Num());
+			bConstantVector &= Components[i]->Kind == MIR::VK_Constant;
+			bWholeVectorCast = bWholeVectorCast && HasMatchingScalarComponentCastChain(Components[0], Components[i], i, Components.Num());
 		}
 		
 		if (bSameComponents && ArithmeticType->ScalarKind != MIR::ScalarKind_LWC)
 		{
 			Printer << TEXT("(") << LowerType(ArithmeticType) << TEXT(")") << LowerValue(Components[0]);
 		}
+		else if (bSameComponents && ArithmeticType->IsVector() && IsScalarToLWCVectorCast(Components[0]))
+		{
+			// Cast scalar to float vector, then promote, for example:  "WSPromote((float3)1.0f)"
+			Printer << "WSPromote" << "((float" << ArithmeticType->NumRows << ")";
+			if (Components[0]->Kind == MIR::VK_Scalar)
+			{
+				// Print the inner non-LWC scalar value
+				const MIR::FScalar* Scalar = (const MIR::FScalar*)Components[0];
+				LowerValue(Scalar->Arg);
+			}
+			else
+			{
+				// Print the inner double constant in its non-LWC form
+				check(Components[0]->Kind == MIR::VK_Constant)
+				Printer << (float)((const MIR::FConstant*)Components[0])->LWCValue;
+			}
+			Printer << ")";
+		}
 		else if (bWholeVectorCast)
 		{
 			LowerVectorCastChain(Components[0], Components.Num());
 		}
+		else if (bConstantVector && ArithmeticType->ScalarKind == MIR::ScalarKind_LWC)
+		{
+			// Special case for LWC constant vectors
+			bool bAllTileValuesZero = true;
+			float TileValues[4];
+			float OffsetValues[4];
+			for (int32 Index = 0; Index < Components.Num(); Index++)
+			{
+				FLargeWorldRenderScalar ConstantLWC(((const MIR::FConstant*)Components[Index])->LWCValue);
+				TileValues[Index] = ConstantLWC.GetTile();
+				OffsetValues[Index] = ConstantLWC.GetOffset();
+				bAllTileValuesZero &= TileValues[Index] == 0.0f;
+			}
+
+			if (bAllTileValuesZero)
+			{
+				// Vector representable as regular floats (all tile values zero), call WSPromote on the offset vector
+				Printer << "WSPromote(";
+			}
+			else
+			{
+				// Vector needs tile values, call LWC constructor and generate tile vector, before generating offset vector below
+				Printer << "MakeLWCVector" << Components.Num() << "(float" << Components.Num() << BeginArgs;
+				for (int32 Index = 0; Index < Components.Num(); Index++)
+				{
+					Printer << ListSeparator << TileValues[Index];
+				}
+				Printer << EndArgs << ", ";
+			}
+
+			// Generate offset vector, plus extra parentheses to close WSPromote or MakeLWCVector call
+			Printer << "float" << Components.Num() << BeginArgs;
+			for (int32 Index = 0; Index < Components.Num(); Index++)
+			{
+				Printer << ListSeparator << OffsetValues[Index];
+			}
+			Printer << EndArgs << ")";
+		}
 		else
 		{
 			if (ArithmeticType->ScalarKind == MIR::ScalarKind_LWC)
diff --git a/Engine/Source/Runtime/Engine/Private/Materials/MaterialIRTypes.cpp b/Engine/Source/Runtime/Engine/Private/Materials/MaterialIRTypes.cpp
--- a/Engine/Source/Runtime/Engine/Private/Materials/MaterialIRTypes.cpp
+++ b/Engine/Source/Runtime/Engine/Private/Materials/MaterialIRTypes.cpp
@@ -359,11 +359,26 @@ const FPrimitiveType* FPrimitiveType::GetFloat4()
 	return GetVector(ScalarKind_Float, 4);
 }
 
+const FPrimitiveType* FPrimitiveType::GetLWC()
+{
+	return GetScalar(ScalarKind_LWC);
+}
+
+const FPrimitiveType* FPrimitiveType::GetLWC2()
+{
+	return GetVector(ScalarKind_LWC, 2);
+}
+
 const FPrimitiveType* FPrimitiveType::GetLWC3()
 {
 	return GetVector(ScalarKind_LWC, 3);
 }
 
+const FPrimitiveType* FPrimitiveType::GetLWC4()
+{
+	return GetVector(ScalarKind_LWC, 4);
+}
+
 const FPrimitiveType* FPrimitiveType::GetScalar(EScalarKind InScalarKind)
 {
 	return Get(InScalarKind, 1, 1);
@@ -415,9 +430,9 @@ const FPrimitiveType* FPrimitiveType::Get(EScalarKind InScalarKind, int NumRows,
 		{ { TypeKind_Primitive }, { TEXT("int4x3") }, 	ScalarKind_Int, 4, 3 },
 		{ { TypeKind_Primitive }, { TEXT("int4x4") }, 	ScalarKind_Int, 4, 4 },
 		{ { TypeKind_Primitive }, { TEXT("float") }, 		ScalarKind_Float, 1, 1 },
-		{ { TypeKind_Primitive }, Invalid, 				ScalarKind_Float, 1, 2 },
-		{ { TypeKind_Primitive }, Invalid, 				ScalarKind_Float, 1, 3 },
-		{ { TypeKind_Primitive }, Invalid, 				ScalarKind_Float, 1, 4 },
+		{ { TypeKind_Primitive }, Invalid, 					ScalarKind_Float, 1, 2 },
+		{ { TypeKind_Primitive }, Invalid, 					ScalarKind_Float, 1, 3 },
+		{ { TypeKind_Primitive }, Invalid, 					ScalarKind_Float, 1, 4 },
 		{ { TypeKind_Primitive }, { TEXT("float2") },   	ScalarKind_Float, 2, 1 },
 		{ { TypeKind_Primitive }, { TEXT("float2x2") }, 	ScalarKind_Float, 2, 2 },
 		{ { TypeKind_Primitive }, { TEXT("float2x3") }, 	ScalarKind_Float, 2, 3 },
diff --git a/Engine/Source/Runtime/Engine/Private/Materials/MaterialIRValueAnalyzer.cpp b/Engine/Source/Runtime/Engine/Private/Materials/MaterialIRValueAnalyzer.cpp
--- a/Engine/Source/Runtime/Engine/Private/Materials/MaterialIRValueAnalyzer.cpp
+++ b/Engine/Source/Runtime/Engine/Private/Materials/MaterialIRValueAnalyzer.cpp
@@ -186,6 +186,7 @@ static void AnalyzePrimitiveUniformParameter(FMaterialIRValueAnalyzer& Analyzer,
 	{
 		case EMaterialParameterType::Scalar: DefaultValue = ParameterMetadata.Value.AsScalar(); break;
 		case EMaterialParameterType::Vector: DefaultValue = ParameterMetadata.Value.AsLinearColor(); break;
+		case EMaterialParameterType::DoubleVector: DefaultValue = ParameterMetadata.Value.AsVector4d(); break;
 		case EMaterialParameterType::Enumeration: DefaultValue = Cast<IMaterialEnumerationProvider>(ParameterMetadata.Enumeration)->GetValueOrDefault(ParameterMetadata.Value.EnumEntry, 0); break;
 		default: UE_MIR_TODO();
 	}
@@ -205,7 +206,7 @@ static void AnalyzePrimitiveUniformParameter(FMaterialIRValueAnalyzer& Analyzer,
 	const MIR::FPrimitiveType* ParameterType = Parameter->Type->AsPrimitive();
 
 	// Only int and float parameters supported for now (TODO)
-	check(ParameterType->ScalarKind == MIR::ScalarKind_Int || ParameterType->ScalarKind == MIR::ScalarKind_Float);
+	check(ParameterType->ScalarKind == MIR::ScalarKind_Int || ParameterType->ScalarKind == MIR::ScalarKind_Float || ParameterType->ScalarKind == MIR::ScalarKind_LWC);
 
 	// Following code calculates GlobalComponentOffset.
 	// GlobalComponentOffset is the i-th component in the array of float4s that make the uniform buffer.
@@ -245,7 +246,7 @@ static void AnalyzePrimitiveUniformParameter(FMaterialIRValueAnalyzer& Analyzer,
 	if (UsedNumComponents >= 4)
 	{
 		check(UsedNumComponents == 4); 
-		GlobalComponentOffset = UniformExpressionSet.AllocateFromUniformBuffer(16) / 4;
+		GlobalComponentOffset = UniformExpressionSet.AllocateFromUniformBuffer(ParameterType->ScalarKind == MIR::ScalarKind_LWC ? 32 : 16) / 4;
 	}
 
 	// Finally, check whether in the used slot there's any slack left, and if so, record it in the appropriate free-list.
@@ -272,6 +273,7 @@ static void AnalyzePrimitiveUniformParameter(FMaterialIRValueAnalyzer& Analyzer,
 		{
 			case MIR::ScalarKind_Int: ParamInsight.ComponentType = FMaterialInsights::FUniformBufferSlotComponentType::CT_Int; break;
 			case MIR::ScalarKind_Float: ParamInsight.ComponentType = FMaterialInsights::FUniformBufferSlotComponentType::CT_Float; break;
+			case MIR::ScalarKind_LWC: ParamInsight.ComponentType = FMaterialInsights::FUniformBufferSlotComponentType::CT_LWC; break;
 			default: UE_MIR_UNREACHABLE();
 		}
 
diff --git a/Engine/Source/Runtime/Engine/Public/Materials/MaterialExpressionDoubleVectorParameter.h b/Engine/Source/Runtime/Engine/Public/Materials/MaterialExpressionDoubleVectorParameter.h
--- a/Engine/Source/Runtime/Engine/Public/Materials/MaterialExpressionDoubleVectorParameter.h
+++ b/Engine/Source/Runtime/Engine/Public/Materials/MaterialExpressionDoubleVectorParameter.h
@@ -19,6 +19,7 @@ class UMaterialExpressionDoubleVectorParameter : public UMaterialExpressionParam
 
 	//~ Begin UMaterialExpression Interface
 #if WITH_EDITOR
+	virtual void Build(MIR::FEmitter& Emitter) override;
 	virtual int32 Compile(class FMaterialCompiler* Compiler, int32 OutputIndex) override;
 	virtual void GetCaption(TArray<FString>& OutCaptions) const override;
 	virtual bool GetParameterValue(FMaterialParameterMetadata& OutMeta) const override
diff --git a/Engine/Source/Runtime/Engine/Public/Materials/MaterialExpressionGenericConstant.h b/Engine/Source/Runtime/Engine/Public/Materials/MaterialExpressionGenericConstant.h
--- a/Engine/Source/Runtime/Engine/Public/Materials/MaterialExpressionGenericConstant.h
+++ b/Engine/Source/Runtime/Engine/Public/Materials/MaterialExpressionGenericConstant.h
@@ -16,6 +16,7 @@ class UMaterialExpressionGenericConstant : public UMaterialExpression
 
 	//~ Begin UMaterialExpression Interface
 #if WITH_EDITOR
+	virtual void Build(MIR::FEmitter& Emitter) override;
 	virtual int32 Compile(class FMaterialCompiler* Compiler, int32 OutputIndex) override;
 	virtual void GetCaption(TArray<FString>& OutCaptions) const override;
 	virtual FString GetDescription() const override;
diff --git a/Engine/Source/Runtime/Engine/Public/Materials/MaterialIR.h b/Engine/Source/Runtime/Engine/Public/Materials/MaterialIR.h
--- a/Engine/Source/Runtime/Engine/Public/Materials/MaterialIR.h
+++ b/Engine/Source/Runtime/Engine/Public/Materials/MaterialIR.h
@@ -279,6 +279,9 @@ using TInteger = int64_t;
 // The floating point type used inside MIR.
 using TFloat = float;
 
+// The double precision / LWC floating point type used inside MIR.
+using TDouble = double;
+
 // A constant value.
 //
 // A constant represents a translation-time known scalar primitive value. Operations on
@@ -291,6 +294,7 @@ struct FConstant : TValue<VK_Constant>
 		bool  		Boolean;
 		TInteger	Integer;
 		TFloat 		Float;
+		TDouble		LWCValue;
 	};
 
 	// Returns whether this constant is a boolean.
@@ -302,6 +306,12 @@ struct FConstant : TValue<VK_Constant>
 	// Returns whether this constant is a float.
 	bool IsFloat() const;
 
+	// Returns whether this constant is LWC.
+	bool IsLWC() const;
+
+	// Returns whether this is an LWC type representable as regular float.
+	bool IsLWCRepresentableAsFloat() const;
+
 	// Returns the constant value of given type T. The type must be bool, integral or floating point.
 	template <typename T>
 	T Get() const
diff --git a/Engine/Source/Runtime/Engine/Public/Materials/MaterialIREmitter.h b/Engine/Source/Runtime/Engine/Public/Materials/MaterialIREmitter.h
--- a/Engine/Source/Runtime/Engine/Public/Materials/MaterialIREmitter.h
+++ b/Engine/Source/Runtime/Engine/Public/Materials/MaterialIREmitter.h
@@ -319,7 +319,10 @@ public:
 
 	// Emits a constant float 4D column vector and returns it.
 	FValueRef ConstantFloat4(UE::Math::TVector4<TFloat> InValue);
-	
+
+	// Returns the given constant LWC scalar.
+	FValueRef ConstantLWC(TDouble InX);
+
 	/*--------------------- Other non-instruction values emission ---------------------*/
 
 	// Emits the poison value.
diff --git a/Engine/Source/Runtime/Engine/Public/Materials/MaterialIRTypes.h b/Engine/Source/Runtime/Engine/Public/Materials/MaterialIRTypes.h
--- a/Engine/Source/Runtime/Engine/Public/Materials/MaterialIRTypes.h
+++ b/Engine/Source/Runtime/Engine/Public/Materials/MaterialIRTypes.h
@@ -157,9 +157,18 @@ struct FPrimitiveType : FType
 	// Returns the floating point 4D vector type.
 	static const FPrimitiveType* GetFloat4();
 
-	// Returns the LWC 3D vector type
+	// Returns the LWC scalar type.
+	static const FPrimitiveType* GetLWC();
+
+	// Returns the LWC 2D vector type.
+	static const FPrimitiveType* GetLWC2();
+
+	// Returns the LWC 3D vector type.
 	static const FPrimitiveType* GetLWC3();
 
+	// Returns the LWC 4D vector type.
+	static const FPrimitiveType* GetLWC4();
+
 	// Returns the scalar type with given kind.
 	static const FPrimitiveType* GetScalar(EScalarKind InScalarKind);
 
diff --git a/Engine/Source/Runtime/Engine/Public/Materials/MaterialInsights.h b/Engine/Source/Runtime/Engine/Public/Materials/MaterialInsights.h
--- a/Engine/Source/Runtime/Engine/Public/Materials/MaterialInsights.h
+++ b/Engine/Source/Runtime/Engine/Public/Materials/MaterialInsights.h
@@ -35,6 +35,7 @@ struct FMaterialInsights
 		CT_Unused,
 		CT_Int,
 		CT_Float,
+		CT_LWC,
 	};
 
 	/**
