Fixed OnSyncBodies calling MoveComponent outside of a transaction.

diff --git a/Engine/Source/Runtime/Engine/Private/PhysicsEngine/Experimental/PhysScene_Chaos.cpp b/Engine/Source/Runtime/Engine/Private/PhysicsEngine/Experimental/PhysScene_Chaos.cpp
--- a/Engine/Source/Runtime/Engine/Private/PhysicsEngine/Experimental/PhysScene_Chaos.cpp
+++ b/Engine/Source/Runtime/Engine/Private/PhysicsEngine/Experimental/PhysScene_Chaos.cpp
@@ -2531,7 +2531,15 @@ void FPhysScene_Chaos::OnSyncBodies(Chaos::FPhysicsSolverBase* Solver)
 
 			if (ComponentTransform.bHasValidTransform)
 			{
+#if UE_WITH_REMOTE_OBJECT_HANDLE
+				static FName TransactionalWorkName(TEXT("FPhysScene_Chaos::OnSyncBodies"));
+				UE::RemoteExecutor::ExecuteTransactional(TransactionalWorkName, [this, &ComponentTransform]
+				{ 
+					ComponentTransform.OwningComp->MoveComponent(ComponentTransform.NewTranslation, ComponentTransform.NewRotation, false, NULL, MOVECOMP_SkipPhysicsMove);
+				});
+#else
 				ComponentTransform.OwningComp->MoveComponent(ComponentTransform.NewTranslation, ComponentTransform.NewRotation, false, NULL, MOVECOMP_SkipPhysicsMove);
+#endif
 			}
 
 			if (IsValid(OwnerPtr))
