Update navigation around instanced static mesh in game when using dynamic navmesh

#rnx

diff --git a/Engine/Source/Runtime/Engine/Private/InstancedStaticMesh.cpp b/Engine/Source/Runtime/Engine/Private/InstancedStaticMesh.cpp
--- a/Engine/Source/Runtime/Engine/Private/InstancedStaticMesh.cpp
+++ b/Engine/Source/Runtime/Engine/Private/InstancedStaticMesh.cpp
@@ -192,6 +192,11 @@ FAutoConsoleVariableRef CVarISMConservativeBoundsThreshold(
 	GConservativeBoundsThreshold,
 	TEXT("Number of instances in an ISM before we start using conservative bounds. Set to -1 to disable conservative bounds."));
 
+static TAutoConsoleVariable<int32> CVarISMInvalidateCachedNavigationBoundsOnPartialNavigationUpdates(
+	TEXT("r.InstancedStaticMeshes.InvalidateCachedNavigationBoundsOnPartialNavigationUpdates"),
+	1,
+	TEXT("Invalidate cached navigation bounds when doing a partial navigation update"));
+
 class FISMExecHelper : public FSelfRegisteringExec
 {
 	virtual bool Exec_Runtime(class UWorld* InWorld, const TCHAR* Cmd, FOutputDevice& Ar)
@@ -5227,6 +5232,12 @@ void UInstancedStaticMeshComponent::PartialNavigationUpdates(const TConstArrayVi
 	const FBox InstanceBounds = GetInstanceNavigationBounds();
 	if (InstanceBounds.IsValid)
 	{
+		if (CVarISMInvalidateCachedNavigationBoundsOnPartialNavigationUpdates.GetValueOnAnyThread() != 0)
+		{
+			// Make sure to invalidate cached navigation bounds that might have been computed from a different number of instances.
+			CachedBounds[static_cast<int32>(EBoundsCacheType::Navigation)].Invalidate();
+		}
+
 		// Update cached navigation bounds
 		CalcAndCacheNavigationBounds();
 
