[DSTM] Making ProcessLevelStreamingVolumes call transactional if remote object handle support is compiled in.

If we don't automatically migrate the player controller with the pawns, ProcessLevelStreamingVolumes can try to dereference it while it is in another server, and if this call is not transactional, it will result in a crash.

This change was discussed and was part of a previous DSTM change but it got left out by accident from that CL (43371499)

diff --git a/Engine/Source/Runtime/Engine/Private/World.cpp b/Engine/Source/Runtime/Engine/Private/World.cpp
--- a/Engine/Source/Runtime/Engine/Private/World.cpp
+++ b/Engine/Source/Runtime/Engine/Private/World.cpp
@@ -4740,7 +4740,16 @@ void UWorld::InternalUpdateStreamingState()
 
 	// Update streaming levels state using streaming volumes.
 	// Issues level streaming load/unload requests based on local players being inside/outside level streaming volumes.
+
+#if UE_WITH_REMOTE_OBJECT_HANDLE
+	static FName TransactionalWorkName(TEXT("UWorld::InternalUpdateStreamingState"));
+	UE::RemoteExecutor::ExecuteTransactional(TransactionalWorkName,[this]()
+	{
+		ProcessLevelStreamingVolumes();
+	});
+#else
 	ProcessLevelStreamingVolumes();
+#endif
 
 	// Update WorldComposition required streaming levels
 	if (WorldComposition)
