Removed redundant commented-out line.

#rnx

diff --git a/Engine/Source/Runtime/Engine/Private/AssetManager.cpp b/Engine/Source/Runtime/Engine/Private/AssetManager.cpp
--- a/Engine/Source/Runtime/Engine/Private/AssetManager.cpp
+++ b/Engine/Source/Runtime/Engine/Private/AssetManager.cpp
@@ -2110,7 +2110,6 @@ TSharedPtr<FStreamableHandle> UAssetManager::ChangeBundleStateForPrimaryAssets(
 		if (MoviePlayerNumAssets >= 500)
 		{
 			MoviePlayerNumAssets = 0;
-			//UE_LOG(LogTemp, Warning, TEXT("pooo %d"), MoviePlayerNumAssets);
 			FMoviePlayerProxy::BlockingTick();
 		}
 
