PSO Precaching fixes and support for:
- FGeometryCacheVertexVertexFactory and GeometryCacheComponent (including correct vertex element usage extraction - doesn't support manual vertex fetch yet)
- VolumetricCloudRender passes and VolumetricCloudComponent
- CustomDepthRendering taking primitive type into account
- FDeferredLightPS for hair rendering
- SkyPass render target format fixes

diff --git a/Engine/Plugins/Runtime/GeometryCache/Source/GeometryCache/Classes/GeometryCacheComponent.h b/Engine/Plugins/Runtime/GeometryCache/Source/GeometryCache/Classes/GeometryCacheComponent.h
--- a/Engine/Plugins/Runtime/GeometryCache/Source/GeometryCache/Classes/GeometryCacheComponent.h
+++ b/Engine/Plugins/Runtime/GeometryCache/Source/GeometryCache/Classes/GeometryCacheComponent.h
@@ -66,6 +66,7 @@ class UGeometryCacheComponent : public UMeshComponent
 
 	//~ Begin UPrimitiveComponent Interface.
 	UE_API virtual FPrimitiveSceneProxy* CreateSceneProxy() override;
+	UE_API virtual void CollectPSOPrecacheData(const FPSOPrecacheParams& BasePrecachePSOParams, FMaterialInterfacePSOPrecacheParamsList& OutParams) override;
 	//~ End UPrimitiveComponent Interface.
 
 	//~ Begin UMeshComponent Interface.
diff --git a/Engine/Plugins/Runtime/GeometryCache/Source/GeometryCache/Private/GeometryCacheComponent.cpp b/Engine/Plugins/Runtime/GeometryCache/Source/GeometryCache/Private/GeometryCacheComponent.cpp
--- a/Engine/Plugins/Runtime/GeometryCache/Source/GeometryCache/Private/GeometryCacheComponent.cpp
+++ b/Engine/Plugins/Runtime/GeometryCache/Source/GeometryCache/Private/GeometryCacheComponent.cpp
@@ -49,6 +49,9 @@ void UGeometryCacheComponent::FinishDestroy()
 void UGeometryCacheComponent::PostLoad()
 {
 	Super::PostLoad();
+	
+	// Precache PSOs for the used materials
+	PrecachePSOs();
 }
 
 void UGeometryCacheComponent::OnRegister()
@@ -257,6 +260,27 @@ void UGeometryCacheComponent::UpdateLocalBounds()
 	UpdateBounds();
 }
 
+void UGeometryCacheComponent::CollectPSOPrecacheData(const FPSOPrecacheParams& BasePrecachePSOParams, FMaterialInterfacePSOPrecacheParamsList& OutParams)
+{	
+	FGeometryCacheVertexVertexFactory::FDataType Data;
+	FGeomCacheVertexFactory::InitDummyComponentVF(Data);
+
+	FVertexDeclarationElementList Elements;
+	FGeometryCacheVertexVertexFactory::GetVertexElements(GMaxRHIFeatureLevel, EVertexInputStreamType::Default, Data, Elements);
+	
+	for (UMaterialInterface* Material : GetMaterials())
+	{
+		if (Material != nullptr && Material->CheckMaterialUsage_Concurrent(EMaterialUsage::MATUSAGE_GeometryCache))
+		{
+			FMaterialInterfacePSOPrecacheParams& ComponentParams = OutParams[OutParams.AddDefaulted()];
+			ComponentParams.Priority = EPSOPrecachePriority::Medium;
+			ComponentParams.MaterialInterface = Material;
+			ComponentParams.VertexFactoryDataList.Add(FPSOPrecacheVertexFactoryData(&FGeometryCacheVertexVertexFactory::StaticType, Elements));
+			ComponentParams.PSOPrecacheParams = BasePrecachePSOParams;
+		}
+	}
+}
+
 FPrimitiveSceneProxy* UGeometryCacheComponent::CreateSceneProxy()
 {
 	IGeometryCacheStreamingManager::Get().PrefetchData(this);
diff --git a/Engine/Plugins/Runtime/GeometryCache/Source/GeometryCache/Private/GeometryCacheSceneProxy.cpp b/Engine/Plugins/Runtime/GeometryCache/Source/GeometryCache/Private/GeometryCacheSceneProxy.cpp
--- a/Engine/Plugins/Runtime/GeometryCache/Source/GeometryCache/Private/GeometryCacheSceneProxy.cpp
+++ b/Engine/Plugins/Runtime/GeometryCache/Source/GeometryCache/Private/GeometryCacheSceneProxy.cpp
@@ -1432,7 +1432,9 @@ void FGeomCacheTrackProxy::InitRenderResources(FRHICommandListBase& RHICmdList,
 	IndexBuffer.NumValidIndices = 0;
 
 	// Init vertex factory
-	VertexFactory.Init(RHICmdList, &PositionBuffers[0], &PositionBuffers[1], &TangentXBuffer, &TangentZBuffer, &TextureCoordinatesBuffer, &ColorBuffer);
+	FGeometryCacheVertexVertexFactory::FDataType Data;
+	FGeomCacheVertexFactory::InitComponentVF(&PositionBuffers[0], &PositionBuffers[1], &TangentXBuffer, &TangentZBuffer, &TextureCoordinatesBuffer, &ColorBuffer, Data);
+	VertexFactory.SetData(RHICmdList, Data);
 
 	// Enqueue initialization of render resource
 	PositionBuffers[0].InitResource(RHICmdList);
@@ -1497,18 +1499,34 @@ FGeomCacheVertexFactory::FGeomCacheVertexFactory(ERHIFeatureLevel::Type InFeatur
 
 }
 
-void FGeomCacheVertexFactory::Init(FRHICommandListBase& RHICmdList, const FVertexBuffer* PositionBuffer, const FVertexBuffer* MotionBlurDataBuffer, const FVertexBuffer* TangentXBuffer, const FVertexBuffer* TangentZBuffer, const FVertexBuffer* TextureCoordinateBuffer, const FVertexBuffer* ColorBuffer)
+void FGeomCacheVertexFactory::InitDummyComponentVF(FGeometryCacheVertexVertexFactory::FDataType& OutData)
+{
+	// Provide empty buffers for all entries
+	FGeomCacheVertexBuffer PositionBuffers[2];
+	FGeomCacheTangentBuffer TangentXBuffer;
+	FGeomCacheTangentBuffer TangentZBuffer;
+	FGeomCacheVertexBuffer TextureCoordinatesBuffer;
+	FGeomCacheColorBuffer ColorBuffer;
+	InitComponentVF(&PositionBuffers[0], &PositionBuffers[1], &TangentXBuffer, &TangentZBuffer, &TextureCoordinatesBuffer, &ColorBuffer, OutData);
+}
+
+void FGeomCacheVertexFactory::InitComponentVF(const FVertexBuffer* PositionBuffer, const FVertexBuffer* MotionBlurDataBuffer, const FVertexBuffer* TangentXBuffer, const FVertexBuffer* TangentZBuffer, const FVertexBuffer* TextureCoordinateBuffer, const FVertexBuffer* ColorBuffer, FGeometryCacheVertexVertexFactory::FDataType& OutData)
 {
 	// Initialize the vertex factory's stream components.
-	FDataType NewData;
-	NewData.PositionComponent = FVertexStreamComponent(PositionBuffer, 0, sizeof(FVector3f), VET_Float3);
+	OutData.PositionComponent = FVertexStreamComponent(PositionBuffer, 0, sizeof(FVector3f), VET_Float3);
 
-	NewData.TextureCoordinates.Add(FVertexStreamComponent(TextureCoordinateBuffer, 0, sizeof(FVector2f), VET_Float2));
-	NewData.TangentBasisComponents[0] = FVertexStreamComponent(TangentXBuffer, 0, sizeof(FPackedNormal), VET_PackedNormal);
-	NewData.TangentBasisComponents[1] = FVertexStreamComponent(TangentZBuffer, 0, sizeof(FPackedNormal), VET_PackedNormal);
-	NewData.ColorComponent = FVertexStreamComponent(ColorBuffer, 0, sizeof(FColor), VET_Color);
-	NewData.MotionBlurDataComponent = FVertexStreamComponent(MotionBlurDataBuffer, 0, sizeof(FVector3f), VET_Float3);
+	OutData.TextureCoordinates.Add(FVertexStreamComponent(TextureCoordinateBuffer, 0, sizeof(FVector2f), VET_Float2));
+	OutData.TangentBasisComponents[0] = FVertexStreamComponent(TangentXBuffer, 0, sizeof(FPackedNormal), VET_PackedNormal);
+	OutData.TangentBasisComponents[1] = FVertexStreamComponent(TangentZBuffer, 0, sizeof(FPackedNormal), VET_PackedNormal);
+	OutData.ColorComponent = FVertexStreamComponent(ColorBuffer, 0, sizeof(FColor), VET_Color);
+	OutData.MotionBlurDataComponent = FVertexStreamComponent(MotionBlurDataBuffer, 0, sizeof(FVector3f), VET_Float3);
+}
 
+void FGeomCacheVertexFactory::Init(FRHICommandListBase& RHICmdList, const FVertexBuffer* PositionBuffer, const FVertexBuffer* MotionBlurDataBuffer, const FVertexBuffer* TangentXBuffer, const FVertexBuffer* TangentZBuffer, const FVertexBuffer* TextureCoordinateBuffer, const FVertexBuffer* ColorBuffer)
+{
+	// Initialize the vertex factory's stream components.
+	FDataType NewData;
+	InitComponentVF(PositionBuffer, MotionBlurDataBuffer, TangentXBuffer, TangentZBuffer, TextureCoordinateBuffer, ColorBuffer, NewData);
 	SetData(RHICmdList, NewData);
 }
 
diff --git a/Engine/Plugins/Runtime/GeometryCache/Source/GeometryCache/Public/GeometryCacheSceneProxy.h b/Engine/Plugins/Runtime/GeometryCache/Source/GeometryCache/Public/GeometryCacheSceneProxy.h
--- a/Engine/Plugins/Runtime/GeometryCache/Source/GeometryCache/Public/GeometryCacheSceneProxy.h
+++ b/Engine/Plugins/Runtime/GeometryCache/Source/GeometryCache/Public/GeometryCacheSceneProxy.h
@@ -138,7 +138,11 @@ class FGeomCacheVertexFactory : public FGeometryCacheVertexVertexFactory
 public:
 	UE_API FGeomCacheVertexFactory(ERHIFeatureLevel::Type InFeatureLevel);
 
+	UE_DEPRECATED(5.7, "Use InitComponentVF to setup the FDataType and then call SetData on the vertex factory with the command list")
 	UE_API void Init(FRHICommandListBase& RHICmdList, const FVertexBuffer* PositionBuffer, const FVertexBuffer* MotionBlurDataBuffer, const FVertexBuffer* TangentXBuffer, const FVertexBuffer* TangentZBuffer, const FVertexBuffer* TextureCoordinateBuffer, const FVertexBuffer* ColorBuffer);
+
+	UE_API static void InitDummyComponentVF(FGeometryCacheVertexVertexFactory::FDataType& OutData);	
+	UE_API static void InitComponentVF(const FVertexBuffer* PositionBuffer, const FVertexBuffer* MotionBlurDataBuffer, const FVertexBuffer* TangentXBuffer, const FVertexBuffer* TangentZBuffer, const FVertexBuffer* TextureCoordinateBuffer, const FVertexBuffer* ColorBuffer, FGeometryCacheVertexVertexFactory::FDataType& OutData);	
 };
 
 // Hacky base class to avoid 8 bytes of padding after the vtable
diff --git a/Engine/Source/Runtime/Engine/Classes/Components/VolumetricCloudComponent.h b/Engine/Source/Runtime/Engine/Classes/Components/VolumetricCloudComponent.h
--- a/Engine/Source/Runtime/Engine/Classes/Components/VolumetricCloudComponent.h
+++ b/Engine/Source/Runtime/Engine/Classes/Components/VolumetricCloudComponent.h
@@ -210,12 +210,14 @@ protected:
 	virtual void CreateRenderState_Concurrent(FRegisterComponentContext* Context) override;
 	virtual void DestroyRenderState_Concurrent() override;
 	virtual void OnRegister() override;
+	virtual void PrecachePSOs() override;
 	//~ End UActorComponent Interface.
 
 public:
 
 	//~ Begin UObject Interface. 
 	virtual void Serialize(FArchive& Ar) override;
+	virtual void PostLoad() override;
 #if WITH_EDITOR
 	virtual void PostEditChangeProperty(FPropertyChangedEvent& PropertyChangedEvent) override;
 	virtual void OnCookEvent(UE::Cook::ECookEvent CookEvent, UE::Cook::FCookEventContext& CookContext) override;
diff --git a/Engine/Source/Runtime/Engine/Private/Components/VolumetricCloudComponent.cpp b/Engine/Source/Runtime/Engine/Private/Components/VolumetricCloudComponent.cpp
--- a/Engine/Source/Runtime/Engine/Private/Components/VolumetricCloudComponent.cpp
+++ b/Engine/Source/Runtime/Engine/Private/Components/VolumetricCloudComponent.cpp
@@ -12,6 +12,7 @@
 #include "UObject/UE5ReleaseStreamObjectVersion.h"
 #include "SceneInterface.h"
 #include "UObject/UE5MainStreamObjectVersion.h"
+#include "LocalVertexFactory.h"
 
 #if WITH_EDITOR
 #include "MaterialShared.h"
@@ -163,6 +164,12 @@ void UVolumetricCloudComponent::Serialize(FArchive& Ar)
 	}
 }
 
+void UVolumetricCloudComponent::PostLoad()
+{
+	Super::PostLoad();
+	PrecachePSOs();
+}
+
 void UVolumetricCloudComponent::OnRegister()
 {
 	Super::OnRegister();
@@ -172,6 +179,26 @@ void UVolumetricCloudComponent::OnRegister()
 	MaterialLifetimePtr = Material.LoadSynchronous();
 }
 
+void UVolumetricCloudComponent::PrecachePSOs()
+{
+#if UE_WITH_PSO_PRECACHING
+	if (!FApp::CanEverRender() || !IsComponentPSOPrecachingEnabled())
+	{
+		return;
+	}
+
+	if (Material && !Material->HasAnyFlags(RF_NeedPostLoad))
+	{
+		FPSOPrecacheParams PSOPrecacheParams;		
+		FPSOPrecacheVertexFactoryDataList VertexFactoryDataList;		
+		VertexFactoryDataList.Add(FPSOPrecacheVertexFactoryData(&FLocalVertexFactory::StaticType));
+		
+		// Immediately create at high priority and thus doesn't need boosting anymore
+		TArray<FMaterialPSOPrecacheRequestID> MaterialPSOPrecacheRequestIDs;
+		Material->PrecachePSOs(VertexFactoryDataList, PSOPrecacheParams, EPSOPrecachePriority::High, MaterialPSOPrecacheRequestIDs);		
+	}
+#endif
+}
 
 #define CLOUD_DECLARE_BLUEPRINT_SETFUNCTION(MemberType, MemberName) void UVolumetricCloudComponent::Set##MemberName(MemberType NewValue)\
 {\
@@ -203,6 +230,8 @@ void UVolumetricCloudComponent::SetMaterial(UMaterialInterface* NewValue)
 		Material = NewValue;
 		MaterialLifetimePtr = Material.LoadSynchronous();
 		MarkRenderStateDirty();
+
+		PrecachePSOs();
 	}
 }
 
diff --git a/Engine/Source/Runtime/Engine/Private/GeometryCacheVertexFactory.cpp b/Engine/Source/Runtime/Engine/Private/GeometryCacheVertexFactory.cpp
--- a/Engine/Source/Runtime/Engine/Private/GeometryCacheVertexFactory.cpp
+++ b/Engine/Source/Runtime/Engine/Private/GeometryCacheVertexFactory.cpp
@@ -179,93 +179,158 @@ public:
 };
 TGlobalResource<FDummyTangentBuffer> GDummyTangentBuffer;
 
-void FGeometryCacheVertexVertexFactory::InitRHI(FRHICommandListBase& RHICmdList)
+void FGeometryCacheVertexVertexFactory::GetVertexElements(
+	ERHIFeatureLevel::Type FeatureLevel, 
+	EVertexInputStreamType InputStreamType, 
+	FDataType& StreamComponentData, 
+	FVertexDeclarationElementList& OutElements, 
+	FVertexStreamList& InOutStreams, 
+	int32& OutPositionStreamIndex, 
+	int32& OutMotionBlurDataStreamIndex)
 {
-	// Position needs to be separate from the rest (we just theck tangents here)
-	check(Data.PositionComponent.VertexBuffer != Data.TangentBasisComponents[0].VertexBuffer);
-	// Motion Blur data also needs to be separate from the rest
-	check(Data.MotionBlurDataComponent.VertexBuffer != Data.TangentBasisComponents[0].VertexBuffer);
-	check(Data.MotionBlurDataComponent.VertexBuffer != Data.PositionComponent.VertexBuffer);
+	check(InputStreamType == EVertexInputStreamType::Default);
 
-	// If the vertex buffer containing position is not the same vertex buffer containing the rest of the data,
-	// then initialize PositionStream and PositionDeclaration.
-	if (Data.PositionComponent.VertexBuffer != Data.TangentBasisComponents[0].VertexBuffer)
+	if (StreamComponentData.PositionComponent.VertexBuffer != NULL)
 	{
-		{
-			FVertexDeclarationElementList PositionOnlyStreamElements;
-			PositionOnlyStreamElements.Add(AccessStreamComponent(Data.PositionComponent, 0, EVertexInputStreamType::PositionOnly));
-			AddPrimitiveIdStreamElement(EVertexInputStreamType::PositionOnly, PositionOnlyStreamElements, 1, 1);
-			InitDeclaration(PositionOnlyStreamElements, EVertexInputStreamType::PositionOnly);
-		}
-
-		{
-			FVertexDeclarationElementList PositionAndNormalOnlyStreamElements;
-			PositionAndNormalOnlyStreamElements.Add(AccessStreamComponent(Data.PositionComponent, 0, EVertexInputStreamType::PositionAndNormalOnly));
-			PositionAndNormalOnlyStreamElements.Add(AccessStreamComponent(Data.TangentBasisComponents[1], 1, EVertexInputStreamType::PositionAndNormalOnly));
-			AddPrimitiveIdStreamElement(EVertexInputStreamType::PositionAndNormalOnly, PositionAndNormalOnlyStreamElements, 2, 2);
-			InitDeclaration(PositionAndNormalOnlyStreamElements, EVertexInputStreamType::PositionAndNormalOnly);
-		}
-	}
-
-	FVertexDeclarationElementList Elements;
-	if (Data.PositionComponent.VertexBuffer != NULL)
-	{
-		Elements.Add(AccessStreamComponent(Data.PositionComponent, 0));
-		PositionStreamIndex = Elements.Last().StreamIndex;
+		OutElements.Add(AccessStreamComponent(StreamComponentData.PositionComponent, 0, InOutStreams));
+		OutPositionStreamIndex = OutElements.Last().StreamIndex;
 	}
 
 	// only tangent,normal are used by the stream. the binormal is derived in the shader
 	uint8 TangentBasisAttributes[2] = { 1, 2 };
 	for (int32 AxisIndex = 0; AxisIndex < 2; AxisIndex++)
 	{
-		if (Data.TangentBasisComponents[AxisIndex].VertexBuffer != NULL)
+		if (StreamComponentData.TangentBasisComponents[AxisIndex].VertexBuffer != NULL)
 		{
-			Elements.Add(AccessStreamComponent(Data.TangentBasisComponents[AxisIndex], TangentBasisAttributes[AxisIndex]));
+			OutElements.Add(AccessStreamComponent(StreamComponentData.TangentBasisComponents[AxisIndex], TangentBasisAttributes[AxisIndex], InOutStreams));
 		}
 	}
 
-	if (Data.ColorComponent.VertexBuffer)
+	if (StreamComponentData.ColorComponent.VertexBuffer)
 	{
-		Elements.Add(AccessStreamComponent(Data.ColorComponent, 3));
+		OutElements.Add(AccessStreamComponent(StreamComponentData.ColorComponent, 3, InOutStreams));
 	}
 	else
 	{
 		//If the mesh has no color component, set the null color buffer on a new stream with a stride of 0.
 		//This wastes 4 bytes of bandwidth per vertex, but prevents having to compile out twice the number of vertex factories.
 		FVertexStreamComponent NullColorComponent(&GNullColorVertexBuffer, 0, 0, VET_Color);
-		Elements.Add(AccessStreamComponent(NullColorComponent, 3));
+		OutElements.Add(AccessStreamComponent(NullColorComponent, 3, InOutStreams));
 	}
 
-	if (Data.MotionBlurDataComponent.VertexBuffer)
+	if (StreamComponentData.MotionBlurDataComponent.VertexBuffer)
 	{
-		Elements.Add(AccessStreamComponent(Data.MotionBlurDataComponent, 4));
+		OutElements.Add(AccessStreamComponent(StreamComponentData.MotionBlurDataComponent, 4, InOutStreams));
 	}
-	else if (Data.PositionComponent.VertexBuffer != NULL)
+	else if (StreamComponentData.PositionComponent.VertexBuffer != NULL)
 	{
-		Elements.Add(AccessStreamComponent(Data.PositionComponent, 4));
+		OutElements.Add(AccessStreamComponent(StreamComponentData.PositionComponent, 4, InOutStreams));
 	}
-	MotionBlurDataStreamIndex = Elements.Last().StreamIndex;
+	OutMotionBlurDataStreamIndex = OutElements.Last().StreamIndex;
 
-	if (Data.TextureCoordinates.Num())
+	if (StreamComponentData.TextureCoordinates.Num())
 	{
 		const int32 BaseTexCoordAttribute = 5;
-		for (int32 CoordinateIndex = 0; CoordinateIndex < Data.TextureCoordinates.Num(); CoordinateIndex++)
+		for (int32 CoordinateIndex = 0; CoordinateIndex < StreamComponentData.TextureCoordinates.Num(); CoordinateIndex++)
 		{
-			Elements.Add(AccessStreamComponent(
-				Data.TextureCoordinates[CoordinateIndex],
-				BaseTexCoordAttribute + CoordinateIndex
+			OutElements.Add(AccessStreamComponent(
+				StreamComponentData.TextureCoordinates[CoordinateIndex],
+				BaseTexCoordAttribute + CoordinateIndex, 
+				InOutStreams
 			));
 		}
 
-		for (int32 CoordinateIndex = Data.TextureCoordinates.Num(); CoordinateIndex < MAX_STATIC_TEXCOORDS / 2; CoordinateIndex++)
+		for (int32 CoordinateIndex = StreamComponentData.TextureCoordinates.Num(); CoordinateIndex < MAX_STATIC_TEXCOORDS / 2; CoordinateIndex++)
 		{
-			Elements.Add(AccessStreamComponent(
-				Data.TextureCoordinates[Data.TextureCoordinates.Num() - 1],
-				BaseTexCoordAttribute + CoordinateIndex
+			OutElements.Add(AccessStreamComponent(
+				StreamComponentData.TextureCoordinates[StreamComponentData.TextureCoordinates.Num() - 1],
+				BaseTexCoordAttribute + CoordinateIndex,
+				InOutStreams
 			));
 		}
 	}
+}
+
+void FGeometryCacheVertexVertexFactory::GetPSOPrecacheVertexFetchElements(EVertexInputStreamType VertexInputStreamType, FVertexDeclarationElementList& OutElements)
+{
+	OutElements.Add(FVertexElement(0, 0, VET_Float3, 0, sizeof(float)*3u, false));
+	
+	if (VertexInputStreamType == EVertexInputStreamType::PositionAndNormalOnly)
+	{
+		// 2-axis TangentBasis components in a single buffer, hence *2u
+		OutElements.Add(FVertexElement(1, 0, VET_PackedNormal, 1, sizeof(FPackedNormal)*2u, false));
+	}
 
+	if (UseGPUScene(GMaxRHIShaderPlatform, GMaxRHIFeatureLevel) 
+		&& !PlatformGPUSceneUsesUniformBufferView(GMaxRHIShaderPlatform))
+	{
+		switch (VertexInputStreamType)
+		{
+			case EVertexInputStreamType::Default:
+			{
+				OutElements.Add(FVertexElement(1, 0, VET_UInt, 13, sizeof(uint32), true));
+				break;
+			}
+			case EVertexInputStreamType::PositionOnly:
+			{
+				OutElements.Add(FVertexElement(1, 0, VET_UInt, 1, sizeof(uint32), true));
+				break;
+			}
+			case EVertexInputStreamType::PositionAndNormalOnly:
+			{
+				OutElements.Add(FVertexElement(2, 0, VET_UInt, 2, sizeof(uint32), true));
+				break;
+			}
+			default:
+				checkNoEntry();
+		}
+	}
+}
+
+void FGeometryCacheVertexVertexFactory::GetVertexElements(ERHIFeatureLevel::Type FeatureLevel, EVertexInputStreamType InputStreamType, FDataType& StreamComponentData, FVertexDeclarationElementList& OutElements)
+{
+	FVertexStreamList VertexStreams;
+	int32 PositionStreamIndex = -1;
+	int32 MotionBlurDataStreamIndex = -1;
+	GetVertexElements(FeatureLevel, InputStreamType, StreamComponentData, OutElements, VertexStreams, PositionStreamIndex, MotionBlurDataStreamIndex);
+
+	if (UseGPUScene(GMaxRHIShaderPlatform, GMaxRHIFeatureLevel) 
+		&& !PlatformGPUSceneUsesUniformBufferView(GMaxRHIShaderPlatform))
+	{
+		OutElements.Add(FVertexElement(VertexStreams.Num(), 0, VET_UInt, 13, sizeof(uint32), true));
+	}
+}
+
+void FGeometryCacheVertexVertexFactory::InitRHI(FRHICommandListBase& RHICmdList)
+{
+	// Position needs to be separate from the rest (we just theck tangents here)
+	check(Data.PositionComponent.VertexBuffer != Data.TangentBasisComponents[0].VertexBuffer);
+	// Motion Blur data also needs to be separate from the rest
+	check(Data.MotionBlurDataComponent.VertexBuffer != Data.TangentBasisComponents[0].VertexBuffer);
+	check(Data.MotionBlurDataComponent.VertexBuffer != Data.PositionComponent.VertexBuffer);	
+
+	// If the vertex buffer containing position is not the same vertex buffer containing the rest of the data,
+	// then initialize PositionStream and PositionDeclaration.
+	if (Data.PositionComponent.VertexBuffer != Data.TangentBasisComponents[0].VertexBuffer)
+	{
+		{
+			FVertexDeclarationElementList PositionOnlyStreamElements;
+			PositionOnlyStreamElements.Add(AccessStreamComponent(Data.PositionComponent, 0, EVertexInputStreamType::PositionOnly));
+			AddPrimitiveIdStreamElement(EVertexInputStreamType::PositionOnly, PositionOnlyStreamElements, 1, 1);
+			InitDeclaration(PositionOnlyStreamElements, EVertexInputStreamType::PositionOnly);
+		}
+
+		{
+			FVertexDeclarationElementList PositionAndNormalOnlyStreamElements;
+			PositionAndNormalOnlyStreamElements.Add(AccessStreamComponent(Data.PositionComponent, 0, EVertexInputStreamType::PositionAndNormalOnly));
+			PositionAndNormalOnlyStreamElements.Add(AccessStreamComponent(Data.TangentBasisComponents[1], 1, EVertexInputStreamType::PositionAndNormalOnly));
+			AddPrimitiveIdStreamElement(EVertexInputStreamType::PositionAndNormalOnly, PositionAndNormalOnlyStreamElements, 2, 2);
+			InitDeclaration(PositionAndNormalOnlyStreamElements, EVertexInputStreamType::PositionAndNormalOnly);
+		}
+	}
+
+	FVertexDeclarationElementList Elements;
+	GetVertexElements(GetFeatureLevel(), EVertexInputStreamType::Default, Data, Elements, Streams, PositionStreamIndex, MotionBlurDataStreamIndex);
 	AddPrimitiveIdStreamElement(EVertexInputStreamType::Default, Elements, 13, 13);
 
 	check(Streams.Num() > 0);
@@ -407,4 +472,5 @@ IMPLEMENT_VERTEX_FACTORY_TYPE(FGeometryCacheVertexVertexFactory, "/Engine/Privat
 	| EVertexFactoryFlags::SupportsPositionOnly
 	| EVertexFactoryFlags::SupportsRayTracing
 	| EVertexFactoryFlags::SupportsPrimitiveIdStream
+	| EVertexFactoryFlags::SupportsPSOPrecaching
 );
diff --git a/Engine/Source/Runtime/Engine/Public/GeometryCacheVertexFactory.h b/Engine/Source/Runtime/Engine/Public/GeometryCacheVertexFactory.h
--- a/Engine/Source/Runtime/Engine/Public/GeometryCacheVertexFactory.h
+++ b/Engine/Source/Runtime/Engine/Public/GeometryCacheVertexFactory.h
@@ -124,9 +124,15 @@ public:
 
 	ENGINE_API virtual void InitRHI(FRHICommandListBase& RHICmdList) override;
 
+	static ENGINE_API void GetPSOPrecacheVertexFetchElements(EVertexInputStreamType VertexInputStreamType, FVertexDeclarationElementList& OutElements);
+	static ENGINE_API void GetVertexElements(ERHIFeatureLevel::Type FeatureLevel, EVertexInputStreamType InputStreamType, FDataType& StreamComponentData, FVertexDeclarationElementList& OutElements);
+
 	friend FGeometryCacheVertexFactoryShaderParameters;
 	
 protected:
+
+	static void GetVertexElements(ERHIFeatureLevel::Type FeatureLevel, EVertexInputStreamType InputStreamType, FDataType& StreamComponentData, FVertexDeclarationElementList& OutElements, FVertexStreamList& InOutStreams, int32& OutPositionStreamIndex, int32& OutMotionBlurDataStreamIndex);
+
 	// Vertex buffer required for creating the Vertex Declaration
 	FVertexBuffer VBAlias;
 
diff --git a/Engine/Source/Runtime/Renderer/Private/CustomDepthRendering.cpp b/Engine/Source/Runtime/Renderer/Private/CustomDepthRendering.cpp
--- a/Engine/Source/Runtime/Renderer/Private/CustomDepthRendering.cpp
+++ b/Engine/Source/Runtime/Renderer/Private/CustomDepthRendering.cpp
@@ -458,6 +458,7 @@ private:
 		const FMaterial& RESTRICT MaterialResource,
 		ERasterizerFillMode MeshFillMode,
 		ERasterizerCullMode MeshCullMode, 
+		EPrimitiveType PrimitiveType, 
 		TArray<FPSOPrecacheData>& PSOInitializers);
 
 	FMeshPassProcessorRenderState PassDrawRenderState;
@@ -686,7 +687,7 @@ void FCustomDepthPassMeshProcessor::CollectPSOInitializers(const FSceneTexturesC
 	if (!bIgnoreThisMaterial)
 	{
 		const FMaterial* EffectiveMaterial = &Material;
-		if (bUseDefaultMaterial && !bSupportPositionOnlyStream && VertexFactoryData.CustomDefaultVertexDeclaration)
+		if (bUseDefaultMaterial && ((!bSupportPositionOnlyStream && VertexFactoryData.CustomDefaultVertexDeclaration) || PreCacheParams.PrimitiveType != PT_TriangleList))
 		{
 			EMaterialQualityLevel::Type ActiveQualityLevel = GetCachedScalabilityCVars().MaterialQualityLevel;
 			EffectiveMaterial = UMaterial::GetDefaultMaterial(MD_Surface)->GetMaterialResource(FeatureLevel, ActiveQualityLevel);
@@ -709,7 +710,7 @@ void FCustomDepthPassMeshProcessor::CollectPSOInitializers(const FSceneTexturesC
 			const ERasterizerFillMode MeshFillMode = ComputeMeshFillMode(Material, OverrideSettings);
 			const ERasterizerCullMode MeshCullMode = ComputeMeshCullMode(Material, OverrideSettings);
 
-			CollectPSOInitializersInternal<false>(VertexFactoryData, *EffectiveMaterial, MeshFillMode, MeshCullMode, PSOInitializers);
+			CollectPSOInitializersInternal<false>(VertexFactoryData, *EffectiveMaterial, MeshFillMode, MeshCullMode, (EPrimitiveType)PreCacheParams.PrimitiveType, PSOInitializers);
 		}
 	}
 }
@@ -727,18 +728,18 @@ void FCustomDepthPassMeshProcessor::CollectDefaultMaterialPSOInitializers(
 	// Collect PSOs for all possible default material combinations
 	{
 		ERasterizerCullMode MeshCullMode = CM_None;
-		CollectPSOInitializersInternal<true>(VertexFactoryData, Material, MeshFillMode, MeshCullMode, PSOInitializers);
-		CollectPSOInitializersInternal<false>(VertexFactoryData, Material, MeshFillMode, MeshCullMode, PSOInitializers);
+		CollectPSOInitializersInternal<true>(VertexFactoryData, Material, MeshFillMode, MeshCullMode, PT_TriangleList, PSOInitializers);
+		CollectPSOInitializersInternal<false>(VertexFactoryData, Material, MeshFillMode, MeshCullMode, PT_TriangleList, PSOInitializers);
 	}
 	{
 		ERasterizerCullMode MeshCullMode = CM_CW;
-		CollectPSOInitializersInternal<true>(VertexFactoryData, Material, MeshFillMode, MeshCullMode, PSOInitializers);
-		CollectPSOInitializersInternal<false>(VertexFactoryData, Material, MeshFillMode, MeshCullMode, PSOInitializers);
+		CollectPSOInitializersInternal<true>(VertexFactoryData, Material, MeshFillMode, MeshCullMode, PT_TriangleList, PSOInitializers);
+		CollectPSOInitializersInternal<false>(VertexFactoryData, Material, MeshFillMode, MeshCullMode, PT_TriangleList, PSOInitializers);
 	}
 	{
 		ERasterizerCullMode MeshCullMode = CM_CCW;
-		CollectPSOInitializersInternal<true>(VertexFactoryData, Material, MeshFillMode, MeshCullMode, PSOInitializers);
-		CollectPSOInitializersInternal<false>(VertexFactoryData, Material, MeshFillMode, MeshCullMode, PSOInitializers);
+		CollectPSOInitializersInternal<true>(VertexFactoryData, Material, MeshFillMode, MeshCullMode, PT_TriangleList, PSOInitializers);
+		CollectPSOInitializersInternal<false>(VertexFactoryData, Material, MeshFillMode, MeshCullMode, PT_TriangleList, PSOInitializers);
 	}
 }
 
@@ -748,6 +749,7 @@ void FCustomDepthPassMeshProcessor::CollectPSOInitializersInternal(
 	const FMaterial& RESTRICT MaterialResource,
 	ERasterizerFillMode MeshFillMode,
 	ERasterizerCullMode MeshCullMode, 
+	EPrimitiveType PrimitiveType, 
 	TArray<FPSOPrecacheData>& PSOInitializers)
 {
 	TMeshProcessorShaders<
@@ -783,7 +785,7 @@ void FCustomDepthPassMeshProcessor::CollectPSOInitializersInternal(
 		DepthPassShaders,
 		MeshFillMode,
 		MeshCullMode,
-		PT_TriangleList,
+		PrimitiveType,
 		bPositionOnly ? EMeshPassFeatures::PositionOnly : EMeshPassFeatures::Default,
 		true /*bRequired*/,
 		PSOInitializers);
diff --git a/Engine/Source/Runtime/Renderer/Private/LightRendering.cpp b/Engine/Source/Runtime/Renderer/Private/LightRendering.cpp
--- a/Engine/Source/Runtime/Renderer/Private/LightRendering.cpp
+++ b/Engine/Source/Runtime/Renderer/Private/LightRendering.cpp
@@ -3537,9 +3537,11 @@ void DeferredLightGlobalPSOCollector(const FSceneTexturesConfig& SceneTexturesCo
 		if (bIsHairShader)
 		{
 			SetupLightForHairGraphicsPSOState(GraphicsPSOInit);
+			AddRenderTargetInfo(PF_FloatR11G11B10, ETextureCreateFlags::RenderTargetable | ETextureCreateFlags::UAV | ETextureCreateFlags::ShaderResource, RenderTargetsInfo);
 
-			// TODO: find out the render target info for hair first
-			return;
+			TShaderMapRef<FDeferredLightHairVS> VertexShader(GlobalShaderMap);
+			GraphicsPSOInit.BoundShaderState.VertexShaderRHI = VertexShader.GetVertexShader();
+			GraphicsPSOInit.BoundShaderState.PixelShaderRHI = PixelShaderRHI;
 		}
 		else
 		{
diff --git a/Engine/Source/Runtime/Renderer/Private/SkyPassRendering.cpp b/Engine/Source/Runtime/Renderer/Private/SkyPassRendering.cpp
--- a/Engine/Source/Runtime/Renderer/Private/SkyPassRendering.cpp
+++ b/Engine/Source/Runtime/Renderer/Private/SkyPassRendering.cpp
@@ -247,6 +247,14 @@ void FSkyPassMeshProcessor::CollectPSOInitializersInternal(T& SkyPassShaders, co
 	FGraphicsPipelineRenderTargetsInfo RenderTargetsInfo;
 	SetupGBufferRenderTargetInfo(SceneTexturesConfig, RenderTargetsInfo, true /*bSetupDepthStencil*/);
 
+	// Remove all but the SceneColor		
+	RenderTargetsInfo.RenderTargetsEnabled = 1;
+	for (uint32 i = 1; i < MaxSimultaneousRenderTargets; ++i)
+	{
+		RenderTargetsInfo.RenderTargetFormats[i] = 0;
+		RenderTargetsInfo.RenderTargetFlags[i] = ETextureCreateFlags::None;
+	}
+
 	FBasePassMeshProcessor::AddBasePassGraphicsPipelineStateInitializer(
 		FeatureLevel,
 		VertexFactoryData,
diff --git a/Engine/Source/Runtime/Renderer/Private/VolumetricCloudRendering.cpp b/Engine/Source/Runtime/Renderer/Private/VolumetricCloudRendering.cpp
--- a/Engine/Source/Runtime/Renderer/Private/VolumetricCloudRendering.cpp
+++ b/Engine/Source/Runtime/Renderer/Private/VolumetricCloudRendering.cpp
@@ -475,6 +475,14 @@ static int32 GetVolumetricCloudSkyAOResolution(const FSkyLightSceneProxy* SkyLig
 	return 32;
 }
 
+FRDGTextureDesc GetVolumetricCloudSkyAORDGTextureDesc(const FIntPoint& Size)
+{
+	return FRDGTextureDesc::Create2D(
+		Size,
+		PF_FloatR11G11B10,
+		FClearValueBinding::None, TexCreate_ShaderResource | TexCreate_RenderTargetable);
+}
+
 static float GetVolumetricCloudSkyAOExtentKm(const FSkyLightSceneProxy* SkyLight)
 {
 	if (SkyLight)
@@ -1388,13 +1396,13 @@ private:
 
 IMPLEMENT_MATERIAL_SHADER_TYPE(, FVolumetricCloudShadowPS, TEXT("/Engine/Private/VolumetricCloud.usf"), TEXT("MainPS"), SF_Pixel);
 
-
+static const TCHAR* VolumetricCloudRenderShadowPassName = TEXT("VolumetricCloudRenderShadow");
 
 class FVolumetricCloudRenderShadowMeshProcessor : public FMeshPassProcessor
 {
 public:
 	FVolumetricCloudRenderShadowMeshProcessor(const FScene* Scene, ERHIFeatureLevel::Type FeatureLevel, const FViewInfo* InViewIfDynamicMeshCommand, FMeshPassDrawListContext* InDrawListContext)
-		: FMeshPassProcessor(EMeshPass::Num, Scene, FeatureLevel, InViewIfDynamicMeshCommand, InDrawListContext)
+		: FMeshPassProcessor(VolumetricCloudRenderShadowPassName, Scene, FeatureLevel, InViewIfDynamicMeshCommand, InDrawListContext)
 	{
 		PassDrawRenderState.SetBlendState(TStaticBlendState<>::GetRHI());
 		PassDrawRenderState.SetDepthStencilState(TStaticDepthStencilState<false, CF_Always>::GetRHI());
@@ -1480,10 +1488,61 @@ private:
 
 		return true;
 	}
+	
+	virtual void CollectPSOInitializers(const FSceneTexturesConfig& SceneTexturesConfig, const FMaterial& Material, const FPSOPrecacheVertexFactoryData& VertexFactoryData, const FPSOPrecacheParams& PreCacheParams, TArray<FPSOPrecacheData>& PSOInitializers) override
+	{
+		if (Material.GetMaterialDomain() != MD_Volume)
+		{
+			return;
+		}
+
+		const FMeshDrawingPolicyOverrideSettings OverrideSettings = ComputeMeshOverrideSettings(PreCacheParams);
+		ERasterizerFillMode MeshFillMode = FM_Solid;
+		ERasterizerCullMode MeshCullMode = CM_None;
+
+		TMeshProcessorShaders<
+			FRenderVolumetricCloudVS,
+		FVolumetricCloudShadowPS> PassShaders;
+		FMaterialShaderTypes ShaderTypes;
+		ShaderTypes.AddShaderType<FRenderVolumetricCloudVS>();
+		ShaderTypes.AddShaderType<FVolumetricCloudShadowPS>();
+				
+		FMaterialShaders Shaders;
+		if (!Material.TryGetShaders(ShaderTypes, VertexFactoryData.VertexFactoryType, Shaders))
+		{
+			return;
+		}
+		Shaders.TryGetVertexShader(PassShaders.VertexShader);
+		Shaders.TryGetPixelShader(PassShaders.PixelShader);
+
+		FRDGTextureDesc CloudSkyOARDGTextureDesc = GetVolumetricCloudSkyAORDGTextureDesc(FIntPoint(0));
+
+		FGraphicsPipelineRenderTargetsInfo RenderTargetsInfo;
+		RenderTargetsInfo.NumSamples = 1;
+		AddRenderTargetInfo(CloudSkyOARDGTextureDesc.Format, CloudSkyOARDGTextureDesc.Flags, RenderTargetsInfo);
+		
+		AddGraphicsPipelineStateInitializer(
+			VertexFactoryData,
+			Material,
+			PassDrawRenderState,
+			RenderTargetsInfo,
+			PassShaders,
+			MeshFillMode,
+			MeshCullMode,
+			(EPrimitiveType)PreCacheParams.PrimitiveType,
+			EMeshPassFeatures::Default,
+			true /*bRequired*/,
+			PSOInitializers);
+	}
 
 	FMeshPassProcessorRenderState PassDrawRenderState;
 };
 
+IPSOCollector* CreatePSOCollectorVolumetricCloudRenderShadow(ERHIFeatureLevel::Type FeatureLevel)
+{
+	return new FVolumetricCloudRenderShadowMeshProcessor(nullptr, FeatureLevel, nullptr, nullptr);
+}
+FRegisterPSOCollectorCreateFunction RegisterPSOCollectorVolumetricCloudRenderShadow(&CreatePSOCollectorVolumetricCloudRenderShadow, EShadingPath::Deferred, VolumetricCloudRenderShadowPassName);
 
 
 //////////////////////////////////////////////////////////////////////////
@@ -1988,9 +2047,8 @@ void FSceneRenderer::InitVolumetricCloudsForViews(FRDGBuilder& GraphBuilder, boo
 					if (ShouldRenderCloudSkyAO(SkyLight))
 					{
 						const uint32 VolumetricCloudSkyAOResolution = GetVolumetricCloudSkyAOResolution(SkyLight);
-						FRDGTextureRef CloudSkyAOTexture = GraphBuilder.CreateTexture(
-							FRDGTextureDesc::Create2D(FIntPoint(VolumetricCloudSkyAOResolution, VolumetricCloudSkyAOResolution), PF_FloatR11G11B10,
-								FClearValueBinding::None, TexCreate_ShaderResource | TexCreate_RenderTargetable), TEXT("Cloud.SkyAOTexture"));
+						const FRDGTextureDesc CloudSkyAOTextureDesc = GetVolumetricCloudSkyAORDGTextureDesc(FIntPoint(VolumetricCloudSkyAOResolution,VolumetricCloudSkyAOResolution));
+						FRDGTextureRef CloudSkyAOTexture = GraphBuilder.CreateTexture(CloudSkyAOTextureDesc, TEXT("Cloud.SkyAOTexture"));
 
 						// We need to make a copy of the parameters on CPU to morph them because the creation is deferred.
 						FRenderVolumetricCloudGlobalParameters& VolumetricCloudParamsAO = *GraphBuilder.AllocParameters<FRenderVolumetricCloudGlobalParameters>();
