Materials: Fix indeterminism in Material packages. Materials include the LightingGuid from textures via CookedShaderMapIdHash -> ShaderMapId -> TextureReferencesHash -> GetForceRecompileTextureIdsHash -> ForceRecompileTexture->GetLightingGuid().ToString(). UCurveLinearColorAtlas spuriously mark themselves modified themselves on every load (to handle the case where their calculation code changes) and create a new LightingGuid. Create their LightingGuid deterministically to prevent them from setting indeterministic data in their lightingguid and in the Materials that use them.
#rnx

diff --git a/Engine/Source/Runtime/Core/Private/Misc/Guid.cpp b/Engine/Source/Runtime/Core/Private/Misc/Guid.cpp
--- a/Engine/Source/Runtime/Core/Private/Misc/Guid.cpp
+++ b/Engine/Source/Runtime/Core/Private/Misc/Guid.cpp
@@ -7,8 +7,10 @@
 #include "Containers/StaticArray.h"
 #include "HAL/PlatformMisc.h"
 #include "HAL/PlatformString.h"
+#include "HAL/UnrealMemory.h"
 #include "Hash/Blake3.h"
 #include "Math/Int128.h"
+#include "Math/UnrealMathUtility.h"
 #include "Misc/Base64.h"
 #include "Misc/ByteSwap.h"
 #include "Misc/CString.h"
@@ -17,6 +19,7 @@
 #include "Misc/StringBuilder.h"
 #include "Serialization/StructuredArchiveNameHelpers.h"
 #include "Serialization/StructuredArchiveSlots.h"
+#include "Templates/AlignmentTemplates.h"
 #include "UObject/PropertyPortFlags.h"
 
 
@@ -260,11 +263,28 @@ FGuid FGuid::NewGuidFromHash(const FBlake3Hash& Hash)
 {
 	// We use the first 16 bytes of the BLAKE3 hash to create the guid, there is no specific reason why these were
 	// chosen, we could take any pattern or combination of bytes.
-	uint32* HashBytes = (uint32*)Hash.GetBytes();
-	uint32 A = uint32(HashBytes[0]);
-	uint32 B = uint32(HashBytes[1]);
-	uint32 C = uint32(HashBytes[2]);
-	uint32 D = uint32(HashBytes[3]);
+	return NewGuidFromHashBytes(&Hash.GetBytes(), sizeof(Hash.GetBytes()));
+}
+
+FGuid FGuid::NewGuidFromHashBytes(const void* HashData, int64 DataLen)
+{
+	alignas(uint32) uint8 Buffer[16];
+	if (DataLen < 16 || !IsAligned(HashData, alignof(uint32)))
+	{
+		FMemory::Memcpy(Buffer, HashData, FMath::Min(DataLen, 16));
+		if (DataLen < 16)
+		{
+			FMemory::Memzero(&Buffer[DataLen], 16 - DataLen);
+		}
+		HashData = Buffer;
+		DataLen = 16;
+	}
+
+	const uint32* HashInts = reinterpret_cast<const uint32*>(HashData);
+	uint32 A = uint32(HashInts[0]);
+	uint32 B = uint32(HashInts[1]);
+	uint32 C = uint32(HashInts[2]);
+	uint32 D = uint32(HashInts[3]);
 
 	// Convert to a Variant 1 Version 3 UUID, which is meant to be created by hashing the namespace and name with MD5.
 	// We use that version even though we're using BLAKE3 instead of MD5 because we don't have a better alternative.
diff --git a/Engine/Source/Runtime/Core/Public/Misc/Guid.h b/Engine/Source/Runtime/Core/Public/Misc/Guid.h
--- a/Engine/Source/Runtime/Core/Public/Misc/Guid.h
+++ b/Engine/Source/Runtime/Core/Public/Misc/Guid.h
@@ -378,6 +378,14 @@ public:
 	 */
 	static CORE_API FGuid NewGuidFromHash(const FBlake3Hash& Hash);
 
+	/**
+	 * Uses the first 16 bytes from HashData to create a Guid. The first 16 bytes of HashData should be premixed to
+	 * include all cryptographically necessary bytes from the hash to describe larger dataset, so that changing
+	 * any byte of the dataset will change the first 16 bytes of HashData. If HashData is less than 16 bytes, it is
+	 * padded with 0 to create the guid.
+	 */
+	static CORE_API FGuid NewGuidFromHashBytes(const void* HashData, int64 DataLen);
+
 	/**
 	 * Returns a GUID which is a combinationof the two provided ones.
 	 *
diff --git a/Engine/Source/Runtime/Engine/Classes/Curves/CurveLinearColorAtlas.h b/Engine/Source/Runtime/Engine/Classes/Curves/CurveLinearColorAtlas.h
--- a/Engine/Source/Runtime/Engine/Classes/Curves/CurveLinearColorAtlas.h
+++ b/Engine/Source/Runtime/Engine/Classes/Curves/CurveLinearColorAtlas.h
@@ -83,6 +83,13 @@ class UCurveLinearColorAtlas : public UTexture2D
 
 	// Re-render all texture groups
 	ENGINE_API void UpdateTextures();
+
+#endif
+
+#if WITH_EDITORONLY_DATA
+
+	ENGINE_API virtual FGuid BuildLightingGuid() override;
+
 #endif
 
 	ENGINE_API virtual void PostLoad() override;
diff --git a/Engine/Source/Runtime/Engine/Classes/Engine/Texture.h b/Engine/Source/Runtime/Engine/Classes/Engine/Texture.h
--- a/Engine/Source/Runtime/Engine/Classes/Engine/Texture.h
+++ b/Engine/Source/Runtime/Engine/Classes/Engine/Texture.h
@@ -2116,7 +2116,7 @@ public:
 	void SetLightingGuid()
 	{
 #if WITH_EDITORONLY_DATA
-		LightingGuid = FGuid::NewGuid();
+		LightingGuid = BuildLightingGuid();
 #else
 		LightingGuid = FGuid(0, 0, 0, 0);
 #endif // WITH_EDITORONLY_DATA
@@ -2132,6 +2132,16 @@ public:
 	  */
 	ENGINE_API void SetDeterministicLightingGuid();
 
+#if WITH_EDITORONLY_DATA
+	/**
+	 * Creates a Guid for the texture based on the current texture values. Default version returns a new unique Guid.
+	 * Some subclasses override this to calculate it from a hash.
+	 */
+	ENGINE_API virtual FGuid BuildLightingGuid();
+	/** Helper for subclasses that want to implement BuildLightingGuid from a deterministic hash of the texture. */
+	ENGINE_API FGuid BuildLightingGuidFromHash();
+#endif
+
 	/**
 	 * Retrieves the pixel format enum for enum <-> string conversions.
 	 */
diff --git a/Engine/Source/Runtime/Engine/Private/Curves/CurveLinearColorAtlas.cpp b/Engine/Source/Runtime/Engine/Private/Curves/CurveLinearColorAtlas.cpp
--- a/Engine/Source/Runtime/Engine/Private/Curves/CurveLinearColorAtlas.cpp
+++ b/Engine/Source/Runtime/Engine/Private/Curves/CurveLinearColorAtlas.cpp
@@ -342,8 +342,14 @@ void UCurveLinearColorAtlas::UpdateTextures()
 	Source.UnlockMip(0);
 	PostEditChange();
 }
+#endif // WITH_EDITOR
 
-#endif
+#if WITH_EDITORONLY_DATA
+FGuid UCurveLinearColorAtlas::BuildLightingGuid()
+{
+	return BuildLightingGuidFromHash();
+}
+#endif // WITH_EDITORONLY_DATA
 
 bool UCurveLinearColorAtlas::GetCurveIndex(UCurveLinearColor* InCurve, int32& Index)
 {
diff --git a/Engine/Source/Runtime/Engine/Private/Texture.cpp b/Engine/Source/Runtime/Engine/Private/Texture.cpp
--- a/Engine/Source/Runtime/Engine/Private/Texture.cpp
+++ b/Engine/Source/Runtime/Engine/Private/Texture.cpp
@@ -1707,6 +1707,13 @@ void UTexture::SetDeterministicLightingGuid()
 #endif // WITH_EDITORONLY_DATA
 }
 
+#if WITH_EDITORONLY_DATA
+FGuid UTexture::BuildLightingGuid()
+{
+	return FGuid::NewGuid();
+}
+#endif
+
 UEnum* UTexture::GetPixelFormatEnum()
 {
 	// Lookup the pixel format enum so that the pixel format can be serialized by name.
diff --git a/Engine/Source/Runtime/Engine/Private/TextureDerivedData.cpp b/Engine/Source/Runtime/Engine/Private/TextureDerivedData.cpp
--- a/Engine/Source/Runtime/Engine/Private/TextureDerivedData.cpp
+++ b/Engine/Source/Runtime/Engine/Private/TextureDerivedData.cpp
@@ -1715,6 +1715,57 @@ int64 PutDerivedDataInCache(FTexturePlatformData* DerivedData, const FString& De
 
 #endif // #if WITH_EDITOR
 
+#if WITH_EDITORONLY_DATA
+
+FGuid UTexture::BuildLightingGuidFromHash()
+{
+#if !WITH_EDITOR
+	return FGuid::NewGuid();
+#else
+	if (!Source.IsValid())
+	{
+		return FGuid::NewDeterministicGuid(GetPathName());
+	}
+
+	// Use the DDC key rather than reading just the texels, because the DDC key includes some lighting settings
+	// that impact lighting on materials using the texture. Those settings are why we need a LightingGuid that is
+	// separate from the texture's Id; the Id depends only on the texels.
+
+	// Use a single TargetPlatform so we get the same LightingSettings in every editor platform, so that we get the
+	// same LightingGuid no matter which platform we're running the editor on. We arbitrarily pick Windows for that
+	// platform since it is the most commonly-used editor platform.
+	ITargetPlatformManagerModule* TPM = GetTargetPlatformManager();
+	if (!TPM)
+	{
+		return FGuid::NewDeterministicGuid(GetPathName());
+	}
+	ITargetPlatform* SettingsPlatform = TPM->FindTargetPlatform(TEXT("Windows"));
+	if (!SettingsPlatform)
+	{
+		return FGuid::NewDeterministicGuid(GetPathName());
+	}
+
+	TArray<FTextureBuildSettings> BuildSettings;
+	GetBuildSettingsForTargetPlatform(*this, SettingsPlatform, ETextureEncodeSpeed::Final,
+		BuildSettings, nullptr /* OutResultMetadataPerLayer */);
+	if (BuildSettings.IsEmpty() || BuildSettings.Num() < Source.GetNumLayers())
+	{
+		return FGuid::NewDeterministicGuid(GetPathName());
+	}
+
+	FString DerivedDataKey;
+	GetTextureDerivedDataKey(*this, BuildSettings.GetData(), DerivedDataKey);
+
+	FXxHash128Builder Hasher;
+	Hasher.Update(*DerivedDataKey, DerivedDataKey.Len() * sizeof(**DerivedDataKey));
+	uint8 Hash[16];
+	Hasher.Finalize().ToByteArray(Hash);
+	return FGuid::NewGuidFromHashBytes(&Hash, sizeof(Hash));
+#endif
+}
+
+#endif // WITH_EDITORONLY_DATA
+
 /*------------------------------------------------------------------------------
 	Derived data.
 ------------------------------------------------------------------------------*/
