Remove double data cook task on audio file import.

diff --git a/Engine/Source/Runtime/Engine/Private/SoundWave.cpp b/Engine/Source/Runtime/Engine/Private/SoundWave.cpp
--- a/Engine/Source/Runtime/Engine/Private/SoundWave.cpp
+++ b/Engine/Source/Runtime/Engine/Private/SoundWave.cpp
@@ -3130,10 +3130,17 @@ void USoundWave::PostEditChangeProperty(FPropertyChangedEvent& PropertyChangedEv
 
 	if (PropertyChangedEvent.Property == nullptr) 
 	{
-		//an empty event property field might mean the update comes from an undo
-		//we can't discern what properties where reverted so we update the asset and bakes
+		// An empty event property field might mean the update comes from an undo or an import.
+		// We are making an educated wish that an import task would still have async tasks running
+		// when we get to the PostEditChange and an undo would not. This will help prevent us
+		// from updating the asset twice on import, blocking the UI for the previous Update task
+		// to complete. 
+	
+		if (IsAsyncWorkComplete())
+		{
+			UpdateAsset();
+		}
 	
-		UpdateAsset();
 		BakeFFTAnalysis();
 		BakeEnvelopeAnalysis();
 
