[StandaloneHLOD] Handle empty HLOD actors (which are then marked as Editor Only) when creating clusters during streaming generation.
When saving information about Standalone HLOD child actors, we don't know yet if they're going to end up being cooked. Verify that they are in the container before adding them as a reference.

diff --git a/Engine/Source/Runtime/Engine/Private/WorldPartition/WorldPartitionStreamingGeneration.cpp b/Engine/Source/Runtime/Engine/Private/WorldPartition/WorldPartitionStreamingGeneration.cpp
--- a/Engine/Source/Runtime/Engine/Private/WorldPartition/WorldPartitionStreamingGeneration.cpp
+++ b/Engine/Source/Runtime/Engine/Private/WorldPartition/WorldPartitionStreamingGeneration.cpp
@@ -1768,7 +1768,13 @@ class FWorldPartitionStreamingGenerator
 			// Include Standalone HLOD actors "dependencies". Make sure all child actors of a Standalone HLOD actor to end up in the same cluster.
 			if (TSet<FGuid>* StandaloneHLODDependencies = StandaloneHLODActorsDependencies.Find(ActorWithRefs.Key))
 			{
-				ActorWithRefs.Value.Append(StandaloneHLODDependencies->Array());
+				for (const FGuid& StandaloneHLODDependency : *StandaloneHLODDependencies)
+				{
+					if (ActorDescView.GetActorDescViewMap().FindByGuid(StandaloneHLODDependency))
+					{
+						ActorWithRefs.Value.Add(StandaloneHLODDependency);
+					}
+				}
 			}
 		});
 		ContainerCollectionInstanceDescriptor.Clusters = GenerateObjectsClusters(ActorsWithRefs);
