[StandaloneHLOD] Make sure that Standalone HLOD Actors are injected into appropriate containers during streaming generation

diff --git a/Engine/Source/Runtime/Engine/Private/WorldPartition/WorldPartitionStreamingGeneration.cpp b/Engine/Source/Runtime/Engine/Private/WorldPartition/WorldPartitionStreamingGeneration.cpp
--- a/Engine/Source/Runtime/Engine/Private/WorldPartition/WorldPartitionStreamingGeneration.cpp
+++ b/Engine/Source/Runtime/Engine/Private/WorldPartition/WorldPartitionStreamingGeneration.cpp
@@ -978,26 +978,29 @@ class FWorldPartitionStreamingGenerator
 				// Standalone HLOD actors
 				WorldPartitionStandaloneHLODSubsystem->ForEachStandaloneHLODActor([this, &OutStandaloneHLODActorInstances, &OutActorDescViewMap, &InActorDescCollection, &RegisterActorDescView](AWorldPartitionStandaloneHLOD* StandaloneHLODActor)
 				{
-					TUniquePtr<FStreamingGenerationTempActorDescInstance>& StandaloneHLODActorRef = OutStandaloneHLODActorInstances.Add_GetRef(FStreamingGenerationTempActorDescInstance::Create(StandaloneHLODActor, InActorDescCollection));
-					RegisterActorDescView(FStreamingGenerationActorDescView(OutActorDescViewMap, StandaloneHLODActorRef.Get(), false));
-
-					// Save information about "dependencies" between Standalone HLOD actors. We want all external child HLOD actors of a Standalone HLOD actor to end up in the same streaming cell
-					// to maintain HLODActor -> SourceCell mapping, that is used for regular HLODActors to control visibility / loading state.
-					// StandaloneHLODActorsDependencies saved here will be used later during clusters generation.
-					if (UActorDescContainer* StandaloneHLODContainer = StandaloneHLODActorRef->GetActorDesc()->GetChildContainer())
+					if (InActorDescCollection.FindHandlingContainerInstance(StandaloneHLODActor))
 					{
-						for (UActorDescContainer::TConstIterator<AWorldPartitionHLOD> Iterator(StandaloneHLODContainer); Iterator; ++Iterator)
+						TUniquePtr<FStreamingGenerationTempActorDescInstance>& StandaloneHLODActorRef = OutStandaloneHLODActorInstances.Add_GetRef(FStreamingGenerationTempActorDescInstance::Create(StandaloneHLODActor, InActorDescCollection));
+						RegisterActorDescView(FStreamingGenerationActorDescView(OutActorDescViewMap, StandaloneHLODActorRef.Get(), false));
+
+						// Save information about "dependencies" between Standalone HLOD actors. We want all external child HLOD actors of a Standalone HLOD actor to end up in the same streaming cell
+						// to maintain HLODActor -> SourceCell mapping, that is used for regular HLODActors to control visibility / loading state.
+						// StandaloneHLODActorsDependencies saved here will be used later during clusters generation.
+						if (UActorDescContainer* StandaloneHLODContainer = StandaloneHLODActorRef->GetActorDesc()->GetChildContainer())
 						{
-							const FHLODActorDesc* HLODActorDesc = static_cast<const FHLODActorDesc*>(*Iterator);
-							if (HLODActorDesc->GetExternalChildHLODActors().Num() > 1)
+							for (UActorDescContainer::TConstIterator<AWorldPartitionHLOD> Iterator(StandaloneHLODContainer); Iterator; ++Iterator)
 							{
-								const TArray<FGuid>& ExternalChildHLODActors = HLODActorDesc->GetExternalChildHLODActors();
-								for (int32 ExternalChildHLODActorIndex = 0; ExternalChildHLODActorIndex < ExternalChildHLODActors.Num(); ExternalChildHLODActorIndex++)
+								const FHLODActorDesc* HLODActorDesc = static_cast<const FHLODActorDesc*>(*Iterator);
+								if (HLODActorDesc->GetExternalChildHLODActors().Num() > 1)
 								{
-									for (int32 OtherExternalChildHLODActorIndex = ExternalChildHLODActorIndex + 1; OtherExternalChildHLODActorIndex < ExternalChildHLODActors.Num(); OtherExternalChildHLODActorIndex++)
+									const TArray<FGuid>& ExternalChildHLODActors = HLODActorDesc->GetExternalChildHLODActors();
+									for (int32 ExternalChildHLODActorIndex = 0; ExternalChildHLODActorIndex < ExternalChildHLODActors.Num(); ExternalChildHLODActorIndex++)
 									{
-										StandaloneHLODActorsDependencies.FindOrAdd(ExternalChildHLODActors[ExternalChildHLODActorIndex]).Add(ExternalChildHLODActors[OtherExternalChildHLODActorIndex]);
-										StandaloneHLODActorsDependencies.FindOrAdd(ExternalChildHLODActors[OtherExternalChildHLODActorIndex]).Add(ExternalChildHLODActors[ExternalChildHLODActorIndex]);
+										for (int32 OtherExternalChildHLODActorIndex = ExternalChildHLODActorIndex + 1; OtherExternalChildHLODActorIndex < ExternalChildHLODActors.Num(); OtherExternalChildHLODActorIndex++)
+										{
+											StandaloneHLODActorsDependencies.FindOrAdd(ExternalChildHLODActors[ExternalChildHLODActorIndex]).Add(ExternalChildHLODActors[OtherExternalChildHLODActorIndex]);
+											StandaloneHLODActorsDependencies.FindOrAdd(ExternalChildHLODActors[OtherExternalChildHLODActorIndex]).Add(ExternalChildHLODActors[ExternalChildHLODActorIndex]);
+										}
 									}
 								}
 							}
