Use submixes consistently on preview sound playback for all sound objects, determined by an editor preference.

diff --git a/Engine/Plugins/Runtime/Metasound/Source/MetasoundEngine/Public/MetasoundSource.h b/Engine/Plugins/Runtime/Metasound/Source/MetasoundEngine/Public/MetasoundSource.h
--- a/Engine/Plugins/Runtime/Metasound/Source/MetasoundEngine/Public/MetasoundSource.h
+++ b/Engine/Plugins/Runtime/Metasound/Source/MetasoundEngine/Public/MetasoundSource.h
@@ -285,7 +285,7 @@ public:
 	UE_API virtual bool IsParameterValid(const FAudioParameter& InParameter) const override;
 	UE_API virtual bool IsLooping() const override;
 	UE_API virtual bool IsOneShot() const override;
-	virtual bool EnableSubmixSendsOnPreview() const override { return true; }
+	UE_DEPRECATED(5.7, "This is no longer used.") virtual bool EnableSubmixSendsOnPreview() const override { return true; }
 
 	/**
 	 * Attempt to get a generator for a given AudioComponentID.
diff --git a/Engine/Source/Editor/AudioEditor/Classes/AudioEditorSettings.h b/Engine/Source/Editor/AudioEditor/Classes/AudioEditorSettings.h
--- a/Engine/Source/Editor/AudioEditor/Classes/AudioEditorSettings.h
+++ b/Engine/Source/Editor/AudioEditor/Classes/AudioEditorSettings.h
@@ -92,7 +92,11 @@ public:
 	/** Array of available audio platforms. */
 	UPROPERTY(config)
 	TArray<FAudioPlatform> AudioMixerPlatforms;
-	
+
+	/** Use the submix tree when playing preview sounds */
+	UPROPERTY(EditAnywhere, config, Category = EditorPlaybackOptions, DisplayName = "Use Submixes When Previewing Sounds")
+	bool bUseSubmixesForPreviewSound = true;
+
 	/** Set and apply, whether audio attenuation is used for non-game worlds*/
 	UE_API void SetUseAudioAttenuation(bool bInUseAudioAttenuation);
 
diff --git a/Engine/Source/Editor/UnrealEd/Private/EditorEngine.cpp b/Engine/Source/Editor/UnrealEd/Private/EditorEngine.cpp
--- a/Engine/Source/Editor/UnrealEd/Private/EditorEngine.cpp
+++ b/Engine/Source/Editor/UnrealEd/Private/EditorEngine.cpp
@@ -279,6 +279,8 @@
 #include "Engine/TimecodeProvider.h"
 #include "TargetReceipt.h"
 
+#include "AudioEditorSettings.h"
+
 DEFINE_LOG_CATEGORY_STATIC(LogEditor, Log, All);
 
 #define LOCTEXT_NAMESPACE "UnrealEd.Editor"
@@ -2901,6 +2903,19 @@ UAudioComponent* UEditorEngine::PlayPreviewSound(USoundBase* Sound,  USoundNode*
 		AudioComponent->bReverb = false;
 		AudioComponent->bCenterChannelOnly = false;
 		AudioComponent->bIsPreviewSound = true;
+
+#if WITH_EDITORONLY_DATA
+		if (!::IsValid(Sound) || Sound->EditorData.OverrideUseSubmixesOnPreview == EUseSubmixOnPreviewMode::UseEditorPreference)
+		{
+			const UAudioEditorSettings* AudioEditorSettings = GetDefault<UAudioEditorSettings>();
+			AudioComponent->bIsDryPreviewPlayback = !AudioEditorSettings->bUseSubmixesForPreviewSound;
+		}
+		else
+		{
+			AudioComponent->bIsDryPreviewPlayback = (Sound->EditorData.OverrideUseSubmixesOnPreview == EUseSubmixOnPreviewMode::PlayPreviewsWithoutSubmixes) ? true : false;
+		}
+#endif //WITH_EDITORONLY_DATA
+
 		AudioComponent->Play();
 
 		return AudioComponent;
diff --git a/Engine/Source/Runtime/Engine/Classes/Components/AudioComponent.h b/Engine/Source/Runtime/Engine/Classes/Components/AudioComponent.h
--- a/Engine/Source/Runtime/Engine/Classes/Components/AudioComponent.h
+++ b/Engine/Source/Runtime/Engine/Classes/Components/AudioComponent.h
@@ -270,6 +270,11 @@ class UAudioComponent : public USceneComponent, public ISoundParameterController
 	/** Whether or not this sound is a preview sound */
 	uint8 bIsPreviewSound:1;
 
+#if WITH_EDITORONLY_DATA
+	/** Whether or not this sound uses submix sends if it is a preview sound */
+	uint8 bIsDryPreviewPlayback : 1;
+#endif //WITH_EDITORONLY_DATA
+
 	/** Whether or not this audio component has been paused */
 	uint8 bIsPaused:1;
 
diff --git a/Engine/Source/Runtime/Engine/Classes/Sound/SoundBase.h b/Engine/Source/Runtime/Engine/Classes/Sound/SoundBase.h
--- a/Engine/Source/Runtime/Engine/Classes/Sound/SoundBase.h
+++ b/Engine/Source/Runtime/Engine/Classes/Sound/SoundBase.h
@@ -76,6 +76,31 @@ enum class EVirtualizationMode : uint8
 	SeekRestart UE_EXPERIMENTAL(5.6, "This virtualization mode is experimental and may not function as expected in all use cases.")
 };
 
+/**
+ * There is a editor setting to use submixes when previewing sounds. This gives the user
+ * the ability to override that locally for this specific asset.
+ */
+UENUM(BlueprintType)
+enum class EUseSubmixOnPreviewMode : uint8
+{
+	UseEditorPreference,
+	UseSubmixesOnPreview,
+	PlayPreviewsWithoutSubmixes
+};
+
+// This struct allows us to modify the transient editor properties and minimize impact on the properties
+USTRUCT()
+struct FSoundBaseEditorData
+{
+	GENERATED_BODY()
+
+public:
+	/** Override the Editor Preference to use submixes on preview playback */
+	UPROPERTY(EditAnywhere, Transient, SkipSerialization, Category = Developer, meta = (DisplayName = "Use Submixes on preview"))
+	EUseSubmixOnPreviewMode OverrideUseSubmixesOnPreview = EUseSubmixOnPreviewMode::UseEditorPreference;
+};
+
+
 /**
  * The base class for a playable sound object
  */
@@ -89,6 +114,12 @@ public:
 	UPROPERTY(EditAnywhere, BlueprintReadOnly, Category = Sound, meta = (DisplayName = "Class"), AssetRegistrySearchable)
 	TObjectPtr<USoundClass> SoundClassObject;
 
+#if WITH_EDITORONLY_DATA
+	/** Supplement the USoundBase with a list of transient editor-only properties */
+	UPROPERTY(EditAnywhere, Transient, SkipSerialization, Category = Developer, meta = (DisplayName = "Editor Data", ShowOnlyInnerProperties))
+	FSoundBaseEditorData EditorData;
+#endif //WITH_EDITORONLY_DATA
+
 	/** When "au.3dVisualize.Attenuation" has been specified, draw this sound's attenuation shape when the sound is audible. For debugging purposes only. */
 	UPROPERTY(EditAnywhere, Category = Developer, meta = (DisplayName = "Enable Attenuation Debug"))
 	uint8 bDebug : 1;
@@ -367,7 +398,7 @@ public:
 	virtual bool GetAllDefaultParameters(TArray<FAudioParameter>& OutParameters) const { return false; }
 
 	/** Whether or not this sound allows submix sends on preview. */
-	virtual bool EnableSubmixSendsOnPreview() const { return false; }
+	UE_DEPRECATED(5.7, "This is no longer used.") virtual bool EnableSubmixSendsOnPreview() const { return false; }
 
 	/** Only used as an edit condition for AttenuationSettings member, as base classes may choose to provide an attenuation override implementation */
 	UFUNCTION()
diff --git a/Engine/Source/Runtime/Engine/Private/ActiveSound.cpp b/Engine/Source/Runtime/Engine/Private/ActiveSound.cpp
--- a/Engine/Source/Runtime/Engine/Private/ActiveSound.cpp
+++ b/Engine/Source/Runtime/Engine/Private/ActiveSound.cpp
@@ -120,6 +120,10 @@ FActiveSound::FActiveSound()
 	, bReverb(false)
 	, bCenterChannelOnly(false)
 	, bIsPreviewSound(false)
+#if WITH_EDITORONLY_DATA
+	, bIsDryPreviewPlayback(false)
+#endif //WITH_EDITORONLY_DATA
+
 	, bGotInteriorSettings(false)
 	, bApplyInteriorVolumes(false)
 #if !(NO_LOGGING || UE_BUILD_SHIPPING || UE_BUILD_TEST)
diff --git a/Engine/Source/Runtime/Engine/Private/Components/AudioComponent.cpp b/Engine/Source/Runtime/Engine/Private/Components/AudioComponent.cpp
--- a/Engine/Source/Runtime/Engine/Private/Components/AudioComponent.cpp
+++ b/Engine/Source/Runtime/Engine/Private/Components/AudioComponent.cpp
@@ -99,6 +99,9 @@ UAudioComponent::UAudioComponent(const FObjectInitializer& ObjectInitializer)
 	bCanPlayMultipleInstances = false;
 	bDisableParameterUpdatesWhilePlaying = false;
 	bIsPreviewSound = false;
+#if WITH_EDITORONLY_DATA
+	bIsDryPreviewPlayback = false;
+#endif //WITH_EDITORONLY_DATA
 	bIsPaused = false;
 
 	Priority = 1.f;
@@ -771,6 +774,9 @@ void UAudioComponent::PlayInternal(const PlayInternalRequestData& InPlayRequestD
 	NewActiveSound.bReverb = bReverb;
 	NewActiveSound.bCenterChannelOnly = bCenterChannelOnly;
 	NewActiveSound.bIsPreviewSound = bIsPreviewSound;
+#if WITH_EDITORONLY_DATA
+	NewActiveSound.bIsDryPreviewPlayback = bIsDryPreviewPlayback;
+#endif //WITH_EDITORONLY_DATA
 	NewActiveSound.bLocationDefined = !bPreviewComponent;
 	NewActiveSound.bIsPaused = bIsPaused;
 
diff --git a/Engine/Source/Runtime/Engine/Private/SoundWave.cpp b/Engine/Source/Runtime/Engine/Private/SoundWave.cpp
--- a/Engine/Source/Runtime/Engine/Private/SoundWave.cpp
+++ b/Engine/Source/Runtime/Engine/Private/SoundWave.cpp
@@ -3822,16 +3822,21 @@ void USoundWave::Parse(FAudioDevice* AudioDevice, const UPTRINT NodeWaveInstance
 	}
 	else //if this is a preview sound, ignore Submix sends
 	{
+#if WITH_EDITORONLY_DATA
 		// Only allow submix sends if this is a metasound preview
-		if (EnableSubmixSendsOnPreview())
+		if (!WaveInstance->ActiveSound->bIsDryPreviewPlayback)
 		{
+#endif //WITH_EDITORONLY_DATA
 			WaveInstance->bEnableSubmixSends = ParseParams.bEnableSubmixSends;
 			WaveInstance->bEnableBaseSubmix = ParseParams.bEnableBaseSubmix;
+#if WITH_EDITORONLY_DATA
 		}
 		else
 		{
 			WaveInstance->bEnableBaseSubmix = true;
 		}
+#endif //WITH_EDITORONLY_DATA
+
 		WaveInstance->bEnableBusSends = ParseParams.bEnableBusSends;
 	}
 
diff --git a/Engine/Source/Runtime/Engine/Public/ActiveSound.h b/Engine/Source/Runtime/Engine/Public/ActiveSound.h
--- a/Engine/Source/Runtime/Engine/Public/ActiveSound.h
+++ b/Engine/Source/Runtime/Engine/Public/ActiveSound.h
@@ -507,6 +507,11 @@ public:
 	/** Whether or not this active sound is a preview sound */
 	uint8 bIsPreviewSound:1;
 
+#if WITH_EDITORONLY_DATA
+	/** Whether or not this sound uses submix sends if it is a preview sound */
+	uint8 bIsDryPreviewPlayback : 1;
+#endif //WITH_EDITORONLY_DATA
+
 	/** Whether we have queried for the interior settings at least once */
 	uint8 bGotInteriorSettings:1;
 
