Fix a missing call to handle the error message for mismatching data layer load filter.

diff --git a/Engine/Source/Runtime/Engine/Private/WorldPartition/ErrorHandling/WorldPartitionStreamingGenerationTokenizedMessageErrorHandler.cpp b/Engine/Source/Runtime/Engine/Private/WorldPartition/ErrorHandling/WorldPartitionStreamingGenerationTokenizedMessageErrorHandler.cpp
--- a/Engine/Source/Runtime/Engine/Private/WorldPartition/ErrorHandling/WorldPartitionStreamingGenerationTokenizedMessageErrorHandler.cpp
+++ b/Engine/Source/Runtime/Engine/Private/WorldPartition/ErrorHandling/WorldPartitionStreamingGenerationTokenizedMessageErrorHandler.cpp
@@ -104,6 +104,8 @@ void ITokenizedMessageErrorHandler::OnDataLayersLoadFilterMismatch(const IWorldP
 		->AddToken(FActorToken::Create(ActorDescView.GetActorSoftPath().ToString(), ActorDescView.GetGuid(), FText::FromString(GetActorName(ActorDescView))))
 		->AddToken(FTextToken::Create(LOCTEXT("TokenMessage_WorldPartition_ContainsRuntimeDataLayersWithDifferentTypesOfLoadFilter", "contains runtime data layers with different types of Load Filter")))
 		->AddToken(FMapErrorToken::Create(TEXT("WorldPartition_DataLayersLoadFilterMismatch_CheckForErrors")));
+
+	HandleTokenizedMessage(MoveTemp(Message));
 }
 
 void ITokenizedMessageErrorHandler::OnInvalidWorldReference(const IWorldPartitionActorDescInstanceView& ActorDescView, EWorldReferenceInvalidReason Reason)
