Materials: IncrementalCook: Improve the performance of UMaterial and UMaterialInstance during incremental cook's test for whether their package needs to be recooked, to reduce the cost of the null cook (nothing's changed) case for incremental cook. RecordMaterialDependenciesForCook stores the ShaderMapId in the dependency arguments for the HashMaterialDependenciesForCook to evaluate during a null cook, but the ShaderMapId records a lot of redundant ShaderType names that are from the layout used by the ShaderMapId, and loading and hashing those names has a significant cost. Change this to instead store just the identifiers of the layout used, and compute a hash per layout that is shared by all materials using the same layout.
#rnx

diff --git a/Engine/Plugins/Importers/USDImporter/Source/USDExporter/Private/MaterialExporterUSD.cpp b/Engine/Plugins/Importers/USDImporter/Source/USDExporter/Private/MaterialExporterUSD.cpp
--- a/Engine/Plugins/Importers/USDImporter/Source/USDExporter/Private/MaterialExporterUSD.cpp
+++ b/Engine/Plugins/Importers/USDImporter/Source/USDExporter/Private/MaterialExporterUSD.cpp
@@ -85,7 +85,7 @@ namespace UE::MaterialExporterUSD::Private
 				MaterialResource = Parent->GetMaterialResource(GMaxRHIFeatureLevel);
 
 				FMaterialShaderMapId ShaderMapID;
-				MaterialResource->GetShaderMapId(GMaxRHIShaderPlatform, nullptr, ShaderMapID);
+				MaterialResource->BuildShaderMapId(ShaderMapID, GMaxRHIShaderPlatform, nullptr);
 
 				FSHAHash MaterialHash;
 				ShaderMapID.GetMaterialHash(MaterialHash);
diff --git a/Engine/Source/Developer/MaterialBaking/Private/ExportMaterialProxy.h b/Engine/Source/Developer/MaterialBaking/Private/ExportMaterialProxy.h
--- a/Engine/Source/Developer/MaterialBaking/Private/ExportMaterialProxy.h
+++ b/Engine/Source/Developer/MaterialBaking/Private/ExportMaterialProxy.h
@@ -231,7 +231,7 @@ public:
 		const FMaterialResource* Resource = InMaterialInterface->GetMaterialResource(GMaxRHIFeatureLevel);
 
 		FMaterialShaderMapId ResourceId;
-		Resource->GetShaderMapId(GMaxRHIShaderPlatform, nullptr, ResourceId);
+		Resource->BuildShaderMapId(ResourceId, GMaxRHIShaderPlatform, nullptr);
 
 		{
 			TArray<FShaderType*> ShaderTypes;
diff --git a/Engine/Source/Developer/MaterialUtilities/Private/MaterialUtilities.cpp b/Engine/Source/Developer/MaterialUtilities/Private/MaterialUtilities.cpp
--- a/Engine/Source/Developer/MaterialUtilities/Private/MaterialUtilities.cpp
+++ b/Engine/Source/Developer/MaterialUtilities/Private/MaterialUtilities.cpp
@@ -476,7 +476,7 @@ public:
 		FMaterialResource* Resource = InMaterialInterface->GetMaterialResource(GMaxRHIFeatureLevel);
 
 		FMaterialShaderMapId ResourceId;
-		Resource->GetShaderMapId(GMaxRHIShaderPlatform, nullptr, ResourceId);
+		Resource->BuildShaderMapId(ResourceId, GMaxRHIShaderPlatform, nullptr);
 
 		{
 			TArray<FShaderType*> ShaderTypes;
diff --git a/Engine/Source/Editor/UnrealEd/Private/Commandlets/DumpMaterialShaderTypes.cpp b/Engine/Source/Editor/UnrealEd/Private/Commandlets/DumpMaterialShaderTypes.cpp
--- a/Engine/Source/Editor/UnrealEd/Private/Commandlets/DumpMaterialShaderTypes.cpp
+++ b/Engine/Source/Editor/UnrealEd/Private/Commandlets/DumpMaterialShaderTypes.cpp
@@ -596,7 +596,7 @@ static void ProcessSwitchOptimizer(const ITargetPlatform* TargetPlatform, const
 			{
 				FMaterialResource* Resource = FindOrCreateMaterialResource(ResourcesToCache, Material, nullptr, GMaxRHIFeatureLevel, EMaterialQualityLevel::Type(QualityLevel));
 				FMaterialShaderMapId ShaderMapId;
-				Resource->GetShaderMapId(ShaderPlatform, TargetPlatform, ShaderMapId);
+				Resource->BuildShaderMapId(ShaderMapId, ShaderPlatform, TargetPlatform);
 				StaticSwitchGroupSet& InnerHashMap = *ShaderMapHashMap.FindOrAdd(ShaderMapId, StaticSwitchGroupSet());
 				MaterialAndSizeSet& InnerValue = *InnerHashMap.FindOrAdd(ShaderMapId.GetStaticSwitchParameters(), MaterialAndSizeSet());
 
@@ -628,7 +628,7 @@ static void ProcessSwitchOptimizer(const ITargetPlatform* TargetPlatform, const
 				UMaterial* BaseMaterial = MaterialInstance->GetMaterial();
 				FMaterialResource* Resource = FindOrCreateMaterialResource(ResourcesToCache, BaseMaterial, MaterialInstance, GMaxRHIFeatureLevel, EMaterialQualityLevel::Type(QualityLevel));
 				FMaterialShaderMapId ShaderMapId;
-				Resource->GetShaderMapId(ShaderPlatform, TargetPlatform, ShaderMapId);
+				Resource->BuildShaderMapId(ShaderMapId, ShaderPlatform, TargetPlatform);
 				StaticSwitchGroupSet& InnerHashMap = *ShaderMapHashMap.FindOrAdd(ShaderMapId, StaticSwitchGroupSet());
 				MaterialAndSizeSet& InnerValue = *InnerHashMap.FindOrAdd(ShaderMapId.GetStaticSwitchParameters(), MaterialAndSizeSet());
 
diff --git a/Engine/Source/Editor/UnrealEd/Private/Lightmass/LightmassRender.cpp b/Engine/Source/Editor/UnrealEd/Private/Lightmass/LightmassRender.cpp
--- a/Engine/Source/Editor/UnrealEd/Private/Lightmass/LightmassRender.cpp
+++ b/Engine/Source/Editor/UnrealEd/Private/Lightmass/LightmassRender.cpp
@@ -267,7 +267,7 @@ public:
 			if (Resource)
 			{
 				FMaterialShaderMapId ResourceId;
-				Resource->GetShaderMapId(GMaxRHIShaderPlatform, nullptr, ResourceId);
+				Resource->BuildShaderMapId(ResourceId, GMaxRHIShaderPlatform, nullptr);
 
 				{
 					TArray<FShaderType*> ShaderTypes;
diff --git a/Engine/Source/Runtime/Engine/Private/Materials/Material.cpp b/Engine/Source/Runtime/Engine/Private/Materials/Material.cpp
--- a/Engine/Source/Runtime/Engine/Private/Materials/Material.cpp
+++ b/Engine/Source/Runtime/Engine/Private/Materials/Material.cpp
@@ -348,10 +348,11 @@ void FMaterialResource::EndAllowCachingStaticParameterValues()
 }
 #endif // WITH_EDITOR
 
-void FMaterialResource::GetShaderMapId(EShaderPlatform Platform, const ITargetPlatform* TargetPlatform, FMaterialShaderMapId& OutId) const
+void FMaterialResource::BuildShaderMapIdOverride(const FBuildShaderMapIdArgs& Args) const
 {
-	FMaterial::GetShaderMapId(Platform, TargetPlatform, OutId);
+	FMaterial::BuildShaderMapIdOverride(Args);
 #if WITH_EDITOR
+	FMaterialShaderMapId& OutId = *Args.OutId;
 	const FMaterialCachedExpressionData& CachedData = GetCachedExpressionData();
 	CachedData.AppendReferencedFunctionIdsTo(OutId.ReferencedFunctions);
 	CachedData.AppendReferencedParameterCollectionIdsTo(OutId.ReferencedParameterCollections);
@@ -360,7 +361,7 @@ void FMaterialResource::GetShaderMapId(EShaderPlatform Platform, const ITargetPl
 
 	Material->GetForceRecompileTextureIdsHash(OutId.TextureReferencesHash);
 
-	if(MaterialInstance)
+	if (MaterialInstance)
 	{
 		MaterialInstance->GetBasePropertyOverridesHash(OutId.BasePropertyOverridesHash);
 
diff --git a/Engine/Source/Runtime/Engine/Private/Materials/MaterialShader.cpp b/Engine/Source/Runtime/Engine/Private/Materials/MaterialShader.cpp
--- a/Engine/Source/Runtime/Engine/Private/Materials/MaterialShader.cpp
+++ b/Engine/Source/Runtime/Engine/Private/Materials/MaterialShader.cpp
@@ -35,6 +35,7 @@
 #include "PSOPrecacheMaterial.h"
 #include "PSOPrecacheValidation.h"
 #include "ProfilingDebugging/CsvProfiler.h"
+#include "Hash/Blake3.h"
 
 #if WITH_ODSC
 #include "ODSC/ODSCManager.h"
@@ -266,7 +267,10 @@ void RecordOrEmitMaterialShaderMapKey(FMaterialKeyGeneratorContext& Context,
 		KeyGen.Append(GetTargetPlatformManagerRef().ShaderFormatVersion(Format));
 		KeyGen.AppendSeparator();
 
-		ShaderMapAppendKey(Platform, KeyGen);
+		if (Context.HasAllFlags(EMaterialKeyInclude::Globals))
+		{
+			ShaderMapAppendKey(Platform, KeyGen);
+		}
 		ShaderMapId.RecordAndEmit(Context);
 
 		FMemoryHasherBlake3 EnvHasher;
@@ -411,8 +415,36 @@ bool LoadFromCompactBinary(FCbFieldView Field,
 namespace UE::MaterialInterface::Private
 {
 
-// DDC Key for Material struct updates, bump when changing related classes which can be loaded/saved to CompactBinary
-#define MATERIAL_CLASS_DERIVEDDATA_VER TEXT("6623D4F4-1CB7-45E7-A809-4D4660884D49")
+// DDC Key for Material structs saved in incremental cook artifacts, bump when changing related classes which
+// can be loaded/saved to CompactBinary
+#define MATERIAL_CLASS_INCREMENTALCOOK_VER TEXT("B127B459-4F59-4075-905F-0B17A982933E")
+
+TMap<EShaderPlatform, FBlake3Hash> GShaderMapPlatformHashes;
+FRWLock GShaderMapPlatformHashesLock;
+
+FBlake3Hash GetShaderPlatformIncrementalCookHash(EShaderPlatform ShaderPlatform)
+{
+	{
+		FReadScopeLock ScopeLock(GShaderMapPlatformHashesLock);
+		const FBlake3Hash* ExistingHash = GShaderMapPlatformHashes.Find(ShaderPlatform);
+		if (ExistingHash)
+		{
+			return *ExistingHash;
+		}
+	}
+
+	FBlake3 Hasher;
+	FShaderKeyGenerator KeyGen([&Hasher](const void* Data, uint64 Size)
+		{
+			Hasher.Update(Data, Size);
+		});
+	ShaderMapAppendKey(ShaderPlatform, KeyGen);
+	FBlake3Hash Hash = Hasher.Finalize();
+
+	FWriteScopeLock ScopeLock(GShaderMapPlatformHashesLock);
+	GShaderMapPlatformHashes.FindOrAdd(ShaderPlatform, Hash);
+	return Hash;
+}
 
 void HashMaterialStaticClassDependenciesForCook(FAppendToClassSchemaContext& Context)
 {
@@ -447,7 +479,7 @@ void HashMaterialStaticClassDependenciesForCook(FAppendToClassSchemaContext& Con
 		Context.Update(*FeatureLevelNameStr, FeatureLevelNameStr.Len() * sizeof((*FeatureLevelNameStr)[0]));
 	}
 
-	FStringView VersionString(MATERIAL_CLASS_DERIVEDDATA_VER);
+	FStringView VersionString(MATERIAL_CLASS_INCREMENTALCOOK_VER);
 	Context.Update(VersionString.GetData(), VersionString.Len() * sizeof(VersionString[0]));
 }
 
@@ -488,6 +520,12 @@ void HashMaterialDependenciesForCook(FCbFieldViewIterator Args, UE::Cook::FCookD
 			}
 			Context.Update(&Hash->Hash, sizeof(Hash->Hash));
 		}
+		const FMaterialShaderMapLayout& Layout = AcquireMaterialShaderMapLayout(Resource.ShaderPlatform,
+			GetShaderPermutationFlags(Resource.ShaderMapId.LayoutParams), Resource.ShaderParameters);
+		Context.Update(&Layout.ShaderMapHash.Hash, sizeof(Layout.ShaderMapHash.Hash));
+
+		FBlake3Hash ShaderMapPlatformHash = GetShaderPlatformIncrementalCookHash(Resource.ShaderPlatform);
+		Context.Update(&ShaderMapPlatformHash.GetBytes(), sizeof(ShaderMapPlatformHash.GetBytes()));
 	}
 }
 
@@ -518,7 +556,10 @@ void RecordMaterialDependenciesForCook(UE::Cook::FCookEventContext CookContext,
 		FRecordedMaterialResourceForCooking& RecordedResource = RecordedResources.Emplace_GetRef();
 		RecordedResource.ShaderPlatform = ResourceForCooking.Platform;
 		RecordedResource.ShaderParameters = FMaterialShaderParameters(&Resource);
-		Resource.GetShaderMapId(RecordedResource.ShaderPlatform, TargetPlatform, RecordedResource.ShaderMapId);
+
+		EMaterialKeyInclude IncludeFlags = EMaterialKeyInclude::All & ~EMaterialKeyInclude::Globals;
+		Resource.BuildShaderMapId(RecordedResource.ShaderMapId, RecordedResource.ShaderPlatform, TargetPlatform,
+			IncludeFlags);
 
 		TSet<FString>& ExpressionIncludes = Resource.GetCachedExpressionData().EditorOnlyData->ExpressionIncludeFilePaths;
 		// ExpressionIncludeFilePaths is already deterministically sorted, so no need to sort it here.
diff --git a/Engine/Source/Runtime/Engine/Private/Materials/MaterialShared.cpp b/Engine/Source/Runtime/Engine/Private/Materials/MaterialShared.cpp
--- a/Engine/Source/Runtime/Engine/Private/Materials/MaterialShared.cpp
+++ b/Engine/Source/Runtime/Engine/Private/Materials/MaterialShared.cpp
@@ -730,8 +730,9 @@ int32 FMaterialAttributesInput::CompileWithDefault(class FMaterialCompiler* Comp
 }
 #endif  // WITH_EDITOR
 
-void FMaterial::GetShaderMapId(EShaderPlatform Platform, const ITargetPlatform* TargetPlatform, FMaterialShaderMapId& OutId) const
+void FMaterial::BuildShaderMapIdOverride(const FBuildShaderMapIdArgs& Args) const
 { 
+	FMaterialShaderMapId& OutId = *Args.OutId;
 	if (bLoadedCookedShaderMapId)
 	{
 		if (GameThreadShaderMap && (IsInGameThread() || IsInAsyncLoadingThread()))
@@ -750,22 +751,25 @@ void FMaterial::GetShaderMapId(EShaderPlatform Platform, const ITargetPlatform*
 	else
 	{
 #if WITH_EDITOR
-		OutId.LayoutParams.InitializeForPlatform(TargetPlatform);
+		OutId.LayoutParams.InitializeForPlatform(Args.TargetPlatform);
 
 		TArray<FShaderType*> ShaderTypes;
 		TArray<FVertexFactoryType*> VFTypes;
 		TArray<const FShaderPipelineType*> ShaderPipelineTypes;
 
-		GetDependentShaderAndVFTypes(Platform, OutId.LayoutParams, ShaderTypes, ShaderPipelineTypes, VFTypes);
+		if (EnumHasAnyFlags(Args.IncludeFlags, EMaterialKeyInclude::Globals))
+		{
+			GetDependentShaderAndVFTypes(Args.Platform, OutId.LayoutParams, ShaderTypes, ShaderPipelineTypes, VFTypes);
+		}
 
 		OutId.Usage = GetShaderMapUsage();
 		OutId.bUsingNewHLSLGenerator = IsUsingNewHLSLGenerator();
 		OutId.BaseMaterialId = GetMaterialId();
 		OutId.QualityLevel = GetQualityLevel();
 		OutId.FeatureLevel = GetFeatureLevel();
-		OutId.SetShaderDependencies(ShaderTypes, ShaderPipelineTypes, VFTypes, Platform);
-		GetReferencedTexturesHash(Platform, OutId.TextureReferencesHash);
-		GetExpressionIncludesHash(Platform, OutId.ExpressionIncludesHash);
+		OutId.SetShaderDependencies(ShaderTypes, ShaderPipelineTypes, VFTypes, Args.Platform);
+		GetReferencedTexturesHash(Args.Platform, OutId.TextureReferencesHash);
+		GetExpressionIncludesHash(Args.Platform, OutId.ExpressionIncludesHash);
 		GetExternalCodeReferencesHash(OutId.ExternalCodeReferencesHash);
 
 		OutId.SubstrateCompilationConfig = GetSubstrateCompilationConfig();
@@ -774,7 +778,7 @@ void FMaterial::GetShaderMapId(EShaderPlatform Platform, const ITargetPlatform*
 		OutId.QualityLevel = GetQualityLevel();
 		OutId.FeatureLevel = GetFeatureLevel();
 
-		if (TargetPlatform != nullptr)
+		if (Args.TargetPlatform != nullptr)
 		{
 			UE_LOG(LogMaterial, Error, TEXT("FMaterial::GetShaderMapId: TargetPlatform is not null, but a cooked executable cannot target platforms other than its own."));
 		}
@@ -2735,7 +2739,7 @@ bool FMaterial::CacheShaders(EShaderPlatform Platform, EMaterialShaderPrecompile
 {
 	FAllowCachingStaticParameterValues AllowCachingStaticParameterValues(*this);
 	FMaterialShaderMapId NoStaticParametersId;
-	GetShaderMapId(Platform, TargetPlatform, NoStaticParametersId);
+	BuildShaderMapId(NoStaticParametersId, Platform, TargetPlatform);
 	return CacheShaders(NoStaticParametersId, Platform, PrecompileMode, TargetPlatform);
 }
 
@@ -3110,7 +3114,7 @@ void FMaterial::BeginCacheShaders(EShaderPlatform Platform, EMaterialShaderPreco
 {
 	FAllowCachingStaticParameterValues AllowCachingStaticParameterValues(*this);
 	FMaterialShaderMapId NoStaticParametersId;
-	GetShaderMapId(Platform, TargetPlatform, NoStaticParametersId);
+	BuildShaderMapId(NoStaticParametersId, Platform, TargetPlatform);
 	return BeginCacheShaders(NoStaticParametersId, Platform, PrecompileMode, TargetPlatform, MoveTemp(CompletionCallback));
 }
 
@@ -4430,7 +4434,7 @@ bool FMaterial::GetMaterialExpressionSource(FString& OutSource)
 	const ITargetPlatform* TargetPlatform = nullptr;
 
 	FMaterialShaderMapId ShaderMapId;
-	GetShaderMapId(ShaderPlatform, TargetPlatform, ShaderMapId);
+	BuildShaderMapId(ShaderMapId, ShaderPlatform, TargetPlatform);
 	FStaticParameterSet StaticParameterSet;
 	GetStaticParameterSet(ShaderPlatform, StaticParameterSet);
 
@@ -4956,7 +4960,7 @@ void UMaterialInterface::AnalyzeMaterialCompilationInCallback(TFunctionRef<void
 	}
 
 	FMaterialShaderMapId ShaderMapID;
-	MaterialResource->GetShaderMapId(GMaxRHIShaderPlatform, nullptr, ShaderMapID);
+	MaterialResource->BuildShaderMapId(ShaderMapID, GMaxRHIShaderPlatform, nullptr);
 	FStaticParameterSet StaticParamSet;
 	MaterialResource->GetStaticParameterSet(GMaxRHIShaderPlatform, StaticParamSet);
 	FMaterialAnalyzer MaterialTranslator(MaterialResource, TempOutput, StaticParamSet, GMaxRHIShaderPlatform, MaterialResource->GetQualityLevel(), GMaxRHIFeatureLevel);
@@ -5028,7 +5032,7 @@ bool UMaterialInterface::IsTextureReferencedByProperty(EMaterialProperty InPrope
 
 	FMaterialCompilationOutput TempOutput;
 	FMaterialShaderMapId ShaderMapID;
-	MaterialResource->GetShaderMapId(GMaxRHIShaderPlatform, nullptr, ShaderMapID);
+	MaterialResource->BuildShaderMapId(ShaderMapID, GMaxRHIShaderPlatform, nullptr);
 	FStaticParameterSet StaticParamSet;
 	MaterialResource->GetStaticParameterSet(GMaxRHIShaderPlatform, StaticParamSet);
 	FHLSLMaterialTranslator MaterialTranslator(MaterialResource, TempOutput, StaticParamSet, GMaxRHIShaderPlatform, MaterialResource->GetQualityLevel(), GMaxRHIFeatureLevel);
diff --git a/Engine/Source/Runtime/Engine/Public/MaterialKeyGeneratorContext.h b/Engine/Source/Runtime/Engine/Public/MaterialKeyGeneratorContext.h
--- a/Engine/Source/Runtime/Engine/Public/MaterialKeyGeneratorContext.h
+++ b/Engine/Source/Runtime/Engine/Public/MaterialKeyGeneratorContext.h
@@ -5,6 +5,7 @@
 #if WITH_EDITOR
 #include "Containers/StringView.h"
 #include "Containers/UnrealString.h"
+#include "MaterialKeyIncludeEnum.h"
 #include "Misc/EnumClassFlags.h"
 #include "Misc/StringBuilder.h"
 #include "RHIShaderPlatform.h"
@@ -14,25 +15,6 @@
 #include "Templates/Function.h"
 #include "UObject/NameTypes.h"
 
-/**
- * Categories of fields which should be included in the output of an FMaterialKeyGeneratorContext. Code recording or
- * emitting fields to a GeneratorContext calls HasAllFlags for the flags relevant to a given field before
- * writing it.
- */
-enum class EMaterialKeyInclude : uint32
-{
-	/** shadertype, shaderpipelinetype, vertexfactortype dependencies are included in the output */
-	ShaderDependencies = 0x1,
-	/** shader hlsl file hashes are included in the output */
-	SourceAndMaterialState = 0x2,
-	/** global data that applies to all shaders is included in the output */
-	Globals = 0x4,
-	/** data stored in UObject exports in .uasset and .umap files is included in the output */
-	UObjectData = 0x8,
-
-	All = ShaderDependencies | SourceAndMaterialState | Globals | UObjectData,
-};
-
 /**
  * Output class passed to RecordAndEmit functions for material shader data. It receives function calls that either
  * save/load variables to compact binary or emits those variables to an FShaderKeyGenerator to construct the key.
diff --git a/Engine/Source/Runtime/Engine/Public/MaterialKeyIncludeEnum.h b/Engine/Source/Runtime/Engine/Public/MaterialKeyIncludeEnum.h
new file mode 100644
--- /dev/null
+++ b/Engine/Source/Runtime/Engine/Public/MaterialKeyIncludeEnum.h
@@ -0,0 +1,33 @@
+// Copyright Epic Games, Inc. All Rights Reserved.
+
+#pragma once
+
+#include "CoreTypes.h"
+
+/**
+ * Categories of fields which should be included in a ShaderMapId or in the output of a FMaterialKeyGeneratorContext.
+ * GetShaderMapId, and code recording or emitting fields to a FMaterialKeyGeneratorContext, calls HasAllFlags for the
+ * flags relevant to a given field before writing it to the ShaderMapId's ShaderTypes or to the
+ * FMaterialKeyGeneratorContext's RecordAndEmit functions.
+ */
+enum class EMaterialKeyInclude : uint32
+{
+	/** shadertype, shaderpipelinetype, vertexfactortype dependencies are included in the output */
+	ShaderDependencies = 0x1,
+	/** shader hlsl file hashes are included in the output */
+	SourceAndMaterialState = 0x2,
+	/**
+	 * Global data that applies to many or all Materials is included in the output.
+	 * For calculating the ShaderMap Id, this includes
+	 *    The ShaderTypes, ShaderPipelineTypes, and VertexFactoryTypes in the Material's FMaterialShaderMapLayout.
+	 * For the output of FMaterialKeyGeneratorContext this includes globals (cvars, project settings) that impact
+	 *     The Material's FMaterialShaderMapLayout via GetShaderTypeLayoutHash.
+	 *     The ShaderPlatform being emitted via ShaderMapAppendKey.
+	 *     All Materials via FMaterialAttributeDefinitionMap::AppendDDCKey.
+	 */
+	Globals = 0x4,
+	/** data stored in UObject exports in .uasset and .umap files is included in the output */
+	UObjectData = 0x8,
+
+	All = ShaderDependencies | SourceAndMaterialState | Globals | UObjectData,
+};
diff --git a/Engine/Source/Runtime/Engine/Public/MaterialShared.h b/Engine/Source/Runtime/Engine/Public/MaterialShared.h
--- a/Engine/Source/Runtime/Engine/Public/MaterialShared.h
+++ b/Engine/Source/Runtime/Engine/Public/MaterialShared.h
@@ -26,6 +26,7 @@
 #include "Misc/Optional.h"
 #include "Serialization/MemoryWriter.h"
 #include "Serialization/ArchiveProxy.h"
+#include "MaterialKeyIncludeEnum.h"
 #include "MaterialRecursionGuard.h"
 #include "MaterialSceneTextureId.h"
 #include "MaterialShaderPrecompileMode.h"
@@ -2168,7 +2169,40 @@ public:
 	void DiscardShaderMap();
 
 	// Material properties.
-	ENGINE_API virtual void GetShaderMapId(EShaderPlatform Platform, const ITargetPlatform* TargetPlatform, FMaterialShaderMapId& OutId) const;
+
+	struct FBuildShaderMapIdArgs
+	{
+		FMaterialShaderMapId* OutId = nullptr; // Guaranteed non-null in BuildShaderMapIdOverride
+
+		EShaderPlatform Platform = EShaderPlatform::SP_NumPlatforms;
+		const ITargetPlatform* TargetPlatform = nullptr;
+		EMaterialKeyInclude IncludeFlags = EMaterialKeyInclude::All;
+	};
+
+	UE_DEPRECATED(5.7, "Call BuildShaderMapId instead")
+	virtual void GetShaderMapId(EShaderPlatform Platform, const ITargetPlatform* TargetPlatform,
+		FMaterialShaderMapId& OutId) const
+	{
+		return BuildShaderMapId(OutId, Platform, TargetPlatform);
+	}
+
+	/** Returns the ShaderMapId used in DDC and IncrementalCook for the ShaderMap created for this FMaterial. */
+	void BuildShaderMapId(FMaterialShaderMapId& OutId, EShaderPlatform Platform, const ITargetPlatform* TargetPlatform,
+		EMaterialKeyInclude IncludeFlags = EMaterialKeyInclude::All) const
+	{
+		return BuildShaderMapIdOverride({
+			.OutId = &OutId, .Platform = Platform, .TargetPlatform = TargetPlatform, .IncludeFlags = IncludeFlags }
+			);
+	}
+
+protected:
+	/**
+	 * The function that should be override by subclasses that want to modify either the DDC key or the
+	 * IncrementalCook key. */
+	ENGINE_API virtual void BuildShaderMapIdOverride(const FBuildShaderMapIdArgs& Args) const;
+
+public:
+
 #if WITH_EDITORONLY_DATA
 	ENGINE_API virtual void GetStaticParameterSet(EShaderPlatform Platform, FStaticParameterSet& OutSet) const;
 #endif // WITH_EDITORONLY_DATA
@@ -2923,7 +2957,9 @@ public:
 	ENGINE_API virtual FString GetMaterialUsageDescription() const override;
 
 	// FMaterial interface.
-	ENGINE_API virtual void GetShaderMapId(EShaderPlatform Platform, const ITargetPlatform* TargetPlatform, FMaterialShaderMapId& OutId) const override;
+protected:
+	ENGINE_API virtual void BuildShaderMapIdOverride(const FBuildShaderMapIdArgs& Args) const override;
+public:
 #if WITH_EDITORONLY_DATA
 	ENGINE_API virtual void GetStaticParameterSet(EShaderPlatform Platform, FStaticParameterSet& OutSet) const override;
 #endif
