Skeletal Mesh: Fix morph target loading crash in cooked editor builds.

#rnx

diff --git a/Engine/Source/Runtime/Engine/Private/MorphTarget.cpp b/Engine/Source/Runtime/Engine/Private/MorphTarget.cpp
--- a/Engine/Source/Runtime/Engine/Private/MorphTarget.cpp
+++ b/Engine/Source/Runtime/Engine/Private/MorphTarget.cpp
@@ -129,12 +129,8 @@ void UMorphTarget::Serialize( FArchive& Ar )
 	{
 		// Only platforms that support ES31 rendering need compressed vertex data on the CPU side. Non-ES31 platforms operate on compressed GPU
 		// data which is cooked out elsewhere.
-		bool bStoreCompressedVertices = false;
+		bool bStoreCompressedVertices = Ar.CookingTarget() ? Ar.CookingTarget()->SupportsFeature(ETargetPlatformFeatures::MobileRendering) : false;
 		
-#if WITH_EDITOR
-		bStoreCompressedVertices = Ar.CookingTarget()->SupportsFeature(ETargetPlatformFeatures::MobileRendering);
-#endif
-
 		Ar << bStoreCompressedVertices;
 		
 		if (bStoreCompressedVertices)
