[HWRT] Add support for ray tracing proxies in FSplineMeshSceneProxy.

diff --git a/Engine/Source/Runtime/Engine/Private/SplineMeshSceneProxy.cpp b/Engine/Source/Runtime/Engine/Private/SplineMeshSceneProxy.cpp
--- a/Engine/Source/Runtime/Engine/Private/SplineMeshSceneProxy.cpp
+++ b/Engine/Source/Runtime/Engine/Private/SplineMeshSceneProxy.cpp
@@ -131,6 +131,18 @@ void FSplineMeshSceneProxy::SetupMeshBatchForSpline(int32 InLODIndex, FMeshBatch
 	OutMeshBatch.ReverseCulling ^= (SplineParams.StartScale.X < 0) ^ (SplineParams.StartScale.Y < 0);
 }
 
+void FSplineMeshSceneProxy::SetupRayTracingMeshBatchForSpline(int32 InLODIndex, FMeshBatch& OutMeshBatch) const
+{
+	const FStaticMeshVertexFactories& VFs = (*RenderData->RayTracingProxy->LODVertexFactories)[InLODIndex];
+	check(OutMeshBatch.Elements.Num() == 1);
+	OutMeshBatch.VertexFactory = OutMeshBatch.Elements[0].bUserDataIsColorVertexBuffer ? VFs.SplineVertexFactoryOverrideColorVertexBuffer : VFs.SplineVertexFactory;
+	check(OutMeshBatch.VertexFactory);
+	OutMeshBatch.Elements[0].SplineMeshSceneProxy = this;
+	OutMeshBatch.Elements[0].bIsSplineProxy = true;
+	OutMeshBatch.Elements[0].PrimitiveUniformBuffer = GetUniformBuffer();
+	OutMeshBatch.ReverseCulling ^= (SplineParams.StartScale.X < 0) ^ (SplineParams.StartScale.Y < 0);
+}
+
 bool FSplineMeshSceneProxy::GetShadowMeshElement(int32 LODIndex, int32 BatchIndex, uint8 InDepthPriorityGroup, FMeshBatch& OutMeshBatch, bool bDitheredLODTransition) const
 {
 	//checkf(LODIndex == 0, TEXT("Getting spline static mesh element with invalid LOD [%d]"), LODIndex);
@@ -180,6 +192,18 @@ bool FSplineMeshSceneProxy::GetCollisionMeshElement(int32 LODIndex, int32 BatchI
 	return false;
 }
 
+bool FSplineMeshSceneProxy::GetRayTracingMeshElement(int32 LODIndex, int32 BatchIndex, int32 SectionIndex, uint8 InDepthPriorityGroup, FMeshBatch& OutMeshBatch) const
+{
+	//checkf(LODIndex == 0 /*&& SectionIndex == 0*/, TEXT("Getting spline static mesh element with invalid params [%d, %d]"), LODIndex, SectionIndex);
+
+	if (FStaticMeshSceneProxy::GetRayTracingMeshElement(LODIndex, BatchIndex, SectionIndex, InDepthPriorityGroup, OutMeshBatch))
+	{
+		SetupRayTracingMeshBatchForSpline(LODIndex, OutMeshBatch);
+		return true;
+	}
+	return false;
+}
+
 #if RHI_RAYTRACING
 void FSplineMeshSceneProxy::GetDynamicRayTracingInstances(FRayTracingInstanceCollector& Collector)
 {
@@ -192,15 +216,50 @@ void FSplineMeshSceneProxy::GetDynamicRayTracingInstances(FRayTracingInstanceCol
 	checkf(!DynamicRayTracingGeometries.IsEmpty(), TEXT("DynamicRayTracingGeometries has not been initialized correctly"));
 
 	ESceneDepthPriorityGroup PrimitiveDPG = GetStaticDepthPriorityGroup();
-	const int32 LODIndex = FMath::Max(GetLOD(Collector.GetReferenceView()), (int32)GetCurrentFirstLODIdx_RenderThread());
-	const FStaticMeshLODResources& LODModel = RenderData->LODResources[LODIndex];
+
+	FStaticMeshRayTracingProxyLODArray& RayTracingLODs = RenderData->RayTracingProxy->LODs;
+
+	const int32 NumLODs = RayTracingLODs.Num();
+
+	const int32 RayTracingMinLOD = RenderData->RayTracingProxy->bUsingRenderingLODs ? FMath::Max(GetLOD(Collector.GetReferenceView()), (int32)GetCurrentFirstLODIdx_RenderThread()) : 0;
+
+	int32 LODIndex = RayTracingMinLOD;
+
+	if (!RenderData->RayTracingProxy->bUsingRenderingLODs)
+	{
+		// when using WPO, need to mark the geometry group as referenced since VB/IB need to be streamed-in 
+		// TODO: Support streaming only buffers when using dynamic geometry
+		Collector.AddReferencedGeometryGroupForDynamicUpdate(RenderData->RayTracingGeometryGroupHandle);
+
+		// select first LOD with valid VB/IB
+		for (; LODIndex < NumLODs; LODIndex++)
+		{
+			FStaticMeshRayTracingProxyLOD& CurrentRayTracingLOD = RayTracingLODs[LODIndex];
+
+			if (CurrentRayTracingLOD.bBuffersInlined || CurrentRayTracingLOD.AreBuffersStreamedIn())
+			{
+				break;
+			}
+		}
+	}
+
+	// TODO: Support streaming only buffers when using dynamic geometry
+	Collector.AddReferencedGeometryGroup(RenderData->RayTracingGeometryGroupHandle);
+
+	if (LODIndex == INDEX_NONE || LODIndex >= NumLODs)
+	{
+		// if none of the LODs have buffers ready for dynamic BLAS update, don't include primitive in ray tracing scene
+		return;
+	}
+
+	const FStaticMeshRayTracingProxyLOD& RayTracingLOD = RayTracingLODs[LODIndex];
 
 	FRayTracingGeometry& Geometry = DynamicRayTracingGeometries[LODIndex];
 
 	FRayTracingInstance RayTracingInstance;
 
 	const int32 NumBatches = GetNumMeshBatches();
-	const int32 NumRayTracingMaterialEntries = LODModel.Sections.Num() * NumBatches;
+	const int32 NumRayTracingMaterialEntries = RayTracingLOD.Sections->Num() * NumBatches;
 
 	if (NumRayTracingMaterialEntries != CachedRayTracingMaterials.Num() || CachedRayTracingMaterialsLODIndex != LODIndex)
 	{
@@ -209,16 +268,27 @@ void FSplineMeshSceneProxy::GetDynamicRayTracingInstances(FRayTracingInstanceCol
 
 		for (int32 BatchIndex = 0; BatchIndex < NumBatches; BatchIndex++)
 		{
-			for (int32 SectionIndex = 0; SectionIndex < LODModel.Sections.Num(); SectionIndex++)
+			for (int32 SectionIndex = 0; SectionIndex < RayTracingLOD.Sections->Num(); SectionIndex++)
 			{
 				FMeshBatch& MeshBatch = CachedRayTracingMaterials.AddDefaulted_GetRef();
 
-				bool bResult = GetMeshElement(LODIndex, BatchIndex, SectionIndex, PrimitiveDPG, false, false, MeshBatch);
+				bool bResult;
+				if (RenderData->RayTracingProxy->bUsingRenderingLODs)
+				{
+					// when using rendering LODs we can reuse the main GetMeshElement(...)
+					bResult = GetMeshElement(LODIndex, BatchIndex, SectionIndex, PrimitiveDPG, false, false, MeshBatch);
+				}
+				else
+				{
+					// otherwise initialize MeshBatch using ray tracing proxy VB/IB/Section data
+					bResult = GetRayTracingMeshElement(LODIndex, BatchIndex, SectionIndex, PrimitiveDPG, MeshBatch);
+				}
+
 				if (!bResult)
 				{
 					// Hidden material
 					MeshBatch.MaterialRenderProxy = UMaterial::GetDefaultMaterial(MD_Surface)->GetRenderProxy();
-					MeshBatch.VertexFactory = &RenderData->LODVertexFactories[LODIndex].VertexFactory;
+					MeshBatch.VertexFactory = &(*RenderData->RayTracingProxy->LODVertexFactories)[LODIndex].VertexFactory;
 				}
 				MeshBatch.ReverseCulling = bReverseCulling;
 				MeshBatch.SegmentIndex = SectionIndex;
@@ -258,13 +328,15 @@ void FSplineMeshSceneProxy::GetDynamicRayTracingInstances(FRayTracingInstanceCol
 
 	if (bNeedsUpdate && RenderData->LODVertexFactories[LODIndex].VertexFactory.GetType()->SupportsRayTracingDynamicGeometry())
 	{
+		const uint32 NumVertices = RayTracingLOD.VertexBuffers->PositionVertexBuffer.GetNumVertices();
+
 		Collector.AddRayTracingGeometryUpdate(
 			FRayTracingDynamicGeometryUpdateParams
 			{
 				CachedRayTracingMaterials, // TODO: this copy can be avoided if FRayTracingDynamicGeometryUpdateParams supported array views
 				false,
-				(uint32)LODModel.GetNumVertices(),
-				uint32((SIZE_T)LODModel.GetNumVertices() * sizeof(FVector3f)),
+				NumVertices,
+				uint32((size_t)NumVertices * sizeof(FVector3f)),
 				Geometry.Initializer.TotalPrimitiveCount,
 				&Geometry,
 				VertexBuffer,
diff --git a/Engine/Source/Runtime/Engine/Public/SplineMeshSceneProxy.h b/Engine/Source/Runtime/Engine/Public/SplineMeshSceneProxy.h
--- a/Engine/Source/Runtime/Engine/Public/SplineMeshSceneProxy.h
+++ b/Engine/Source/Runtime/Engine/Public/SplineMeshSceneProxy.h
@@ -175,6 +175,7 @@ public:
 	virtual bool GetMeshElement(int32 LODIndex, int32 BatchIndex, int32 ElementIndex, uint8 InDepthPriorityGroup, bool bUseSelectionOutline, bool bAllowPreCulledIndices, FMeshBatch& OutMeshBatch) const override;
 	virtual bool GetWireframeMeshElement(int32 LODIndex, int32 BatchIndex, const FMaterialRenderProxy* WireframeRenderProxy, uint8 InDepthPriorityGroup, bool bAllowPreCulledIndices, FMeshBatch& OutMeshBatch) const override;
 	virtual bool GetCollisionMeshElement(int32 LODIndex, int32 BatchIndex, int32 ElementIndex, uint8 InDepthPriorityGroup, const FMaterialRenderProxy* RenderProxy, FMeshBatch& OutMeshBatch) const override;
+	virtual bool GetRayTracingMeshElement(int32 LODIndex, int32 BatchIndex, int32 SectionIndex, uint8 InDepthPriorityGroup, FMeshBatch& OutMeshBatch) const;
 #if RHI_RAYTRACING
 	virtual bool HasRayTracingRepresentation() const override { return true; }
 	virtual bool IsRayTracingRelevant() const override { return true; }
@@ -195,6 +196,7 @@ private:
 	};
 
 	void SetupMeshBatchForSpline(int32 InLODIndex, FMeshBatch& OutMeshBatch) const;
+	void SetupRayTracingMeshBatchForSpline(int32 InLODIndex, FMeshBatch& OutMeshBatch) const;
 
 private:
 	TArray<FLODResources> LODResources;
