Lowered the level of certain logs output during actor migration.

diff --git a/Engine/Source/Runtime/Engine/Private/Actor.cpp b/Engine/Source/Runtime/Engine/Private/Actor.cpp
--- a/Engine/Source/Runtime/Engine/Private/Actor.cpp
+++ b/Engine/Source/Runtime/Engine/Private/Actor.cpp
@@ -1362,7 +1362,7 @@ void AActor::PostMigrate(const FUObjectMigrationContext& MigrationContext)
 			}
 			else
 			{
-				UE_LOG(LogActor, Warning, TEXT("Loaned object %s isn't registered in the NetGUID cache."), *GetName());
+				UE_LOG(LogActor, Verbose, TEXT("Loaned object %s isn't registered in the NetGUID cache."), *GetName());
 			}
 		}
 	}
@@ -1463,7 +1463,7 @@ void AActor::PostMigrate(const FUObjectMigrationContext& MigrationContext)
 			}
 			else
 			{
-				UE_LOG(LogActor, Warning, TEXT("Borrowed object %s is already registered in the NetGUID cache."), *GetName());
+				UE_LOG(LogActor, Verbose, TEXT("Borrowed object %s is already registered in the NetGUID cache."), *GetName());
 			}
 		}
 
