Updated OnSyncBodies to put more of the code in a transaction.

diff --git a/Engine/Source/Runtime/Engine/Private/PhysicsEngine/Experimental/PhysScene_Chaos.cpp b/Engine/Source/Runtime/Engine/Private/PhysicsEngine/Experimental/PhysScene_Chaos.cpp
--- a/Engine/Source/Runtime/Engine/Private/PhysicsEngine/Experimental/PhysScene_Chaos.cpp
+++ b/Engine/Source/Runtime/Engine/Private/PhysicsEngine/Experimental/PhysScene_Chaos.cpp
@@ -2525,36 +2525,38 @@ void FPhysScene_Chaos::OnSyncBodies(Chaos::FPhysicsSolverBase* Solver)
 
 	for (const FPhysScenePendingComponentTransform_Chaos& ComponentTransform : Dispatcher.PendingTransforms)
 	{
-		if (ComponentTransform.OwningComp != nullptr)
+		auto DoWorkCallback = [this, &ComponentTransform]
 		{
-			AActor* OwnerPtr = ComponentTransform.OwningComp->GetOwner();
-
-			if (ComponentTransform.bHasValidTransform)
+			if (ComponentTransform.OwningComp != nullptr)
 			{
-#if UE_WITH_REMOTE_OBJECT_HANDLE
-				static FName TransactionalWorkName(TEXT("FPhysScene_Chaos::OnSyncBodies"));
-				UE::RemoteExecutor::ExecuteTransactional(TransactionalWorkName, [this, &ComponentTransform]
-				{ 
+				AActor* OwnerPtr = ComponentTransform.OwningComp->GetOwner();
+
+				if (ComponentTransform.bHasValidTransform)
+				{
 					ComponentTransform.OwningComp->MoveComponent(ComponentTransform.NewTranslation, ComponentTransform.NewRotation, false, NULL, MOVECOMP_SkipPhysicsMove);
-				});
-#else
-				ComponentTransform.OwningComp->MoveComponent(ComponentTransform.NewTranslation, ComponentTransform.NewRotation, false, NULL, MOVECOMP_SkipPhysicsMove);
-#endif
-			}
+				}
 
-			if (IsValid(OwnerPtr))
-			{
-				OwnerPtr->CheckStillInWorld();
+				if (IsValid(OwnerPtr))
+				{
+					OwnerPtr->CheckStillInWorld();
+				}
 			}
-		}
 
-		if (ComponentTransform.OwningComp != nullptr)
-		{
-			if (ComponentTransform.WakeEvent != Chaos::EWakeEventEntry::None)
+			if (ComponentTransform.OwningComp != nullptr)
 			{
-				ComponentTransform.OwningComp->DispatchWakeEvents(ComponentTransform.WakeEvent == EWakeEventEntry::Awake ? ESleepEvent::SET_Wakeup : ESleepEvent::SET_Sleep, NAME_None);
+				if (ComponentTransform.WakeEvent != Chaos::EWakeEventEntry::None)
+				{
+					ComponentTransform.OwningComp->DispatchWakeEvents(ComponentTransform.WakeEvent == EWakeEventEntry::Awake ? ESleepEvent::SET_Wakeup : ESleepEvent::SET_Sleep, NAME_None);
+				}
 			}
-		}
+		};
+
+#if UE_WITH_REMOTE_OBJECT_HANDLE
+		static FName TransactionalWorkName(TEXT("FPhysScene_Chaos::OnSyncBodies"));
+		UE::RemoteExecutor::ExecuteTransactional(TransactionalWorkName, DoWorkCallback);
+#else
+		DoWorkCallback();
+#endif
 	}
 }
 
