Allow all LOD levels for FSkeletalMeshSceneProxy to be rendered

Clamping has side effects that are not well understood at the moment.

This undoes the attempt to help PSO precaching deal with bone count mismatches.

diff --git a/Engine/Source/Runtime/Engine/Private/SkeletalMeshSceneProxy.cpp b/Engine/Source/Runtime/Engine/Private/SkeletalMeshSceneProxy.cpp
--- a/Engine/Source/Runtime/Engine/Private/SkeletalMeshSceneProxy.cpp
+++ b/Engine/Source/Runtime/Engine/Private/SkeletalMeshSceneProxy.cpp
@@ -488,7 +488,8 @@ void FSkeletalMeshSceneProxy::DrawStaticElements(FStaticPrimitiveDrawInterface*
 
 		int32 NumLODs = SkeletalMeshRenderData->LODRenderData.Num();
 
-		for (int32 LODIndex = ClampedLODIndex; LODIndex < NumLODs; ++LODIndex)
+		// TODO: Use ClampedLODIndex but without meshes disappearing
+		for (int32 LODIndex = 0; LODIndex < NumLODs; ++LODIndex)
 		{
 			const FSkeletalMeshLODRenderData& LODData = SkeletalMeshRenderData->LODRenderData[LODIndex];
 
