[materials] optimize material instance postload by directly hashing the portions of shadermap ID used to construct the unique asset name, rather than constructing the full string representation then hashing.

diff --git a/Engine/Source/Runtime/Engine/Private/Materials/MaterialShared.cpp b/Engine/Source/Runtime/Engine/Private/Materials/MaterialShared.cpp
--- a/Engine/Source/Runtime/Engine/Private/Materials/MaterialShared.cpp
+++ b/Engine/Source/Runtime/Engine/Private/Materials/MaterialShared.cpp
@@ -3040,17 +3040,18 @@ bool FMaterial::CacheShaders(const FMaterialShaderMapId& ShaderMapId, EShaderPla
 #if WITH_EDITOR
 FString FMaterial::GetUniqueAssetName(EShaderPlatform Platform, const FMaterialShaderMapId& ShaderMapId) const
 {
-	FString IdKeyStr;
-	// append the portion of the DDC key string associated with the shadermapid, but exclude the source code
-	// hashes and material function/parameter collection guids, such that this name remains stable when edits
-	// to this data are applied (including source version bumps)
-	ShaderMapId.AppendKeyString(IdKeyStr, /* bIncludeSourceAndMaterialState */ false);
-	// append the base material path as well to differentiate materials with the same name and different paths.
+	FXxHash64Builder Hasher;
+	// Construct a hash representing the material shadermap ID key, but excluding source code hashes and material function/parameter 
+	// collection guids, such that this hash remains stable when edits to this data are applied (including source version bumps)
+	FMaterialKeyGeneratorContext KeyGenCtx([&Hasher](const void* Data, uint64 Size) { Hasher.Update(Data, Size); }, Platform);
+	KeyGenCtx.RemoveFlags(EMaterialKeyInclude::SourceAndMaterialState | EMaterialKeyInclude::Globals | EMaterialKeyInclude::ShaderDependencies);
+	const_cast<FMaterialShaderMapId&>(ShaderMapId).RecordAndEmit(KeyGenCtx);
+	// Hash the base material path as well to differentiate materials with the same name and different paths.
 	// note we explicitly _do not_ use the path of the asset itself as if the asset is a material instance we 
 	// want it to properly deduplicate against other instances which might end up pointing to the same shadermap.
-	IdKeyStr.Append(GetBaseMaterialPathName());
-	uint64 Hash = CityHash64((const char*)*IdKeyStr, IdKeyStr.Len() * sizeof(TCHAR));
-	return FString::Printf(TEXT("%s_%llx"), *GetFriendlyName(), Hash);
+	FString BaseMaterialPath = GetBaseMaterialPathName();
+	Hasher.Update(BaseMaterialPath.GetCharArray().GetData(), BaseMaterialPath.Len() * sizeof(TCHAR));
+	return FString::Printf(TEXT("%s_%llx"), *GetFriendlyName(), Hasher.Finalize().Hash);
 }
 #endif // WITH_EDITOR
 
