Removing misleading log line "Disallowed to write first packet in batch, with Iris this is not good!" that causes confusion

diff --git a/Engine/Source/Runtime/Engine/Private/Net/Experimental/Iris/DataStreamChannel.cpp b/Engine/Source/Runtime/Engine/Private/Net/Experimental/Iris/DataStreamChannel.cpp
--- a/Engine/Source/Runtime/Engine/Private/Net/Experimental/Iris/DataStreamChannel.cpp
+++ b/Engine/Source/Runtime/Engine/Private/Net/Experimental/Iris/DataStreamChannel.cpp
@@ -244,7 +244,6 @@ void UDataStreamChannel::Tick()
 	// We probably want separate bandwidth management for iris as we are not pre-filling sendbuffer before call to NetReady.
 	if (!IsNetReady())
 	{
-		UE_CLOG(bHandshakeComplete, LogIris, Log, TEXT("Disallowed to write first packet in batch, with Iris this is not good!"))
 		return;
 	}
 
