Replace some usages of FORCEINLINE with inline in Niagara modules.

diff --git a/Engine/Plugins/FX/Niagara/Source/NiagaraEditor/Public/ViewModels/Stack/NiagaraStackGraphUtilities.h b/Engine/Plugins/FX/Niagara/Source/NiagaraEditor/Public/ViewModels/Stack/NiagaraStackGraphUtilities.h
--- a/Engine/Plugins/FX/Niagara/Source/NiagaraEditor/Public/ViewModels/Stack/NiagaraStackGraphUtilities.h
+++ b/Engine/Plugins/FX/Niagara/Source/NiagaraEditor/Public/ViewModels/Stack/NiagaraStackGraphUtilities.h
@@ -176,7 +176,7 @@ namespace FNiagaraStackGraphUtilities
 		}
 	};
 
-	FORCEINLINE uint32 GetTypeHash(const FInputDataCacheKey& InputDataCacheKey)
+	inline uint32 GetTypeHash(const FInputDataCacheKey& InputDataCacheKey)
 	{
 		uint32 Hash = 0;
 		Hash = HashCombine(Hash, GetTypeHash(InputDataCacheKey.NodeGuid));
diff --git a/Engine/Plugins/FX/Niagara/Source/NiagaraShader/Public/NiagaraShared.h b/Engine/Plugins/FX/Niagara/Source/NiagaraShader/Public/NiagaraShared.h
--- a/Engine/Plugins/FX/Niagara/Source/NiagaraShader/Public/NiagaraShared.h
+++ b/Engine/Plugins/FX/Niagara/Source/NiagaraShader/Public/NiagaraShared.h
@@ -892,10 +892,10 @@ public:
 
 	NIAGARASHADER_API bool IsShaderMapComplete() const;
 
-	FORCEINLINE bool IsShaderMapComplete_RenderThread() const { check(IsInRenderingThread()); return CachedData_RenderThread.bIsComplete != 0; }
-	FORCEINLINE int32 GetNumPermutations_RenderThread() const { check(IsInRenderingThread()); return CachedData_RenderThread.NumPermutations; }
-	FORCEINLINE bool IsExternalConstantBufferUsed_RenderThread(int32 Index) const { return (CachedData_RenderThread.bExternalConstantBufferUsed & (1 << Index)) != 0; }
-	FORCEINLINE bool IsViewUniformBufferUsed_RenderThread() const { return CachedData_RenderThread.bViewUniformBufferUsed != 0; }
+	inline bool IsShaderMapComplete_RenderThread() const { check(IsInRenderingThread()); return CachedData_RenderThread.bIsComplete != 0; }
+	inline int32 GetNumPermutations_RenderThread() const { check(IsInRenderingThread()); return CachedData_RenderThread.NumPermutations; }
+	inline bool IsExternalConstantBufferUsed_RenderThread(int32 Index) const { return (CachedData_RenderThread.bExternalConstantBufferUsed & (1 << Index)) != 0; }
+	inline bool IsViewUniformBufferUsed_RenderThread() const { return CachedData_RenderThread.bViewUniformBufferUsed != 0; }
 
 	FNiagaraCompileHash GetBaseCompileHash()const { return BaseCompileHash; }
 
diff --git a/Engine/Plugins/FX/Niagara/Source/NiagaraVertexFactories/Public/NiagaraMeshVertexFactory.h b/Engine/Plugins/FX/Niagara/Source/NiagaraVertexFactories/Public/NiagaraMeshVertexFactory.h
--- a/Engine/Plugins/FX/Niagara/Source/NiagaraVertexFactories/Public/NiagaraMeshVertexFactory.h
+++ b/Engine/Plugins/FX/Niagara/Source/NiagaraVertexFactories/Public/NiagaraMeshVertexFactory.h
@@ -163,7 +163,7 @@ public:
 	/**
 	* Set the uniform buffer for this vertex factory.
 	*/
-	FORCEINLINE void SetUniformBuffer(const FNiagaraMeshUniformBufferRef& InMeshParticleUniformBuffer)
+	inline void SetUniformBuffer(const FNiagaraMeshUniformBufferRef& InMeshParticleUniformBuffer)
 	{
 		MeshParticleUniformBuffer = InMeshParticleUniformBuffer;
 	}
@@ -171,32 +171,32 @@ public:
 	/**
 	* Retrieve the uniform buffer for this vertex factory.
 	*/
-	FORCEINLINE FRHIUniformBuffer* GetUniformBuffer() const
+	inline FRHIUniformBuffer* GetUniformBuffer() const
 	{
 		return MeshParticleUniformBuffer;
 	}
 
-	FORCEINLINE FRHIShaderResourceView* GetTangentsSRV() const
+	inline FRHIShaderResourceView* GetTangentsSRV() const
 	{
 		return Data.TangentsSRV;
 	}
 
-	FORCEINLINE FRHIShaderResourceView* GetTextureCoordinatesSRV() const
+	inline FRHIShaderResourceView* GetTextureCoordinatesSRV() const
 	{
 		return Data.TextureCoordinatesSRV;
 	}
 
-	FORCEINLINE FRHIShaderResourceView* GetColorComponentsSRV() const
+	inline FRHIShaderResourceView* GetColorComponentsSRV() const
 	{
 		return Data.ColorComponentsSRV;
 	}
 
-	FORCEINLINE uint32 GetColorIndexMask() const
+	inline uint32 GetColorIndexMask() const
 	{
 		return Data.ColorIndexMask;
 	}
 
-	FORCEINLINE int GetNumTexcoords() const
+	inline int GetNumTexcoords() const
 	{
 		return Data.NumTexCoords;
 	}
diff --git a/Engine/Plugins/FX/Niagara/Source/NiagaraVertexFactories/Public/NiagaraRibbonVertexFactory.h b/Engine/Plugins/FX/Niagara/Source/NiagaraVertexFactories/Public/NiagaraRibbonVertexFactory.h
--- a/Engine/Plugins/FX/Niagara/Source/NiagaraVertexFactories/Public/NiagaraRibbonVertexFactory.h
+++ b/Engine/Plugins/FX/Niagara/Source/NiagaraVertexFactories/Public/NiagaraRibbonVertexFactory.h
@@ -130,7 +130,7 @@ public:
 	/**
 	* Set the uniform buffer for this vertex factory.
 	*/
-	FORCEINLINE void SetRibbonUniformBuffer(FNiagaraRibbonUniformBufferRef InSpriteUniformBuffer)
+	inline void SetRibbonUniformBuffer(FNiagaraRibbonUniformBufferRef InSpriteUniformBuffer)
 	{
 		NiagaraRibbonUniformBuffer = InSpriteUniformBuffer;
 	}
@@ -138,7 +138,7 @@ public:
 	/**
 	* Retrieve the uniform buffer for this vertex factory.
 	*/
-	FORCEINLINE FNiagaraRibbonUniformBufferRef GetRibbonUniformBuffer() const
+	inline FNiagaraRibbonUniformBufferRef GetRibbonUniformBuffer() const
 	{
 		return NiagaraRibbonUniformBuffer;
 	}
diff --git a/Engine/Source/Runtime/Engine/Classes/Particles/FXBudget.h b/Engine/Source/Runtime/Engine/Classes/Particles/FXBudget.h
--- a/Engine/Source/Runtime/Engine/Classes/Particles/FXBudget.h
+++ b/Engine/Source/Runtime/Engine/Classes/Particles/FXBudget.h
@@ -38,15 +38,15 @@ public:
 	 **/
 	static ENGINE_API FFXTimeData GetAdjustedUsage();
 	/** Returns the highest single adjusted usage value. */
-	FORCEINLINE static float GetWorstAdjustedUsage() { return WorstAdjustedUsage; }
-	FORCEINLINE static void SetWorstAdjustedUsage(float NewAdjustedUsage){ WorstAdjustedUsage = NewAdjustedUsage; }
+	inline static float GetWorstAdjustedUsage() { return WorstAdjustedUsage; }
+	inline static void SetWorstAdjustedUsage(float NewAdjustedUsage){ WorstAdjustedUsage = NewAdjustedUsage; }
 
 	static ENGINE_API void Reset();
 
 	static ENGINE_API TSharedPtr<FParticlePerfStatsListener_FXBudget, ESPMode::ThreadSafe> StatsListener;
 
 	static ENGINE_API void OnEnabledCVarChanged(IConsoleVariable* CVar);
-	FORCEINLINE static bool Enabled(){ return bEnabled; }
+	inline static bool Enabled(){ return bEnabled; }
 	static ENGINE_API void SetEnabled(bool bInEnabled);
 
 	static ENGINE_API bool bEnabled;
@@ -61,14 +61,14 @@ private:
 class FFXBudget
 {
 public:
-	FORCEINLINE static FFXTimeData GetTime(){ return FFXTimeData(); }
-	FORCEINLINE static FFXTimeData GetBudget() { return FFXTimeData(); }
-	FORCEINLINE static FFXTimeData GetUsage() { return FFXTimeData(); }
-	FORCEINLINE static FFXTimeData GetAdjustedUsage() { return FFXTimeData(); }
-	FORCEINLINE static float GetWorstAdjustedUsage() { return 0.0f; }
-	FORCEINLINE static bool Enabled() { return false; }
-	FORCEINLINE static void SetEnabled(bool bInEnabled) { }
+	inline static FFXTimeData GetTime(){ return FFXTimeData(); }
+	inline static FFXTimeData GetBudget() { return FFXTimeData(); }
+	inline static FFXTimeData GetUsage() { return FFXTimeData(); }
+	inline static FFXTimeData GetAdjustedUsage() { return FFXTimeData(); }
+	inline static float GetWorstAdjustedUsage() { return 0.0f; }
+	inline static bool Enabled() { return false; }
+	inline static void SetEnabled(bool bInEnabled) { }
 
-	FORCEINLINE static void Reset(){}
+	inline static void Reset(){}
 };
 #endif
diff --git a/Engine/Source/Runtime/Engine/Classes/Particles/Location/ParticleModuleLocationBoneSocket.h b/Engine/Source/Runtime/Engine/Classes/Particles/Location/ParticleModuleLocationBoneSocket.h
--- a/Engine/Source/Runtime/Engine/Classes/Particles/Location/ParticleModuleLocationBoneSocket.h
+++ b/Engine/Source/Runtime/Engine/Classes/Particles/Location/ParticleModuleLocationBoneSocket.h
@@ -185,7 +185,7 @@ class UParticleModuleLocationBoneSocket : public UParticleModuleLocationBase
 	ENGINE_API void UpdatePrevBoneLocationsAndVelocities(FModuleLocationBoneSocketInstancePayload* InstancePayload, USkeletalMeshComponent* SourceComponent, float DeltaTime);
 
 	//If we're updating our position each frame, there's no point in inheriting bone velocity
-	FORCEINLINE bool InheritingBoneVelocity()const { return bInheritBoneVelocity && !bUpdatePositionEachFrame; }
+	inline bool InheritingBoneVelocity()const { return bInheritBoneVelocity && !bUpdatePositionEachFrame; }
 
 	ENGINE_API void SetSourceIndexMode();
 
diff --git a/Engine/Source/Runtime/Engine/Classes/Particles/Location/ParticleModuleLocationSkelVertSurface.h b/Engine/Source/Runtime/Engine/Classes/Particles/Location/ParticleModuleLocationSkelVertSurface.h
--- a/Engine/Source/Runtime/Engine/Classes/Particles/Location/ParticleModuleLocationSkelVertSurface.h
+++ b/Engine/Source/Runtime/Engine/Classes/Particles/Location/ParticleModuleLocationSkelVertSurface.h
@@ -60,7 +60,7 @@ struct FModuleLocationVertSurfaceInstancePayload
 		BoneVelocities = TPreallocatedArrayProxy<FVector3f>((uint8*)this + StructOffset, FixedArraySize);
 	}
 
-	FORCEINLINE bool MeshIsValid() { return SourceComponent.IsValid() && MeshMinLOD != INDEX_NONE; }
+	inline bool MeshIsValid() { return SourceComponent.IsValid() && MeshMinLOD != INDEX_NONE; }
 };
 
 /** ModuleLocationVertSurface per-particle payload - only used if updating each frame */
diff --git a/Engine/Source/Runtime/Engine/Classes/Particles/Parameter/ParticleModuleParameterDynamic.h b/Engine/Source/Runtime/Engine/Classes/Particles/Parameter/ParticleModuleParameterDynamic.h
--- a/Engine/Source/Runtime/Engine/Classes/Particles/Parameter/ParticleModuleParameterDynamic.h
+++ b/Engine/Source/Runtime/Engine/Classes/Particles/Parameter/ParticleModuleParameterDynamic.h
@@ -167,7 +167,7 @@ class UParticleModuleParameterDynamic : public UParticleModuleParameterBase
 	 *
 	 *	@return	float			The value for the parameter.
 	 */
-	FORCEINLINE float GetParameterValue(const FContext& Context, FEmitterDynamicParameter& InDynParams, FBaseParticle& Particle, FParticleEmitterInstance* Owner, struct FRandomStream* InRandomStream)
+	inline float GetParameterValue(const FContext& Context, FEmitterDynamicParameter& InDynParams, FBaseParticle& Particle, FParticleEmitterInstance* Owner, struct FRandomStream* InRandomStream)
 	{
 		float ScaleValue = 1.0f;
 		float DistributionValue = 1.0f;
@@ -210,7 +210,7 @@ class UParticleModuleParameterDynamic : public UParticleModuleParameterBase
 	 *
 	 *	@return	float			The value for the parameter.
 	 */
-	FORCEINLINE float GetParameterValue_UserSet(const FContext& Context, FEmitterDynamicParameter& InDynParams, FBaseParticle& Particle, FParticleEmitterInstance* Owner, struct FRandomStream* InRandomStream)
+	inline float GetParameterValue_UserSet(const FContext& Context, FEmitterDynamicParameter& InDynParams, FBaseParticle& Particle, FParticleEmitterInstance* Owner, struct FRandomStream* InRandomStream)
 	{
 		return InDynParams.ParamValue.GetValue(InDynParams.bUseEmitterTime ? Owner->EmitterTime : Particle.RelativeTime, Context.GetDistributionData(), InRandomStream);
 	}
diff --git a/Engine/Source/Runtime/Engine/Classes/Particles/ParticleEmitter.h b/Engine/Source/Runtime/Engine/Classes/Particles/ParticleEmitter.h
--- a/Engine/Source/Runtime/Engine/Classes/Particles/ParticleEmitter.h
+++ b/Engine/Source/Runtime/Engine/Classes/Particles/ParticleEmitter.h
@@ -392,7 +392,7 @@ class UParticleEmitter : public UObject
 	* Returns the stat ID of the object...
 	* We can't use the normal version of this because those names are meaningless; we need the special name in the emitter
 	**/
-	FORCEINLINE TStatId GetStatID(bool bForDeferredUse = false) const
+	inline TStatId GetStatID(bool bForDeferredUse = false) const
 	{
 #if STATS
 		// this is done to avoid even registering stats for a disabled group (unless we plan on using it later)
@@ -408,7 +408,7 @@ class UParticleEmitter : public UObject
 		return TStatId(); // not doing stats at the moment, or ever
 	}
 
-	FORCEINLINE TStatId GetStatIDRT(bool bForDeferredUse = false) const
+	inline TStatId GetStatIDRT(bool bForDeferredUse = false) const
 	{
 #if STATS
 		// this is done to avoid even registering stats for a disabled group (unless we plan on using it later)
diff --git a/Engine/Source/Runtime/Engine/Classes/Particles/ParticleModule.h b/Engine/Source/Runtime/Engine/Classes/Particles/ParticleModule.h
--- a/Engine/Source/Runtime/Engine/Classes/Particles/ParticleModule.h
+++ b/Engine/Source/Runtime/Engine/Classes/Particles/ParticleModule.h
@@ -137,7 +137,7 @@ struct FParticleRandomSeedInfo
 		, bRandomlySelectSeedArray(false)
 		{
 		}
-		FORCEINLINE int32 GetInstancePayloadSize() const
+		inline int32 GetInstancePayloadSize() const
 		{
 			return ((RandomSeeds.Num() > 0) ? sizeof(FParticleRandomSeedInstancePayload) : 0);
 		}
diff --git a/Engine/Source/Runtime/Engine/Classes/Particles/ParticleModuleRequired.h b/Engine/Source/Runtime/Engine/Classes/Particles/ParticleModuleRequired.h
--- a/Engine/Source/Runtime/Engine/Classes/Particles/ParticleModuleRequired.h
+++ b/Engine/Source/Runtime/Engine/Classes/Particles/ParticleModuleRequired.h
@@ -37,7 +37,7 @@ enum class EParticleUVFlipMode : uint8
 };
 
 /** Flips the sign of a particle's base size based on it's UV flip mode. */
-FORCEINLINE void AdjustParticleBaseSizeForUVFlipping(FVector& OutSize, EParticleUVFlipMode FlipMode, FRandomStream& InRandomStream)
+inline void AdjustParticleBaseSizeForUVFlipping(FVector& OutSize, EParticleUVFlipMode FlipMode, FRandomStream& InRandomStream)
 {
 	static const float HalfRandMax = 0.5f;
 
diff --git a/Engine/Source/Runtime/Engine/Classes/Particles/ParticlePerfStats.h b/Engine/Source/Runtime/Engine/Classes/Particles/ParticlePerfStats.h
--- a/Engine/Source/Runtime/Engine/Classes/Particles/ParticlePerfStats.h
+++ b/Engine/Source/Runtime/Engine/Classes/Particles/ParticlePerfStats.h
@@ -99,7 +99,7 @@ struct FParticlePerfStats_GT
 		return *this;
 	}
 
-	FORCEINLINE void Reset()
+	inline void Reset()
 	{
 		NumInstances = 0;
 		TickGameThreadCycles = 0;
@@ -110,11 +110,11 @@ struct FParticlePerfStats_GT
 		WaitCycles = 0;
 	}
 
-	FORCEINLINE uint64 GetTotalCycles_GTOnly() const { return TickGameThreadCycles + FinalizeCycles + ActivationCycles + WaitCycles; }
-	FORCEINLINE uint64 GetPerInstanceAvgCycles_GTOnly() const { return NumInstances > 0 ? GetTotalCycles_GTOnly() / NumInstances : 0; }
+	inline uint64 GetTotalCycles_GTOnly() const { return TickGameThreadCycles + FinalizeCycles + ActivationCycles + WaitCycles; }
+	inline uint64 GetPerInstanceAvgCycles_GTOnly() const { return NumInstances > 0 ? GetTotalCycles_GTOnly() / NumInstances : 0; }
 
-	FORCEINLINE uint64 GetTotalCycles() const { return GetTotalCycles_GTOnly() + TickConcurrentCycles + EndOfFrameCycles; }
-	FORCEINLINE uint64 GetPerInstanceAvgCycles() const { return NumInstances > 0 ? GetTotalCycles() / NumInstances : 0; }
+	inline uint64 GetTotalCycles() const { return GetTotalCycles_GTOnly() + TickConcurrentCycles + EndOfFrameCycles; }
+	inline uint64 GetPerInstanceAvgCycles() const { return NumInstances > 0 ? GetTotalCycles() / NumInstances : 0; }
 };
 
 /** Stats gathered on the render thread. */
@@ -125,14 +125,14 @@ struct FParticlePerfStats_RT
 	uint64 GetDynamicMeshElementsCycles = 0;
 	
 	FParticlePerfStats_RT()	{ Reset();	}
-	FORCEINLINE void Reset()
+	inline void Reset()
 	{
 		NumInstances = 0;
 		RenderUpdateCycles = 0;
 		GetDynamicMeshElementsCycles = 0;
 	}
-	FORCEINLINE uint64 GetTotalCycles() const { return RenderUpdateCycles + GetDynamicMeshElementsCycles; }
-	FORCEINLINE uint64 GetPerInstanceAvgCycles() const { return NumInstances > 0 ? (RenderUpdateCycles + GetDynamicMeshElementsCycles) / NumInstances : 0; }
+	inline uint64 GetTotalCycles() const { return RenderUpdateCycles + GetDynamicMeshElementsCycles; }
+	inline uint64 GetPerInstanceAvgCycles() const { return NumInstances > 0 ? (RenderUpdateCycles + GetDynamicMeshElementsCycles) / NumInstances : 0; }
 
 	FParticlePerfStats_RT& operator+=(FParticlePerfStats_RT& Other)
 	{
@@ -149,11 +149,11 @@ struct FParticlePerfStats_GPU
 	uint64 NumInstances = 0;
 	uint64 TotalMicroseconds = 0;
 
-	FORCEINLINE uint64 GetTotalMicroseconds() const { return TotalMicroseconds; }
-	FORCEINLINE uint64 GetPerInstanceAvgMicroseconds() const { return NumInstances > 0 ? GetTotalMicroseconds() / NumInstances : 0; }
+	inline uint64 GetTotalMicroseconds() const { return TotalMicroseconds; }
+	inline uint64 GetPerInstanceAvgMicroseconds() const { return NumInstances > 0 ? GetTotalMicroseconds() / NumInstances : 0; }
 
 	FParticlePerfStats_GPU() { Reset(); }
-	FORCEINLINE void Reset()
+	inline void Reset()
 	{
 		NumInstances = 0;
 		TotalMicroseconds = 0;
@@ -177,12 +177,12 @@ struct FParticlePerfStats
 	ENGINE_API void Tick();
 	ENGINE_API void TickRT();
 
-	FORCEINLINE static bool GetCSVStatsEnabled() { return bCSVStatsEnabled.Load(EMemoryOrder::Relaxed); }
-	FORCEINLINE static bool GetStatsEnabled() { return bStatsEnabled.Load(EMemoryOrder::Relaxed); }
-	FORCEINLINE static bool GetGatherWorldStats() { return WorldStatsReaders.Load(EMemoryOrder::Relaxed) > 0; }
-	FORCEINLINE static bool GetGatherSystemStats() { return SystemStatsReaders.Load(EMemoryOrder::Relaxed) > 0; }
-	FORCEINLINE static bool GetGatherComponentStats() { return ComponentStatsReaders.Load(EMemoryOrder::Relaxed) > 0; }
-	FORCEINLINE static bool ShouldGatherStats() 
+	inline static bool GetCSVStatsEnabled() { return bCSVStatsEnabled.Load(EMemoryOrder::Relaxed); }
+	inline static bool GetStatsEnabled() { return bStatsEnabled.Load(EMemoryOrder::Relaxed); }
+	inline static bool GetGatherWorldStats() { return WorldStatsReaders.Load(EMemoryOrder::Relaxed) > 0; }
+	inline static bool GetGatherSystemStats() { return SystemStatsReaders.Load(EMemoryOrder::Relaxed) > 0; }
+	inline static bool GetGatherComponentStats() { return ComponentStatsReaders.Load(EMemoryOrder::Relaxed) > 0; }
+	inline static bool ShouldGatherStats() 
 	{
 		return GetStatsEnabled() && 
 			(GetGatherWorldStats() 
@@ -196,16 +196,16 @@ struct FParticlePerfStats
 	}
 
 
-	FORCEINLINE static void SetCSVStatsEnabled(bool bEnabled) { bCSVStatsEnabled.Store(bEnabled); }
-	FORCEINLINE static void SetStatsEnabled(bool bEnabled) { bStatsEnabled.Store(bEnabled); }
-	FORCEINLINE static void AddWorldStatReader() { ++WorldStatsReaders; }
-	FORCEINLINE static void RemoveWorldStatReader() { --WorldStatsReaders; }
-	FORCEINLINE static void AddSystemStatReader() { ++SystemStatsReaders; }
-	FORCEINLINE static void RemoveSystemStatReader() { --SystemStatsReaders; }
-	FORCEINLINE static void AddComponentStatReader() { ++ComponentStatsReaders; }
-	FORCEINLINE static void RemoveComponentStatReader() { --ComponentStatsReaders; }
+	inline static void SetCSVStatsEnabled(bool bEnabled) { bCSVStatsEnabled.Store(bEnabled); }
+	inline static void SetStatsEnabled(bool bEnabled) { bStatsEnabled.Store(bEnabled); }
+	inline static void AddWorldStatReader() { ++WorldStatsReaders; }
+	inline static void RemoveWorldStatReader() { --WorldStatsReaders; }
+	inline static void AddSystemStatReader() { ++SystemStatsReaders; }
+	inline static void RemoveSystemStatReader() { --SystemStatsReaders; }
+	inline static void AddComponentStatReader() { ++ComponentStatsReaders; }
+	inline static void RemoveComponentStatReader() { --ComponentStatsReaders; }
 
-	static FORCEINLINE FParticlePerfStats* GetStats(const UWorld* World)
+	static inline FParticlePerfStats* GetStats(const UWorld* World)
 	{
 		if (World && GetGatherWorldStats() && GetStatsEnabled())
 		{
@@ -214,7 +214,7 @@ struct FParticlePerfStats
 		return nullptr;
 	}
 	
-	static FORCEINLINE FParticlePerfStats* GetStats(const UFXSystemAsset* System)
+	static inline FParticlePerfStats* GetStats(const UFXSystemAsset* System)
 	{
 	#if WITH_PER_SYSTEM_PARTICLE_PERF_STATS
 		if (System && GetGatherSystemStats() && GetStatsEnabled())
@@ -225,7 +225,7 @@ struct FParticlePerfStats
 		return nullptr;
 	}
 
-	static FORCEINLINE FParticlePerfStats* GetStats(const UFXSystemComponent* Component)
+	static inline FParticlePerfStats* GetStats(const UFXSystemComponent* Component)
 	{
 	#if WITH_PER_COMPONENT_PARTICLE_PERF_STATS
 		if (Component && GetGatherComponentStats() && GetStatsEnabled())
@@ -253,19 +253,19 @@ struct FParticlePerfStats
 	FParticlePerfStats_GPU GPUStats;
 
 	/** Returns the current frame Game Thread stats. */
-	FORCEINLINE FParticlePerfStats_GT& GetGameThreadStats()
+	inline FParticlePerfStats_GT& GetGameThreadStats()
 	{
 		return GameThreadStats; 
 	}
 
 	/** Returns the current frame Render Thread stats. */
-	FORCEINLINE FParticlePerfStats_RT& GetRenderThreadStats()
+	inline FParticlePerfStats_RT& GetRenderThreadStats()
 	{
 		return RenderThreadStats;
 	}
 
 	/** Returns the current frame GPU stats. */
-	FORCEINLINE FParticlePerfStats_GPU& GetGPUStats()
+	inline FParticlePerfStats_GPU& GetGPUStats()
 	{
 		return GPUStats;
 	}
@@ -310,67 +310,67 @@ struct FParticlePerfStatsContext
 #endif
 	{}
 
-	FORCEINLINE FParticlePerfStatsContext(FParticlePerfStats* InWorldStats, FParticlePerfStats* InSystemStats, FParticlePerfStats* InComponentStats)
+	inline FParticlePerfStatsContext(FParticlePerfStats* InWorldStats, FParticlePerfStats* InSystemStats, FParticlePerfStats* InComponentStats)
 	{
 		SetWorldStats(InWorldStats);
 		SetSystemStats(InSystemStats);
 		SetComponentStats(InComponentStats);
 	}
 
-	FORCEINLINE FParticlePerfStatsContext(FParticlePerfStats* InWorldStats, FParticlePerfStats* InSystemStats)
+	inline FParticlePerfStatsContext(FParticlePerfStats* InWorldStats, FParticlePerfStats* InSystemStats)
 	{
 		SetWorldStats(InWorldStats);
 		SetSystemStats(InSystemStats);
 	}
 
-	FORCEINLINE FParticlePerfStatsContext(FParticlePerfStats* InComponentStats)
+	inline FParticlePerfStatsContext(FParticlePerfStats* InComponentStats)
 	{
 		SetComponentStats(InComponentStats);
 	}
 
-	FORCEINLINE FParticlePerfStatsContext(const UWorld* InWorld, const UFXSystemAsset* InSystem, const UFXSystemComponent* InComponent)
+	inline FParticlePerfStatsContext(const UWorld* InWorld, const UFXSystemAsset* InSystem, const UFXSystemComponent* InComponent)
 	{
 		SetWorldStats(FParticlePerfStats::GetStats(InWorld));
 		SetSystemStats(FParticlePerfStats::GetStats(InSystem));
 		SetComponentStats(FParticlePerfStats::GetStats(InComponent));
 	}
 
-	FORCEINLINE FParticlePerfStatsContext(const UWorld* InWorld, const UFXSystemAsset* InSystem)
+	inline FParticlePerfStatsContext(const UWorld* InWorld, const UFXSystemAsset* InSystem)
 	{
 		SetWorldStats(FParticlePerfStats::GetStats(InWorld));
 		SetSystemStats(FParticlePerfStats::GetStats(InSystem));
 	}
 
-	FORCEINLINE FParticlePerfStatsContext(const UFXSystemComponent* InComponent)
+	inline FParticlePerfStatsContext(const UFXSystemComponent* InComponent)
 	{
 		SetComponentStats(FParticlePerfStats::GetStats(InComponent));
 	}
 
-	FORCEINLINE bool IsValid()
+	inline bool IsValid()
 	{
 		return GetWorldStats() != nullptr || GetSystemStats() != nullptr || GetComponentStats() != nullptr;
 	}
 
 	FParticlePerfStats* WorldStats = nullptr;
-	FORCEINLINE FParticlePerfStats* GetWorldStats() { return WorldStats; }
-	FORCEINLINE void SetWorldStats(FParticlePerfStats* Stats) { WorldStats = Stats; }
+	inline FParticlePerfStats* GetWorldStats() { return WorldStats; }
+	inline void SetWorldStats(FParticlePerfStats* Stats) { WorldStats = Stats; }
 
 #if WITH_PER_SYSTEM_PARTICLE_PERF_STATS
 	FParticlePerfStats* SystemStats = nullptr;
-	FORCEINLINE FParticlePerfStats* GetSystemStats() { return SystemStats; }
-	FORCEINLINE void SetSystemStats(FParticlePerfStats* Stats) { SystemStats = Stats; }
+	inline FParticlePerfStats* GetSystemStats() { return SystemStats; }
+	inline void SetSystemStats(FParticlePerfStats* Stats) { SystemStats = Stats; }
 #else
-	FORCEINLINE FParticlePerfStats* GetSystemStats() { return nullptr; }
-	FORCEINLINE void SetSystemStats(FParticlePerfStats* Stats) { }
+	inline FParticlePerfStats* GetSystemStats() { return nullptr; }
+	inline void SetSystemStats(FParticlePerfStats* Stats) { }
 #endif
 
 #if WITH_PER_COMPONENT_PARTICLE_PERF_STATS
 	FParticlePerfStats* ComponentStats = nullptr;
-	FORCEINLINE FParticlePerfStats* GetComponentStats() { return ComponentStats; }
-	FORCEINLINE void SetComponentStats(FParticlePerfStats* Stats) { ComponentStats = Stats; }
+	inline FParticlePerfStats* GetComponentStats() { return ComponentStats; }
+	inline void SetComponentStats(FParticlePerfStats* Stats) { ComponentStats = Stats; }
 #else
-	FORCEINLINE FParticlePerfStats* GetComponentStats() { return nullptr; }
-	FORCEINLINE void SetComponentStats(FParticlePerfStats* Stats) { }
+	inline FParticlePerfStats* GetComponentStats() { return nullptr; }
+	inline void SetComponentStats(FParticlePerfStats* Stats) { }
 #endif
 };
 
@@ -379,7 +379,7 @@ typedef TFunction<void(FParticlePerfStats* Stats, uint64 Cycles)> FParticlePerfS
 template<typename TWriterFunc>
 struct FParticlePerfStatScope
 {
-	FORCEINLINE FParticlePerfStatScope(FParticlePerfStatsContext InContext, int32 InCount=0)
+	inline FParticlePerfStatScope(FParticlePerfStatsContext InContext, int32 InCount=0)
 	: Context(InContext)
 	, StartCycles(INDEX_NONE)
 	, Count(InCount)
@@ -390,7 +390,7 @@ struct FParticlePerfStatScope
 		}
 	}
 
-	FORCEINLINE ~FParticlePerfStatScope()
+	inline ~FParticlePerfStatScope()
 	{
 		if (StartCycles != INDEX_NONE)
 		{
@@ -409,7 +409,7 @@ struct FParticlePerfStatScope
 #define PARTICLE_PERF_STAT_CYCLES_COMMON(CONTEXT, THREAD, NAME)\
 struct FParticlePerfStatsWriterCycles_##THREAD##_##NAME\
 {\
-	FORCEINLINE static void Write(FParticlePerfStats* Stats, uint64 Cycles, int32 Count)\
+	inline static void Write(FParticlePerfStats* Stats, uint64 Cycles, int32 Count)\
 	{\
 		if (Stats){ Stats->Get##THREAD##Stats().NAME##Cycles += Cycles;	}\
 	}\
@@ -419,7 +419,7 @@ FParticlePerfStatScope<FParticlePerfStatsWriterCycles_##THREAD##_##NAME> ANONYMO
 #define PARTICLE_PERF_STAT_CYCLES_WITH_COUNT_COMMON(CONTEXT, THREAD, NAME, COUNT)\
 struct FParticlePerfStatsWriterCyclesAndCount_##THREAD##_##NAME\
 {\
-	FORCEINLINE static void Write(FParticlePerfStats* Stats, uint64 Cycles, int32 Count)\
+	inline static void Write(FParticlePerfStats* Stats, uint64 Cycles, int32 Count)\
 	{\
 		if(Stats)\
 		{\
@@ -446,12 +446,12 @@ FParticlePerfStatScope<FParticlePerfStatsWriterCyclesAndCount_##THREAD##_##NAME>
 
 struct FParticlePerfStatsContext
 {
-	FORCEINLINE FParticlePerfStatsContext(FParticlePerfStats* InWorldStats, FParticlePerfStats* InSystemStats, FParticlePerfStats* InComponentStats){}
-	FORCEINLINE FParticlePerfStatsContext(FParticlePerfStats* InWorldStats, FParticlePerfStats* InSystemStats) {}
-	FORCEINLINE FParticlePerfStatsContext(FParticlePerfStats* InComponentStats) {}
-	FORCEINLINE FParticlePerfStatsContext(UWorld* InWorld, UFXSystemAsset* InSystem, const UFXSystemComponent* InComponent) {}
-	FORCEINLINE FParticlePerfStatsContext(UWorld* InWorld, UFXSystemAsset* InSystem) {}
-	FORCEINLINE FParticlePerfStatsContext(const UFXSystemComponent* InComponent) {}
+	inline FParticlePerfStatsContext(FParticlePerfStats* InWorldStats, FParticlePerfStats* InSystemStats, FParticlePerfStats* InComponentStats){}
+	inline FParticlePerfStatsContext(FParticlePerfStats* InWorldStats, FParticlePerfStats* InSystemStats) {}
+	inline FParticlePerfStatsContext(FParticlePerfStats* InComponentStats) {}
+	inline FParticlePerfStatsContext(UWorld* InWorld, UFXSystemAsset* InSystem, const UFXSystemComponent* InComponent) {}
+	inline FParticlePerfStatsContext(UWorld* InWorld, UFXSystemAsset* InSystem) {}
+	inline FParticlePerfStatsContext(const UFXSystemComponent* InComponent) {}
 };
 
 #endif //WITH_PARTICLE_PERF_STATS
diff --git a/Engine/Source/Runtime/Engine/Classes/Particles/ParticlePerfStatsManager.h b/Engine/Source/Runtime/Engine/Classes/Particles/ParticlePerfStatsManager.h
--- a/Engine/Source/Runtime/Engine/Classes/Particles/ParticlePerfStatsManager.h
+++ b/Engine/Source/Runtime/Engine/Classes/Particles/ParticlePerfStatsManager.h
@@ -28,31 +28,31 @@ struct FAccumulatedParticlePerfStats_GT
 	ENGINE_API void Tick(FParticlePerfStats& Stats);
 
 	/** Returns the total cycles used by all GameThread stats. */
-	FORCEINLINE uint64 GetTotalCycles() const { return AccumulatedStats.GetTotalCycles(); }
-	FORCEINLINE uint64 GetTotalCycles_GTOnly() const { return AccumulatedStats.GetTotalCycles_GTOnly(); }
+	inline uint64 GetTotalCycles() const { return AccumulatedStats.GetTotalCycles(); }
+	inline uint64 GetTotalCycles_GTOnly() const { return AccumulatedStats.GetTotalCycles_GTOnly(); }
 
 	/** Returns the average cycles per frame by all GameThread stats. */
-	FORCEINLINE uint64 GetPerFrameAvgCycles() const { return NumFrames > 0 ? AccumulatedStats.GetTotalCycles() / NumFrames : 0; }
-	FORCEINLINE uint64 GetPerFrameAvgCycles_GTOnly() const { return NumFrames > 0 ? AccumulatedStats.GetTotalCycles_GTOnly() / NumFrames : 0; }
+	inline uint64 GetPerFrameAvgCycles() const { return NumFrames > 0 ? AccumulatedStats.GetTotalCycles() / NumFrames : 0; }
+	inline uint64 GetPerFrameAvgCycles_GTOnly() const { return NumFrames > 0 ? AccumulatedStats.GetTotalCycles_GTOnly() / NumFrames : 0; }
 	/** Returns the max cycles per frame by all GameThread stats. */
-	FORCEINLINE uint64 GetPerFrameMaxCycles(int32 Index = 0) const { return MaxPerFrameTotalCycles[Index]; }
+	inline uint64 GetPerFrameMaxCycles(int32 Index = 0) const { return MaxPerFrameTotalCycles[Index]; }
 
 	/** Returns the average time in µs per frame by all GameThread stats. */
-	FORCEINLINE float GetPerFrameAvg() const { return float(FPlatformTime::ToMilliseconds64(GetPerFrameAvgCycles()) * 1000.0); }
-	FORCEINLINE float GetPerFrameAvg_GTOnly() const { return float(FPlatformTime::ToMilliseconds64(GetPerFrameAvgCycles_GTOnly()) * 1000.0); }
+	inline float GetPerFrameAvg() const { return float(FPlatformTime::ToMilliseconds64(GetPerFrameAvgCycles()) * 1000.0); }
+	inline float GetPerFrameAvg_GTOnly() const { return float(FPlatformTime::ToMilliseconds64(GetPerFrameAvgCycles_GTOnly()) * 1000.0); }
 	/** Returns the max time in µs per frame by all GameThread stats. */
-	FORCEINLINE float GetPerFrameMax(int32 Index = 0) const { return float(FPlatformTime::ToMilliseconds64(GetPerFrameMaxCycles(Index)) * 1000.0); }
+	inline float GetPerFrameMax(int32 Index = 0) const { return float(FPlatformTime::ToMilliseconds64(GetPerFrameMaxCycles(Index)) * 1000.0); }
 
 	/** Returns the average cycles per instance by all GameThread stats. */
-	FORCEINLINE uint64 GetPerInstanceAvgCycles() const { return AccumulatedStats.GetPerInstanceAvgCycles(); }
+	inline uint64 GetPerInstanceAvgCycles() const { return AccumulatedStats.GetPerInstanceAvgCycles(); }
 	/** Returns the max cycles per instance by all GameThread stats. */
-	FORCEINLINE uint64 GetPerInstanceMaxCycles(int32 Index = 0) const { return MaxPerInstanceCycles[Index]; }
+	inline uint64 GetPerInstanceMaxCycles(int32 Index = 0) const { return MaxPerInstanceCycles[Index]; }
 
 	/** Returns the average time in µs per instance by all GameThread stats. */
-	FORCEINLINE float GetPerInstanceAvg() const { return float(FPlatformTime::ToMilliseconds64(GetPerInstanceAvgCycles()) * 1000.0); }
-	FORCEINLINE float GetPerInstanceAvg_GTOnly() const { return float(FPlatformTime::ToMilliseconds64(AccumulatedStats.GetPerInstanceAvgCycles_GTOnly()) * 1000.0); }
+	inline float GetPerInstanceAvg() const { return float(FPlatformTime::ToMilliseconds64(GetPerInstanceAvgCycles()) * 1000.0); }
+	inline float GetPerInstanceAvg_GTOnly() const { return float(FPlatformTime::ToMilliseconds64(AccumulatedStats.GetPerInstanceAvgCycles_GTOnly()) * 1000.0); }
 	/** Returns the max time in µs per instance by all GameThread stats. */
-	FORCEINLINE float GetPerInstanceMax(int32 Index = 0) const { return float(FPlatformTime::ToMilliseconds64(GetPerInstanceMaxCycles(Index)) * 1000.0); }
+	inline float GetPerInstanceMax(int32 Index = 0) const { return float(FPlatformTime::ToMilliseconds64(GetPerInstanceMaxCycles(Index)) * 1000.0); }
 };
 
 struct FAccumulatedParticlePerfStats_RT
@@ -68,27 +68,27 @@ struct FAccumulatedParticlePerfStats_RT
 	ENGINE_API void Tick(FParticlePerfStats& Stats);
 
 	/** Returns the total cycles used by all RenderThread stats. */
-	FORCEINLINE uint64 GetTotalCycles() const { return AccumulatedStats.GetTotalCycles(); }
+	inline uint64 GetTotalCycles() const { return AccumulatedStats.GetTotalCycles(); }
 
 	/** Returns the average cycles per frame by all RenderThread stats. */
-	FORCEINLINE uint64 GetPerFrameAvgCycles() const { return NumFrames > 0 ? AccumulatedStats.GetTotalCycles() / NumFrames : 0; }
+	inline uint64 GetPerFrameAvgCycles() const { return NumFrames > 0 ? AccumulatedStats.GetTotalCycles() / NumFrames : 0; }
 	/** Returns the max cycles per frame by all RenderThread stats. */
-	FORCEINLINE uint64 GetPerFrameMaxCycles(int32 Index = 0) const { return MaxPerFrameTotalCycles[Index]; }
+	inline uint64 GetPerFrameMaxCycles(int32 Index = 0) const { return MaxPerFrameTotalCycles[Index]; }
 
 	/** Returns the average time in µs per frame by all RenderThread stats. */
-	FORCEINLINE float GetPerFrameAvg() const { return float(FPlatformTime::ToMilliseconds64(GetPerFrameAvgCycles()) * 1000.0); }
+	inline float GetPerFrameAvg() const { return float(FPlatformTime::ToMilliseconds64(GetPerFrameAvgCycles()) * 1000.0); }
 	/** Returns the max time in µs per frame by all RenderThread stats. */
-	FORCEINLINE float GetPerFrameMax(int32 Index = 0) const { return float(FPlatformTime::ToMilliseconds64(GetPerFrameMaxCycles(Index)) * 1000.0); }
+	inline float GetPerFrameMax(int32 Index = 0) const { return float(FPlatformTime::ToMilliseconds64(GetPerFrameMaxCycles(Index)) * 1000.0); }
 
 	/** Returns the average cycles per instance by all RenderThread stats. */
-	FORCEINLINE uint64 GetPerInstanceAvgCycles() const { return AccumulatedStats.GetPerInstanceAvgCycles(); }
+	inline uint64 GetPerInstanceAvgCycles() const { return AccumulatedStats.GetPerInstanceAvgCycles(); }
 	/** Returns the max cycles per instance by all RenderThread stats. */
-	FORCEINLINE uint64 GetPerInstanceMaxCycles(int32 Index = 0) const { return MaxPerInstanceCycles[Index];  }
+	inline uint64 GetPerInstanceMaxCycles(int32 Index = 0) const { return MaxPerInstanceCycles[Index];  }
 
 	/** Returns the average time in µs per instance by all RenderThread stats. */
-	FORCEINLINE float GetPerInstanceAvg() const { return float(FPlatformTime::ToMilliseconds64(GetPerInstanceAvgCycles()) * 1000.0); }
+	inline float GetPerInstanceAvg() const { return float(FPlatformTime::ToMilliseconds64(GetPerInstanceAvgCycles()) * 1000.0); }
 	/** Returns the max time in µs per instance by all RenderThread stats. */
-	FORCEINLINE float GetPerInstanceMax(int32 Index = 0) const { return float(FPlatformTime::ToMilliseconds64(GetPerInstanceMaxCycles(Index)) * 1000.0); }
+	inline float GetPerInstanceMax(int32 Index = 0) const { return float(FPlatformTime::ToMilliseconds64(GetPerInstanceMaxCycles(Index)) * 1000.0); }
 };
 
 struct FAccumulatedParticlePerfStats_GPU
@@ -104,13 +104,13 @@ struct FAccumulatedParticlePerfStats_GPU
 	ENGINE_API void Tick(FParticlePerfStats& Stats);
 
 	/** Returns the total microseconds used by GPU. */
-	FORCEINLINE uint64 GetTotalMicroseconds() const { return AccumulatedStats.GetTotalMicroseconds(); }
+	inline uint64 GetTotalMicroseconds() const { return AccumulatedStats.GetTotalMicroseconds(); }
 
-	FORCEINLINE uint64 GetPerFrameAvgMicroseconds() const { return NumFrames > 0 ? AccumulatedStats.GetTotalMicroseconds() / NumFrames : 0; }
-	FORCEINLINE uint64 GetPerFrameMaxMicroseconds(int32 Index = 0) const { return MaxPerFrameTotalMicroseconds[Index]; }
+	inline uint64 GetPerFrameAvgMicroseconds() const { return NumFrames > 0 ? AccumulatedStats.GetTotalMicroseconds() / NumFrames : 0; }
+	inline uint64 GetPerFrameMaxMicroseconds(int32 Index = 0) const { return MaxPerFrameTotalMicroseconds[Index]; }
 
-	FORCEINLINE uint64 GetPerInstanceAvgMicroseconds() const { return AccumulatedStats.GetPerInstanceAvgMicroseconds(); }
-	FORCEINLINE uint64 GetPerInstanceMaxMicroseconds(int32 Index = 0) const { return MaxPerInstanceMicroseconds[Index]; }
+	inline uint64 GetPerInstanceAvgMicroseconds() const { return AccumulatedStats.GetPerInstanceAvgMicroseconds(); }
+	inline uint64 GetPerInstanceMaxMicroseconds(int32 Index = 0) const { return MaxPerInstanceMicroseconds[Index]; }
 };
 
 /** Utility class for accumulating many frames worth of stats data. */
@@ -132,13 +132,13 @@ struct FAccumulatedParticlePerfStats
 	static ENGINE_API void ResetMaxArray(TArray<uint64, TInlineAllocator<ACCUMULATED_PARTICLE_PERF_STAT_MAX_SAMPLES>>& MaxArray);
 
 	/** Returns the current Game Thread stats. */
-	FORCEINLINE FAccumulatedParticlePerfStats_GT& GetGameThreadStats()
+	inline FAccumulatedParticlePerfStats_GT& GetGameThreadStats()
 	{
 		return GameThreadStats;
 	}
 
 	/** Returns the RenderThread stats with an optional flush when on the GameThread. */
-	FORCEINLINE FAccumulatedParticlePerfStats_RT& GetRenderThreadStats(bool bFlushForGameThread = false)
+	inline FAccumulatedParticlePerfStats_RT& GetRenderThreadStats(bool bFlushForGameThread = false)
 	{
 		if ( IsInGameThread() )
 		{
@@ -155,7 +155,7 @@ struct FAccumulatedParticlePerfStats
 	}
 
 	/** Returns the GPU stats with an optional flush when on the GameThread. */
-	FORCEINLINE FAccumulatedParticlePerfStats_GPU& GetGPUStats(bool bFlushForGameThread = false)
+	inline FAccumulatedParticlePerfStats_GPU& GetGPUStats(bool bFlushForGameThread = false)
 	{
 		if (IsInGameThread())
 		{
diff --git a/Engine/Source/Runtime/Engine/Classes/Particles/ParticleSystem.h b/Engine/Source/Runtime/Engine/Classes/Particles/ParticleSystem.h
--- a/Engine/Source/Runtime/Engine/Classes/Particles/ParticleSystem.h
+++ b/Engine/Source/Runtime/Engine/Classes/Particles/ParticleSystem.h
@@ -560,7 +560,7 @@ public:
 	ENGINE_API void BuildEmitters();
 
 	/** return true if this psys can tick in any thread */
-	FORCEINLINE bool CanTickInAnyThread()
+	inline bool CanTickInAnyThread()
 	{
 		if (!bIsElligibleForAsyncTickComputed)
 		{
@@ -590,7 +590,7 @@ public:
 	EParticleSignificanceLevel GetLowestSignificance()const { return LowestSignificance; }
 	bool ShouldManageSignificance()const { return bShouldManageSignificance; }
 
-	FORCEINLINE bool AllowManagedTicking()const { return bAllowManagedTicking; }
+	inline bool AllowManagedTicking()const { return bAllowManagedTicking; }
 private:
 
 	/** The highest significance of any emitter. Clamped by MaxSignificanceLevel.*/
diff --git a/Engine/Source/Runtime/Engine/Classes/Particles/ParticleSystemComponent.h b/Engine/Source/Runtime/Engine/Classes/Particles/ParticleSystemComponent.h
--- a/Engine/Source/Runtime/Engine/Classes/Particles/ParticleSystemComponent.h
+++ b/Engine/Source/Runtime/Engine/Classes/Particles/ParticleSystemComponent.h
@@ -630,14 +630,14 @@ public:
 #endif
 	ENGINE_API bool ShouldBeTickManaged()const;
 
-	FORCEINLINE bool IsTickManaged()const { return ManagerHandle != INDEX_NONE && !IsPendingManagerRemove(); }
-	FORCEINLINE int32 GetManagerHandle()const { return ManagerHandle; }
-	FORCEINLINE void SetManagerHandle(int32 InHandle) { ManagerHandle = InHandle; }
+	inline bool IsTickManaged()const { return ManagerHandle != INDEX_NONE && !IsPendingManagerRemove(); }
+	inline int32 GetManagerHandle()const { return ManagerHandle; }
+	inline void SetManagerHandle(int32 InHandle) { ManagerHandle = InHandle; }
 
-	FORCEINLINE int32 IsPendingManagerAdd()const { return bPendingManagerAdd; }
-	FORCEINLINE void SetPendingManagerAdd(bool bValue) { bPendingManagerAdd = bValue; }
-	FORCEINLINE int32 IsPendingManagerRemove()const { return bPendingManagerRemove; }
-	FORCEINLINE void SetPendingManagerRemove(bool bValue) { bPendingManagerRemove = bValue; }
+	inline int32 IsPendingManagerAdd()const { return bPendingManagerAdd; }
+	inline void SetPendingManagerAdd(bool bValue) { bPendingManagerAdd = bValue; }
+	inline int32 IsPendingManagerRemove()const { return bPendingManagerRemove; }
+	inline void SetPendingManagerRemove(bool bValue) { bPendingManagerRemove = bValue; }
 
 	ENGINE_API FPSCTickData& GetManagerTickData();
 	ENGINE_API FParticleSystemWorldManager* GetWorldManager()const;
@@ -891,13 +891,13 @@ public:
 	ENGINE_API void OnSignificanceChanged(bool bSignificant, bool bApplyToEmitters, bool bAsync=false);
 
 	/** Called from game code when the component begins having it's significance managed. */
-	FORCEINLINE void SetManagingSignificance(bool bManageSignificance)
+	inline void SetManagingSignificance(bool bManageSignificance)
 	{
 		bWasManagingSignificance = bIsManagingSignificance;
 		bIsManagingSignificance = bManageSignificance;
 	}
 
-	FORCEINLINE const FVector3f& GetLWCTile() const { return LWCTile;  }
+	inline const FVector3f& GetLWCTile() const { return LWCTile;  }
 	
 	/** Returns the approximate distance squared from this component to the passed location. */
 	ENGINE_API float GetApproxDistanceSquared(FVector Point)const;
@@ -912,7 +912,7 @@ public:
 
 	ENGINE_API void Complete();
 	
-	FORCEINLINE const FTransform& GetAsyncComponentToWorld()
+	inline const FTransform& GetAsyncComponentToWorld()
 	{
 		if (!bParallelRenderThreadUpdate && !IsInGameThread())
 		{
@@ -922,7 +922,7 @@ public:
 		return GetComponentTransform();
 	}
 
-	FORCEINLINE const TArray<struct FParticleSysParam>& GetAsyncInstanceParameters()
+	inline const TArray<struct FParticleSysParam>& GetAsyncInstanceParameters()
 	{
 		if (!bParallelRenderThreadUpdate && !IsInGameThread())
 		{
@@ -932,7 +932,7 @@ public:
 		return InstanceParameters;
 	}
 
-	FORCEINLINE const FBoxSphereBounds& GetAsyncBounds()
+	inline const FBoxSphereBounds& GetAsyncBounds()
 	{
 		if (!bParallelRenderThreadUpdate && !IsInGameThread())
 		{
@@ -942,7 +942,7 @@ public:
 		return Bounds;
 	}
 
-	FORCEINLINE const FVector& GetAsyncPartSysVelocity()
+	inline const FVector& GetAsyncPartSysVelocity()
 	{
 		if (!bParallelRenderThreadUpdate && !IsInGameThread())
 		{
@@ -952,12 +952,12 @@ public:
 		return PartSysVelocity;
 	}
 
-	FORCEINLINE const TArray<FVector>& GetPlayerLocations()
+	inline const TArray<FVector>& GetPlayerLocations()
 	{
 		return PlayerLocations;
 	}
 
-	FORCEINLINE const TArray<float>& GetPlayerLODDistanceFactor()
+	inline const TArray<float>& GetPlayerLODDistanceFactor()
 	{
 		return PlayerLODDistanceFactor;
 	}
@@ -1320,7 +1320,7 @@ public:
 
 private:
 	/** In some cases the async work for this PSC can be created externally by the manager. */
-	FORCEINLINE void SetAsyncWork(FGraphEventRef& InAsyncWork) { AsyncWork = InAsyncWork; }
+	inline void SetAsyncWork(FGraphEventRef& InAsyncWork) { AsyncWork = InAsyncWork; }
 
 	/** Cached copy of the transform for async work */
 	FTransform AsyncComponentToWorld;
@@ -1382,7 +1382,7 @@ public:
 		SILENT, // this would only be appropriate for editor only or other unusual things that we never see in game
 	};
 	/** If there is async work outstanding, force it to be completed now **/
-	FORCEINLINE void ForceAsyncWorkCompletion(EForceAsyncWorkCompletion Behavior, bool bDefinitelyGameThread, bool InSkipUpdateDynamicDataDuringTick)
+	inline void ForceAsyncWorkCompletion(EForceAsyncWorkCompletion Behavior, bool bDefinitelyGameThread, bool InSkipUpdateDynamicDataDuringTick)
 	{
 		if (AsyncWork.GetReference())
 		{
@@ -1394,7 +1394,7 @@ public:
 	}
 
 	/** If there is async work outstanding, force it to be completed now **/
-	FORCEINLINE void ForceAsyncWorkCompletion(EForceAsyncWorkCompletion Behavior, bool bDefinitelyGameThread = true) const
+	inline void ForceAsyncWorkCompletion(EForceAsyncWorkCompletion Behavior, bool bDefinitelyGameThread = true) const
 	{
 		if (AsyncWork.GetReference())
 		{
@@ -1472,7 +1472,7 @@ public:
 	ENGINE_API virtual bool ParticleLineCheck(FHitResult& Hit, AActor* SourceActor, const FVector& End, const FVector& Start, const FVector& HalfExtent, const FCollisionObjectQueryParams& ObjectParams);
 
 	/** return true if this psys can tick in any thread */
-	FORCEINLINE bool CanTickInAnyThread()
+	inline bool CanTickInAnyThread()
 	{
 		if (!bIsElligibleForAsyncTickComputed)
 		{
@@ -1543,7 +1543,7 @@ protected:
 	ENGINE_API virtual void UpdateDynamicData();
 
 public:
-	FORCEINLINE int32 GetCurrentLODIndex() const
+	inline int32 GetCurrentLODIndex() const
 	{
 		return LODLevel;
 	}
@@ -1716,14 +1716,14 @@ private:
 	ENGINE_API bool ShouldComputeLODFromGameThread();
 
 public:
-	FORCEINLINE FParticlePerfStatsContext GetPerfStatsContext(){ return FParticlePerfStatsContext(GetWorld(), Template, this); }
+	inline FParticlePerfStatsContext GetPerfStatsContext(){ return FParticlePerfStatsContext(GetWorld(), Template, this); }
 };
 
 
 //////////////////////////////////////////////////////////////////////////
 // ParticleSystemComponent inlines
 
-FORCEINLINE_DEBUGGABLE void UParticleSystemComponent::SetAutoAttachParams(USceneComponent* Parent, FName SocketName, EAttachLocation::Type LocationType)
+inline void UParticleSystemComponent::SetAutoAttachParams(USceneComponent* Parent, FName SocketName, EAttachLocation::Type LocationType)
 {
 	AutoAttachParent = Parent;
 	AutoAttachSocketName = SocketName;
diff --git a/Engine/Source/Runtime/Engine/Classes/Particles/ParticleSystemManager.h b/Engine/Source/Runtime/Engine/Classes/Particles/ParticleSystemManager.h
--- a/Engine/Source/Runtime/Engine/Classes/Particles/ParticleSystemManager.h
+++ b/Engine/Source/Runtime/Engine/Classes/Particles/ParticleSystemManager.h
@@ -106,16 +106,16 @@ private:
 		void Add(int32 Handle);
 		void Remove(int32 Handle);
 		void Reset();
-		FORCEINLINE int32 Num()const { return TickList.Num(); }
-		FORCEINLINE int32& operator[](const int32 Index) { return TickList[Index]; }
-		FORCEINLINE TArray<int32>& Get() { return TickList; }
+		inline int32 Num()const { return TickList.Num(); }
+		inline int32& operator[](const int32 Index) { return TickList[Index]; }
+		inline TArray<int32>& Get() { return TickList; }
 	};
 
 public:
 	static void OnStartup();
 	static void OnShutdown();
 
-	FORCEINLINE static FParticleSystemWorldManager* Get(UWorld* World)
+	inline static FParticleSystemWorldManager* Get(UWorld* World)
 	{
 		FParticleSystemWorldManager** Ret =  WorldManagers.Find(World);
 		if (Ret)
@@ -133,8 +133,8 @@ public:
 	bool RegisterComponent(UParticleSystemComponent* PSC);
 	void UnregisterComponent(UParticleSystemComponent* PSC);
 
-	FORCEINLINE FPSCTickData& GetTickData(int32 Handle) { return PSCTickData[Handle]; }
-	FORCEINLINE UParticleSystemComponent* GetManagedComponent(int32 Handle) { return ManagedPSCs[Handle]; }
+	inline FPSCTickData& GetTickData(int32 Handle) { return PSCTickData[Handle]; }
+	inline UParticleSystemComponent* GetManagedComponent(int32 Handle) { return ManagedPSCs[Handle]; }
 
 	void Tick(ETickingGroup TickGroup, float DeltaTime, ELevelTick TickType, ENamedThreads::Type CurrentThread, const FGraphEventRef& MyCompletionGraphEvent);
 
@@ -158,8 +158,8 @@ private:
 
 	template<bool bAsync>
 	void ProcessTickList(float DeltaTime, ELevelTick TickType, ETickingGroup TickGroup, TArray<FTickList>& TickLists, const FGraphEventRef& TickGroupCompletionGraphEvent);
-	FORCEINLINE void FlushAsyncTicks(const FGraphEventRef& TickGroupCompletionGraphEvent);
-	FORCEINLINE void QueueAsyncTick(int32 Handle, const FGraphEventRef& TickGroupCompletionGraphEvent);
+	inline void FlushAsyncTicks(const FGraphEventRef& TickGroupCompletionGraphEvent);
+	inline void QueueAsyncTick(int32 Handle, const FGraphEventRef& TickGroupCompletionGraphEvent);
 
 	void AddPSC(UParticleSystemComponent* PSC);
 	void RemovePSC(int32 PSCIndex);
