DiffTool: Ensure const correctness

#rnx

diff --git a/Engine/Source/Runtime/Engine/Public/DiffResults.h b/Engine/Source/Runtime/Engine/Public/DiffResults.h
--- a/Engine/Source/Runtime/Engine/Public/DiffResults.h
+++ b/Engine/Source/Runtime/Engine/Public/DiffResults.h
@@ -105,10 +105,10 @@ struct FDiffSingleResult
 	class UEdGraphPin* Pin2;
 
 	/** First top-level object involved in a diff */
-	UObject* Object1;
+	const UObject* Object1;
 
 	/** Second top-level object involved in a diff */
-	UObject* Object2;
+	const UObject* Object2;
 
 	/** String describing the error to the user */
 	FText DisplayString;
