[MaterialExpression]
- Stop relying on FApp::IsGame as it can dynamically change and is unreliable

diff --git a/Engine/Source/Runtime/Engine/Private/Materials/MaterialExpressions.cpp b/Engine/Source/Runtime/Engine/Private/Materials/MaterialExpressions.cpp
--- a/Engine/Source/Runtime/Engine/Private/Materials/MaterialExpressions.cpp
+++ b/Engine/Source/Runtime/Engine/Private/Materials/MaterialExpressions.cpp
@@ -16,6 +16,7 @@
 #include "UObject/UObjectAnnotation.h"
 #include "RenderUtils.h"
 #include "UObject/UObjectIterator.h"
+#include "UObject/Package.h"
 #include "SubstrateDefinitions.h"
 #include "UObject/ConstructorHelpers.h"
 #include "Engine/Engine.h"
@@ -2061,8 +2062,13 @@ void UMaterialExpression::UpdateParameterGuid(bool bForceGeneration, bool bAllow
 	if (bIsParameterExpression)
 	{
 		// If we are in the editor, and we don't have a valid GUID yet, generate one.
-		if(GIsEditor && !FApp::IsGame())
+		if (GIsEditor && !IsRunningGame())
 		{
+			if (GetPackage()->HasAnyPackageFlags(PKG_PlayInEditor))
+			{
+				return;
+			}
+
 			FGuid& Guid = GetParameterExpressionId();
 
 			if (bForceGeneration || !Guid.IsValid())
