Only process level streaming volumes transactionally if the world has already begun play.

diff --git a/Engine/Source/Runtime/Engine/Private/World.cpp b/Engine/Source/Runtime/Engine/Private/World.cpp
--- a/Engine/Source/Runtime/Engine/Private/World.cpp
+++ b/Engine/Source/Runtime/Engine/Private/World.cpp
@@ -4742,14 +4742,20 @@ void UWorld::InternalUpdateStreamingState()
 	// Issues level streaming load/unload requests based on local players being inside/outside level streaming volumes.
 
 #if UE_WITH_REMOTE_OBJECT_HANDLE
-	static FName TransactionalWorkName(TEXT("UWorld::InternalUpdateStreamingState"));
-	UE::RemoteExecutor::ExecuteTransactional(TransactionalWorkName,[this]()
+	// Only process transactionally if the world has begun play. This allows time for the transactional infrastructure to be set up.
+	if (HasBegunPlay())
 	{
-		ProcessLevelStreamingVolumes();
-	});
-#else
-	ProcessLevelStreamingVolumes();
+		static FName TransactionalWorkName(TEXT("UWorld::InternalUpdateStreamingState"));
+		UE::RemoteExecutor::ExecuteTransactional(TransactionalWorkName,[this]()
+		{
+			ProcessLevelStreamingVolumes();
+		});
+	}
+	else
 #endif
+	{
+		ProcessLevelStreamingVolumes();
+	}
 
 	// Update WorldComposition required streaming levels
 	if (WorldComposition)
