Fix incorrect packing/encoding when SimpleClearCoat is used.

SimpleClearCoat requires to be at the bottom. Otherwise it will be flagged has having simple volume SSS which will result in incorrect encoding.

#test SubstrateTest

diff --git a/Engine/Source/Runtime/Engine/Private/Materials/MaterialExpressionSubstrate.cpp b/Engine/Source/Runtime/Engine/Private/Materials/MaterialExpressionSubstrate.cpp
--- a/Engine/Source/Runtime/Engine/Private/Materials/MaterialExpressionSubstrate.cpp
+++ b/Engine/Source/Runtime/Engine/Private/Materials/MaterialExpressionSubstrate.cpp
@@ -904,7 +904,7 @@ int32 UMaterialExpressionSubstrateSlabBSDF::Compile(class FMaterialCompiler* Com
 		GlintValueCodeChunk,
 		GlintUVCodeChunk,
 		SpecularProfileCodeChunk != INDEX_NONE ? SpecularProfileCodeChunk : Compiler->Constant(0.0f),
-		SubstrateOperator.bIsBottom > 0 ? true : false,
+		SubstrateOperator.bIsBottom > 0 ? true : false, 				// bIsAtTheBottomOfTopology
 		NormalCodeChunk,
 		TangentCodeChunk,
 		Compiler->GetSubstrateSharedLocalBasisIndexMacro(NewRegisteredSharedLocalBasis),
@@ -975,7 +975,7 @@ int32 UMaterialExpressionSubstrateSlabBSDF::CompileDefaultSlab(FMaterialCompiler
 		GlintValueCodeChunk,
 		GlintUVCodeChunk,
 		Compiler->Constant(0.0f),
-		false,
+		false, 															// bIsAtTheBottomOfTopology
 		NormalCodeChunk,
 		TangentCodeChunk,
 		Compiler->GetSubstrateSharedLocalBasisIndexMacro(NewRegisteredSharedLocalBasis),
@@ -1402,7 +1402,7 @@ int32 UMaterialExpressionSubstrateSimpleClearCoatBSDF::Compile(class FMaterialCo
 		Compiler->Constant(1.0f),								// GlintValue
 		Compiler->Constant2(0.0f, 0.0f),						// GlintUV
 		Compiler->Constant(0.0f),								// SpecularProfile
-		false,													// bIsAtTheBottomOfTopology, always false for SimpleClearCoat
+		true,													// bIsAtTheBottomOfTopology
 		NormalCodeChunk,
 		NullTangentCodeChunk,
 		Compiler->GetSubstrateSharedLocalBasisIndexMacro(NewRegisteredSharedLocalBasis),
