Material Cache, fixed material layer's not compiling properly
- Compiling the full attribute set is an unmaintained path, compile them separately and aggregate
- Added explitic attribute set

diff --git a/Engine/Source/Editor/DetailCustomizations/Private/DetailCustomizations.cpp b/Engine/Source/Editor/DetailCustomizations/Private/DetailCustomizations.cpp
--- a/Engine/Source/Editor/DetailCustomizations/Private/DetailCustomizations.cpp
+++ b/Engine/Source/Editor/DetailCustomizations/Private/DetailCustomizations.cpp
@@ -485,6 +485,7 @@ void FDetailCustomizationsModule::RegisterObjectCustomizations()
 
 	RegisterCustomClassLayout("MaterialExpressionGetMaterialAttributes", FOnGetDetailCustomizationInstance::CreateStatic(&FMaterialAttributePropertyDetails::MakeInstance));
 	RegisterCustomClassLayout("MaterialExpressionSetMaterialAttributes", FOnGetDetailCustomizationInstance::CreateStatic(&FMaterialAttributePropertyDetails::MakeInstance));
+	RegisterCustomClassLayout("MaterialExpressionMaterialCache", FOnGetDetailCustomizationInstance::CreateStatic(&FMaterialAttributePropertyDetails::MakeInstance));
 	RegisterCustomClassLayout("MaterialExpressionTextureBase", FOnGetDetailCustomizationInstance::CreateStatic(&FMaterialExpressionTextureBaseDetails::MakeInstance));
 	RegisterCustomClassLayout("MaterialInstanceDynamic", FOnGetDetailCustomizationInstance::CreateStatic(&FMaterialInstanceDynamicDetails::MakeInstance));
 	RegisterCustomClassLayout("MaterialParameterCollection", FOnGetDetailCustomizationInstance::CreateStatic(&FMaterialParameterCollectionDetails::MakeInstance));
diff --git a/Engine/Source/Editor/DetailCustomizations/Private/MaterialAttributePropertyDetails.cpp b/Engine/Source/Editor/DetailCustomizations/Private/MaterialAttributePropertyDetails.cpp
--- a/Engine/Source/Editor/DetailCustomizations/Private/MaterialAttributePropertyDetails.cpp
+++ b/Engine/Source/Editor/DetailCustomizations/Private/MaterialAttributePropertyDetails.cpp
@@ -21,6 +21,7 @@
 #include "PropertyHandle.h"
 #include "Templates/Tuple.h"
 #include "RenderUtils.h"
+#include "Materials/MaterialExpressionMaterialCache.h"
 
 class SToolTip;
 
@@ -52,6 +53,7 @@ void FMaterialAttributePropertyDetails::CustomizeDetails(IDetailLayoutBuilder& D
 	// Fetch root property we're dealing with
 	TSharedPtr<IPropertyHandle> PropertyGetArray = DetailLayout.GetProperty(GET_MEMBER_NAME_CHECKED(UMaterialExpressionGetMaterialAttributes, AttributeGetTypes));
 	TSharedPtr<IPropertyHandle> PropertySetArray = DetailLayout.GetProperty(GET_MEMBER_NAME_CHECKED(UMaterialExpressionSetMaterialAttributes, AttributeSetTypes));
+	TSharedPtr<IPropertyHandle> PropertyCache    = DetailLayout.GetProperty(GET_MEMBER_NAME_CHECKED(UMaterialExpressionMaterialCache, AttributeTypes));
 	TSharedPtr<IPropertyHandle> PropertyArray;
 
 	if (PropertyGetArray->IsValidHandle())
@@ -62,6 +64,10 @@ void FMaterialAttributePropertyDetails::CustomizeDetails(IDetailLayoutBuilder& D
 	{
 		PropertyArray = PropertySetArray;
 	}
+	else if (PropertyCache->IsValidHandle())
+	{
+		PropertyArray = PropertyCache;
+	}
 	
 	check(PropertyArray->IsValidHandle());
 
diff --git a/Engine/Source/Runtime/Engine/Private/Materials/MaterialExpressions.cpp b/Engine/Source/Runtime/Engine/Private/Materials/MaterialExpressions.cpp
--- a/Engine/Source/Runtime/Engine/Private/Materials/MaterialExpressions.cpp
+++ b/Engine/Source/Runtime/Engine/Private/Materials/MaterialExpressions.cpp
@@ -20238,6 +20238,16 @@ UMaterialExpressionMaterialCache::UMaterialExpressionMaterialCache(const FObject
 	Outputs.Reset();
 	Outputs.Add(FExpressionOutput(TEXT(""), 0, 0, 0, 0, 0));
 #endif // WITH_EDITOR
+
+	// Default set of attributes
+	AttributeTypes = {
+		FMaterialAttributeDefinitionMap::GetID(MP_BaseColor),
+		FMaterialAttributeDefinitionMap::GetID(MP_Roughness),
+		FMaterialAttributeDefinitionMap::GetID(MP_Normal),
+		FMaterialAttributeDefinitionMap::GetID(MP_Specular),
+		FMaterialAttributeDefinitionMap::GetID(MP_Metallic),
+		FMaterialAttributeDefinitionMap::GetID(MP_Opacity)
+	};
 }
 
 #if WITH_EDITOR
@@ -20246,42 +20256,64 @@ int32 UMaterialExpressionMaterialCache::Compile(class FMaterialCompiler* Compile
 	// TODO[MP]: Validate the sub-graph against the supported expression set
 	// It's fine if the output is used as is, but the input should be strictly validated.
 
+	if (!Attributes.GetTracedInput().Expression)
+	{
+		return Compiler->Errorf(TEXT("Missing input attributes"));
+	}
+
+	for (FGuid AttributeSetType : AttributeTypes)
+	{
+		switch (EMaterialProperty Property = FMaterialAttributeDefinitionMap::GetProperty(AttributeSetType))
+		{
+			default:
+				return Compiler->Errorf(TEXT("Unsupported attribute type '%s"), *FMaterialAttributeDefinitionMap::GetAttributeName(Property));
+			case MP_BaseColor:
+			case MP_Roughness:
+			case MP_Normal:
+			case MP_Specular:
+			case MP_Metallic:
+			case MP_Opacity:
+				break;
+		}
+	}
+
 	// If we're compiling the output's graph, emit the input attributes
 	if (Compiler->GetTopCustomOutput() == this)
 	{
-		if (!Attributes.GetTracedInput().Expression)
+		int32 AttributeSet = Compiler->DefaultMaterialAttributes();
+
+		// Compile each attribute separately, there is a path for compiling the full set as a single code index
+		// However, it's somewhat unmaintained and fails in different ways.
+		for (FGuid AttributeSetType : AttributeTypes)
 		{
-			return Compiler->Errorf(TEXT("Missing input attributes"));
-		}
+			Compiler->PushMaterialAttribute(AttributeSetType);
+			int32 AttributeIndex = Attributes.Compile(Compiler);
+			Compiler->PopMaterialAttribute();
 
-		// Compile the full input tree, not per-attribute sub-graphs
-		Compiler->PushMaterialAttribute(FMaterialAttributeDefinitionMap::GetID(MP_MaterialAttributes));
-		int32 AttributeIndex = Attributes.Compile(Compiler);
-		Compiler->PopMaterialAttribute();
+			// May fail for unrelated reasons
+			if (AttributeIndex == INDEX_NONE)
+			{
+				return Compiler->Errorf(TEXT("Failed to compile attribute '%s"), *FMaterialAttributeDefinitionMap::GetAttributeName(AttributeSetType));
+			}
 
-		return Compiler->CustomOutput(this, 0, AttributeIndex, EMaterialCustomOutputFlags::AllowAttributeConnection);
+			AttributeSet = Compiler->SetMaterialAttribute(AttributeSet, AttributeIndex, AttributeSetType);
+		}
+
+		// Emit entire set
+		return Compiler->CustomOutput(this, 0, AttributeSet, EMaterialCustomOutputFlags::AllowAttributeConnection);
 	}
 
 	// We're compiling sub-graphs now, get the current attribute
 	const FGuid AttributeID = Compiler->GetMaterialAttribute();
 
-	// Supported attribute?
-	EMaterialProperty Property = FMaterialAttributeDefinitionMap::GetProperty(AttributeID);
-	switch (Property)
+	// If not part of the exported set, pass-through or default it
+	if (!AttributeTypes.Contains(AttributeID))
 	{
-	default:
-		// Not supported, default it
-		return FMaterialAttributeDefinitionMap::CompileDefaultExpression(Compiler, AttributeID);
-	case MP_BaseColor:
-	case MP_Roughness:
-	case MP_Normal:
-	case MP_Specular:
-	case MP_Metallic:
-	case MP_Opacity:
-		break;
+		int32 AttributeIndex = Attributes.Compile(Compiler);
+		return AttributeIndex == INDEX_NONE ? FMaterialAttributeDefinitionMap::CompileDefaultExpression(Compiler, AttributeID) : AttributeIndex;;
 	}
 
-	// Otherwise, same the cached attributes
+	// Otherwise, sample the cached attributes
 	int32 AttributeIndex = Compiler->SampleMaterialCache(0);
 	if (!AttributeID.IsValid())
 	{
diff --git a/Engine/Source/Runtime/Engine/Public/Materials/MaterialExpressionMaterialCache.h b/Engine/Source/Runtime/Engine/Public/Materials/MaterialExpressionMaterialCache.h
--- a/Engine/Source/Runtime/Engine/Public/Materials/MaterialExpressionMaterialCache.h
+++ b/Engine/Source/Runtime/Engine/Public/Materials/MaterialExpressionMaterialCache.h
@@ -3,12 +3,13 @@
 #pragma once
 
 #include "CoreMinimal.h"
+#include "Misc/Guid.h"
 #include "UObject/ObjectMacros.h"
 #include "MaterialExpressionIO.h"
 #include "MaterialExpressionCustomOutput.h"
 #include "MaterialExpressionMaterialCache.generated.h"
 
-UCLASS(meta = (DisplayName = "MaterialCache"))
+UCLASS(collapsecategories, hidecategories = Object, meta = (DisplayName = "MaterialCache"))
 class UMaterialExpressionMaterialCache : public UMaterialExpressionCustomOutput
 {
 	GENERATED_BODY()
@@ -16,6 +17,9 @@ class UMaterialExpressionMaterialCache : public UMaterialExpressionCustomOutput
 public:
 	UPROPERTY()
 	FMaterialAttributesInput Attributes;
+
+	UPROPERTY(EditAnywhere, Category=MaterialAttributes)
+	TArray<FGuid> AttributeTypes;
 	
 	virtual FString GetFunctionName() const override { return TEXT("GetMaterialCache"); }
 	virtual FString GetDisplayName() const override { return TEXT("MaterialCache"); }
