Removed unused primitive flags PRIMITIVE_SCENE_DATA_FLAG_CACHE_SHADOW_AS_STATIC and PRIMITIVE_SCENE_DATA_FLAG_WRITES_CUSTOM_DEPTH_STENCIL

#tests Should not affect any functionality
#rnx

diff --git a/Engine/Shaders/Shared/SceneDefinitions.h b/Engine/Shaders/Shared/SceneDefinitions.h
--- a/Engine/Shaders/Shared/SceneDefinitions.h
+++ b/Engine/Shaders/Shared/SceneDefinitions.h
@@ -34,7 +34,7 @@
 #define PRIMITIVE_SCENE_DATA_FLAG_USE_SINGLE_SAMPLE_SHADOW_SL				0x2
 #define PRIMITIVE_SCENE_DATA_FLAG_USE_VOLUMETRIC_LM_SHADOW_SL				0x4
 #define PRIMITIVE_SCENE_DATA_FLAG_DECAL_RECEIVER							0x8
-#define PRIMITIVE_SCENE_DATA_FLAG_CACHE_SHADOW_AS_STATIC					0x10
+#define PRIMITIVE_SCENE_DATA_FLAG_UNUSED_1									0x10
 #define PRIMITIVE_SCENE_DATA_FLAG_OUTPUT_VELOCITY							0x20
 #define PRIMITIVE_SCENE_DATA_FLAG_DETERMINANT_SIGN							0x40
 #define PRIMITIVE_SCENE_DATA_FLAG_HAS_RAY_TRACING_REPRESENTATION			0x80
@@ -50,7 +50,7 @@
 #define PRIMITIVE_SCENE_DATA_FLAG_INSTANCE_DRAW_DISTANCE_CULL				0x20000
 #define PRIMITIVE_SCENE_DATA_FLAG_WPO_DISABLE_DISTANCE						0x40000
 #define PRIMITIVE_SCENE_DATA_FLAG_HAS_ALWAYS_EVALUATE_WPO_MATERIALS			0x80000
-#define PRIMITIVE_SCENE_DATA_FLAG_WRITES_CUSTOM_DEPTH_STENCIL				0x100000
+#define PRIMITIVE_SCENE_DATA_FLAG_UNUSED_2									0x100000
 #define PRIMITIVE_SCENE_DATA_FLAG_DISABLE_MATERIAL_INVALIDATIONS			0x200000
 #define PRIMITIVE_SCENE_DATA_FLAG_HOLDOUT									0x400000
 #define PRIMITIVE_SCENE_DATA_FLAG_SPLINE_MESH								0x800000
diff --git a/Engine/Source/Runtime/Engine/Private/PrimitiveSceneProxy.cpp b/Engine/Source/Runtime/Engine/Private/PrimitiveSceneProxy.cpp
--- a/Engine/Source/Runtime/Engine/Private/PrimitiveSceneProxy.cpp
+++ b/Engine/Source/Runtime/Engine/Private/PrimitiveSceneProxy.cpp
@@ -932,7 +932,6 @@ void FPrimitiveSceneProxy::BuildUniformShaderParameters(FPrimitiveUniformShaderP
 	if (PrimitiveSceneInfo != nullptr)
 	{
 		Builder.LightmapDataIndex(PrimitiveSceneInfo->GetLightmapDataOffset())
-			.CacheShadowAsStatic(PrimitiveSceneInfo->ShouldCacheShadowAsStatic())
 			.InstanceSceneDataOffset(PrimitiveSceneInfo->GetInstanceSceneDataOffset())
 			.NumInstanceSceneDataEntries(PrimitiveSceneInfo->GetNumInstanceSceneDataEntries())
 			.InstancePayloadDataOffset(PrimitiveSceneInfo->GetInstancePayloadDataOffset())
diff --git a/Engine/Source/Runtime/Engine/Public/PrimitiveUniformShaderParametersBuilder.h b/Engine/Source/Runtime/Engine/Public/PrimitiveUniformShaderParametersBuilder.h
--- a/Engine/Source/Runtime/Engine/Public/PrimitiveUniformShaderParametersBuilder.h
+++ b/Engine/Source/Runtime/Engine/Public/PrimitiveUniformShaderParametersBuilder.h
@@ -39,7 +39,6 @@ public:
 		bReceivesDecals								= false;
 		bUseSingleSampleShadowFromStationaryLights	= false;
 		bUseVolumetricLightmap						= false;
-		bCacheShadowAsStatic						= false;
 		bOutputVelocity								= false;
 		bHasRayTracingRepresentation				= false;
 		bHasCapsuleRepresentation					= false;
@@ -54,7 +53,6 @@ public:
 		bHasInstanceDrawDistanceCull				= false;
 		bHasWPODisableDistance						= false;
 		bHasAlwaysEvaluateWPOMaterials				= false;
-		bWritesCustomDepthStencil					= false;
 		bReverseCulling								= false;
 		bHoldout									= false;
 		bDisableMaterialInvalidations				= false;
@@ -118,7 +116,6 @@ public:
 	PRIMITIVE_UNIFORM_BUILDER_FLAG_METHOD(bool,			CastShadow);
 	PRIMITIVE_UNIFORM_BUILDER_FLAG_METHOD(bool,			UseSingleSampleShadowFromStationaryLights);
 	PRIMITIVE_UNIFORM_BUILDER_FLAG_METHOD(bool,			UseVolumetricLightmap);
-	PRIMITIVE_UNIFORM_BUILDER_FLAG_METHOD(bool,			CacheShadowAsStatic);
 	PRIMITIVE_UNIFORM_BUILDER_FLAG_METHOD(bool,			OutputVelocity);
 	PRIMITIVE_UNIFORM_BUILDER_FLAG_METHOD(bool,			EvaluateWorldPositionOffset);
 	PRIMITIVE_UNIFORM_BUILDER_FLAG_METHOD(bool,			VisibleInGame);
@@ -312,7 +309,6 @@ public:
 			break;
 		}
 		Parameters.CustomStencilValueAndMask = (GPUStencilMask << 8u) | StencilValue;
-		bWritesCustomDepthStencil = true;
 
 		return *this;
 	}
@@ -416,7 +412,6 @@ public:
 		Parameters.Flags |= bHasCapsuleRepresentation ? PRIMITIVE_SCENE_DATA_FLAG_HAS_CAPSULE_REPRESENTATION : 0u;
 		Parameters.Flags |= bUseSingleSampleShadowFromStationaryLights ? PRIMITIVE_SCENE_DATA_FLAG_USE_SINGLE_SAMPLE_SHADOW_SL : 0u;
 		Parameters.Flags |= (bUseVolumetricLightmap && bUseSingleSampleShadowFromStationaryLights) ? PRIMITIVE_SCENE_DATA_FLAG_USE_VOLUMETRIC_LM_SHADOW_SL : 0u;
-		Parameters.Flags |= bCacheShadowAsStatic ? PRIMITIVE_SCENE_DATA_FLAG_CACHE_SHADOW_AS_STATIC : 0u;
 		Parameters.Flags |= bOutputVelocity ? PRIMITIVE_SCENE_DATA_FLAG_OUTPUT_VELOCITY : 0u;
 		Parameters.Flags |= bEvaluateWorldPositionOffset ? PRIMITIVE_SCENE_DATA_FLAG_EVALUATE_WORLD_POSITION_OFFSET : 0u;
 		Parameters.Flags |= (Parameters.LocalToRelativeWorld.RotDeterminant() < 0.0f) ? PRIMITIVE_SCENE_DATA_FLAG_DETERMINANT_SIGN : 0u;
@@ -431,7 +426,6 @@ public:
 		Parameters.Flags |= bHasInstanceDrawDistanceCull ? PRIMITIVE_SCENE_DATA_FLAG_INSTANCE_DRAW_DISTANCE_CULL : 0u;
 		Parameters.Flags |= bHasWPODisableDistance ? PRIMITIVE_SCENE_DATA_FLAG_WPO_DISABLE_DISTANCE : 0u;
 		Parameters.Flags |= bHasAlwaysEvaluateWPOMaterials ? PRIMITIVE_SCENE_DATA_FLAG_HAS_ALWAYS_EVALUATE_WPO_MATERIALS : 0u;
-		Parameters.Flags |= bWritesCustomDepthStencil ? PRIMITIVE_SCENE_DATA_FLAG_WRITES_CUSTOM_DEPTH_STENCIL : 0u;
 #if SUPPORT_REVERSE_CULLING_IN_NANITE
 		Parameters.Flags |= bReverseCulling ? PRIMITIVE_SCENE_DATA_FLAG_REVERSE_CULLING : 0u;
 #endif
@@ -477,7 +471,6 @@ private:
 	uint32 bReceivesDecals : 1;
 	uint32 bUseSingleSampleShadowFromStationaryLights : 1;
 	uint32 bUseVolumetricLightmap : 1;
-	uint32 bCacheShadowAsStatic : 1;
 	uint32 bOutputVelocity : 1;
 	uint32 bEvaluateWorldPositionOffset : 1;
 	uint32 bCastShadow : 1;
@@ -502,7 +495,6 @@ private:
 	uint32 bHasInstanceDrawDistanceCull : 1;
 	uint32 bHasWPODisableDistance : 1;
 	uint32 bHasAlwaysEvaluateWPOMaterials : 1;
-	uint32 bWritesCustomDepthStencil : 1;
 	uint32 bReverseCulling: 1;
 	uint32 bHoldout : 1;
 	uint32 bDisableMaterialInvalidations : 1;
diff --git a/Engine/Source/Runtime/Renderer/Private/GPUScene.cpp b/Engine/Source/Runtime/Renderer/Private/GPUScene.cpp
--- a/Engine/Source/Runtime/Renderer/Private/GPUScene.cpp
+++ b/Engine/Source/Runtime/Renderer/Private/GPUScene.cpp
@@ -1669,7 +1669,6 @@ struct FUploadDataSourceAdapterDynamicPrimitives
 		FPrimitiveUniformShaderParameters Tmp = *PrimData.ShaderParams;
 		Tmp.InstanceSceneDataOffset = InstanceIDStartOffset + PrimData.LocalInstanceSceneDataOffset;
 		Tmp.NumInstanceSceneDataEntries = PrimData.NumInstances;
-		Tmp.Flags &= ~(PRIMITIVE_SCENE_DATA_FLAG_CACHE_SHADOW_AS_STATIC);	// Dyanmic primtives are never cached as static for shadows
 		if (PrimitiveData[ItemIndex].LocalPayloadDataOffset != INDEX_NONE)
 		{
 			Tmp.InstancePayloadDataOffset = PayloadStartOffset + PrimData.LocalPayloadDataOffset;
diff --git a/Engine/Source/Runtime/Renderer/Private/PrimitiveSceneInfo.cpp b/Engine/Source/Runtime/Renderer/Private/PrimitiveSceneInfo.cpp
--- a/Engine/Source/Runtime/Renderer/Private/PrimitiveSceneInfo.cpp
+++ b/Engine/Source/Runtime/Renderer/Private/PrimitiveSceneInfo.cpp
@@ -314,7 +314,6 @@ FPrimitiveSceneInfo::FPrimitiveSceneInfo(const FPrimitiveSceneInfoAdapter& InAda
 	bIndirectLightingCacheBufferDirty(false),
 	bRegisteredLightmapVirtualTextureProducerCallback(false),
 	bRegisteredWithVelocityData(false),
-	bCacheShadowAsStatic((InAdapter.Mobility != EComponentMobility::Movable && InAdapter.SceneProxy->GetShadowCacheInvalidationBehavior() != EShadowCacheInvalidationBehavior::Always) || InAdapter.SceneProxy->GetShadowCacheInvalidationBehavior() == EShadowCacheInvalidationBehavior::Static),
 	bNaniteRasterBinsRenderCustomDepth(false),
 	bPendingAddToScene(false),
 	bPendingAddStaticMeshes(false),
diff --git a/Engine/Source/Runtime/Renderer/Public/PrimitiveSceneInfo.h b/Engine/Source/Runtime/Renderer/Public/PrimitiveSceneInfo.h
--- a/Engine/Source/Runtime/Renderer/Public/PrimitiveSceneInfo.h
+++ b/Engine/Source/Runtime/Renderer/Public/PrimitiveSceneInfo.h
@@ -540,9 +540,10 @@ public:
 	FString GetFullnameForDebuggingOnly() const;
 	/** Return primitive Owner actor name (for debugging only). */
 	FString GetOwnerActorNameOrLabelForDebuggingOnly() const;
+	UE_DEPRECATED(5.7, "Unused")
 	inline bool ShouldCacheShadowAsStatic() const
 	{
-		return bCacheShadowAsStatic;
+		return false;
 	}
 
 	inline FMeshDrawCommandPrimitiveIdInfo GetMDCIdInfo() const { return FMeshDrawCommandPrimitiveIdInfo(PackedIndex, PersistentIndex, InstanceSceneDataOffset);}
@@ -612,9 +613,6 @@ private:
 	/** True if the primitive registered with velocity data and needs to remove itself when being removed from the scene. */
 	bool bRegisteredWithVelocityData : 1;
 
-	/** True if the primitive should be treated as static for the purpose of caching shadows */
-	bool bCacheShadowAsStatic : 1;
-
 	/** True if the Nanite raster bins were registered with custom depth enabled */
 	bool bNaniteRasterBinsRenderCustomDepth : 1;
 
