[LightMap/ShadowMap]
- Fix serialization from parallel loading thread

diff --git a/Engine/Source/Runtime/Engine/Public/LightMap.h b/Engine/Source/Runtime/Engine/Public/LightMap.h
--- a/Engine/Source/Runtime/Engine/Public/LightMap.h
+++ b/Engine/Source/Runtime/Engine/Public/LightMap.h
@@ -88,12 +88,12 @@ public:
 	// Reference counting.
 	void AddRef()
 	{
-		check(IsInGameThread() || IsInAsyncLoadingThread());
+		check(IsInGameThread() || IsInAsyncLoadingThread() || IsInParallelLoadingThread());
 		NumRefs++;
 	}
 	void Release()
 	{
-		check(IsInGameThread() || IsInAsyncLoadingThread());
+		check(IsInGameThread() || IsInAsyncLoadingThread() || IsInParallelLoadingThread());
 		checkSlow(NumRefs > 0);
 		if(--NumRefs == 0)
 		{
diff --git a/Engine/Source/Runtime/Engine/Public/ShadowMap.h b/Engine/Source/Runtime/Engine/Public/ShadowMap.h
--- a/Engine/Source/Runtime/Engine/Public/ShadowMap.h
+++ b/Engine/Source/Runtime/Engine/Public/ShadowMap.h
@@ -272,12 +272,12 @@ public:
 	// Reference counting.
 	void AddRef()
 	{
-		check(IsInGameThread() || IsInAsyncLoadingThread());
+		check(IsInGameThread() || IsInAsyncLoadingThread() || IsInParallelLoadingThread());
 		NumRefs++;
 	}
 	void Release()
 	{
-		check(IsInGameThread() || IsInAsyncLoadingThread());
+		check(IsInGameThread() || IsInAsyncLoadingThread() || IsInParallelLoadingThread());
 		checkSlow(NumRefs > 0);
 		if (--NumRefs == 0)
 		{
