Splines - Fix Reparam Table - Revert to using old interp curve reparam table for the time being.

This addresses a crash and inaccuracy concerns. We should investigate the underlying issues in the future

diff --git a/Engine/Source/Runtime/Engine/Private/Curves/Spline.cpp b/Engine/Source/Runtime/Engine/Private/Curves/Spline.cpp
--- a/Engine/Source/Runtime/Engine/Private/Curves/Spline.cpp
+++ b/Engine/Source/Runtime/Engine/Private/Curves/Spline.cpp
@@ -229,8 +229,7 @@ private:
 	
 	static inline const FName RotationAttrName = FName("Rotation");
 	static inline const FName ScaleAttrName = FName("Scale");
-	static inline const FName ParamToLengthAttrName = FName("ParamToLength");
-	static inline const FName NormalizedLengthToParamAttrName = FName("LengthToParam");
+	FInterpCurveFloat ReparamTable;
 
     /**
      * We keep these up to date with Points.
@@ -1111,10 +1110,6 @@ TArray<FName> FSpline::GetFloatPropertyChannels() const
 	if (NewData && SupportsAttributes())
 	{
 		ChannelNames = NewData->GetAttributeChannelNamesByValueType<float>();
-
-		// Don't report these internal attribute channels to the caller
-		ChannelNames.Remove(NewData->ParamToLengthAttrName);
-		ChannelNames.Remove(NewData->NormalizedLengthToParamAttrName);
 	}
 	return ChannelNames;
 }
@@ -1621,10 +1616,12 @@ void FSpline::OnSplineImplementationChanged()
 
 FNewSpline::FNewSpline()
 {
+	// This is a hack. We just need to make sure that this spline type registers itself with the spline registry.
+	// This should be removed in the future once we can register this spline type automatically without instantiation.
+	TTangentBezierSpline<float> AutoRegister;
+	
 	CreateAttributeChannel<TTangentBezierSpline<FQuat>>(RotationAttrName);
 	CreateAttributeChannel<TTangentBezierSpline<FVector>>(ScaleAttrName);
-	CreateAttributeChannel<TTangentBezierSpline<float>>(ParamToLengthAttrName);
-	CreateAttributeChannel<TTangentBezierSpline<float>>(NormalizedLengthToParamAttrName);
 	
 	ValidateRotScale();
 }
@@ -1640,8 +1637,6 @@ FNewSpline::FNewSpline(const FLegacySpline& Other)
 {
 	CreateAttributeChannel<TTangentBezierSpline<FQuat>>(RotationAttrName);
 	CreateAttributeChannel<TTangentBezierSpline<FVector>>(ScaleAttrName);
-	CreateAttributeChannel<TTangentBezierSpline<float>>(ParamToLengthAttrName);
-	CreateAttributeChannel<TTangentBezierSpline<float>>(NormalizedLengthToParamAttrName);
 	
 	TArray<FTangentBezierControlPoint> Points;
     for (int32 Idx = 0; Idx < Other.GetNumControlPoints(); Idx++)
@@ -1678,8 +1673,6 @@ FNewSpline::FNewSpline(const FSplineCurves& Other)
 {
 	CreateAttributeChannel<TTangentBezierSpline<FQuat>>(RotationAttrName);
 	CreateAttributeChannel<TTangentBezierSpline<FVector>>(ScaleAttrName);
-	CreateAttributeChannel<TTangentBezierSpline<float>>(ParamToLengthAttrName);
-	CreateAttributeChannel<TTangentBezierSpline<float>>(NormalizedLengthToParamAttrName);
 	
 	TArray<FTangentBezierControlPoint> Points;
     for (int32 Idx = 0; Idx < Other.Position.Points.Num(); Idx++)
@@ -2208,15 +2201,23 @@ float FNewSpline::GetParameterAtDistance(float Distance)
 {
 	UpdateReparamTable();
 	
-	const float NormalizedDistance = FMath::Clamp(Distance / FMath::Max(GetSplineLength(), UE_SMALL_NUMBER), 0.f, 1.f);
-	return EvaluateAttribute<float>(NormalizedLengthToParamAttrName, ToInternalSplineSpace(NormalizedDistance));
+	return ReparamTable.Eval(Distance);
 }
 
 float FNewSpline::GetDistanceAtParameter(float Parameter)
 {
 	UpdateReparamTable();
 
-	return EvaluateAttribute<float>(ParamToLengthAttrName, ToInternalSplineSpace(Parameter));
+	const float ParameterMax = PositionCurve.Points.Last().InVal;
+	const float Key = (Parameter / ParameterMax) * (ReparamTable.Points.Num() - 1);
+	const int32 LowerKey = FMath::FloorToInt32(Key);
+	ensureAlways(LowerKey >= 0 && LowerKey < ReparamTable.Points.Num());
+	const int32 UpperKey = FMath::CeilToInt32(Key);
+	ensureAlways(UpperKey >= 0 && UpperKey < ReparamTable.Points.Num());
+	const float Alpha = FMath::Frac(Key);
+	const float Distance = FMath::Lerp(ReparamTable.Points[LowerKey].InVal, ReparamTable.Points[UpperKey].InVal, Alpha);
+
+	return Distance;
 }
 
 float FNewSpline::FindNearest(const FVector& InLocation, float& OutSquaredDistance) const
@@ -2320,39 +2321,22 @@ void FNewSpline::UpdateReparamTable()
 	}
 
 	const int32 NumSegments = GetSpline().GetNumberOfSegments();
-	
+
 	// Start by clearing it
-	ClearAttributeChannel(ParamToLengthAttrName);
-	ClearAttributeChannel(NormalizedLengthToParamAttrName);
-	
-	// Here we get Param -> Length mapping.
+	ReparamTable.Points.Reset(NumSegments * ReparamStepsPerSegment + 1);
 	float AccumulatedLength = 0.0f;
 	for (int32 SegmentIndex = 0; SegmentIndex < NumSegments; ++SegmentIndex)
 	{
 		for (int32 Step = 0; Step < ReparamStepsPerSegment; ++Step)
 		{
-			// These are with respect to current segment.
-			const float SegmentParam = static_cast<float>(Step) / ReparamStepsPerSegment;
-			const float OffsetLength = GetSegmentLength(SegmentIndex, SegmentParam);
-	
-			// Map the above values to 'spline space' (i.e. [0, NumSegments]) by normalizing and converting.
-			const float SplineParam = SegmentIndex + SegmentParam;
-			const float CurrentLength = OffsetLength + AccumulatedLength;
-			AddAttributeValue<float>(ParamToLengthAttrName, CurrentLength, SplineParam);
+			const float Param = static_cast<float>(Step) / ReparamStepsPerSegment;
+			const float SegmentLength = (Step == 0) ? 0.0f : GetSegmentLength(SegmentIndex, Param);
+
+			ReparamTable.Points.Emplace(SegmentLength + AccumulatedLength, SegmentIndex + Param, 0.0f, 0.0f, CIM_Linear);
 		}
-	
 		AccumulatedLength += GetSegmentLength(SegmentIndex, 1.0f);
 	}
-	AddAttributeValue<float>(ParamToLengthAttrName, AccumulatedLength, static_cast<float>(NumSegments));
-	
-	// Here we get Length -> Param mapping
-	const int32 TotalSteps = NumSegments * ReparamStepsPerSegment + 1;
-	for (int32 Step = 0; Step <= TotalSteps; ++Step)
-	{
-		const float SplineParam = FromNormalizedSpace(FMath::Clamp(static_cast<float>(Step) / TotalSteps, 0.0f, 1.0f));
-		const float NormalizedLength = EvaluateAttribute<float>(ParamToLengthAttrName, ToInternalSplineSpace(SplineParam)) / AccumulatedLength;
-		AddAttributeValue<float>(NormalizedLengthToParamAttrName, SplineParam, NormalizedLength);
-	}
+	ReparamTable.Points.Emplace(AccumulatedLength, static_cast<float>(NumSegments), 0.0f, 0.0f, CIM_Linear);
 
 	bReparamTableDirty = false;
 }
