Replace some usages of FORCEINLINE with inline in WorldBuilding modules.

For context about this change, see this doc: https://docs.google.com/document/d/1LjpT7NpzTecfXdFNyv5KM6sCu67827BRljNi1iaWVSE/edit?tab=t.0#heading=h.fvtz8gjukz7x

Reviewers were automatically determined. Please comment if anyone else should be added.

#rnx

diff --git a/Engine/Plugins/PCG/Source/PCG/Public/Data/PCGDataPtrWrapper.h b/Engine/Plugins/PCG/Source/PCG/Public/Data/PCGDataPtrWrapper.h
--- a/Engine/Plugins/PCG/Source/PCG/Public/Data/PCGDataPtrWrapper.h
+++ b/Engine/Plugins/PCG/Source/PCG/Public/Data/PCGDataPtrWrapper.h
@@ -74,13 +74,13 @@ struct TStructOpsTypeTraits<FPCGDataPtrWrapper> : public TStructOpsTypeTraitsBas
 };
 
 template<typename To>
-FORCEINLINE TCopyQualifiersFromTo_T<const TObjectPtr<const UPCGData>, To>* Cast(const FPCGDataPtrWrapper& InSrc)
+inline TCopyQualifiersFromTo_T<const TObjectPtr<const UPCGData>, To>* Cast(const FPCGDataPtrWrapper& InSrc)
 {
 	return Cast<To>(InSrc.GetObjectPtr());
 }
 
 template<typename To>
-FORCEINLINE TCopyQualifiersFromTo_T<const TObjectPtr<const UPCGData>, To>* CastChecked(const FPCGDataPtrWrapper& InSrc, ECastCheckedType::Type CheckType = ECastCheckedType::NullChecked)
+inline TCopyQualifiersFromTo_T<const TObjectPtr<const UPCGData>, To>* CastChecked(const FPCGDataPtrWrapper& InSrc, ECastCheckedType::Type CheckType = ECastCheckedType::NullChecked)
 {
 	return CastChecked<To>(InSrc.GetObjectPtr(), CheckType);
 }
diff --git a/Engine/Plugins/PCG/Source/PCG/Public/Data/PCGPointData.h b/Engine/Plugins/PCG/Source/PCG/Public/Data/PCGPointData.h
--- a/Engine/Plugins/PCG/Source/PCG/Public/Data/PCGPointData.h
+++ b/Engine/Plugins/PCG/Source/PCG/Public/Data/PCGPointData.h
@@ -43,23 +43,23 @@ struct PCG_API UE_DEPRECATED(5.6, "Use PCGPointOctree::FPointRefSemantics instea
 
 	typedef TInlineAllocator<MaxElementsPerLeaf> ElementAllocator;
 
-	FORCEINLINE static const FBoxSphereBounds& GetBoundingBox(const FPCGPointRef& InPoint)
+	inline static const FBoxSphereBounds& GetBoundingBox(const FPCGPointRef& InPoint)
 	{
 		return InPoint.Bounds;
 	}
 
-	FORCEINLINE static const bool AreElementsEqual(const FPCGPointRef& A, const FPCGPointRef& B)
+	inline static const bool AreElementsEqual(const FPCGPointRef& A, const FPCGPointRef& B)
 	{
 		// TODO: verify if that's sufficient
 		return A.Point == B.Point;
 	}
 
-	FORCEINLINE static void ApplyOffset(FPCGPointRef& InPoint)
+	inline static void ApplyOffset(FPCGPointRef& InPoint)
 	{
 		ensureMsgf(false, TEXT("Not implemented"));
 	}
 
-	FORCEINLINE static void SetElementId(const FPCGPointRef& Element, FOctreeElementId2 OctreeElementID)
+	inline static void SetElementId(const FPCGPointRef& Element, FOctreeElementId2 OctreeElementID)
 	{
 	}
 };
diff --git a/Engine/Plugins/PCG/Source/PCG/Public/Grid/PCGComponentOctree.h b/Engine/Plugins/PCG/Source/PCG/Public/Grid/PCGComponentOctree.h
--- a/Engine/Plugins/PCG/Source/PCG/Public/Grid/PCGComponentOctree.h
+++ b/Engine/Plugins/PCG/Source/PCG/Public/Grid/PCGComponentOctree.h
@@ -40,22 +40,22 @@ struct FPCGComponentRefSemantics
 
 	typedef TInlineAllocator<MaxElementsPerLeaf> ElementAllocator;
 
-	FORCEINLINE static const FBoxSphereBounds& GetBoundingBox(const FPCGComponentRef& InVolume)
+	inline static const FBoxSphereBounds& GetBoundingBox(const FPCGComponentRef& InVolume)
 	{
 		return InVolume.Bounds;
 	}
 
-	FORCEINLINE static const bool AreElementsEqual(const FPCGComponentRef& A, const FPCGComponentRef& B)
+	inline static const bool AreElementsEqual(const FPCGComponentRef& A, const FPCGComponentRef& B)
 	{
 		return A.Component == B.Component;
 	}
 
-	FORCEINLINE static void ApplyOffset(FPCGComponentRef& InVolume, const FVector& Offset)
+	inline static void ApplyOffset(FPCGComponentRef& InVolume, const FVector& Offset)
 	{
 		InVolume.Bounds.Origin += Offset;
 	}
 
-	FORCEINLINE static void SetElementId(const FPCGComponentRef& Element, FOctreeElementId2 OctreeElementID)
+	inline static void SetElementId(const FPCGComponentRef& Element, FOctreeElementId2 OctreeElementID)
 	{
 		Element.IdShared->Id = OctreeElementID;
 	}
diff --git a/Engine/Plugins/PCG/Source/PCG/Public/Helpers/PCGPointHelpers.h b/Engine/Plugins/PCG/Source/PCG/Public/Helpers/PCGPointHelpers.h
--- a/Engine/Plugins/PCG/Source/PCG/Public/Helpers/PCGPointHelpers.h
+++ b/Engine/Plugins/PCG/Source/PCG/Public/Helpers/PCGPointHelpers.h
@@ -9,33 +9,33 @@
 
 namespace PCGPointHelpers
 {
-	FORCEINLINE FBox GetLocalBounds(const FVector& InBoundsMin, const FVector& InBoundsMax)
+	inline FBox GetLocalBounds(const FVector& InBoundsMin, const FVector& InBoundsMax)
 	{
 		return FBox(InBoundsMin, InBoundsMax);
 	}
 
-	FORCEINLINE void SetLocalBounds(const FBox& InBounds, FVector& OutBoundsMin, FVector& OutBoundsMax)
+	inline void SetLocalBounds(const FBox& InBounds, FVector& OutBoundsMin, FVector& OutBoundsMax)
 	{
 		OutBoundsMin = InBounds.Min;
 		OutBoundsMax = InBounds.Max;
 	}
 
-	FORCEINLINE FBox GetLocalDensityBounds(float InSteepness, const FVector& InBoundsMin, const FVector& InBoundsMax)
+	inline FBox GetLocalDensityBounds(float InSteepness, const FVector& InBoundsMin, const FVector& InBoundsMax)
 	{
 		return FBox((2 - InSteepness) * InBoundsMin, (2 - InSteepness) * InBoundsMax);
 	}
 
-	FORCEINLINE FVector GetLocalCenter(const FVector& InBoundsMin, const FVector& InBoundsMax)
+	inline FVector GetLocalCenter(const FVector& InBoundsMin, const FVector& InBoundsMax)
 	{
 		return (InBoundsMax + InBoundsMin) / 2.0;
 	}
 
-	FORCEINLINE FVector GetExtents(const FVector& InBoundsMin, const FVector& InBoundsMax)
+	inline FVector GetExtents(const FVector& InBoundsMin, const FVector& InBoundsMax)
 	{
 		return (InBoundsMax - InBoundsMin) / 2.0;
 	}
 
-	FORCEINLINE FBoxSphereBounds GetDensityBounds(const FTransform& InTransform, float InSteepness, const FVector& InBoundsMin, const FVector& InBoundsMax)
+	inline FBoxSphereBounds GetDensityBounds(const FTransform& InTransform, float InSteepness, const FVector& InBoundsMin, const FVector& InBoundsMax)
 	{
 		if (InTransform.IsRotationNormalized())
 		{
@@ -49,36 +49,36 @@ namespace PCGPointHelpers
 		}
 	}
 	
-	FORCEINLINE void SetExtents(const FVector& InExtents, FVector& InOutBoundsMin, FVector& InOutBoundsMax)
+	inline void SetExtents(const FVector& InExtents, FVector& InOutBoundsMin, FVector& InOutBoundsMax)
 	{
 		const FVector Center = PCGPointHelpers::GetLocalCenter(InOutBoundsMin, InOutBoundsMax);
 		InOutBoundsMin = Center - InExtents;
 		InOutBoundsMax = Center + InExtents;
 	}
 		
-	FORCEINLINE FVector GetScaledExtents(const FTransform& InTransform, const FVector& InBoundsMin, const FVector& InBoundsMax)
+	inline FVector GetScaledExtents(const FTransform& InTransform, const FVector& InBoundsMin, const FVector& InBoundsMax)
 	{ 
 		return PCGPointHelpers::GetExtents(InBoundsMin, InBoundsMax) * InTransform.GetScale3D();
 	}
 
-	FORCEINLINE void SetLocalCenter(const FVector& InCenter, FVector& InOutBoundsMin, FVector& InOutBoundsMax)
+	inline void SetLocalCenter(const FVector& InCenter, FVector& InOutBoundsMin, FVector& InOutBoundsMax)
 	{
 		const FVector Delta = InCenter - PCGPointHelpers::GetLocalCenter(InOutBoundsMin, InOutBoundsMax);
 		InOutBoundsMin += Delta;
 		InOutBoundsMax += Delta;
 	}
 
-	FORCEINLINE FVector GetLocalSize(const FVector& InBoundsMin, const FVector& InBoundsMax)
+	inline FVector GetLocalSize(const FVector& InBoundsMin, const FVector& InBoundsMax)
 	{ 
 		return InBoundsMax - InBoundsMin; 
 	}
 	
-	FORCEINLINE FVector GetScaledLocalSize(const FTransform& InTransform, const FVector& InBoundsMin, const FVector& InBoundsMax)
+	inline FVector GetScaledLocalSize(const FTransform& InTransform, const FVector& InBoundsMin, const FVector& InBoundsMax)
 	{ 
 		return PCGPointHelpers::GetLocalSize(InBoundsMin, InBoundsMax) * InTransform.GetScale3D();
 	}
 
-	FORCEINLINE void ApplyScaleToBounds(FTransform& InOutTransform, FVector& InOutBoundsMin, FVector& InOutBoundsMax)
+	inline void ApplyScaleToBounds(FTransform& InOutTransform, FVector& InOutBoundsMin, FVector& InOutBoundsMax)
 	{
 		const FVector PointScale = InOutTransform.GetScale3D();
 		InOutTransform.SetScale3D(PointScale.GetSignVector());
@@ -86,7 +86,7 @@ namespace PCGPointHelpers
 		InOutBoundsMax *= PointScale.GetAbs();
 	}
 
-	FORCEINLINE void ResetPointCenter(const FVector& BoundsRatio, FTransform& InOutTransform, FVector& InOutBoundsMin, FVector& InOutBoundsMax)
+	inline void ResetPointCenter(const FVector& BoundsRatio, FTransform& InOutTransform, FVector& InOutBoundsMin, FVector& InOutBoundsMax)
 	{
 		const FVector NewCenterLocal = FMath::Lerp(InOutBoundsMin, InOutBoundsMax, BoundsRatio);
 
diff --git a/Engine/Plugins/PCG/Source/PCG/Public/PCGPointArray.h b/Engine/Plugins/PCG/Source/PCG/Public/PCGPointArray.h
--- a/Engine/Plugins/PCG/Source/PCG/Public/PCGPointArray.h
+++ b/Engine/Plugins/PCG/Source/PCG/Public/PCGPointArray.h
@@ -34,7 +34,7 @@ public:
 	}
 
 	/** Defined for serialization purposes(FPCGPointArray trait WithIdenticalViaEquality) */
-	[[nodiscard]] FORCEINLINE bool operator==(const FPCGPointArrayProperty<T>& Other) const
+	[[nodiscard]] inline bool operator==(const FPCGPointArrayProperty<T>& Other) const
 	{
 		if (NumValues != Other.NumValues || Values.Num() != Other.Values.Num())
 		{
@@ -70,13 +70,13 @@ public:
 	}
 	
 	/** If 'Values' isn't currently allocated, allocate it and copy 'Value' at every index of 'Values[0, NumValues - 1]' */
-	FORCEINLINE void Allocate(bool bInitializeValues)
+	inline void Allocate(bool bInitializeValues)
 	{
 		SetNum(NumValues, /*bAllocate=*/true, bInitializeValues);
 	}
 
 	/** Set the 'NumValues', optionally allocating 'Values' */
-	FORCEINLINE void SetNum(int32 NewNum, bool bAllocate, bool bInitializeValues)
+	inline void SetNum(int32 NewNum, bool bAllocate, bool bInitializeValues)
 	{
 		NumValues = NewNum;
 
@@ -99,19 +99,19 @@ public:
 		}
 	};
 
-	FORCEINLINE int32 Num() const
+	inline int32 Num() const
 	{
 		return NumValues;
 	}
 
 	/** Free 'Values', effectively resets all values to 'Value' */
-	FORCEINLINE void Free()
+	inline void Free()
 	{
 		Values.Reset();
 	}
 
 	/** Set all values('Value/Values') to InValue */
-	FORCEINLINE void SetValue(const T& InValue)
+	inline void SetValue(const T& InValue)
 	{
 		for (int32 i = 0; i < Values.Num(); ++i)
 		{
@@ -121,7 +121,7 @@ public:
 	}
 
 	/** Return value at Index */
-	[[nodiscard]] FORCEINLINE const T& GetValue(int32 Index) const
+	[[nodiscard]] inline const T& GetValue(int32 Index) const
 	{
 		ensure(Index >= 0 && Index < NumValues);
 		if (Values.Num() == 0)
@@ -133,7 +133,7 @@ public:
 	}
 
 	/** Move ranges of values(moves them back into the array).Used by the PCGAsync api */
-	FORCEINLINE void MoveRange(int32 RangeStartIndex, int32 MoveToIndex, int32 NumElements)
+	inline void MoveRange(int32 RangeStartIndex, int32 MoveToIndex, int32 NumElements)
 	{
 		if (IsAllocated())
 		{
@@ -143,14 +143,14 @@ public:
 	}
 
 	/* Returns true if 'Values' is allocated */
-	[[nodiscard]] FORCEINLINE bool IsAllocated() const
+	[[nodiscard]] inline bool IsAllocated() const
 	{
 		ensure(Values.Num() == 0 || Values.Num() == NumValues);
 		return Values.Num() != 0;
 	}
 
 	/** Returns a TArray copy of size NumValues */
-	[[nodiscard]] FORCEINLINE TArray<T> GetCopy() const
+	[[nodiscard]] inline TArray<T> GetCopy() const
 	{
 		if (IsAllocated())
 		{
@@ -168,7 +168,7 @@ public:
 	}
 
 	/** Returns a TPCGValueRange<T> of NumValues and optionally allocate 'Values' */
-	[[nodiscard]] FORCEINLINE TPCGValueRange<T> GetValueRange(bool bAllocate, bool bInitializeValues)
+	[[nodiscard]] inline TPCGValueRange<T> GetValueRange(bool bAllocate, bool bInitializeValues)
 	{
 		if (bAllocate)
 		{
@@ -180,24 +180,24 @@ public:
 	}
 
 	/** Returns a TConstPCGValueRange<T> of NumValues */
-	[[nodiscard]] FORCEINLINE TConstPCGValueRange<T> GetConstValueRange() const
+	[[nodiscard]] inline TConstPCGValueRange<T> GetConstValueRange() const
 	{
 		return TConstPCGValueRange<T>(IsAllocated() ? MakeConstStridedView(Values) : MakeConstStridedView(sizeof(Value), &Value, 1), NumValues);
 	}
 		
-	[[nodiscard]] FORCEINLINE SIZE_T GetSizeBytes() const
+	[[nodiscard]] inline SIZE_T GetSizeBytes() const
 	{
 		return Values.GetAllocatedSize() + sizeof(Value) + sizeof(NumValues) + sizeof(Values);
 	}
 
-	FORCEINLINE void Serialize(FArchive& Ar)
+	inline void Serialize(FArchive& Ar)
 	{
 		Ar << NumValues;
 		Ar << Value;
 		Ar << Values;
 	}
 
-	FORCEINLINE void CopyTo(FPCGPointArrayProperty<T>& OutProperty, int32 StartReadIndex, int32 StartWriteIndex, int32 Count) const
+	inline void CopyTo(FPCGPointArrayProperty<T>& OutProperty, int32 StartReadIndex, int32 StartWriteIndex, int32 Count) const
 	{
 		check(StartReadIndex + Count <= NumValues && StartWriteIndex + Count <= OutProperty.NumValues);
 		if (!IsAllocated())
@@ -221,7 +221,7 @@ public:
 		}
 	}
 
-	FORCEINLINE void CopyUnallocatedProperty(FPCGPointArrayProperty<T>& OutProperty) const
+	inline void CopyUnallocatedProperty(FPCGPointArrayProperty<T>& OutProperty) const
 	{
 		if (!IsAllocated() && !OutProperty.IsAllocated())
 		{
diff --git a/Engine/Plugins/PCG/Source/PCG/Public/Utils/PCGPointOctree.h b/Engine/Plugins/PCG/Source/PCG/Public/Utils/PCGPointOctree.h
--- a/Engine/Plugins/PCG/Source/PCG/Public/Utils/PCGPointOctree.h
+++ b/Engine/Plugins/PCG/Source/PCG/Public/Utils/PCGPointOctree.h
@@ -29,23 +29,23 @@ namespace PCGPointOctree
 
 		typedef TInlineAllocator<MaxElementsPerLeaf> ElementAllocator;
 
-		FORCEINLINE static const FBoxSphereBounds& GetBoundingBox(const FPointRef& InPoint)
+		inline static const FBoxSphereBounds& GetBoundingBox(const FPointRef& InPoint)
 		{
 			return InPoint.Bounds;
 		}
 
-		FORCEINLINE static const bool AreElementsEqual(const FPointRef& A, const FPointRef& B)
+		inline static const bool AreElementsEqual(const FPointRef& A, const FPointRef& B)
 		{
 			// TODO: verify if that's sufficient
 			return A.Index == B.Index;
 		}
 
-		FORCEINLINE static void ApplyOffset(FPointRef& InPoint)
+		inline static void ApplyOffset(FPointRef& InPoint)
 		{
 			ensureMsgf(false, TEXT("Not implemented"));
 		}
 
-		FORCEINLINE static void SetElementId(const FPointRef& Element, FOctreeElementId2 OctreeElementID)
+		inline static void SetElementId(const FPointRef& Element, FOctreeElementId2 OctreeElementID)
 		{
 		}
 	};
diff --git a/Engine/Plugins/PCG/Source/PCG/Public/Utils/PCGValueRange.h b/Engine/Plugins/PCG/Source/PCG/Public/Utils/PCGValueRange.h
--- a/Engine/Plugins/PCG/Source/PCG/Public/Utils/PCGValueRange.h
+++ b/Engine/Plugins/PCG/Source/PCG/Public/Utils/PCGValueRange.h
@@ -63,37 +63,37 @@ public:
 	{
 	}
 		
-	[[nodiscard]] FORCEINLINE TOptional<ElementType> GetSingleValue() const
+	[[nodiscard]] inline TOptional<ElementType> GetSingleValue() const
 	{
 		return ViewNum() == 1 && Num() > 0 ? TOptional<ElementType>(GetElement(0)) : TOptional<ElementType>();
 	}
 
-	[[nodiscard]] FORCEINLINE bool IsValidIndex(SizeType Index) const
+	[[nodiscard]] inline bool IsValidIndex(SizeType Index) const
 	{
 		return (Index >= 0) && (Index < NumElements);
 	}
 
-	[[nodiscard]] FORCEINLINE bool IsEmpty() const
+	[[nodiscard]] inline bool IsEmpty() const
 	{
 		return NumElements == 0;
 	}
 
-	[[nodiscard]] FORCEINLINE SizeType Num() const
+	[[nodiscard]] inline SizeType Num() const
 	{
 		return NumElements;
 	}
 
-	[[nodiscard]] FORCEINLINE SizeType ViewNum() const
+	[[nodiscard]] inline SizeType ViewNum() const
 	{
 		return FMath::Min(NumElements, ElementView.Num());
 	}
 
-	[[nodiscard]] FORCEINLINE ElementType& operator[](SizeType Index) const
+	[[nodiscard]] inline ElementType& operator[](SizeType Index) const
 	{
 		return GetElement(Index);
 	}
 
-	FORCEINLINE void Set(const ElementType& Value)
+	inline void Set(const ElementType& Value)
 	{
 		// Set all values in underlying view to same value
 		for (SizeType i = 0; i < ElementView.Num(); ++i)
@@ -107,39 +107,39 @@ public:
 		const TPCGValueRange* Owner;
 		SizeType Index;
 
-		FORCEINLINE FIterator& operator++()
+		inline FIterator& operator++()
 		{
 			++Index;
 			return *this;
 		}
 
-		FORCEINLINE ElementType& operator*()
+		inline ElementType& operator*()
 		{
 			return Owner->GetElementUnsafe(Index);
 		}
 
-		[[nodiscard]] FORCEINLINE bool operator==(const FIterator& Other) const
+		[[nodiscard]] inline bool operator==(const FIterator& Other) const
 		{
 			return Owner == Other.Owner
 				&& Index == Other.Index;
 		}
 
-		[[nodiscard]] FORCEINLINE bool operator!=(const FIterator& Other) const
+		[[nodiscard]] inline bool operator!=(const FIterator& Other) const
 		{
 			return !(*this == Other);
 		}
 	};
 
-	[[nodiscard]] FORCEINLINE FIterator begin() const { return FIterator{ this, 0 }; }
-	[[nodiscard]] FORCEINLINE FIterator end() const { return FIterator{ this, Num() }; }
+	[[nodiscard]] inline FIterator begin() const { return FIterator{ this, 0 }; }
+	[[nodiscard]] inline FIterator end() const { return FIterator{ this, Num() }; }
 
 private:
-	FORCEINLINE void RangeCheck(SizeType Index) const
+	inline void RangeCheck(SizeType Index) const
 	{
 		checkf((Index >= 0) & (Index < NumElements), TEXT("Array index out of bounds: %" INT64_FMT " from an array of size %" INT64_FMT), int64(Index), int64(NumElements))
 	}
 	
-	[[nodiscard]] FORCEINLINE ElementType& GetElementUnsafe(SizeType Index) const
+	[[nodiscard]] inline ElementType& GetElementUnsafe(SizeType Index) const
 	{
 		const SizeType ViewIndex = Index % ElementView.Num();
 		ElementType& Value = ElementView.GetUnsafe(ViewIndex);
@@ -166,7 +166,7 @@ private:
 		}
 	}
 
-	[[nodiscard]] FORCEINLINE ElementType& GetElement(SizeType Index) const
+	[[nodiscard]] inline ElementType& GetElement(SizeType Index) const
 	{
 		RangeCheck(Index);
 		return GetElementUnsafe(Index);
diff --git a/Engine/Plugins/Runtime/GeoReferencing/Source/GeoReferencing/Public/GeographicCoordinates.h b/Engine/Plugins/Runtime/GeoReferencing/Source/GeoReferencing/Public/GeographicCoordinates.h
--- a/Engine/Plugins/Runtime/GeoReferencing/Source/GeoReferencing/Public/GeographicCoordinates.h
+++ b/Engine/Plugins/Runtime/GeoReferencing/Source/GeoReferencing/Public/GeographicCoordinates.h
@@ -43,7 +43,7 @@ class GEOREFERENCING_API UGeographicCoordinatesFunctionLibrary : public UBluepri
 	 * Converts a GeographicCoordinates value to localized formatted text, in the form 'X= Y= Z='
 	 **/
 	UFUNCTION(BlueprintPure, meta = (DisplayName = "ToFullText", AdvancedDisplay = "1", BlueprintAutocast), Category = "GeoReferencing")
-	static FORCEINLINE FText ToFullText(UPARAM(ref) FGeographicCoordinates& GeographicCoordinates, int32 IntegralDigitsLatLon = 8, int32 IntegralDigitsAlti = 2, bool bAsDMS = false)
+	static inline FText ToFullText(UPARAM(ref) FGeographicCoordinates& GeographicCoordinates, int32 IntegralDigitsLatLon = 8, int32 IntegralDigitsAlti = 2, bool bAsDMS = false)
 	{
 		return GeographicCoordinates.ToFullText(IntegralDigitsLatLon, IntegralDigitsAlti, bAsDMS);
 	}
@@ -52,7 +52,7 @@ class GEOREFERENCING_API UGeographicCoordinatesFunctionLibrary : public UBluepri
 	 * Converts a GeographicCoordinates value to formatted text, in the form '(X, Y, Z)'
 	 **/
 	UFUNCTION(BlueprintPure, meta = (DisplayName = "ToCompactText", AdvancedDisplay = "1", BlueprintAutocast), Category = "GeoReferencing")
-	static FORCEINLINE FText ToCompactText(UPARAM(ref) FGeographicCoordinates& GeographicCoordinates, int32 IntegralDigitsLatLon = 8, int32 IntegralDigitsAlti = 2, bool bAsDMS = false)
+	static inline FText ToCompactText(UPARAM(ref) FGeographicCoordinates& GeographicCoordinates, int32 IntegralDigitsLatLon = 8, int32 IntegralDigitsAlti = 2, bool bAsDMS = false)
 	{
 		return GeographicCoordinates.ToCompactText(IntegralDigitsLatLon, IntegralDigitsAlti, bAsDMS);
 	}
diff --git a/Engine/Source/Runtime/Engine/Public/WorldPartition/ActorDescContainerCollection.h b/Engine/Source/Runtime/Engine/Public/WorldPartition/ActorDescContainerCollection.h
--- a/Engine/Source/Runtime/Engine/Public/WorldPartition/ActorDescContainerCollection.h
+++ b/Engine/Source/Runtime/Engine/Public/WorldPartition/ActorDescContainerCollection.h
@@ -131,7 +131,7 @@ public:
 		 *
 		 * @return	Current suitable actor desc
 		 */
-		FORCEINLINE ReturnType operator*() const
+		inline ReturnType operator*() const
 		{
 			return StaticCast<ReturnType>(**ActorsIterator);
 		}
@@ -141,7 +141,7 @@ public:
 		 *
 		 * @return	Current suitable actor desc
 		 */
-		FORCEINLINE ReturnType operator->() const
+		inline ReturnType operator->() const
 		{
 			return StaticCast<ReturnType>(**ActorsIterator);
 		}
@@ -151,13 +151,13 @@ public:
 		 *
 		 * @return true if iterator points to a suitable actor desc, false if it has reached the end
 		 */
-		FORCEINLINE explicit operator bool() const
+		inline explicit operator bool() const
 		{
 			return ActorsIterator ? (bool)*ActorsIterator : false;
 		}
 
 	protected:
-		FORCEINLINE void AdvanceToRelevantActorInNextContainer()
+		inline void AdvanceToRelevantActorInNextContainer()
 		{
 			while (!(*ActorsIterator) && ContainerIterator)
 			{
diff --git a/Engine/Source/Runtime/Engine/Public/WorldPartition/ActorDescContainerInstanceCollection.h b/Engine/Source/Runtime/Engine/Public/WorldPartition/ActorDescContainerInstanceCollection.h
--- a/Engine/Source/Runtime/Engine/Public/WorldPartition/ActorDescContainerInstanceCollection.h
+++ b/Engine/Source/Runtime/Engine/Public/WorldPartition/ActorDescContainerInstanceCollection.h
@@ -135,7 +135,7 @@ public:
 		 *
 		 * @return	Current suitable actor desc
 		 */
-		FORCEINLINE ReturnType operator*() const
+		inline ReturnType operator*() const
 		{
 			return StaticCast<ReturnType>(**ActorsIterator);
 		}
@@ -145,7 +145,7 @@ public:
 		 *
 		 * @return	Current suitable actor desc
 		 */
-		FORCEINLINE ReturnType operator->() const
+		inline ReturnType operator->() const
 		{
 			return StaticCast<ReturnType>(**ActorsIterator);
 		}
@@ -155,13 +155,13 @@ public:
 		 *
 		 * @return true if iterator points to a suitable actor desc, false if it has reached the end
 		 */
-		FORCEINLINE explicit operator bool() const
+		inline explicit operator bool() const
 		{
 			return ActorsIterator ? (bool)*ActorsIterator : false;
 		}
 
 	protected:
-		FORCEINLINE void AdvanceToRelevantActorInNextContainer()
+		inline void AdvanceToRelevantActorInNextContainer()
 		{
 			while (!(*ActorsIterator) && ContainerIterator)
 			{
diff --git a/Engine/Source/Runtime/Engine/Public/WorldPartition/ActorDescList.h b/Engine/Source/Runtime/Engine/Public/WorldPartition/ActorDescList.h
--- a/Engine/Source/Runtime/Engine/Public/WorldPartition/ActorDescList.h
+++ b/Engine/Source/Runtime/Engine/Public/WorldPartition/ActorDescList.h
@@ -24,7 +24,7 @@ class TActorDescList
 
 	struct FActorGuidKeyFuncs : TDefaultMapKeyFuncs<FGuid, TUniquePtr<DescType>*, false>
 	{
-		static FORCEINLINE uint32 GetKeyHash(const FGuid& Key)
+		static inline uint32 GetKeyHash(const FGuid& Key)
 		{
 			return HashCombineFast(HashCombineFast(Key.A, Key.B), HashCombineFast(Key.C, Key.D));
 		}
@@ -80,7 +80,7 @@ public:
 		/**
 		 * Iterates to next suitable actor desc
 		 */
-		FORCEINLINE void operator++()
+		inline void operator++()
 		{
 			do
 			{
@@ -91,7 +91,7 @@ public:
 		/**
 		 * Removes the current iterator element
 		 */
-		FORCEINLINE void RemoveCurrent()
+		inline void RemoveCurrent()
 		{
 			check(!bIsProxy);
 			ActorsIterator.RemoveCurrent();
@@ -102,7 +102,7 @@ public:
 		 *
 		 * @return	Current suitable actor desc
 		 */
-		FORCEINLINE ReturnType operator*() const
+		inline ReturnType operator*() const
 		{
 			return StaticCast<ReturnType>(ActorsIterator->Value->Get());
 		}
@@ -112,7 +112,7 @@ public:
 		 *
 		 * @return	Current suitable actor desc
 		 */
-		FORCEINLINE ReturnType operator->() const
+		inline ReturnType operator->() const
 		{
 			return StaticCast<ReturnType>(ActorsIterator->Value->Get());
 		}
@@ -123,7 +123,7 @@ public:
 		 *
 		 * @return true if iterator points to a suitable actor desc, false if it has reached the end
 		 */
-		FORCEINLINE explicit operator bool() const
+		inline explicit operator bool() const
 		{
 			return (bool)ActorsIterator;
 		}
@@ -133,14 +133,14 @@ public:
 		 *
 		 * @return the actor class
 		 */
-		FORCEINLINE UClass* GetActorClass() const { return ActorClass; }
+		inline UClass* GetActorClass() const { return ActorClass; }
 
 	protected:
 		/**
 		 * Determines whether the iterator currently points to a valid actor desc or not
 		 * @return true if we should skip the actor desc
 		 */
-		FORCEINLINE bool ShouldSkip() const
+		inline bool ShouldSkip() const
 		{
 			if (!ActorsIterator)
 			{
diff --git a/Engine/Source/Runtime/Engine/Public/WorldPartition/DataLayer/ActorDataLayer.h b/Engine/Source/Runtime/Engine/Public/WorldPartition/DataLayer/ActorDataLayer.h
--- a/Engine/Source/Runtime/Engine/Public/WorldPartition/DataLayer/ActorDataLayer.h
+++ b/Engine/Source/Runtime/Engine/Public/WorldPartition/DataLayer/ActorDataLayer.h
@@ -26,10 +26,10 @@ struct FActorDataLayer
 	: Name(InName)
 	{}
 
-	FORCEINLINE bool operator==(const FActorDataLayer& Other) const { return Name == Other.Name; }
-	FORCEINLINE bool operator<(const FActorDataLayer& Other) const { return Name.FastLess(Other.Name); }
+	inline bool operator==(const FActorDataLayer& Other) const { return Name == Other.Name; }
+	inline bool operator<(const FActorDataLayer& Other) const { return Name.FastLess(Other.Name); }
 
-	FORCEINLINE operator FName () const { return Name; }
+	inline operator FName () const { return Name; }
 
 	/** The name of this layer */
 	UPROPERTY(BlueprintReadOnly, VisibleAnywhere, Category = DataLayer)
diff --git a/Engine/Source/Runtime/Engine/Public/WorldPartition/DataLayer/DataLayerEditorContext.h b/Engine/Source/Runtime/Engine/Public/WorldPartition/DataLayer/DataLayerEditorContext.h
--- a/Engine/Source/Runtime/Engine/Public/WorldPartition/DataLayer/DataLayerEditorContext.h
+++ b/Engine/Source/Runtime/Engine/Public/WorldPartition/DataLayer/DataLayerEditorContext.h
@@ -19,9 +19,9 @@ public:
 	static const uint32 EmptyHash = 0;
 	FDataLayerEditorContext() : Hash(FDataLayerEditorContext::EmptyHash) {}
 	ENGINE_API FDataLayerEditorContext(UWorld* InWorld, const TArray<FName>& InDataLayerInstances);
-	FORCEINLINE bool IsEmpty() const { return (Hash == FDataLayerEditorContext::EmptyHash) && DataLayerInstances.IsEmpty(); }
-	FORCEINLINE uint32 GetHash() const { return Hash; }
-	FORCEINLINE const TArray<FName>& GetDataLayerInstanceNames() const { return DataLayerInstances; }
+	inline bool IsEmpty() const { return (Hash == FDataLayerEditorContext::EmptyHash) && DataLayerInstances.IsEmpty(); }
+	inline uint32 GetHash() const { return Hash; }
+	inline const TArray<FName>& GetDataLayerInstanceNames() const { return DataLayerInstances; }
 private:
 	uint32 Hash;
 	TArray<FName> DataLayerInstances;
diff --git a/Engine/Source/Runtime/Engine/Public/WorldPartition/WorldPartition.h b/Engine/Source/Runtime/Engine/Public/WorldPartition/WorldPartition.h
--- a/Engine/Source/Runtime/Engine/Public/WorldPartition/WorldPartition.h
+++ b/Engine/Source/Runtime/Engine/Public/WorldPartition/WorldPartition.h
@@ -203,7 +203,7 @@ public:
 	ENGINE_API const FTransform& GetInstanceTransform() const;
 	//~ End UActorDescContainer Interface
 
-	FORCEINLINE bool HasInstanceTransform() const { return InstanceTransform.IsSet(); }
+	inline bool HasInstanceTransform() const { return InstanceTransform.IsSet(); }
 
 	//~ Begin UObject Interface
 #if WITH_EDITOR
diff --git a/Engine/Source/Runtime/Engine/Public/WorldPartition/WorldPartitionActorContainerID.h b/Engine/Source/Runtime/Engine/Public/WorldPartition/WorldPartitionActorContainerID.h
--- a/Engine/Source/Runtime/Engine/Public/WorldPartition/WorldPartitionActorContainerID.h
+++ b/Engine/Source/Runtime/Engine/Public/WorldPartition/WorldPartitionActorContainerID.h
@@ -115,7 +115,7 @@ struct FActorContainerID
 		return Guid.IsValid() ? FGuid::Combine(Guid, InActorGuid) : InActorGuid;
 	}
 
-	friend FORCEINLINE uint32 GetTypeHash(const FActorContainerID& InContainerID)
+	friend inline uint32 GetTypeHash(const FActorContainerID& InContainerID)
 	{
 		return GetTypeHash(InContainerID.Guid);
 	}
diff --git a/Engine/Source/Runtime/Engine/Public/WorldPartition/WorldPartitionActorDescViewMap.h b/Engine/Source/Runtime/Engine/Public/WorldPartition/WorldPartitionActorDescViewMap.h
--- a/Engine/Source/Runtime/Engine/Public/WorldPartition/WorldPartitionActorDescViewMap.h
+++ b/Engine/Source/Runtime/Engine/Public/WorldPartition/WorldPartitionActorDescViewMap.h
@@ -68,7 +68,7 @@ public:
 		return Emplace(InActorDesc->GetGuid(), Type(InActorDesc));
 	}
 
-	FORCEINLINE int32 Num() const
+	inline int32 Num() const
 	{
 		return ActorDescViewList.Num();
 	}
diff --git a/Engine/Source/Runtime/Engine/Public/WorldPartition/WorldPartitionDebugHelper.h b/Engine/Source/Runtime/Engine/Public/WorldPartition/WorldPartitionDebugHelper.h
--- a/Engine/Source/Runtime/Engine/Public/WorldPartition/WorldPartitionDebugHelper.h
+++ b/Engine/Source/Runtime/Engine/Public/WorldPartition/WorldPartitionDebugHelper.h
@@ -22,9 +22,9 @@ public:
 	static bool IsDebugCellNameShown(const FString& Name);
 	static void DrawText(UCanvas* Canvas, const FString& Text, const UFont* Font, const FColor& Color, FVector2D& Pos, float* MaxTextWidth = nullptr);
 	static void DrawLegendItem(UCanvas* Canvas, const FString& Text, const UFont* Font, const FColor& Color, const FColor& TextColor, FVector2D& Pos, float* MaxItemWidth = nullptr);
-	FORCEINLINE static bool IsRuntimeSpatialHashCellStreamingPriorityShown() { return ShowRuntimeSpatialHashCellStreamingPriorityMode != 0; }
-	FORCEINLINE static int32 GetRuntimeSpatialHashCellStreamingPriorityMode() { return ShowRuntimeSpatialHashCellStreamingPriorityMode; }
-	FORCEINLINE static bool CanDrawContentBundles() { return bCanDrawContentBundles; }
+	inline static bool IsRuntimeSpatialHashCellStreamingPriorityShown() { return ShowRuntimeSpatialHashCellStreamingPriorityMode != 0; }
+	inline static int32 GetRuntimeSpatialHashCellStreamingPriorityMode() { return ShowRuntimeSpatialHashCellStreamingPriorityMode; }
+	inline static bool CanDrawContentBundles() { return bCanDrawContentBundles; }
 
 private:
 	static FAutoConsoleCommand DebugFilterByStreamingStatusCommand;
diff --git a/Engine/Source/Runtime/Engine/Public/WorldPartition/WorldPartitionHandle.h b/Engine/Source/Runtime/Engine/Public/WorldPartition/WorldPartitionHandle.h
--- a/Engine/Source/Runtime/Engine/Public/WorldPartition/WorldPartitionHandle.h
+++ b/Engine/Source/Runtime/Engine/Public/WorldPartition/WorldPartitionHandle.h
@@ -111,12 +111,12 @@ template <typename Impl>
 class TWorldPartitionHandle
 {
 public:
-	FORCEINLINE TWorldPartitionHandle()
+	inline TWorldPartitionHandle()
 		: ContainerInstance(nullptr)
 		, ActorDescInstance(nullptr)
 	{}
 
-	FORCEINLINE TWorldPartitionHandle(TUniquePtr<FWorldPartitionActorDescInstance>* InActorDescInstance)
+	inline TWorldPartitionHandle(TUniquePtr<FWorldPartitionActorDescInstance>* InActorDescInstance)
 		: ContainerInstance(Impl::GetActorDescContainerInstance(InActorDescInstance))
 		, ActorDescInstance(InActorDescInstance)
 	{
@@ -126,7 +126,7 @@ public:
 		}
 	}
 
-	FORCEINLINE TWorldPartitionHandle(FWorldPartitionActorDescInstance* InActorDescInstance)
+	inline TWorldPartitionHandle(FWorldPartitionActorDescInstance* InActorDescInstance)
 		: ContainerInstance(Impl::GetActorDescContainerInstance(InActorDescInstance))
 		, ActorDescInstance(nullptr)
 	{
@@ -141,7 +141,7 @@ public:
 		}
 	}
 
-	FORCEINLINE TWorldPartitionHandle(UActorDescContainerInstance* InContainerInstance, const FGuid& InActorGuid)
+	inline TWorldPartitionHandle(UActorDescContainerInstance* InContainerInstance, const FGuid& InActorGuid)
 		: ContainerInstance(InContainerInstance)
 		, ActorDescInstance(Impl::GetActorDescInstance(InContainerInstance, InActorGuid))
 	{
@@ -152,7 +152,7 @@ public:
 	}
 
 	template<class U>
-	FORCEINLINE TWorldPartitionHandle(TActorDescContainerInstanceCollection<U>* ContainerCollection, const FGuid& ActorGuid)
+	inline TWorldPartitionHandle(TActorDescContainerInstanceCollection<U>* ContainerCollection, const FGuid& ActorGuid)
 		: ContainerInstance(Impl::GetActorDescContainerInstance(ContainerCollection, ActorGuid))
 		, ActorDescInstance(nullptr)
 	{
@@ -168,14 +168,14 @@ public:
 	}
 
 	UE_DEPRECATED(5.4, "Use FWorldPartitionActorDescInstance version")
-	FORCEINLINE TWorldPartitionHandle(TUniquePtr<FWorldPartitionActorDesc>* InActorDesc)
+	inline TWorldPartitionHandle(TUniquePtr<FWorldPartitionActorDesc>* InActorDesc)
 		: ContainerInstance(nullptr)
 		, ActorDescInstance(nullptr)
 	{
 	}
 
 	UE_DEPRECATED(5.4, "Use FActorDescContainerInstance version")
-	FORCEINLINE TWorldPartitionHandle(UActorDescContainer* InContainer, const FGuid& ActorGuid)
+	inline TWorldPartitionHandle(UActorDescContainer* InContainer, const FGuid& ActorGuid)
 		: ContainerInstance(nullptr)
 		, ActorDescInstance(nullptr)
 	{
@@ -189,45 +189,45 @@ public:
 	};
 
 	template<class U>
-	FORCEINLINE TWorldPartitionHandle(TActorDescContainerCollection<U>* ContainerCollection, const FGuid& ActorGuid)
+	inline TWorldPartitionHandle(TActorDescContainerCollection<U>* ContainerCollection, const FGuid& ActorGuid)
 		: ContainerInstance(nullptr)
 		, ActorDescInstance(nullptr)
 	{
 		TActorDescContainerCollectionDeprecated<U>::Deprecated();
 	}
 	
-	FORCEINLINE TWorldPartitionHandle(const TWorldPartitionHandle& Other)
+	inline TWorldPartitionHandle(const TWorldPartitionHandle& Other)
 		: ContainerInstance(nullptr)
 		, ActorDescInstance(nullptr)
 	{
 		*this = Other;
 	}
 
-	FORCEINLINE TWorldPartitionHandle(TWorldPartitionHandle&& Other)
+	inline TWorldPartitionHandle(TWorldPartitionHandle&& Other)
 		: ContainerInstance(nullptr)
 		, ActorDescInstance(nullptr)
 	{
 		*this = MoveTemp(Other);
 	}
 
-	FORCEINLINE TWorldPartitionHandle<FWorldPartitionHandleImpl> ToHandle() const
+	inline TWorldPartitionHandle<FWorldPartitionHandleImpl> ToHandle() const
 	{
 		return Impl::ToHandle(*this);
 	}
 
-	FORCEINLINE TWorldPartitionHandle<FWorldPartitionReferenceImpl> ToReference() const
+	inline TWorldPartitionHandle<FWorldPartitionReferenceImpl> ToReference() const
 	{
 		return Impl::ToReference(*this);
 	}
 
 	template <typename T>
-	FORCEINLINE TWorldPartitionHandle<Impl>(TWorldPartitionHandle<T>&& Other)
+	inline TWorldPartitionHandle<Impl>(TWorldPartitionHandle<T>&& Other)
 		: ActorDescInstance(nullptr)
 	{
 		*this = MoveTemp(Other);
 	}
 
-	FORCEINLINE ~TWorldPartitionHandle()
+	inline ~TWorldPartitionHandle()
 	{
 		if (IsValid())
 		{
@@ -235,7 +235,7 @@ public:
 		}
 	}
 
-	FORCEINLINE TWorldPartitionHandle& operator=(const TWorldPartitionHandle& Other)
+	inline TWorldPartitionHandle& operator=(const TWorldPartitionHandle& Other)
 	{
 		if (this != &Other)
 		{
@@ -256,7 +256,7 @@ public:
 		return *this;
 	}
 
-	FORCEINLINE TWorldPartitionHandle<Impl>& operator=(TWorldPartitionHandle&& Other)
+	inline TWorldPartitionHandle<Impl>& operator=(TWorldPartitionHandle&& Other)
 	{
 		if (this != &Other)
 		{
@@ -277,7 +277,7 @@ public:
 
 	// Conversions
 	template <typename T>
-	FORCEINLINE TWorldPartitionHandle<Impl>& operator=(const TWorldPartitionHandle<T>& Other)
+	inline TWorldPartitionHandle<Impl>& operator=(const TWorldPartitionHandle<T>& Other)
 	{
 		if (IsValid())
 		{
@@ -296,7 +296,7 @@ public:
 	}
 
 	template <typename T>
-	FORCEINLINE TWorldPartitionHandle<Impl>& operator=(TWorldPartitionHandle<T>&& Other)
+	inline TWorldPartitionHandle<Impl>& operator=(TWorldPartitionHandle<T>&& Other)
 	{
 		if (IsValid())
 		{
@@ -318,38 +318,38 @@ public:
 		return *this;
 	}
 
-	FORCEINLINE FWorldPartitionActorDescInstance* operator->() const
+	inline FWorldPartitionActorDescInstance* operator->() const
 	{
 		return GetInstance();
 	}
 
-	FORCEINLINE FWorldPartitionActorDescInstance* operator*() const
+	inline FWorldPartitionActorDescInstance* operator*() const
 	{
 		return GetInstance();
 	}
 
-	FORCEINLINE FWorldPartitionActorDescInstance* GetInstance() const
+	inline FWorldPartitionActorDescInstance* GetInstance() const
 	{
 		return IsValid() ? ActorDescInstance->Get() : nullptr;
 	}
 
-	FORCEINLINE bool IsValid() const
+	inline bool IsValid() const
 	{
 		return ContainerInstance.IsValid() && ActorDescInstance && ActorDescInstance->IsValid();
 	}
 
-	FORCEINLINE bool IsLoaded() const
+	inline bool IsLoaded() const
 	{
 		return IsValid() && Impl::IsLoaded(ActorDescInstance->Get());
 	}
 
-	FORCEINLINE AActor* GetActor() const
+	inline AActor* GetActor() const
 	{
 		return IsValid() ? Impl::GetActor(ActorDescInstance->Get()) : nullptr;
 	}
 
 	UE_DEPRECATED(5.4, "Use GetInstance instead")
-	FORCEINLINE FWorldPartitionActorDesc* Get() const
+	inline FWorldPartitionActorDesc* Get() const
 	{
 		if (ContainerInstance.IsValid() && ActorDescInstance && ActorDescInstance->IsValid())
 		{
@@ -359,7 +359,7 @@ public:
 		return nullptr;
 	}
 
-	FORCEINLINE void Reset()
+	inline void Reset()
 	{
 		if (IsValid())
 		{
@@ -370,41 +370,41 @@ public:
 		ActorDescInstance = nullptr;
 	}
 	
-	friend FORCEINLINE uint32 GetTypeHash(const TWorldPartitionHandle<Impl>& HandleBase)
+	friend inline uint32 GetTypeHash(const TWorldPartitionHandle<Impl>& HandleBase)
 	{
 		return ::PointerHash(HandleBase.ActorDescInstance);
 	}
 
-	FORCEINLINE bool operator==(const TWorldPartitionHandle& Other) const
+	inline bool operator==(const TWorldPartitionHandle& Other) const
 	{
 		return ActorDescInstance == Other.ActorDescInstance;
 	}
 
-	FORCEINLINE bool operator!=(const TWorldPartitionHandle& Other) const
+	inline bool operator!=(const TWorldPartitionHandle& Other) const
 	{
 		return !(*this == Other);
 	}
 
 	// Conversions
 	template <typename T>
-	FORCEINLINE bool operator==(const TWorldPartitionHandle<T>& Other) const
+	inline bool operator==(const TWorldPartitionHandle<T>& Other) const
 	{
 		return GetInstance() == Other.GetInstance();
 	}
 
 	template <typename T>
-	FORCEINLINE bool operator!=(const TWorldPartitionHandle<T>& Other) const
+	inline bool operator!=(const TWorldPartitionHandle<T>& Other) const
 	{
 		return !(*this == Other);
 	}
 
 public:
-	FORCEINLINE void IncRefCount()
+	inline void IncRefCount()
 	{
 		Impl::IncRefCount(ActorDescInstance->Get());
 	}
 
-	FORCEINLINE void DecRefCount()
+	inline void DecRefCount()
 	{
 		Impl::DecRefCount(ActorDescInstance->Get());
 	}
diff --git a/Engine/Source/Runtime/Engine/Public/WorldPartition/WorldPartitionLevelHelper.h b/Engine/Source/Runtime/Engine/Public/WorldPartition/WorldPartitionLevelHelper.h
--- a/Engine/Source/Runtime/Engine/Public/WorldPartition/WorldPartitionLevelHelper.h
+++ b/Engine/Source/Runtime/Engine/Public/WorldPartition/WorldPartitionLevelHelper.h
@@ -131,7 +131,7 @@ private:
 
 		TArray<FActorPropertyOverride> ActorPropertyOverrides;
 		FTransform ContainerTransform;
-		FORCEINLINE bool IsDefault()
+		inline bool IsDefault()
 		{
 			return ActorPropertyOverrides.IsEmpty();
 		}
diff --git a/Engine/Source/Runtime/Engine/Public/WorldPartition/WorldPartitionRuntimeCell.h b/Engine/Source/Runtime/Engine/Public/WorldPartition/WorldPartitionRuntimeCell.h
--- a/Engine/Source/Runtime/Engine/Public/WorldPartition/WorldPartitionRuntimeCell.h
+++ b/Engine/Source/Runtime/Engine/Public/WorldPartition/WorldPartitionRuntimeCell.h
@@ -287,7 +287,7 @@ class UWorldPartitionRuntimeCell : public UObject, public IWorldPartitionCell, p
 
 	ENGINE_API UDataLayerManager* GetDataLayerManager() const;
 	ENGINE_API EDataLayerRuntimeState GetCellEffectiveWantedState(const FWorldPartitionStreamingContext& Context) const;
-	FORCEINLINE bool HasDataLayers() const { return !DataLayers.IsEmpty(); }
+	inline bool HasDataLayers() const { return !DataLayers.IsEmpty(); }
 
 	//~Begin Deprecation
 	UE_DEPRECATED(5.3, "IsLoading is deprecated.")
@@ -370,7 +370,7 @@ class UWorldPartitionRuntimeCell : public UObject, public IWorldPartitionCell, p
 #endif
 
 protected:
-	FORCEINLINE const TArray<FName>& GetDataLayersInline() const { return DataLayers.ToArray(); }
+	inline const TArray<FName>& GetDataLayersInline() const { return DataLayers.ToArray(); }
 
 	//@todo_ow: Implement ServerOnlyVisible and refactor ClientOnlyVisible.
 	//          Instead of this function, server would not not wait for client level visibility 
diff --git a/Engine/Source/Runtime/Engine/Public/WorldPartition/WorldPartitionStreamingGeneration.h b/Engine/Source/Runtime/Engine/Public/WorldPartition/WorldPartitionStreamingGeneration.h
--- a/Engine/Source/Runtime/Engine/Public/WorldPartition/WorldPartitionStreamingGeneration.h
+++ b/Engine/Source/Runtime/Engine/Public/WorldPartition/WorldPartitionStreamingGeneration.h
@@ -180,7 +180,7 @@ public:
 
 	PRAGMA_ENABLE_DEPRECATION_WARNINGS
 			
-	FORCEINLINE int32 Num() const
+	inline int32 Num() const
 	{
 		return ActorDescViewList.Num();
 	}
diff --git a/Engine/Source/Runtime/Engine/Public/WorldPartition/WorldPartitionStreamingSource.h b/Engine/Source/Runtime/Engine/Public/WorldPartition/WorldPartitionStreamingSource.h
--- a/Engine/Source/Runtime/Engine/Public/WorldPartition/WorldPartitionStreamingSource.h
+++ b/Engine/Source/Runtime/Engine/Public/WorldPartition/WorldPartitionStreamingSource.h
@@ -163,7 +163,7 @@ enum class EStreamingSourceTargetBehavior : uint8
 class FStreamingSourceShapeHelper
 {
 public:
-	FORCEINLINE static bool IsSourceAffectingGrid(const TSet<FName>& InSourceTargetGrids, EStreamingSourceTargetBehavior InSourceTargetBehavior, FName InGridName)
+	inline static bool IsSourceAffectingGrid(const TSet<FName>& InSourceTargetGrids, EStreamingSourceTargetBehavior InSourceTargetBehavior, FName InGridName)
 	{
 		switch (InSourceTargetBehavior)
 		{
@@ -178,7 +178,7 @@ public:
 		return false;
 	}
 
-	FORCEINLINE static void ForEachShape(float InGridLoadingRange, float InDefaultRadius, bool bInProjectIn2D, const FVector& InLocation, const FRotator& InRotation, const TArray<FStreamingSourceShape>& InShapes, TFunctionRef<void(const FSphericalSector&)> InOperation, float InExtraRadius = 0.f, float InExtraAngle = 0.f)
+	inline static void ForEachShape(float InGridLoadingRange, float InDefaultRadius, bool bInProjectIn2D, const FVector& InLocation, const FRotator& InRotation, const TArray<FStreamingSourceShape>& InShapes, TFunctionRef<void(const FSphericalSector&)> InOperation, float InExtraRadius = 0.f, float InExtraAngle = 0.f)
 	{
 		const FTransform Transform(bInProjectIn2D ? FRotator(0, InRotation.Yaw, 0) : InRotation, InLocation);
 		if (InShapes.IsEmpty())
@@ -305,7 +305,7 @@ struct FWorldPartitionStreamingQuerySource
 	TArray<FStreamingSourceShape> Shapes;
 
 	/** Helper method that iterates over all shapes. If none is provided, it will still pass a sphere shape using Radius or grid's loading range (see bUseGridLoadingRange). */
-	FORCEINLINE void ForEachShape(float InGridLoadingRange, bool bInProjectIn2D, TFunctionRef<void(const FSphericalSector&)> InOperation) const
+	inline void ForEachShape(float InGridLoadingRange, bool bInProjectIn2D, TFunctionRef<void(const FSphericalSector&)> InOperation) const
 	{
 		if (bSpatialQuery)
 		{
@@ -314,7 +314,7 @@ struct FWorldPartitionStreamingQuerySource
 	}
 
 	/** Helper method that iterates over all shapes. If none is provided, it will still pass a sphere shape using Radius or grid's loading range (see bUseGridLoadingRange). */
-	FORCEINLINE void ForEachShape(float InGridLoadingRange, FName InGridName, bool bInProjectIn2D, TFunctionRef<void(const FSphericalSector&)> InOperation) const
+	inline void ForEachShape(float InGridLoadingRange, FName InGridName, bool bInProjectIn2D, TFunctionRef<void(const FSphericalSector&)> InOperation) const
 	{
 		if (FStreamingSourceShapeHelper::IsSourceAffectingGrid(TargetGrids, TargetBehavior, InGridName))
 		{
