Splines - Fix ensure when querying for distance at input key

diff --git a/Engine/Source/Runtime/Engine/Private/Curves/Spline.cpp b/Engine/Source/Runtime/Engine/Private/Curves/Spline.cpp
--- a/Engine/Source/Runtime/Engine/Private/Curves/Spline.cpp
+++ b/Engine/Source/Runtime/Engine/Private/Curves/Spline.cpp
@@ -2208,7 +2208,7 @@ float FNewSpline::GetDistanceAtParameter(float Parameter)
 {
 	UpdateReparamTable();
 
-	const float ParameterMax = PositionCurve.Points.Last().InVal;
+	const float ParameterMax = static_cast<float>(GetNumberOfSegments());
 	const float Key = (Parameter / ParameterMax) * (ReparamTable.Points.Num() - 1);
 	const int32 LowerKey = FMath::FloorToInt32(Key);
 	ensureAlways(LowerKey >= 0 && LowerKey < ReparamTable.Points.Num());
