TimerManager now has some limited transactional support for DSTM.

diff --git a/Engine/Source/Runtime/Engine/Private/Actor.cpp b/Engine/Source/Runtime/Engine/Private/Actor.cpp
--- a/Engine/Source/Runtime/Engine/Private/Actor.cpp
+++ b/Engine/Source/Runtime/Engine/Private/Actor.cpp
@@ -6544,13 +6544,7 @@ float AActor::ActorGetDistanceToCollision(const FVector& Point, ECollisionChanne
 
 void AActor::LifeSpanExpired()
 {
-#if UE_WITH_REMOTE_OBJECT_HANDLE
-	// World Timer Manager is not yet transactional, so let's make this call transactional to make up for it
-	static FName TransactionalWorkName(TEXT("AActor::LifeSpanExpired"));
-	UE::RemoteExecutor::ExecuteTransactional(TransactionalWorkName, [this] { Destroy(); });
-#else
 	Destroy();
-#endif
 }
 
 void AActor::SetLifeSpan( float InLifespan )
diff --git a/Engine/Source/Runtime/Engine/Private/TimerManager.cpp b/Engine/Source/Runtime/Engine/Private/TimerManager.cpp
--- a/Engine/Source/Runtime/Engine/Private/TimerManager.cpp
+++ b/Engine/Source/Runtime/Engine/Private/TimerManager.cpp
@@ -346,18 +346,38 @@ void FTimerUnifiedDelegate::Execute() const
 			if (FuncDelegate.IsBound())
 			{
 				FScopeCycleCounterUObject Context(FuncDelegate.GetUObject());
+#if UE_WITH_REMOTE_OBJECT_HANDLE
+				const FName WorkName { TEXT("TimerManager::Execute FTimerDelegate") };
+				UE::RemoteExecutor::ExecuteTransactional(WorkName, [FuncDelegate, this]()
+					{
+						FuncDelegate.Execute();
+					});
+#else
 				FuncDelegate.Execute();
+#endif
 			}
 			break;
 		}
 	case FTimerDelegateVariant::IndexOfType<FTimerDynamicDelegate>():
 		{
 			const FTimerDynamicDelegate& FuncDynDelegate = VariantDelegate.Get<FTimerDynamicDelegate>();
+#if UE_WITH_REMOTE_OBJECT_HANDLE
+			const FName WorkName { TEXT("TimerManager::Execute FTimerDynamicDelegate") };
+			UE::RemoteExecutor::ExecuteTransactional(WorkName, [FuncDynDelegate, this]()
+			{
+				if (FuncDynDelegate.IsBound())
+				{
+					// stat scope is handled by UObject::ProcessEvent for the UFunction.
+					FuncDynDelegate.ProcessDelegate<UObject>(nullptr);
+				}
+			});
+#else
 			if (FuncDynDelegate.IsBound())
 			{
 				// stat scope is handled by UObject::ProcessEvent for the UFunction.
 				FuncDynDelegate.ProcessDelegate<UObject>(nullptr);
 			}
+#endif
 			break;
 		}
 	case FTimerDelegateVariant::IndexOfType<FTimerFunction>():
@@ -586,6 +606,10 @@ void FTimerManager::InternalSetTimer(FTimerHandle& InOutHandle, FTimerUnifiedDel
 	// not currently threadsafe
 	check(IsInGameThread());
 
+#if UE_WITH_REMOTE_OBJECT_HANDLE && UE_AUTORTFM
+	ensureAlwaysMsgf(!InDelegate.VariantDelegate.IsType<FTimerFunction>(), TEXT("Non-UObject Delegates (TFunctions) are not safe with DSTM. Use FTimerDelegate::CreateUObject and don't capture raw pointers."));
+#endif
+
 	if (FindTimer(InOutHandle))
 	{
 		// if the timer is already set, just clear it and we'll re-add it, since 
