Better implementation for transactionalizing TimerManager delegates.
Previous missed an IsBound() call.

#rnx

diff --git a/Engine/Source/Runtime/Engine/Private/TimerManager.cpp b/Engine/Source/Runtime/Engine/Private/TimerManager.cpp
--- a/Engine/Source/Runtime/Engine/Private/TimerManager.cpp
+++ b/Engine/Source/Runtime/Engine/Private/TimerManager.cpp
@@ -343,19 +343,20 @@ void FTimerUnifiedDelegate::Execute() const
 	case FTimerDelegateVariant::IndexOfType<FTimerDelegate>():
 		{
 			const FTimerDelegate& FuncDelegate = VariantDelegate.Get<FTimerDelegate>();
-			if (FuncDelegate.IsBound())
+
+#if UE_WITH_REMOTE_OBJECT_HANDLE
+			const FName WorkName{ TEXT("TimerManager::Execute FTimerDelegate") };
+			UE::RemoteExecutor::ExecuteTransactional(WorkName, [&FuncDelegate]()
 			{
-				FScopeCycleCounterUObject Context(FuncDelegate.GetUObject());
+#endif
+				if (FuncDelegate.IsBound())
+				{
+					FScopeCycleCounterUObject Context(FuncDelegate.GetUObject());
+					FuncDelegate.Execute();
+				}
 #if UE_WITH_REMOTE_OBJECT_HANDLE
-				const FName WorkName { TEXT("TimerManager::Execute FTimerDelegate") };
-				UE::RemoteExecutor::ExecuteTransactional(WorkName, [FuncDelegate, this]()
-					{
-						FuncDelegate.Execute();
-					});
-#else
-				FuncDelegate.Execute();
+			});
 #endif
-			}
 			break;
 		}
 	case FTimerDelegateVariant::IndexOfType<FTimerDynamicDelegate>():
@@ -363,20 +364,16 @@ void FTimerUnifiedDelegate::Execute() const
 			const FTimerDynamicDelegate& FuncDynDelegate = VariantDelegate.Get<FTimerDynamicDelegate>();
 #if UE_WITH_REMOTE_OBJECT_HANDLE
 			const FName WorkName { TEXT("TimerManager::Execute FTimerDynamicDelegate") };
-			UE::RemoteExecutor::ExecuteTransactional(WorkName, [FuncDynDelegate, this]()
+			UE::RemoteExecutor::ExecuteTransactional(WorkName, [&FuncDynDelegate]()
 			{
+#endif
 				if (FuncDynDelegate.IsBound())
 				{
 					// stat scope is handled by UObject::ProcessEvent for the UFunction.
 					FuncDynDelegate.ProcessDelegate<UObject>(nullptr);
 				}
+#if UE_WITH_REMOTE_OBJECT_HANDLE
 			});
-#else
-			if (FuncDynDelegate.IsBound())
-			{
-				// stat scope is handled by UObject::ProcessEvent for the UFunction.
-				FuncDynDelegate.ProcessDelegate<UObject>(nullptr);
-			}
 #endif
 			break;
 		}
