New material translator:
- Refactored the way types are expressed in MIR. Instead of having each type map to a singleton and passing reference around, we took

#rnx

diff --git a/Engine/Config/BaseMaterialExpressions.ini b/Engine/Config/BaseMaterialExpressions.ini
--- a/Engine/Config/BaseMaterialExpressions.ini
+++ b/Engine/Config/BaseMaterialExpressions.ini
@@ -1055,7 +1055,7 @@ bShaderInputData=true
 +MenuCategories=NSLOCTEXT("MaterialExpression", "Parameters", "Parameters")
 
 [/Script/Engine.MaterialExpressionAggregate]
-+MenuCategories=NSLOCTEXT("MaterialExpression", "MaterialAggregates", "Material Aggregates")
++MenuCategories=NSLOCTEXT("MaterialExpression", "Aggregates", "Aggregates")
 
 [/Script/Engine.MaterialExpressionBlend]
 +MenuCategories=NSLOCTEXT("MaterialExpression", "Math", "Math")
diff --git a/Engine/Source/Runtime/Engine/Private/Materials/MaterialExpressionsIR.cpp b/Engine/Source/Runtime/Engine/Private/Materials/MaterialExpressionsIR.cpp
--- a/Engine/Source/Runtime/Engine/Private/Materials/MaterialExpressionsIR.cpp
+++ b/Engine/Source/Runtime/Engine/Private/Materials/MaterialExpressionsIR.cpp
@@ -394,7 +394,7 @@ void UMaterialExpressionVectorParameter::Build(MIR::FEmitter& Em)
 void UMaterialExpressionDoubleVectorParameter::Build(MIR::FEmitter& Em)
 {
 	FValueRef Value = BuildMaterialExpressionParameter(Em, this);
-	Em.Output(0, Em.Cast(Value, MIR::FPrimitiveType::GetLWC3()));
+	Em.Output(0, Em.Cast(Value, MIR::FType::MakeLWC(3)));
 	Em.Output(1, Em.Subscript(Value, 0));
 	Em.Output(2, Em.Subscript(Value, 1));
 	Em.Output(3, Em.Subscript(Value, 2));
@@ -494,10 +494,10 @@ void UMaterialExpressionAppendVector::Build(MIR::FEmitter& Em)
 
 	UE_MIR_CHECKPOINT(Em);
 
-	const MIR::FPrimitiveType* AType = AVal->Type->AsPrimitive();
-	const MIR::FPrimitiveType* BType = BVal ? BVal->Type->AsPrimitive() : nullptr;
+	MIR::FPrimitive AType = AVal->Type.GetPrimitive();
+	MIR::FPrimitive BType = BVal ? BVal->Type.GetPrimitive() : MIR::FPrimitive{};
 
-	int Dimensions = AType->NumRows + (BType ? BType->NumRows : 0);
+	int Dimensions = AType.NumRows + (BVal ? BType.NumRows : 0);
 	if (Dimensions > 4)
 	{
 		Em.Errorf(TEXT("The resulting vector would have %d component (it can have at most 4)."), Dimensions);
@@ -507,20 +507,20 @@ void UMaterialExpressionAppendVector::Build(MIR::FEmitter& Em)
 	check(Dimensions >= 2 && Dimensions <= 4);
 
 	// Construct the output vector type.
-	MIR::EScalarKind ResultKind = (AType->ScalarKind == MIR::ScalarKind_LWC || (BType && BType->ScalarKind == MIR::ScalarKind_LWC)) ? MIR::ScalarKind_LWC : MIR::ScalarKind_Float;
-	const MIR::FPrimitiveType* ResultType = MIR::FPrimitiveType::GetVector(ResultKind, Dimensions);
+	MIR::EScalarKind ResultKind = (AType.IsLWC() || (BVal && BType.IsLWC())) ? MIR::EScalarKind::Double : MIR::EScalarKind::Float;
+	MIR::FType ResultType = MIR::FType::MakeVector(ResultKind, Dimensions);
 
 	// Set up each output vector component.  These need CastToScalarKind in case we are appending LWC and non-LWC.
 	FValueRef Components[4] = { nullptr, nullptr, nullptr, nullptr };
 	int ComponentIndex = 0;
-	for (int i = 0; i < AType->NumRows; ++i, ++ComponentIndex)
+	for (int i = 0; i < AType.NumRows; ++i, ++ComponentIndex)
 	{
 		Components[ComponentIndex] = Em.CastToScalarKind(Em.Subscript(AVal, i), ResultKind);
 	}
 
-	if (BType)
+	if (BVal)
 	{
-		for (int i = 0; i < BType->NumRows; ++i, ++ComponentIndex)
+		for (int i = 0; i < BType.NumRows; ++i, ++ComponentIndex)
 		{
 			Components[ComponentIndex] = Em.CastToScalarKind(Em.Subscript(BVal, i), ResultKind);
 		}
@@ -574,9 +574,9 @@ void UMaterialExpressionLength::Build(MIR::FEmitter& Em)
 void UMaterialExpressionNormalize::Build(MIR::FEmitter& Em)
 {
 	FValueRef InputValue = Em.CastToFloatKind(Em.Input(&VectorInput));
-	if (InputValue->IsScalar())
+	if (InputValue->Type.IsScalar())
 	{
-		Em.Output(0, Em.ConstantOne(MIR::ScalarKind_Float));
+		Em.Output(0, Em.ConstantOne(MIR::EScalarKind::Float));
 	}
 	else
 	{
@@ -682,7 +682,7 @@ static FValueRef PositiveClampedPow(MIR::FEmitter& Em, FValueRef Base, FValueRef
 		return PrimitiveBase.ToPoison();
 	}
 
-	const MIR::FPrimitiveType* ValuePrimitiveType = Base->Type->AsPrimitive();
+	TOptional<MIR::FPrimitive> ValuePrimitiveType = Base->Type.AsPrimitive();
 	return Em.Select(
 				Em.LessThanOrEquals(Base, Em.ConstantFloat(2.980233e-8f)),
 				Em.ConstantZero(ValuePrimitiveType->ScalarKind),
@@ -872,7 +872,7 @@ void UMaterialExpressionOneMinus::Build(MIR::FEmitter& Em)
 	FValueRef Value = Em.InputDefaultFloat(&Input, 0.0f);
 
 	// Make a "One" value of the same type and dimension as input's.
-	FValueRef One = Em.ConstantOne(Value->Type->AsPrimitive()->ScalarKind);
+	FValueRef One = Em.ConstantOne(Value->Type.AsPrimitive()->ScalarKind);
 
 	// And flow the subtraction out of the expression's only output.
 	Em.Output(0, Em.Subtract(One, Value));
@@ -899,7 +899,7 @@ void UMaterialExpressionIfThenElse::Build(MIR::FEmitter& Em)
 	FValueRef ThenValue = Em.Input(&True);
 	FValueRef ElseValue = Em.Input(&False);
 	
-	const MIR::FType* CommonType = Em.GetCommonType(ThenValue->Type, ElseValue->Type);
+	MIR::FType CommonType = Em.GetCommonType(ThenValue->Type, ElseValue->Type);
 
 	UE_MIR_CHECKPOINT(Em); // Make sure the common type is valid
 
@@ -933,7 +933,7 @@ void UMaterialExpressionIf::Build(MIR::FEmitter& Em)
 	// Less than comparison -- if equals value isn't present (see below), AGreaterThanBValue will also be returned for the equal case.
 	FValueRef ALessThanBConditionValue = Em.LessThan(AValue, BValue);
 
-	if (ALessThanBConditionValue->Type->IsBoolScalar())
+	if (ALessThanBConditionValue->Type.IsBoolScalar())
 	{
 		OutputValue = Em.Branch(ALessThanBConditionValue, ALessThanBValue, AGreaterThanBValue);
 	}
@@ -945,9 +945,9 @@ void UMaterialExpressionIf::Build(MIR::FEmitter& Em)
 	// Equals value is optional -- if present, generate an additional conditional.
 	if (AEqualsBValue.IsValid())
 	{
-	FValueRef AEqualsBConditionValue = EmitAlmostEquals(Em, AValue, BValue, EqualsThreshold);
+		FValueRef AEqualsBConditionValue = EmitAlmostEquals(Em, AValue, BValue, EqualsThreshold);
 
-		if (AEqualsBConditionValue->Type->IsBoolScalar())
+		if (AEqualsBConditionValue->Type.IsBoolScalar())
 		{
 			OutputValue = Em.Branch(AEqualsBConditionValue, AEqualsBValue, OutputValue);
 		}
@@ -1090,7 +1090,7 @@ static FValueRef BuildTextureObjectParameter(MIR::FEmitter& Em, UMaterialExpress
 	}
 
 	FValueRef ParameterValue = Em.Parameter(Expr->GetParameterName(), Param, Expr->SamplerType);
-	if (!ParameterValue->Type->IsTexture())
+	if (!ParameterValue->Type.IsTexture2D())
 	{
 		Em.Error(TEXT("Parameter is not a texture"));
 		return nullptr;
@@ -1136,7 +1136,7 @@ void UMaterialExpressionTextureCoordinate::Build(MIR::FEmitter& Em)
 
 void UMaterialExpressionTime::Build(MIR::FEmitter& Em)
 {
-	const MIR::FType* ScalarFloatType = MIR::FPrimitiveType::GetScalar(MIR::ScalarKind_Float);
+	MIR::FType ScalarFloatType = MIR::FType::MakeScalar(MIR::EScalarKind::Float);
 
 	// When pausing the game is ignored for this time expression, use real-time instead of game-time.
 	const TCHAR* TimeFieldName = bIgnorePause ? TEXT("View.<PREV>RealTime") : TEXT("View.<PREV>GameTime");
@@ -1410,15 +1410,15 @@ void UMaterialExpressionConvert::Build(MIR::FEmitter& Em)
 				InputValues[InputIndex] = Em.InputDefaultFloat4(&ConvertInput.ExpressionInput, ConvertInput.DefaultValue);
 
 				// Expect type to be primitive.
-				const MIR::FPrimitiveType* InputPrimitiveType = InputValues[InputIndex]->Type->AsPrimitive();
+				TOptional<MIR::FPrimitive> InputPrimitiveType = InputValues[InputIndex]->Type.AsPrimitive();
 				if (!InputPrimitiveType)
 				{
-					Em.Errorf(TEXT("Input `%d` of type `%s` is not primitive."), InputComponentIndex, *InputValues[InputIndex]->Type->GetSpelling());
+					Em.Errorf(TEXT("Input `%d` of type `%s` is not primitive."), InputComponentIndex, *InputValues[InputIndex]->Type.GetSpelling());
 					continue;
 				}
 
 				// Determine the target type.
-				const MIR::FType* InputType = MIR::FPrimitiveType::GetVector(InputPrimitiveType->ScalarKind, MaterialExpressionConvertType::GetComponentCount(ConvertInput.Type));
+				MIR::FType InputType = MIR::FType::MakeVector(InputPrimitiveType->ScalarKind, MaterialExpressionConvertType::GetComponentCount(ConvertInput.Type));
 
 				// Cast the input value to the target type.
 				InputValues[InputIndex] = Em.Cast(InputValues[InputIndex], InputType);
@@ -1467,12 +1467,12 @@ static FValueRef BuildViewProperty(MIR::FEmitter& Em, EMaterialExposedViewProper
 	const bool bHasCustomInverseCode = PropertyMeta.InvPropertyCode != nullptr;
 
 	FString HLSLCode = bInvProperty && bHasCustomInverseCode ? PropertyMeta.InvPropertyCode : PropertyMeta.PropertyCode;
-	const MIR::FType* HLSLCodeType = MIR::FType::FromMaterialValueType(PropertyMeta.Type);
+	MIR::FType HLSLCodeType = MIR::FType::FromMaterialValueType(PropertyMeta.Type);
 
 	FValueRef Result = Em.InlineHLSL(HLSLCodeType, MoveTemp(HLSLCode), {}, MIR::EValueFlags::SubstituteTagsInInlineHLSL);
 
 	// CastToNonLWCIfDisabled
-	const MIR::FPrimitiveType* PrimitiveType = HLSLCodeType->AsPrimitive();
+	TOptional<MIR::FPrimitive> PrimitiveType = HLSLCodeType.AsPrimitive();
 	if (PrimitiveType && PrimitiveType->IsLWC() && !UE::MaterialTranslatorUtils::IsLWCEnabled())
 	{
 		Result = Em.CastToFloatKind(Result);
@@ -1508,7 +1508,7 @@ void UMaterialExpressionCameraPositionWS::Build(MIR::FEmitter& Em)
 
 void UMaterialExpressionPixelNormalWS::Build(MIR::FEmitter& Em)
 {
-	FValueRef Output = Em.InlineHLSL(MIR::FPrimitiveType::GetFloat3(), "Parameters.WorldNormal", {}, MIR::EValueFlags::None, MIR::EGraphProperties::ReadsPixelNormal);
+	FValueRef Output = Em.InlineHLSL(MIR::FType::MakeFloat(3), "Parameters.WorldNormal", {}, MIR::EValueFlags::None, MIR::EGraphProperties::ReadsPixelNormal);
 	Em.Output(0, Output);
 }
 
@@ -1584,8 +1584,8 @@ void UMaterialExpressionTruncateLWC::Build(MIR::FEmitter& Em)
 
 	if (LWCTruncateMode == 1 || LWCTruncateMode == 2)
 	{
-		const MIR::FPrimitiveType* Primitive = InputValue->Type->AsPrimitive();
-		if (Primitive && Primitive->ScalarKind == MIR::EScalarKind::ScalarKind_LWC)
+		TOptional<MIR::FPrimitive> Primitive = InputValue->Type.AsPrimitive();
+		if (Primitive && Primitive->ScalarKind == MIR::EScalarKind::Double)
 		{
 			Em.Output(0, Em.CastToFloatKind(Em.Input(&Input)));
 			return;
@@ -1828,7 +1828,7 @@ static FValueRef BlendAggregate(MIR::FEmitter& Em,
 						   FValueRef B,
 						   FValueRef Alpha)
 {
-	const UMaterialAggregate* MaterialAggregate = A->Type->AsAggregate()->Aggregate;
+	const UMaterialAggregate* MaterialAggregate = A->Type.AsAggregate();
 	MIR::TTemporaryArray<MIR::FValueRef> AttributeValues{ MaterialAggregate->Attributes.Num() };
 
 	for (int32 i = 0; i < MaterialAggregate->Attributes.Num(); ++i)
@@ -1864,7 +1864,7 @@ static FValueRef BlendAggregate(MIR::FEmitter& Em,
 static FValueRef Blend(MIR::FEmitter& Em, EMaterialExpressionBlendMode PixelAttributeBlendMode, EMaterialExpressionBlendMode VertexAttributeBlendMode, FValueRef A, FValueRef B, FValueRef Alpha)
 {
 	// Find the common type between arguments
-	if (const MIR::FType* CommonType = Em.GetCommonType(A->Type, B->Type))
+	if (MIR::FType CommonType = Em.GetCommonType(A->Type, B->Type))
 	{
 		// And cast both arguments to the common type
 		A = Em.Cast(A, CommonType);
@@ -1875,24 +1875,24 @@ static FValueRef Blend(MIR::FEmitter& Em, EMaterialExpressionBlendMode PixelAttr
 		return Em.Poison();
 	}
 
-	if (A->Type->IsAnyFloat())
+	if (A->Type.IsAnyFloat())
 	{
 		// Blend floating point values using linear interpolation
 		return Em.Lerp(A, B, Alpha);
 	}
-	else if (A->Type->IsInteger())
+	else if (A->Type.IsInteger())
 	{
 		// "Blend" integer values by selecting A or B based on whether alpha is less than 0.5.
 		return Em.Select(Em.LessThan(Alpha, Em.ConstantFloat(0.5f)), A, B);
 	}
-	else if (const MIR::FAggregateType* AggregateType = A->Type->AsAggregate())
+	else if (A->Type.AsAggregate())
 	{
 		// Arguments are aggregates, so recursively blend each attribute pair.
 		return BlendAggregate(Em, PixelAttributeBlendMode, VertexAttributeBlendMode, A, B, Alpha);
 	}
 	else
 	{
-		Em.Errorf(TEXT("Cannot blend values of type '%s'."), *A->Type->GetSpelling());
+		Em.Errorf(TEXT("Cannot blend values of type '%s'."), *A->Type.GetSpelling());
 		return Em.Poison();
 	}
 }
diff --git a/Engine/Source/Runtime/Engine/Private/Materials/MaterialIR.cpp b/Engine/Source/Runtime/Engine/Private/Materials/MaterialIR.cpp
--- a/Engine/Source/Runtime/Engine/Private/Materials/MaterialIR.cpp
+++ b/Engine/Source/Runtime/Engine/Private/Materials/MaterialIR.cpp
@@ -123,34 +123,6 @@ uint32 FValue::GetSizeInBytes() const
 	return 0;
 }
 
-bool FValue::IsA(EValueKind InKind) const
-{
-	return Kind == InKind;
-}
-
-bool FValue::IsPoison() const
-{
-	return Kind == VK_Poison;
-}
-
-bool FValue::Equals(const FValue* Other) const
-{
-	if (this == Other)
-	{
-		return true;
-	}
-
-	// Get the size of this value in bytes. It should match that of Other, since the value kinds are the same.
-	uint32 SizeInBytes = GetSizeInBytes();
-	if (SizeInBytes != Other->GetSizeInBytes())
-	{
-		return false;
-	}
-
-	// Values are PODs by design, therefore simply comparing bytes is sufficient.
-	return FMemory::Memcmp(this, Other, SizeInBytes) == 0;
-}
-
 TConstArrayView<FValue*> FValue::GetUses() const
 {
 	// Values have no uses by definition.
@@ -246,33 +218,23 @@ TConstArrayView<FValue*> FValue::GetUsesForStage(MIR::EStage Stage) const
 	return GetUses();
 }
 
-bool FValue::IsScalar() const
-{
-	return Type->AsScalar() != nullptr;
-}
-
-bool FValue::IsVector() const
-{
-	return Type->AsVector() != nullptr;
-}
-
 bool FValue::IsTrue() const
 {
 	const MIR::FConstant* Constant = As<MIR::FConstant>();
-	return Constant && Constant->IsBool() && Constant->Boolean == true;
+	return Constant && Constant->Type.IsBoolean() && Constant->Boolean == true;
 }
 
 bool FValue::IsFalse() const
 {
 	const MIR::FConstant* Constant = As<MIR::FConstant>();
-	return Constant && Constant->IsBool() && Constant->Boolean == false;
+	return Constant && Constant->Type.IsBoolean() && Constant->Boolean == false;
 }
 
 bool FValue::AreAllTrue() const
 {
-	if (const MIR::FComposite* Dimensional = As<MIR::FComposite>())
+	if (const MIR::FComposite* Composite = As<MIR::FComposite>())
 	{
-		for (const MIR::FValue* Component : Dimensional->GetComponents())
+		for (const MIR::FValue* Component : Composite->GetComponents())
 		{
 			if (!Component->IsTrue())
 			{
@@ -289,9 +251,9 @@ bool FValue::AreAllTrue() const
 
 bool FValue::AreAllFalse() const
 {
-	if (const MIR::FComposite* Dimensional = As<MIR::FComposite>())
+	if (const MIR::FComposite* Composite = As<MIR::FComposite>())
 	{
-		for (const MIR::FValue* Component : Dimensional->GetComponents())
+		for (const MIR::FValue* Component : Composite->GetComponents())
 		{
 			if (!Component->IsFalse())
 			{
@@ -308,9 +270,9 @@ bool FValue::AreAllFalse() const
 
 bool FValue::AreAllExactlyZero() const
 {
-	if (const MIR::FComposite* Dimensional = As<MIR::FComposite>())
+	if (const MIR::FComposite* Composite = As<MIR::FComposite>())
 	{
-		for (const MIR::FValue* Component : Dimensional->GetComponents())
+		for (const MIR::FValue* Component : Composite->GetComponents())
 		{
 			if (!Component->AreAllExactlyZero())
 			{
@@ -320,17 +282,17 @@ bool FValue::AreAllExactlyZero() const
 	}
 	else if (const MIR::FConstant* Constant = As<MIR::FConstant>())
 	{
-		return (Constant->IsInteger() && Constant->Integer == 0)
-			|| (Constant->IsFloat() && Constant->Float == 0.0f);
+		return (Constant->Type.IsInteger() && Constant->Integer == 0)
+			|| (Constant->Type.IsFloat() && Constant->Float == 0.0f);
 	}
 	return false;
 }
 
 bool FValue::AreAllNearlyZero() const
 {
-	if (const MIR::FComposite* Dimensional = As<MIR::FComposite>())
+	if (const MIR::FComposite* Composite = As<MIR::FComposite>())
 	{
-		for (const MIR::FValue* Component : Dimensional->GetComponents())
+		for (const MIR::FValue* Component : Composite->GetComponents())
 		{
 			if (!Component->AreAllNearlyZero())
 			{
@@ -340,17 +302,17 @@ bool FValue::AreAllNearlyZero() const
 	}
 	else if (const MIR::FConstant* Constant = As<MIR::FConstant>())
 	{
-		return (Constant->IsInteger() && Constant->Integer == 0)
-			|| (Constant->IsFloat() && FMath::IsNearlyZero(Constant->Float));
+		return (Constant->Type.IsInteger() && Constant->Integer == 0)
+			|| (Constant->Type.IsFloat() && FMath::IsNearlyZero(Constant->Float));
 	}
 	return false;
 }
 
 bool FValue::AreAllExactlyOne() const
 {
-	if (const MIR::FComposite* Dimensional = As<MIR::FComposite>())
+	if (const MIR::FComposite* Composite = As<MIR::FComposite>())
 	{
-		for (const MIR::FValue* Component : Dimensional->GetComponents())
+		for (const MIR::FValue* Component : Composite->GetComponents())
 		{
 			if (!Component->AreAllExactlyOne())
 			{
@@ -360,17 +322,17 @@ bool FValue::AreAllExactlyOne() const
 	}
 	else if (const MIR::FConstant* Constant = As<MIR::FConstant>())
 	{
-		return (Constant->IsInteger() && Constant->Integer == 1)
-			|| (Constant->IsFloat() && Constant->Float == 1.0f);
+		return (Constant->Type.IsInteger() && Constant->Integer == 1)
+			|| (Constant->Type.IsFloat() && Constant->Float == 1.0f);
 	}
 	return false;
 }
 
 bool FValue::AreAllNearlyOne() const
 {
-	if (const MIR::FComposite* Dimensional = As<MIR::FComposite>())
+	if (const MIR::FComposite* Composite = As<MIR::FComposite>())
 	{
-		for (const MIR::FValue* Component : Dimensional->GetComponents())
+		for (const MIR::FValue* Component : Composite->GetComponents())
 		{
 			if (!Component->AreAllNearlyOne())
 			{
@@ -380,12 +342,30 @@ bool FValue::AreAllNearlyOne() const
 	}
 	else if (const MIR::FConstant* Constant = As<MIR::FConstant>())
 	{
-		return (Constant->IsInteger() && Constant->Integer == 1)
-			|| (Constant->IsFloat() && FMath::IsNearlyEqual(Constant->Float, 1.0f));
+		return (Constant->Type.IsInteger() && Constant->Integer == 1)
+			|| (Constant->Type.IsFloat() && FMath::IsNearlyEqual(Constant->Float, 1.0f));
 	}
 	return false;
 }
 
+bool FValue::Equals(const FValue* Other) const
+{
+	if (this == Other)
+	{
+		return true;
+	}
+
+	// Get the size of this value in bytes. It should match that of Other, since the value kinds are the same.
+	uint32 SizeInBytes = GetSizeInBytes();
+	if (SizeInBytes != Other->GetSizeInBytes())
+	{
+		return false;
+	}
+
+	// Values are PODs by design, therefore simply comparing bytes is sufficient.
+	return FMemory::Memcmp(this, Other, SizeInBytes) == 0;
+}
+
 FInstruction* AsInstruction(FValue* Value)
 {
 	return Value && (Value->Kind > VK_InstructionBegin && Value->Kind < VK_InstructionEnd) ? static_cast<FInstruction*>(Value) : nullptr;
@@ -401,42 +381,12 @@ FPoison* FPoison::Get()
 	static FPoison Poison = [] {
 		FPoison P;
 		P.Kind = VK_Poison;
-		P.Type = FType::GetPoison();
+		P.Type = FType::MakePoison();
 		return P;
 	} ();
 	return &Poison;
 }
 
-bool FConstant::IsBool() const
-{
-	return Type->IsBoolScalar();
-}
-
-bool FConstant::IsInteger() const
-{
-	return Type == FPrimitiveType::GetInt();
-}
-
-bool FConstant::IsFloat() const
-{
-	return Type == FPrimitiveType::GetFloat();
-}
-
-bool FConstant::IsLWC() const
-{
-	return Type == FPrimitiveType::GetLWC();
-}
-
-bool FConstant::IsLWCRepresentableAsFloat() const
-{
-	if (Type != FPrimitiveType::GetLWC())
-	{
-		return false;
-	}
-	FLargeWorldRenderScalar ConstantLWC(LWCValue);
-	return ConstantLWC.GetTile() == 0.0f;
-}
-
 const TCHAR* LexToString(EExternalInput Input)
 {
 	switch (Input)
@@ -490,17 +440,17 @@ const TCHAR* LexToString(EExternalInput Input)
 	}
 }
 
-const FType* GetExternalInputType(EExternalInput Id)
+FType GetExternalInputType(EExternalInput Id)
 {
 	if (IsExternalInputTexCoordOrPartialDerivative(Id))
 	{
-		return FPrimitiveType::GetFloat2();
+		return FType::MakeFloat(2);
 	}
 	
 	switch (Id)
 	{
-		case EExternalInput::ViewMaterialTextureMipBias: return FPrimitiveType::GetFloat();
-		case EExternalInput::ViewMaterialTextureDerivativeMultiply: return FPrimitiveType::GetFloat();
+		case EExternalInput::ViewMaterialTextureMipBias: return FType::MakeFloat();
+		case EExternalInput::ViewMaterialTextureDerivativeMultiply: return FType::MakeFloat();
 
 		case EExternalInput::ActorPosition_Absolute:
 		case EExternalInput::ObjectPosition_Absolute:
@@ -510,7 +460,7 @@ const FType* GetExternalInputType(EExternalInput Id)
 		case EExternalInput::WorldPosition_AbsoluteNoOffsets_Ddx:		// analytic derivative expressions.  So we cast them to LWC on load.  The shader compiler should
 		case EExternalInput::WorldPosition_Absolute_Ddy:				// be able to optimize the math (the tile offsets will all be zero, cancelling things out).
 		case EExternalInput::WorldPosition_AbsoluteNoOffsets_Ddy:
-			return FPrimitiveType::GetLWC3();
+			return FType::MakeLWC(3);
 
 		case EExternalInput::ActorPosition_CameraRelative:
 		case EExternalInput::ObjectPosition_CameraRelative:
@@ -520,7 +470,7 @@ const FType* GetExternalInputType(EExternalInput Id)
 		case EExternalInput::WorldPosition_CameraRelativeNoOffsets_Ddx:
 		case EExternalInput::WorldPosition_CameraRelative_Ddy:
 		case EExternalInput::WorldPosition_CameraRelativeNoOffsets_Ddy:
-			return FPrimitiveType::GetFloat3();
+			return FType::MakeFloat(3);
 
 		default: UE_MIR_UNREACHABLE();
 	}
@@ -614,13 +564,13 @@ TComposite<1> GCompositeNatvisPrototype;
 TConstArrayView<FValue*> FComposite::GetComponents() const
 {
 	int32 NumComponents;
-	if (const FPrimitiveType* PrimitiveType = Type->AsPrimitive())
+	if (TOptional<FPrimitive> Primitive = Type.AsPrimitive())
 	{
-		NumComponents = PrimitiveType->NumComponents();
+		NumComponents = Primitive->NumComponents();
 	}
-	else if (const FAggregateType* AggregateType = Type->AsAggregate())
+	else if (const UMaterialAggregate* Aggregate = Type.AsAggregate())
 	{
-		NumComponents = AggregateType->Aggregate->Attributes.Num();
+		NumComponents = Aggregate->Attributes.Num();
 	}
 	else
 	{
diff --git a/Engine/Source/Runtime/Engine/Private/Materials/MaterialIRDebug.cpp b/Engine/Source/Runtime/Engine/Private/Materials/MaterialIRDebug.cpp
--- a/Engine/Source/Runtime/Engine/Private/Materials/MaterialIRDebug.cpp
+++ b/Engine/Source/Runtime/Engine/Private/Materials/MaterialIRDebug.cpp
@@ -61,7 +61,7 @@ struct FDebugDumpIRUseGraphState
 		Out.Appendf(TEXT("\"%p\" [label=< <b>%s</b>  (%s) <br/> "),
 						Value,
 						LexToString(Value->Kind),
-						Value->Type ? *Value->Type->GetSpelling() : TEXT("???"));
+						Value->Type ? *Value->Type.GetSpelling() : TEXT("???"));
 
 		DumpValueInfo(Value);
 
@@ -111,12 +111,12 @@ struct FDebugDumpIRUseGraphState
 	{
 		if (const FConstant* Constant = Value->As<FConstant>())
 		{
-			switch (Constant->Type->AsPrimitive()->ScalarKind)
+			switch (Constant->Type.GetPrimitive().ScalarKind)
 			{
-				case ScalarKind_Bool:	Out.Append(Constant->Boolean ? TEXT("true") : TEXT("false")); break;
-				case ScalarKind_Int:	Out.Appendf(TEXT("%lld"), Constant->Integer); break;
-				case ScalarKind_Float:	Out.Appendf(TEXT("%f"), Constant->Float); break;
-				case ScalarKind_LWC:	Out.Appendf(TEXT("%f"), Constant->LWCValue); break;
+				case EScalarKind::Bool:  Out.Append(Constant->Boolean ? TEXT("true") : TEXT("false")); break;
+				case EScalarKind::Int:	 Out.Appendf(TEXT("%lld"), Constant->Integer); break;
+				case EScalarKind::Float: Out.Appendf(TEXT("%f"), Constant->Float); break;
+				case EScalarKind::Double:   Out.Appendf(TEXT("%f"), Constant->Double); break;
 				default: UE_MIR_UNREACHABLE();
 			}
 		}
@@ -134,7 +134,7 @@ struct FDebugDumpIRUseGraphState
 		}
 		else if (const FSubscript* Subscript = Value->As<FSubscript>())
 		{
-			if (Subscript->Arg->Type->AsVector())
+			if (Subscript->Arg->Type.IsVector())
 			{
 				static const TCHAR* Suffix[] = { TEXT(".x"), TEXT(".y"), TEXT(".z"), TEXT(".w") };
 				check(Subscript->Index < 4); 
@@ -155,7 +155,7 @@ struct FDebugDumpIRUseGraphState
 	{
 		if (const FComposite* Composite = Value->As<FComposite>())
 		{
-			if (Composite->Type->AsPrimitive()->IsVector())
+			if (Composite->Type.IsVector())
 			{
 				check(UseIndex < 4);
 				Out.AppendChar(TEXT("xyzw")[UseIndex]);
@@ -267,7 +267,7 @@ struct FDebugDumpIRState
 				}
 
 				// First the type...
-				Out.Appendf(TEXT("%s "), *Use->Type->GetSpelling());
+				Out.Appendf(TEXT("%s "), *Use->Type.GetSpelling());
 
 				// If this use is in a block different from current's, dump the block in "{}" first.
 				const FBlock* UseBlock = Instr->GetDesiredBlockForUse(CurrentStage, UseIndex);
@@ -307,7 +307,7 @@ struct FDebugDumpIRState
 		}
 		else if (const FSubscript* Subscript = Instr->As<FSubscript>())
 		{
-			if (Subscript->Arg->Type->AsVector())
+			if (Subscript->Arg->Type.IsVector())
 			{
 				check(Subscript->Index < 4);
 				Out.Appendf(TEXT(" .%c"), TEXT("XYZW")[Subscript->Index]);
@@ -356,21 +356,21 @@ struct FDebugDumpIRState
 		}
 		else if (const FConstant* Constant = Value->As<FConstant>())
 		{
-			if (Constant->IsBool())
+			if (Constant->Type.IsBoolScalar())
 			{
 				Out.Append(Constant->Boolean ? TEXT("true") : TEXT("false"));
 			}
-			else if (Constant->IsInteger())
+			else if (Constant->Type.IsInteger())
 			{
 				Out.Appendf(TEXT("%lld"), Constant->Integer);
 			}
-			else if (Constant->IsFloat())
+			else if (Constant->Type.IsFloat())
 			{
 				Out.Appendf(TEXT("%.5ff"), Constant->Float);
 			}
-			else if (Constant->IsLWC())
+			else if (Constant->Type.IsLWC())
 			{
-				Out.Appendf(TEXT("%.8f"), Constant->LWCValue);
+				Out.Appendf(TEXT("%.8f"), Constant->Double);
 			}
 			else
 			{
diff --git a/Engine/Source/Runtime/Engine/Private/Materials/MaterialIREmitter.cpp b/Engine/Source/Runtime/Engine/Private/Materials/MaterialIREmitter.cpp
--- a/Engine/Source/Runtime/Engine/Private/Materials/MaterialIREmitter.cpp
+++ b/Engine/Source/Runtime/Engine/Private/Materials/MaterialIREmitter.cpp
@@ -85,11 +85,16 @@ struct FEmitter::FPrivate
 	static uint8* Allocate(FEmitter& Emitter, int32 Size, int32 Alignment)
 	{
 		// TODO: use a linear allocator inside the module instead.
-		uint8* Bytes = (uint8*) FMemory::Malloc(Size, Alignment); 
+		uint8* Bytes = (uint8*)FMemory::Malloc(Size, Alignment); 
 		FMemory::Memzero(Bytes, Size);
 		return Bytes;
 	}
 
+	static void Deallocate(FEmitter& Emitter, void* Ptr)
+	{
+		FMemory::Free(Ptr);
+	}
+
 	// Pushes a new value to the list of values.
 	static void PushNewValue(FEmitter& Emitter, FValue* Value)
 	{
@@ -99,7 +104,7 @@ struct FEmitter::FPrivate
 };
 
 // Creates a new composite value of specified type and returns it.
-static FComposite* NewCompositeValue(FEmitter& Emitter, const FType* Type, int NumComponents)
+static FComposite* NewCompositeValue(FEmitter& Emitter, FType Type, int NumComponents)
 {
 	int32 SizeInBytes = sizeof(FComposite) + sizeof(FValue*) * NumComponents;
 
@@ -118,7 +123,7 @@ static FValueRef EmitNew(FEmitter& Emitter, FValue* Value)
 {
 	if (FValue* Existing = FEmitter::FPrivate::FindValue(Emitter, Value))
 	{
-		delete Value;
+		FEmitter::FPrivate::Deallocate(Emitter, Value);
 		return Existing;
 	}
 
@@ -127,7 +132,7 @@ static FValueRef EmitNew(FEmitter& Emitter, FValue* Value)
 }
 
 template <typename T>
-static T MakePrototype(const FType* InType)
+static T MakePrototype(FType InType)
 {
 	static_assert(std::is_trivially_constructible_v<T> && std::is_trivially_destructible_v<T> &&  std::is_trivially_copy_constructible_v<T> && std::is_trivially_copy_assignable_v<T>,
 		"FValues are expected to be trivial types.");
@@ -259,38 +264,18 @@ void FEmitter::Error(FStringView Message)
 
 /*--------------------------------- Type handling ----------------------------------*/
 
-const FType* FEmitter::GetCommonType(const FType* A, const FType* B)
+FType FEmitter::GetCommonType(FType A, FType B)
 {
-	if (const FType* CommonType = TryGetCommonType(A, B))
+	if (FType CommonType = TryGetCommonType(A, B))
 	{
 		return CommonType;
 	}
 
-	Errorf(TEXT("No common type between '%s' and '%s'."), *A->GetSpelling(), *B->GetSpelling());
-	return nullptr;
-}
-
-const FAggregateType* FEmitter::GetMaterialAggregateType(const UMaterialAggregate* Aggregate)
-{
-	// First, see if the aggregate type hasn't been mapped already.
-	FAggregateType* Type{};
-	if (Internal::Find(AggregateTypeMap, Aggregate, Type))
-	{
-		return Type;
-	}
-
-	// If note, create a MIR type container for this aggregate.
-	Type = new (FPrivate::Allocate(*this, sizeof(FAggregateType), alignof(FAggregateType))) FAggregateType{};
-	Type->Kind = TypeKind_Aggregate;
-	Type->Aggregate = Aggregate;
-
-	// Add the new aggregate type to the map.
-	AggregateTypeMap.Add(Aggregate, Type);
-
-	return Type;
+	Errorf(TEXT("No common type between '%s' and '%s'."), *A.GetSpelling(), *B.GetSpelling());
+	return {};
 }
 
-const FType* FEmitter::GetMaterialAggregateAttributeType(const UMaterialAggregate* Aggregate, int32 AttributeIndex)
+FType FEmitter::GetMaterialAggregateAttributeType(const UMaterialAggregate* Aggregate, int32 AttributeIndex)
 {
 	check(AttributeIndex >= 0);
 
@@ -298,25 +283,25 @@ const FType* FEmitter::GetMaterialAggregateAttributeType(const UMaterialAggregat
 	{
 		Errorf(TEXT("Invalid attribute index %d for material aggregate '%s'. Index is out of range (Num = %d)."),
 			   AttributeIndex, *Aggregate->GetName(), Aggregate->Attributes.Num());
-		return nullptr;
+		return FType::MakePoison();
 	}
 
 	switch (Aggregate->Attributes[AttributeIndex].Type)
 	{
-		case EMaterialAggregateAttributeType::Bool1: return FPrimitiveType::GetVector(ScalarKind_Bool, 1);
-		case EMaterialAggregateAttributeType::Bool2: return FPrimitiveType::GetVector(ScalarKind_Bool, 2);
-		case EMaterialAggregateAttributeType::Bool3: return FPrimitiveType::GetVector(ScalarKind_Bool, 3);
-		case EMaterialAggregateAttributeType::Bool4: return FPrimitiveType::GetVector(ScalarKind_Bool, 4);
+		case EMaterialAggregateAttributeType::Bool1: return FType::MakeVector(EScalarKind::Bool, 1);
+		case EMaterialAggregateAttributeType::Bool2: return FType::MakeVector(EScalarKind::Bool, 2);
+		case EMaterialAggregateAttributeType::Bool3: return FType::MakeVector(EScalarKind::Bool, 3);
+		case EMaterialAggregateAttributeType::Bool4: return FType::MakeVector(EScalarKind::Bool, 4);
 		case EMaterialAggregateAttributeType::ShadingModel:
-		case EMaterialAggregateAttributeType::UInt1: return FPrimitiveType::GetVector(ScalarKind_Int, 1);
-		case EMaterialAggregateAttributeType::UInt2: return FPrimitiveType::GetVector(ScalarKind_Int, 2);
-		case EMaterialAggregateAttributeType::UInt3: return FPrimitiveType::GetVector(ScalarKind_Int, 3);
-		case EMaterialAggregateAttributeType::UInt4: return FPrimitiveType::GetVector(ScalarKind_Int, 4);
-		case EMaterialAggregateAttributeType::Float1: return FPrimitiveType::GetVector(ScalarKind_Float, 1);
-		case EMaterialAggregateAttributeType::Float2: return FPrimitiveType::GetVector(ScalarKind_Float, 2);
-		case EMaterialAggregateAttributeType::Float3: return FPrimitiveType::GetVector(ScalarKind_Float, 3);
-		case EMaterialAggregateAttributeType::Float4: return FPrimitiveType::GetVector(ScalarKind_Float, 4);
-		case EMaterialAggregateAttributeType::Aggregate: return GetMaterialAggregateType(Aggregate->Attributes[AttributeIndex].Aggregate.Get());
+		case EMaterialAggregateAttributeType::UInt1: return FType::MakeVector(EScalarKind::Int, 1);
+		case EMaterialAggregateAttributeType::UInt2: return FType::MakeVector(EScalarKind::Int, 2);
+		case EMaterialAggregateAttributeType::UInt3: return FType::MakeVector(EScalarKind::Int, 3);
+		case EMaterialAggregateAttributeType::UInt4: return FType::MakeVector(EScalarKind::Int, 4);
+		case EMaterialAggregateAttributeType::Float1: return FType::MakeVector(EScalarKind::Float, 1);
+		case EMaterialAggregateAttributeType::Float2: return FType::MakeVector(EScalarKind::Float, 2);
+		case EMaterialAggregateAttributeType::Float3: return FType::MakeVector(EScalarKind::Float, 3);
+		case EMaterialAggregateAttributeType::Float4: return FType::MakeVector(EScalarKind::Float, 4);
+		case EMaterialAggregateAttributeType::Aggregate: return FType::MakeAggregate(Aggregate->Attributes[AttributeIndex].Aggregate.Get());
 		default: UE_MIR_UNREACHABLE();
 	}
 }
@@ -325,7 +310,7 @@ const FType* FEmitter::GetMaterialAggregateAttributeType(const UMaterialAggregat
 
 FValueRef FEmitter::TryInput(const FExpressionInput* InInput)
 {
-	return FValueRef{ MIR::Internal::FetchValueFromExpressionInput(BuilderImpl, InInput), InInput };
+	return FValueRef{ Internal::FetchValueFromExpressionInput(BuilderImpl, InInput), InInput };
 }
 
 FValueRef FEmitter::Input(const FExpressionInput* InInput)
@@ -395,9 +380,9 @@ FValueRef FEmitter::InputDefaultFloat4(const FExpressionInput* Input, UE::Math::
 
 FValueRef FEmitter::CheckTypeIsKind(FValueRef Value, ETypeKind Kind)
 {
-	if (Value.IsValid() && Value->Type->Kind != Kind)
+	if (Value.IsValid() && Value->Type.Is(Kind))
 	{
-		Errorf(Value, TEXT("Expected a '%s' value, got a '%s' instead."), TypeKindToString(Kind), *Value->Type->GetSpelling());
+		Errorf(Value, TEXT("Expected a '%s' value, got a '%s' instead."), TypeKindToString(Kind), *Value->Type.GetSpelling());
 		return Value.ToPoison();
 	}
 	return Value;
@@ -405,9 +390,9 @@ FValueRef FEmitter::CheckTypeIsKind(FValueRef Value, ETypeKind Kind)
 
 FValueRef FEmitter::CheckIsPrimitive(FValueRef Value)
 {
-	if (Value.IsValid() && !Value->Type->AsPrimitive())
+	if (Value.IsValid() && !Value->Type.AsPrimitive())
 	{
-		Errorf(Value, TEXT("Expected a primitive value, got a '%s' instead."), *Value->Type->GetSpelling());
+		Errorf(Value, TEXT("Expected a primitive value, got a '%s' instead."), *Value->Type.GetSpelling());
 		return Value.ToPoison();
 	}
 	return Value;
@@ -415,9 +400,19 @@ FValueRef FEmitter::CheckIsPrimitive(FValueRef Value)
 
 FValueRef FEmitter::CheckIsArithmetic(FValueRef Value)
 {
-	if (Value.IsValid() && !Value->Type->AsArithmetic())
+	if (Value.IsValid() && !Value->Type.IsArithmetic())
+	{
+		Errorf(Value, TEXT("Expected an arithmetic value, got a '%s' instead."), *Value->Type.GetSpelling());
+		return Value.ToPoison();
+	}
+	return Value;
+}
+
+FValueRef FEmitter::CheckIsBoolean(FValueRef Value)
+{
+	if (Value.IsValid() && !Value->Type.IsBoolean())
 	{
-		Errorf(Value, TEXT("Expected an arithmetic value, got a '%s' instead."), *Value->Type->GetSpelling());
+		Errorf(Value, TEXT("Expected a boolean value, got a '%s' instead."), *Value->Type.GetSpelling());
 		return Value.ToPoison();
 	}
 	return Value;
@@ -425,9 +420,9 @@ FValueRef FEmitter::CheckIsArithmetic(FValueRef Value)
 
 FValueRef FEmitter::CheckIsInteger(FValueRef Value)
 {
-	if (Value.IsValid() && !Value->Type->IsInteger())
+	if (Value.IsValid() && !Value->Type.IsInteger())
 	{
-		Errorf(Value, TEXT("Expected an integer value, got a '%s' instead."), *Value->Type->GetSpelling());
+		Errorf(Value, TEXT("Expected an integer value, got a '%s' instead."), *Value->Type.GetSpelling());
 		return Value.ToPoison();
 	}
 	return Value;
@@ -435,9 +430,9 @@ FValueRef FEmitter::CheckIsInteger(FValueRef Value)
 
 FValueRef FEmitter::CheckIsScalar(FValueRef Value)
 {
-	if (Value.IsValid() && (!Value->Type->AsPrimitive() || !Value->Type->AsPrimitive()->IsScalar()))
+	if (Value.IsValid() && !Value->Type.IsScalar())
 	{
-		Errorf(Value, TEXT("Expected a scalar value, got a '%s' instead."), *Value->Type->GetSpelling());
+		Errorf(Value, TEXT("Expected a scalar value, got a '%s' instead."), *Value->Type.GetSpelling());
 		return Value.ToPoison();
 	}
 	return Value;
@@ -445,9 +440,9 @@ FValueRef FEmitter::CheckIsScalar(FValueRef Value)
 
 FValueRef FEmitter::CheckIsScalarOrVector(FValueRef Value)
 {
-	if (Value.IsValid() && (!Value->Type->AsPrimitive() || Value->Type->AsPrimitive()->IsMatrix()))
+	if (Value.IsValid() && (!Value->Type.AsPrimitive() || Value->Type.GetPrimitive().IsMatrix()))
 	{
-		Errorf(Value, TEXT("Expected a scalar or vector value, got a '%s' instead."), *Value->Type->GetSpelling());
+		Errorf(Value, TEXT("Expected a scalar or vector value, got a '%s' instead."), *Value->Type.GetSpelling());
 		return Value.ToPoison();
 	}
 	return Value;
@@ -455,9 +450,9 @@ FValueRef FEmitter::CheckIsScalarOrVector(FValueRef Value)
 
 FValueRef FEmitter::CheckIsTexture(FValueRef Value)
 {
-	if (Value.IsValid() && !Value->Type->IsTexture())
+	if (Value.IsValid() && !Value->Type.IsTexture2D())
 	{
-		Errorf(Value, TEXT("Expected a texture value, got a '%s' instead."), *Value->Type->GetSpelling());
+		Errorf(Value, TEXT("Expected a texture value, got a '%s' instead."), *Value->Type.GetSpelling());
 		return Value.ToPoison();
 	}
 	return Value;
@@ -467,14 +462,15 @@ FValueRef FEmitter::CheckIsAggregate(FValueRef Value, const UMaterialAggregate*
 {
 	if (Value.IsValid())
 	{
-		if (!Value->Type->AsAggregate())
+		const UMaterialAggregate* ValueAggregate = Value->Type.AsAggregate();
+		if (!ValueAggregate)
 		{
-			Errorf(Value, TEXT("Expected an aggregate value, got a '%s' instead."), *Value->Type->GetSpelling());
+			Errorf(Value, TEXT("Expected an aggregate value, got a '%s' instead."), *Value->Type.GetSpelling());
 			return Value.ToPoison();
 		}
-		if (Aggregate && Value->Type->AsAggregate()->Aggregate != Aggregate)
+		if (Aggregate && ValueAggregate != Aggregate)
 		{
-			Errorf(Value, TEXT("Expected a value of aggregate type '%s', got a value of aggregate type '%s' instead."), *Aggregate->GetName(), *Value->Type->AsAggregate()->Aggregate->GetName());
+			Errorf(Value, TEXT("Expected a value of aggregate type '%s', got a value of aggregate type '%s' instead."), *Aggregate->GetName(), *Value->Type.AsAggregate()->GetName());
 			return Value.ToPoison();
 		}
 	}
@@ -487,15 +483,15 @@ bool FEmitter::ToConstantBool(FValueRef Value)
 	{
 		return false;
 	}
-	const MIR::FConstant* Constant = Value->As<MIR::FConstant>();
+	const FConstant* Constant = Value->As<FConstant>();
     if (!Constant)
     {
         Errorf(Value, TEXT("Expected a constant bool value, got a non-constant value instead."));
         return false;
     }
-	if (Constant->Type != FPrimitiveType::GetBool())
+	if (Constant->Type != FType::MakeBool())
     {
-        Errorf(Value, TEXT("Expected a constant bool value, got a '%s' instead."), *Constant->Type->GetSpelling());
+        Errorf(Value, TEXT("Expected a constant bool value, got a '%s' instead."), *Constant->Type.GetSpelling());
         return false;
     }
 	return Constant->Boolean;
@@ -513,7 +509,7 @@ FEmitter& FEmitter::Output(const FExpressionOutput* Output, FValueRef Value)
 {
 	if (Value)
 	{
-		MIR::Internal::BindValueToExpressionOutput(BuilderImpl, Output, Value);
+		Internal::BindValueToExpressionOutput(BuilderImpl, Output, Value);
 	}
 	return *this;
 }
@@ -526,28 +522,28 @@ FValueRef FEmitter::ConstantFromShaderValue(const UE::Shader::FValue& InValue)
 
 	switch (InValue.Type.ValueType)
 	{
-		case UE::Shader::EValueType::Float1: return ConstantFloat(InValue.AsFloatScalar());
-		case UE::Shader::EValueType::Float2: return ConstantFloat2(UE::Math::TVector2<TFloat>{ InValue.Component[0].Float, InValue.Component[1].Float });
-		case UE::Shader::EValueType::Float3: return ConstantFloat3(UE::Math::TVector<TFloat>{ InValue.Component[0].Float, InValue.Component[1].Float, InValue.Component[2].Float });
-		case UE::Shader::EValueType::Float4: return ConstantFloat4(UE::Math::TVector4<TFloat>{ InValue.Component[0].Float, InValue.Component[1].Float, InValue.Component[2].Float, InValue.Component[3].Float });
-
 		case UE::Shader::EValueType::Int1: return ConstantInt(InValue.AsFloatScalar());
 		case UE::Shader::EValueType::Int2: return ConstantInt2({ InValue.Component[0].Int, InValue.Component[1].Int });
 		case UE::Shader::EValueType::Int3: return ConstantInt3({ InValue.Component[0].Int, InValue.Component[1].Int, InValue.Component[2].Int });
 		case UE::Shader::EValueType::Int4: return ConstantInt4({ InValue.Component[0].Int, InValue.Component[1].Int, InValue.Component[2].Int, InValue.Component[3].Int });
+	
+		case UE::Shader::EValueType::Float1: return ConstantFloat(InValue.AsFloatScalar());
+		case UE::Shader::EValueType::Float2: return ConstantFloat2(UE::Math::TVector2<TFloat>{ InValue.Component[0].Float, InValue.Component[1].Float });
+		case UE::Shader::EValueType::Float3: return ConstantFloat3(UE::Math::TVector<TFloat>{ InValue.Component[0].Float, InValue.Component[1].Float, InValue.Component[2].Float });
+		case UE::Shader::EValueType::Float4: return ConstantFloat4(UE::Math::TVector4<TFloat>{ InValue.Component[0].Float, InValue.Component[1].Float, InValue.Component[2].Float, InValue.Component[3].Float });
 
-		case UE::Shader::EValueType::Double1: return ConstantLWC(InValue.Component[0].Double);
-		case UE::Shader::EValueType::Double2: return Vector2(ConstantLWC(InValue.Component[0].Double), ConstantLWC(InValue.Component[1].Double));
-		case UE::Shader::EValueType::Double3: return Vector3(ConstantLWC(InValue.Component[0].Double), ConstantLWC(InValue.Component[1].Double), ConstantLWC(InValue.Component[2].Double));
-		case UE::Shader::EValueType::Double4: return Vector4(ConstantLWC(InValue.Component[0].Double), ConstantLWC(InValue.Component[1].Double), ConstantLWC(InValue.Component[2].Double), ConstantLWC(InValue.Component[3].Double));
+		case UE::Shader::EValueType::Double1: return ConstantDouble(InValue.Component[0].Double);
+		case UE::Shader::EValueType::Double2: return Vector2(ConstantDouble(InValue.Component[0].Double), ConstantDouble(InValue.Component[1].Double));
+		case UE::Shader::EValueType::Double3: return Vector3(ConstantDouble(InValue.Component[0].Double), ConstantDouble(InValue.Component[1].Double), ConstantDouble(InValue.Component[2].Double));
+		case UE::Shader::EValueType::Double4: return Vector4(ConstantDouble(InValue.Component[0].Double), ConstantDouble(InValue.Component[1].Double), ConstantDouble(InValue.Component[2].Double), ConstantDouble(InValue.Component[3].Double));
 	}
 
 	UE_MIR_UNREACHABLE();
 }
 
-FValueRef FEmitter::ConstantDefault(const FType* Type)
+FValueRef FEmitter::ConstantDefault(FType Type)
 {
-	if (const FPrimitiveType* PrimitiveType = Type->AsPrimitive())
+	if (TOptional<FPrimitive> PrimitiveType = Type.AsPrimitive())
 	{
 		FValueRef Zero = ConstantZero(PrimitiveType->ScalarKind);
 		if (PrimitiveType->IsScalar())
@@ -564,13 +560,13 @@ FValueRef FEmitter::ConstantDefault(const FType* Type)
 			return EmitNew(*this, Composite);
 		}
 	}
-	else if (const FAggregateType* AggregateType = Type->AsAggregate())
+	else if (const UMaterialAggregate* TypeAggregate = Type.AsAggregate())
 	{
-		return Aggregate(AggregateType->Aggregate);
+		return Aggregate(TypeAggregate);
 	}
 	else
 	{
-		Errorf(TEXT("Type '%s' has no default. Expected primitive or aggregate type."), *Type->GetSpelling());
+		Errorf(TEXT("Type '%s' has no default. Expected primitive or aggregate type."), *Type.GetSpelling());
 		return Poison();
 	}
 }
@@ -579,10 +575,10 @@ FValueRef FEmitter::ConstantZero(EScalarKind Kind)
 {
 	switch (Kind)
 	{
-		case ScalarKind_Bool: return ConstantFalse();
-		case ScalarKind_Int: return ConstantInt(0);
-		case ScalarKind_Float: return ConstantFloat(0.0f);
-		case ScalarKind_LWC: return ConstantLWC(0.0);
+		case EScalarKind::Bool: return ConstantFalse();
+		case EScalarKind::Int: return ConstantInt(0);
+		case EScalarKind::Float: return ConstantFloat(0.0f);
+		case EScalarKind::Double: return ConstantDouble(0.0);
 		default: UE_MIR_UNREACHABLE();
 	}
 }
@@ -591,22 +587,22 @@ FValueRef FEmitter::ConstantOne(EScalarKind Kind)
 {
 	switch (Kind)
 	{
-		case ScalarKind_Bool: return ConstantTrue();
-		case ScalarKind_Int: return ConstantInt(1);
-		case ScalarKind_Float: return ConstantFloat(1.0f);
-		case ScalarKind_LWC: return ConstantLWC(1.0);
+		case EScalarKind::Bool: return ConstantTrue();
+		case EScalarKind::Int: return ConstantInt(1);
+		case EScalarKind::Float: return ConstantFloat(1.0f);
+		case EScalarKind::Double: return ConstantDouble(1.0);
 		default: UE_MIR_UNREACHABLE();
 	}
 }
 
-FValueRef FEmitter::ConstantScalar(EScalarKind Kind, TFloat FromFloat)
+FValueRef FEmitter::ConstantScalar(EScalarKind Kind, TDouble Value)
 {
 	switch (Kind)
 	{
-		case ScalarKind_Bool: return ConstantBool(FromFloat != 0.0f);
-		case ScalarKind_Int: return ConstantInt(TInteger(FromFloat));
-		case ScalarKind_Float: return ConstantFloat(FromFloat);
-		case ScalarKind_LWC: return ConstantLWC(FromFloat);
+		case EScalarKind::Bool: return ConstantBool(Value != 0.0f);
+		case EScalarKind::Int: return ConstantInt(TInteger(Value));
+		case EScalarKind::Float: return ConstantFloat(TFloat(Value));
+		case EScalarKind::Double: return ConstantDouble(Value);
 		default: UE_MIR_UNREACHABLE();
 	}
 }
@@ -652,7 +648,7 @@ FValueRef FEmitter::ConstantBool4(bool InX, bool InY, bool InZ, bool InW)
 
 FValueRef FEmitter::ConstantInt(TInteger InX)
 {
-	FConstant Scalar = MakePrototype<FConstant>(FPrimitiveType::GetScalar(ScalarKind_Int));
+	FConstant Scalar = MakePrototype<FConstant>(FType::MakeScalar(EScalarKind::Int));
 	Scalar.Integer = InX;
 	return EmitPrototype(*this, Scalar);
 }
@@ -683,7 +679,7 @@ FValueRef FEmitter::ConstantInt4(UE::Math::TIntVector4<TInteger> InValue)
 
 FValueRef FEmitter::ConstantFloat(TFloat InX)
 {
-	FConstant Scalar = MakePrototype<FConstant>(FPrimitiveType::GetScalar(ScalarKind_Float));
+	FConstant Scalar = MakePrototype<FConstant>(FType::MakeScalar(EScalarKind::Float));
 	Scalar.Float = InX;
 	return EmitPrototype(*this, Scalar);
 }
@@ -712,10 +708,10 @@ FValueRef FEmitter::ConstantFloat4(UE::Math::TVector4<TFloat> InValue)
 	return Vector4(X, Y, Z, W);
 }
 
-FValueRef FEmitter::ConstantLWC(TDouble InX)
+FValueRef FEmitter::ConstantDouble(TDouble InX)
 {
-	FConstant Scalar = MakePrototype<FConstant>(FPrimitiveType::GetScalar(ScalarKind_LWC));
-	Scalar.LWCValue = InX;
+	FConstant Scalar = MakePrototype<FConstant>(FType::MakeLWC());
+	Scalar.Double = InX;
 	return EmitPrototype(*this, Scalar);
 }
 
@@ -732,7 +728,7 @@ FValueRef FEmitter::Poison()
 
 FValueRef FEmitter::ExternalInput(EExternalInput Id)
 {
-	MIR::FExternalInput Prototype = MakePrototype<MIR::FExternalInput>(GetExternalInputType(Id));
+	FExternalInput Prototype = MakePrototype<FExternalInput>(GetExternalInputType(Id));
 	Prototype.Id = Id;
 	return EmitPrototype(*this, Prototype);
 }
@@ -744,7 +740,7 @@ FValueRef FEmitter::TextureObject(UTexture* Texture, EMaterialSamplerType Sample
 		case MCT_Texture2D:
 		case MCT_TextureVirtual:
 		{
-			FTextureObject Proto = MakePrototype<FTextureObject>(FObjectType::GetTexture2D());
+			FTextureObject Proto = MakePrototype<FTextureObject>(FType::MakeTexture2D());
 			Proto.Texture = Texture;
 			Proto.SamplerType = SamplerType;
 			return EmitPrototype(*this, Proto);
@@ -790,43 +786,29 @@ FValueRef FEmitter::Parameter(FName Name, FMaterialParameterMetadata& Metadata,
 				return CustomPrimitiveData(Metadata.PrimitiveDataIndex);
 			}
 
-			const FType* Type = Metadata.Value.Type == EMaterialParameterType::Enumeration
-				? FPrimitiveType::GetInt()
-				: FPrimitiveType::GetFloat();
+			FType Type = Metadata.Value.Type == EMaterialParameterType::Enumeration
+				? FType::MakeInt()
+				: FType::MakeFloat();
 
-			MIR::FUniformParameter Prototype = MakePrototype<MIR::FUniformParameter>(Type);
+			FUniformParameter Prototype = MakePrototype<FUniformParameter>(Type);
 			Prototype.ParameterIdInModule = RegisterParameter(Module, Info, Metadata);
 			return EmitPrototype(*this, Prototype);
 		}
 
 		case EMaterialParameterType::Vector:
-		{
-			if (Metadata.PrimitiveDataIndex != INDEX_NONE)
-			{
-				MIR::FValue* X = CustomPrimitiveData(Metadata.PrimitiveDataIndex + 0);
-				MIR::FValue* Y = CustomPrimitiveData(Metadata.PrimitiveDataIndex + 1);
-				MIR::FValue* Z = CustomPrimitiveData(Metadata.PrimitiveDataIndex + 2);
-				MIR::FValue* W = CustomPrimitiveData(Metadata.PrimitiveDataIndex + 3);
-				return Vector4(X, Y, Z, W);
-			}
-
-			MIR::FUniformParameter Prototype = MakePrototype<MIR::FUniformParameter>(FPrimitiveType::GetFloat4());
-			Prototype.ParameterIdInModule = RegisterParameter(Module, Info, Metadata);
-			return EmitPrototype(*this, Prototype);
-		}
-
 		case EMaterialParameterType::DoubleVector:
 		{
 			if (Metadata.PrimitiveDataIndex != INDEX_NONE)
 			{
-				MIR::FValue* X = CustomPrimitiveData(Metadata.PrimitiveDataIndex + 0);
-				MIR::FValue* Y = CustomPrimitiveData(Metadata.PrimitiveDataIndex + 1);
-				MIR::FValue* Z = CustomPrimitiveData(Metadata.PrimitiveDataIndex + 2);
-				MIR::FValue* W = CustomPrimitiveData(Metadata.PrimitiveDataIndex + 3);
+				FValue* X = CustomPrimitiveData(Metadata.PrimitiveDataIndex + 0);
+				FValue* Y = CustomPrimitiveData(Metadata.PrimitiveDataIndex + 1);
+				FValue* Z = CustomPrimitiveData(Metadata.PrimitiveDataIndex + 2);
+				FValue* W = CustomPrimitiveData(Metadata.PrimitiveDataIndex + 3);
 				return Vector4(X, Y, Z, W);
 			}
 
-			MIR::FUniformParameter Prototype = MakePrototype<MIR::FUniformParameter>(FPrimitiveType::GetLWC4());
+			EScalarKind ScalarKind = Metadata.Value.Type == EMaterialParameterType::Vector ? EScalarKind::Float : EScalarKind::Double;
+			FUniformParameter Prototype = MakePrototype<FUniformParameter>(FType::MakePrimitive(ScalarKind, 4));
 			Prototype.ParameterIdInModule = RegisterParameter(Module, Info, Metadata);
 			return EmitPrototype(*this, Prototype);
 		}
@@ -838,7 +820,7 @@ FValueRef FEmitter::Parameter(FName Name, FMaterialParameterMetadata& Metadata,
 				case MCT_Texture2D:
 				case MCT_TextureVirtual:
 				{
-					MIR::FUniformParameter Prototype = MakePrototype<MIR::FUniformParameter>(FObjectType::GetTexture2D());
+					FUniformParameter Prototype = MakePrototype<FUniformParameter>(FType::MakeTexture2D());
 					Prototype.ParameterIdInModule = RegisterParameter(Module, Info, Metadata);
 					Prototype.SamplerType = SamplerType;
 					return EmitPrototype(*this, Prototype);
@@ -851,7 +833,7 @@ FValueRef FEmitter::Parameter(FName Name, FMaterialParameterMetadata& Metadata,
 				}
 				break;
 			}
-			return MIR::FPoison::Get();
+			return FPoison::Get();
 		}
 
 		case EMaterialParameterType::StaticSwitch:
@@ -901,16 +883,16 @@ FValueRef FEmitter::CustomPrimitiveData(uint32 PrimitiveDataIndex)
 
 FSetMaterialOutput* FEmitter::SetMaterialOutput(EMaterialProperty InProperty, FValue* Arg)
 {
-	FSetMaterialOutput Proto = MakePrototype<FSetMaterialOutput>(nullptr);
-	Proto.Property  			= InProperty;
-	Proto.Arg 					= Arg;
+	FSetMaterialOutput Proto = MakePrototype<FSetMaterialOutput>({});
+	Proto.Property           = InProperty;
+	Proto.Arg                = Arg;
 
 	// Initialize the instruction block to the root of each stage it is evaluated in.
 	for (int i = 0; i < NumStages; ++i)
 	{
 		if (MaterialOutputEvaluatesInStage(InProperty, (EStage)i))
 		{
-			Proto.Block[i]	= Module->RootBlock[i];
+			Proto.Block[i] = &Module->RootBlocks[i];
 		}
 	}
 
@@ -936,10 +918,10 @@ FValueRef FEmitter::Vector2(FValueRef InX, FValueRef InY)
 		return Poison();
 	}
 
-	check(InX->Type->AsScalar());
+	check(InX->Type.IsScalar());
 	check(InX->Type == InY->Type);
 
-	TComposite<2> Vector = MakePrototype<TComposite<2>>(FPrimitiveType::GetVector(InX->Type->AsPrimitive()->ScalarKind, 2));
+	TComposite<2> Vector = MakePrototype<TComposite<2>>(FType::MakeVector(InX->Type.AsPrimitive()->ScalarKind, 2));
 	TArrayView<FValue*> Components = Vector.GetMutableComponents();
 	Components[0] = InX;
 	Components[1] = InY;
@@ -962,11 +944,11 @@ FValueRef FEmitter::Vector3(FValueRef InX, FValueRef InY, FValueRef InZ)
 		return Poison();
 	}
 
-	check(InX->Type->AsScalar());
+	check(InX->Type.IsScalar());
 	check(InX->Type == InY->Type);
 	check(InY->Type == InZ->Type);
 
-	TComposite<3> Vector = MakePrototype<TComposite<3>>(FPrimitiveType::GetVector(InX->Type->AsPrimitive()->ScalarKind, 3));
+	TComposite<3> Vector = MakePrototype<TComposite<3>>(FType::MakeVector(InX->Type.AsPrimitive()->ScalarKind, 3));
 	TArrayView<FValue*> Components = Vector.GetMutableComponents();
 	Components[0] = InX;
 	Components[1] = InY;
@@ -990,12 +972,12 @@ FValueRef FEmitter::Vector4(FValueRef InX, FValueRef InY, FValueRef InZ, FValueR
 		return Poison();
 	}
 
-	check(InX->Type->AsScalar());
+	check(InX->Type.IsScalar());
 	check(InX->Type == InY->Type);
 	check(InY->Type == InZ->Type);
 	check(InZ->Type == InW->Type);
 
-	TComposite<4> Vector = MakePrototype<TComposite<4>>(FPrimitiveType::GetVector(InX->Type->AsPrimitive()->ScalarKind, 4));
+	TComposite<4> Vector = MakePrototype<TComposite<4>>(FType::MakeVector(InX->Type.AsPrimitive()->ScalarKind, 4));
 	TArrayView<FValue*> Components = Vector.GetMutableComponents();
 	Components[0] = InX;
 	Components[1] = InY;
@@ -1030,14 +1012,14 @@ FValueRef FEmitter::Aggregate(const UMaterialAggregate* InAggregate, FValueRef I
 	}
 
 	// Create the new composite to store the aggregate attribute values.
-	FComposite* AggregateValue = NewCompositeValue(*this, GetMaterialAggregateType(InAggregate), InAggregate->Attributes.Num());
+	FComposite* AggregateValue = NewCompositeValue(*this, FType::MakeAggregate(InAggregate), InAggregate->Attributes.Num());
 	
 	// Assign all components of the new composite value.
 	TArrayView<FValue*> Components = AggregateValue->GetMutableComponents();
 	for (int32 i = 0; i < Components.Num(); ++i)
 	{
 		// Get the the ith aggregate aggregate MIR type.
-		const FType* AttributeType = GetMaterialAggregateAttributeType(InAggregate, i);
+		FType AttributeType = GetMaterialAggregateAttributeType(InAggregate, i);
 
 		if (i < AttributeValues.Num())
 		{
@@ -1071,7 +1053,7 @@ FValueRef FEmitter::Aggregate(const UMaterialAggregate* InAggregate, FValueRef I
 	}
 
 	// Create the new composite to store the aggregate attribute values.
-	FComposite* AggregateValue = NewCompositeValue(*this, GetMaterialAggregateType(InAggregate), InAggregate->Attributes.Num());
+	FComposite* AggregateValue = NewCompositeValue(*this, FType::MakeAggregate(InAggregate), InAggregate->Attributes.Num());
 	
 	// Assign all components of the new composite value.
 	TArrayView<FValue*> Components = AggregateValue->GetMutableComponents();
@@ -1085,7 +1067,7 @@ FValueRef FEmitter::Aggregate(const UMaterialAggregate* InAggregate, FValueRef I
 			});
 
 		// Get the the ith aggregate aggregate MIR type.
-		const FType* AttributeType = GetMaterialAggregateAttributeType(InAggregate, i);
+		FType AttributeType = GetMaterialAggregateAttributeType(InAggregate, i);
 
 		if (Assignment && Assignment->Value)
 		{
@@ -1168,13 +1150,13 @@ static T FoldScalarOperator(FEmitter& Emitter, EOperator Operator, T A, T B, T C
 		{
 			case UO_ACos: return FGenericPlatformMath::Acos(A);
 			case UO_ACosFast: return FGenericPlatformMath::Acos(A);
-			case UO_ACosh: return MIR::ACosh(A);
+			case UO_ACosh: return ACosh(A);
 			case UO_ASin: return FGenericPlatformMath::Asin(A);
 			case UO_ASinFast: return FGenericPlatformMath::Asin(A);
-			case UO_ASinh: return MIR::ASinh(A);
+			case UO_ASinh: return ASinh(A);
 			case UO_ATan: return FGenericPlatformMath::Atan(A);
 			case UO_ATanFast: return FGenericPlatformMath::Atan(A);
-			case UO_ATanh: return MIR::ATanh(A);
+			case UO_ATanh: return ATanh(A);
 			case UO_Ceil:
 				if constexpr (std::is_same<T, float>::value)
 				{
@@ -1293,7 +1275,7 @@ static FValueRef TrySimplifyOperator(FEmitter& Emitter, EOperator& Op, FValueRef
 	{
 		/* Unary Operators */
 		case UO_Length:
-			if (A->Type->AsScalar())
+			if (A->Type.GetPrimitive().IsScalar())
 			{
 				Op = UO_Abs;
 			}
@@ -1341,7 +1323,7 @@ static FValueRef TrySimplifyOperator(FEmitter& Emitter, EOperator& Op, FValueRef
 			}
 			else if (A->Equals(B))
 			{
-				return Emitter.ConstantZero(A->Type->AsPrimitive()->ScalarKind);
+				return Emitter.ConstantZero(A->Type.GetPrimitive().ScalarKind);
 			}
 			break;
 
@@ -1363,14 +1345,14 @@ static FValueRef TrySimplifyOperator(FEmitter& Emitter, EOperator& Op, FValueRef
 			}
 			else if (A->Equals(B))
 			{
-				return Emitter.ConstantOne(A->Type->AsPrimitive()->ScalarKind);
+				return Emitter.ConstantOne(A->Type.GetPrimitive().ScalarKind);
 			}
 			break;
 
 		case BO_Modulo:
 			if (A->AreAllNearlyZero() || B->AreAllNearlyOne())
 			{
-				return Emitter.ConstantZero(A->Type->AsPrimitive()->ScalarKind);
+				return Emitter.ConstantZero(A->Type.GetPrimitive().ScalarKind);
 			}
 			break;
 
@@ -1413,7 +1395,7 @@ static FValueRef TrySimplifyOperator(FEmitter& Emitter, EOperator& Op, FValueRef
 			else if (B->AreAllNearlyZero())
 			{
 				// If the exponent is 0.
-				return Emitter.ConstantOne(A->Type->AsPrimitive()->ScalarKind);
+				return Emitter.ConstantOne(A->Type.GetPrimitive().ScalarKind);
 			}
 			else if (B->AreAllExactlyOne())
 			{
@@ -1476,7 +1458,7 @@ static FValueRef TrySimplifyOperator(FEmitter& Emitter, EOperator& Op, FValueRef
 // It returns either the result of the operator or null if it could not be folded.
 static FValueRef TryFoldOperatorScalar(FEmitter& Emitter, EOperator Op, FValueRef A, FValueRef B, FValueRef C)
 {
-	const FPrimitiveType* PrimitiveType = A->Type->AsPrimitive();
+	TOptional<FPrimitive> PrimitiveType = A->Type.AsPrimitive();
 
 	// Try to simplify the operator. This could potentially change Op, A, B and C.
 	if (FValueRef Simplified = TrySimplifyOperator(Emitter, Op, A, B, C))
@@ -1487,7 +1469,7 @@ static FValueRef TryFoldOperatorScalar(FEmitter& Emitter, EOperator Op, FValueRe
 	// If TrySimplifyOperator did not already fold the `select` operator, there is nothing else to do.
 	if (Op == TO_Select)
 	{
-		return nullptr;
+		return {};
 	}
 
 	// Verify that both lhs and rhs are constants, otherwise we cannot fold the operation.
@@ -1496,7 +1478,7 @@ static FValueRef TryFoldOperatorScalar(FEmitter& Emitter, EOperator Op, FValueRe
 	const FConstant* CConstant = As<FConstant>(C.Value);
 	if (!AConstant || (IsBinaryOperator(Op) && !BConstant) || (IsTernaryOperator(Op) && (!BConstant || !CConstant)))
 	{
-		return nullptr;
+		return {};
 	}
 
 	// Call the appropriate helper function depending on what type of operator this is
@@ -1505,16 +1487,16 @@ static FValueRef TryFoldOperatorScalar(FEmitter& Emitter, EOperator Op, FValueRe
 		bool Result;
 		switch (PrimitiveType->ScalarKind)
 		{
-			case ScalarKind_Int:
+			case EScalarKind::Int:
 				Result = FoldComparisonOperatorScalar<TInteger>(Op, AConstant->Integer, BConstant->Integer);
 				break;
 
-			case ScalarKind_Float:
+			case EScalarKind::Float:
 				Result = FoldComparisonOperatorScalar<TFloat>(Op, AConstant->Float, BConstant->Float);
 				break;
 
-			case ScalarKind_LWC:
-				Result = FoldComparisonOperatorScalar<TDouble>(Op, AConstant->LWCValue, BConstant->LWCValue);
+			case EScalarKind::Double:
+				Result = FoldComparisonOperatorScalar<TDouble>(Op, AConstant->Double, BConstant->Double);
 				break;
 
 			default:
@@ -1526,28 +1508,28 @@ static FValueRef TryFoldOperatorScalar(FEmitter& Emitter, EOperator Op, FValueRe
 	{
 		switch (PrimitiveType->ScalarKind)
 		{
-			case ScalarKind_Bool:
+			case EScalarKind::Bool:
 			{
 				bool Result = FoldScalarOperator<TInteger>(Emitter, Op, AConstant->Boolean, BConstant ? BConstant->Boolean : 0, 0) & 0x1;
 				return Emitter.ConstantBool(Result);
 			}
 
-			case ScalarKind_Int:
+			case EScalarKind::Int:
 			{
 				TInteger Result = FoldScalarOperator<TInteger>(Emitter, Op, AConstant->Integer, BConstant ? BConstant->Integer : 0, CConstant ? CConstant->Integer : 0);
 				return Emitter.ConstantInt(Result);
 			}
 
-			case ScalarKind_Float:
+			case EScalarKind::Float:
 			{
 				TFloat Result = FoldScalarOperator<TFloat>(Emitter, Op, AConstant->Float, BConstant ? BConstant->Float : 0, CConstant ? CConstant->Float : 0);
 				return Emitter.ConstantFloat(Result);
 			}
 
-			case ScalarKind_LWC:
+			case EScalarKind::Double:
 			{
-				TDouble Result = FoldScalarOperator<TDouble>(Emitter, Op, AConstant->LWCValue, BConstant ? BConstant->LWCValue : 0, CConstant ? CConstant->LWCValue : 0);
-				return Emitter.ConstantLWC(Result);
+				TDouble Result = FoldScalarOperator<TDouble>(Emitter, Op, AConstant->Double, BConstant ? BConstant->Double : 0, CConstant ? CConstant->Double : 0);
+				return Emitter.ConstantDouble(Result);
 			}
 
 			default:
@@ -1621,7 +1603,7 @@ static const FOperatorSignature* GetOperatorSignature(EOperator Op)
 		const FOperatorSignature TernaryArithmeticLWC = { { OPTF_CastToCommonArithmeticTypeLWC, OPTF_CastToCommonArithmeticTypeLWC, OPTF_CastToCommonArithmeticTypeLWC }, ORT_FirstArgumentTypeLWC };
 		const FOperatorSignature TernaryFloatLWCDemote = { { OPTF_CastToCommonArithmeticTypeLWC | OPTF_CastToAnyFloat, OPTF_CastToCommonArithmeticTypeLWC, OPTF_CastToCommonArithmeticTypeLWC }, ORT_FirstArgumentType };
 
-		static FOperatorSignature S[MIR::OperatorCount];
+		static FOperatorSignature S[OperatorCount];
 
 		/* unary operators */
 		S[UO_BitwiseNot] 				= { { OPTF_CheckIsInteger }, ORT_FirstArgumentType };
@@ -1706,7 +1688,7 @@ static const FOperatorSignature* GetOperatorSignature(EOperator Op)
  
 /// Validates that the types of the arguments are valid for specified operator.
 /// If valid, it returns the type of the result. Otherwise if it is not valid, it returns nullptr.
-static const FPrimitiveType* ValidateOperatorAndGetResultType(FEmitter& Emitter, EOperator Op, FValueRef& A, FValueRef& B, FValueRef& C)
+static FType ValidateOperatorAndGetResultType(FEmitter& Emitter, EOperator Op, FValueRef& A, FValueRef& B, FValueRef& C)
 {
 	// Argument A must have always been provided.
 	check(A);
@@ -1715,30 +1697,41 @@ static const FPrimitiveType* ValidateOperatorAndGetResultType(FEmitter& Emitter,
 	check(!C || B);
 
 	// Verify that B argument has been provided if operator is binary.
-	check(!MIR::IsBinaryOperator(Op) || B);
+	check(!IsBinaryOperator(Op) || B);
 
 	// Verify that C argument has been provided if operator is ternary.
-	check(!MIR::IsTernaryOperator(Op) || C);
+	check(!IsTernaryOperator(Op) || C);
+
+	// Make sure the first argument has primitive type first, since the following operations
+	// assume this.
+	if (!Emitter.CheckIsPrimitive(A).IsValid())
+	{
+		return FType::MakePoison();
+	}
 
 	// For Clamp, we do a special case and demote the first LWC argument if the second and third arguments (min / max) are non-LWC.
 	// We want to do this before fetching FirstArgumentPrimitiveType.
-	if (Op == TO_Clamp && A->Type->IsLWC() && !B->Type->IsLWC() && !C->Type->IsLWC())
+	if (Op == TO_Clamp && A->Type.GetPrimitive().IsLWC() && !B->Type.IsLWC() && !C->Type.IsLWC())
 	{
 		A = Emitter.CastToFloatKind(A);
 	}
 
-	// Verify that the first argument type is primitive.
-	FValueRef Arguments[] = { A, B, C, nullptr };
-	const FPrimitiveType* FirstArgumentPrimitiveType = A->Type->AsPrimitive();
-	static const TCHAR* ArgumentsStr[] = { TEXT("first"), TEXT("second"), TEXT("third") };
+	// Get the operator signature information.
 	const FOperatorSignature* Signature = GetOperatorSignature(Op);
 
 	// Handle automatic cast to float for operators that don't support LWC inputs
-	if (FirstArgumentPrimitiveType->ScalarKind == EScalarKind::ScalarKind_LWC && !(Signature->ParameterTypes[0] & OPTF_LWCInput))
+	FType FirstArgumentType = A->Type;
+	if (FirstArgumentType.GetPrimitive().IsLWC() && !(Signature->ParameterTypes[0] & OPTF_LWCInput))
 	{
-		FirstArgumentPrimitiveType = FirstArgumentPrimitiveType->WithScalarKind(EScalarKind::ScalarKind_Float);
+		FirstArgumentType = A->Type.GetPrimitive().ToScalarKind(EScalarKind::Float);
 	}
-	const FType* ArgumentsCommonType = FirstArgumentPrimitiveType;
+
+	// This will hold the common type between arguments types.
+	FType ArgumentsCommonType = FirstArgumentType;
+
+	// Verify that the first argument type is primitive.
+	FValueRef Arguments[] = { A, B, C, nullptr };
+	static const TCHAR* ArgumentsStr[] = { TEXT("first"), TEXT("second"), TEXT("third") };
 
 	bool bValid = true;
 	for (int32 i = 0; Arguments[i]; ++i)
@@ -1747,10 +1740,10 @@ static const FPrimitiveType* ValidateOperatorAndGetResultType(FEmitter& Emitter,
 		Arguments[i] = Emitter.CheckIsPrimitive(Arguments[i]);
 		if (!Arguments[i].IsValid())
 		{
-			return nullptr;
+			return FType::MakePoison();
 		}
-		const FPrimitiveType* ArgumentPrimitiveType = Arguments[i]->Type->AsPrimitive();
 
+		FType ArgType = Arguments[i]->Type;
 		EOperatorParameterTypeFilter Filter = Signature->ParameterTypes[i];
 		check(Filter != OPTF_Unknown); // No signature specified for this operator.
 
@@ -1759,33 +1752,33 @@ static const FPrimitiveType* ValidateOperatorAndGetResultType(FEmitter& Emitter,
 			check(i > 0); // This check can't apply to the first argument.
 
 			// Cast this argument to the first argument primitive type
-			ArgumentPrimitiveType = FirstArgumentPrimitiveType;
-			Arguments[i] = Emitter.Cast(Arguments[i], ArgumentPrimitiveType);
+			ArgType = FirstArgumentType;
+			Arguments[i] = Emitter.Cast(Arguments[i], FirstArgumentType);
 			bValid &= Arguments[i].IsValid();
 		}
-		else if (ArgumentPrimitiveType->ScalarKind == EScalarKind::ScalarKind_LWC && !(Filter & OPTF_LWCInput))
+		else if (ArgType.GetPrimitive().IsLWC() && !(Filter & OPTF_LWCInput))
 		{
 			// Handle automatic cast to float for operators that don't support LWC inputs
-			ArgumentPrimitiveType = ArgumentPrimitiveType->WithScalarKind(EScalarKind::ScalarKind_Float);
-			Arguments[i] = Emitter.Cast(Arguments[i], ArgumentPrimitiveType);
+			ArgType = ArgType.GetPrimitive().ToScalarKind(EScalarKind::Float);
+			Arguments[i] = Emitter.Cast(Arguments[i], ArgType);
 		}
 		else if (Filter & OPTF_CastToAnyFloat)
 		{
-			if (!ScalarKindIsAnyFloat(ArgumentPrimitiveType->ScalarKind))
+			if (!ArgType.GetPrimitive().IsAnyFloat())
 			{
-				ArgumentPrimitiveType = ArgumentPrimitiveType->WithScalarKind(MIR::ScalarKind_Float);
-				Arguments[i] = Emitter.Cast(Arguments[i], ArgumentPrimitiveType);
+				ArgType = ArgType.GetPrimitive().ToScalarKind(EScalarKind::Float);
+				Arguments[i] = Emitter.Cast(Arguments[i], ArgType);
 				bValid &= Arguments[i].IsValid();
 			}
 		}
 		
 		if (Filter & OPTF_CheckIsBoolean)
 		{
-			if (ArgumentPrimitiveType->ScalarKind != ScalarKind_Bool)
-			{
-				Emitter.Errorf(Arguments[i], TEXT("Expected a boolean."));
-				bValid = false;
-			}
+			bValid &= Emitter.CheckIsArithmetic(Arguments[i]).IsValid(); 
+		}
+		else if (Filter & OPTF_CheckIsInteger)
+		{
+			bValid &= Emitter.CheckIsInteger(Arguments[i]).IsValid(); 
 		}
 		
 		if (Filter & OPTF_CheckIsArithmetic)
@@ -1793,11 +1786,6 @@ static const FPrimitiveType* ValidateOperatorAndGetResultType(FEmitter& Emitter,
 			bValid &= Emitter.CheckIsArithmetic(Arguments[i]).IsValid(); 
 		}
 
-		if (Filter & OPTF_CheckIsInteger)
-		{
-			bValid &= Emitter.CheckIsInteger(Arguments[i]).IsValid(); 
-		}
-		
 		if (Filter & OPTF_CheckIsNotMatrix)
 		{
 			bValid &= Emitter.CheckIsScalarOrVector(Arguments[i]).IsValid(); 
@@ -1805,13 +1793,14 @@ static const FPrimitiveType* ValidateOperatorAndGetResultType(FEmitter& Emitter,
 
 		if (Filter & OPTF_CheckIsVector3)
 		{
-			if (!ArgumentPrimitiveType->IsVector() || ArgumentPrimitiveType->NumComponents() != 3)
+			if (!ArgType.GetPrimitive().IsVector() || ArgType.GetPrimitive().NumComponents() != 3)
 			{
 				Emitter.Errorf(Arguments[i], TEXT("Expected a 3D vector."));
 				bValid = false;
 			}
 		}
 
+		// The following checks are only applicable if the argument is constant.
 		if (FConstant* Constant = Arguments[i]->As<FConstant>())
 		{
 			if (Filter & OPTF_CheckIsNonZeroFloatConst)
@@ -1858,7 +1847,7 @@ static const FPrimitiveType* ValidateOperatorAndGetResultType(FEmitter& Emitter,
 		// Update the common type.
 		if (i >= 1)
 		{
-			ArgumentsCommonType = Emitter.TryGetCommonType(ArgumentsCommonType, ArgumentPrimitiveType);
+			ArgumentsCommonType = Emitter.TryGetCommonType(ArgumentsCommonType, ArgType);
 		}
 	}
 	
@@ -1866,19 +1855,19 @@ static const FPrimitiveType* ValidateOperatorAndGetResultType(FEmitter& Emitter,
 	if (Op == TO_Select)
 	{
 		// Cast the second and third argument types to primitive. This is safe as it was already checked earlier.
-		const FPrimitiveType* BPrimitiveType = static_cast<const FPrimitiveType*>(Arguments[1]->Type);
-		const FPrimitiveType* CPrimitiveType = static_cast<const FPrimitiveType*>(Arguments[2]->Type);
+		const FPrimitive& BPrimitiveType = Arguments[1]->Type.GetPrimitive();
+		const FPrimitive& CPrimitiveType = Arguments[2]->Type.GetPrimitive();
 
 		// Compute the maximum number of vector components between all arguments. We know they're scalar or vectors, as it was checked before.
-		int32 MaxNumComponents = FMath::Max3(FirstArgumentPrimitiveType->NumComponents(), BPrimitiveType->NumComponents(), CPrimitiveType->NumComponents());
+		int32 MaxNumComponents = FMath::Max3(FirstArgumentType.GetPrimitive().NumComponents(), BPrimitiveType.NumComponents(), CPrimitiveType.NumComponents());
 
 		// Cast the first argument (the boolean condition) to a bool vector of the maximum number of components.
-		Arguments[0] = Emitter.Cast(Arguments[0], FPrimitiveType::Get(ScalarKind_Bool, MaxNumComponents, 1));
+		Arguments[0] = Emitter.Cast(Arguments[0], FType::MakePrimitive(EScalarKind::Bool, MaxNumComponents, 1));
 
 		// Compute the common type between the second and third argument types with a number of components equal to the max of all three.
-		const FType* CommonTypeBetweenSecondAndThirdArguments = Emitter.TryGetCommonType(
-			FPrimitiveType::Get(BPrimitiveType->ScalarKind, MaxNumComponents, 1),
-			FPrimitiveType::Get(CPrimitiveType->ScalarKind, MaxNumComponents, 1));
+		FType CommonTypeBetweenSecondAndThirdArguments = Emitter.TryGetCommonType(
+			FType::MakePrimitive(BPrimitiveType.ScalarKind, MaxNumComponents, 1),
+			FType::MakePrimitive(CPrimitiveType.ScalarKind, MaxNumComponents, 1));
 		
 		// Getting the common type should always be possible.
 		check(CommonTypeBetweenSecondAndThirdArguments);
@@ -1898,21 +1887,22 @@ static const FPrimitiveType* ValidateOperatorAndGetResultType(FEmitter& Emitter,
 			EOperatorParameterTypeFilter Filter = Signature->ParameterTypes[i];
 			if (Filter & OPTF_CastToCommonType)
 			{
-				check(ArgumentsCommonType && ArgumentsCommonType->AsPrimitive());
+				check(ArgumentsCommonType && ArgumentsCommonType.AsPrimitive());
 
 				// Lerp doesn't accept LWC for its third input, so we need to check if LWC is allowed per input when casting to the common type.
-				if (ArgumentsCommonType->IsLWC() && !(Filter & OPTF_LWCInput))
+				if (ArgumentsCommonType.GetPrimitive().IsLWC() && !(Filter & OPTF_LWCInput))
 				{
-					Arguments[i] = Emitter.Cast(Arguments[i], ArgumentsCommonType->AsPrimitive()->WithScalarKind(ScalarKind_Float));
+					Arguments[i] = Emitter.Cast(Arguments[i], ArgumentsCommonType.GetPrimitive().ToScalarKind(EScalarKind::Float));
 				}
 				else
 				{
 					Arguments[i] = Emitter.Cast(Arguments[i], ArgumentsCommonType);
 				}
+
 				bValid &= !Arguments[i]->IsPoison();
 			}
 			// Fmod doesn't accept LWC for its second input, so we need to potentially cast that to non-LWC.
-			else if (Arguments[i]->Type->IsLWC() && !(Filter & OPTF_LWCInput))
+			else if (Arguments[i]->Type.GetPrimitive().IsLWC() && !(Filter & OPTF_LWCInput))
 			{
 				Arguments[i] = Emitter.CastToFloatKind(Arguments[i]);
 			}
@@ -1926,11 +1916,13 @@ static const FPrimitiveType* ValidateOperatorAndGetResultType(FEmitter& Emitter,
 
 	if (!bValid)
 	{
-		return {};
+		return FType::MakePoison();
 	}
 	
 	// Update the first argument type, as it might have changed.
-	FirstArgumentPrimitiveType = Arguments[0]->Type->AsPrimitive();
+	FirstArgumentType = Arguments[0]->Type;
+
+	const FPrimitive& FirstArgumentPrimitiveType = FirstArgumentType.GetPrimitive();
 
 	// Finally, determine operator result type.
 	switch (Signature->ReturnType)
@@ -1939,20 +1931,20 @@ static const FPrimitiveType* ValidateOperatorAndGetResultType(FEmitter& Emitter,
 			UE_MIR_UNREACHABLE(); // missing operator signature declaration
 
 		case ORT_FirstArgumentType:
-			// Some operators accept LWC as input, but always output non-LWC (examples:  modulus, sign, sine, cosine, saturate, frac).  Anything that may output LWC should use ORT_FirstArgumentTypeLWC.
-			return (FirstArgumentPrimitiveType->ScalarKind == EScalarKind::ScalarKind_LWC) ? FirstArgumentPrimitiveType->WithScalarKind(EScalarKind::ScalarKind_Float) : FirstArgumentPrimitiveType;
+			// Some operators accept LWC as input, but always output non-LWC (examples: modulo, sign, sine, cosine, saturate, frac). Anything that may output LWC should use ORT_FirstArgumentTypeLWC.
+			return FirstArgumentPrimitiveType.IsLWC() ? FirstArgumentPrimitiveType.ToScalarKind(EScalarKind::Float) : FirstArgumentType;
 
 		case ORT_FirstArgumentTypeLWC:
-			return FirstArgumentPrimitiveType;
+			return FirstArgumentType;
 
 		case ORT_BooleanWithFirstArgumentDimensions:
-			return FPrimitiveType::Get(ScalarKind_Bool, FirstArgumentPrimitiveType->NumRows, FirstArgumentPrimitiveType->NumColumns);
+			return FType::MakePrimitive(EScalarKind::Bool, FirstArgumentPrimitiveType.NumRows, FirstArgumentPrimitiveType.NumColumns);
 
  		case ORT_FirstArgumentTypeToScalarLWC:
-			return FPrimitiveType::GetScalar(FirstArgumentPrimitiveType->ScalarKind);
+			return FirstArgumentPrimitiveType.ToScalar();
 
 		case ORT_SecondArgumentType:
-			return static_cast<const FPrimitiveType*>(B->Type);
+			return B->Type;
 	}
 
 	UE_MIR_UNREACHABLE();
@@ -1969,7 +1961,7 @@ static bool IsComponentwiseOperator(EOperator Op)
 /// If a value is returned, it will be a composite with some component folded to a constant. If some argument
 /// isn't a composite, or all arguments components are non-constant, the folding will not be carried out.
 /// If no folding is carried out, this function simply returns nullptr.
-static FValue* TryFoldComponentwiseOperator(FEmitter& Emitter, EOperator Op, FValue* A, FValue* B, FValue* C, const FPrimitiveType* ResultType)
+static FValue* TryFoldComponentwiseOperator(FEmitter& Emitter, EOperator Op, FValue* A, FValue* B, FValue* C, FType ResultType)
 {
 	// Check that at least one component of the resulting composite value would folded.
 	// If all components of resulting composite value are not folded, then instead of emitting
@@ -1981,9 +1973,9 @@ static FValue* TryFoldComponentwiseOperator(FEmitter& Emitter, EOperator Op, FVa
 	bool bResultIsIdenticalToC = true;
 
 	// Allocate the temporary array to store the folded component results
-	TTemporaryArray<FValue*> TempResultComponents{ ResultType->NumComponents() };
+	TTemporaryArray<FValue*> TempResultComponents{ ResultType.GetPrimitive().NumComponents() };
 	
-	for (int32 i = 0; i < ResultType->NumComponents(); ++i)
+	for (int32 i = 0; i < TempResultComponents.Num(); ++i)
 	{
 		// Extract the arguments individual components
 		FValue* AComponent = Emitter.Subscript(A, i);
@@ -2026,23 +2018,23 @@ static FValue* TryFoldComponentwiseOperator(FEmitter& Emitter, EOperator Op, FVa
 	if (bSomeResultComponentWasFolded)
 	{
 		// If result type is scalar, simply return the single folded result (instead of creating a composite value)
-		if (ResultType->IsScalar())
+		if (ResultType.GetPrimitive().IsScalar())
 		{
 			check(TempResultComponents[0]);
 			return TempResultComponents[0];
 		}
 
 		// Make the new composite value
-		FComposite* Result = NewCompositeValue(Emitter, ResultType, ResultType->NumComponents());
+		FComposite* Result = NewCompositeValue(Emitter, ResultType, TempResultComponents.Num());
 
 		// Fetch the components array from the result composite
 		TArrayView<FValue*> ResultComponents = Result->GetMutableComponents();
 
 		// Also cache the type of a single component
-		const FPrimitiveType* ComponentType = ResultType->ToScalar();
+		FType ComponentType = ResultType.GetPrimitive().ToScalar();
 
 		// Create the operator instruction for each component pair
-		for (int32 i = 0; i < ResultType->NumComponents(); ++i)
+		for (int32 i = 0; i < TempResultComponents.Num(); ++i)
 		{
 			// Reuse cached result if possible
 			ResultComponents[i] = TempResultComponents[i];
@@ -2071,7 +2063,7 @@ static bool TryUnpackConstantScalarOrVector(FValue* V, TArrayView<FConstant*> Ou
 {
 	// V not specified? Or not a scalar/vector?
 	FComposite* Composite = As<FComposite>(V);
-	if (!Composite || V->Type->AsPrimitive()->IsMatrix())
+	if (!Composite || V->Type.AsPrimitive()->IsMatrix())
 	{
 		return false;
 	}
@@ -2091,7 +2083,7 @@ static bool TryUnpackConstantScalarOrVector(FValue* V, TArrayView<FConstant*> Ou
 }
 
 /// Computes the dot product on two arrays of constant float components.
-static MIR::TFloat ConstantDotFloat(TArrayView<FConstant*> AComponents, TArrayView<FConstant*> BComponents, int32 NumComponents)
+static TFloat ConstantDotFloat(TArrayView<FConstant*> AComponents, TArrayView<FConstant*> BComponents, int32 NumComponents)
 {
 	float Result = 0.0f;
 	for (int32 i = 0; i < NumComponents; ++i)
@@ -2100,25 +2092,25 @@ static MIR::TFloat ConstantDotFloat(TArrayView<FConstant*> AComponents, TArrayVi
 	}
 	return Result;
 }
+
 static MIR::TDouble ConstantDotDouble(TArrayView<FConstant*> AComponents, TArrayView<FConstant*> BComponents, int32 NumComponents)
 {
 	double Result = 0.0;
 	for (int32 i = 0; i < NumComponents; ++i)
 	{
-		Result += AComponents[i]->LWCValue * BComponents[i]->LWCValue;
+		Result += AComponents[i]->Double * BComponents[i]->Double;
 	}
 	return Result;
 }
 
-
 /// Tries to fold the operator, that is to evaluate its result now at translation time if its arguments are constant.
 /// If the operator could not be folded in any way, it returns nullptr.
-static FValue* TryFoldOperator(FEmitter& Emitter, EOperator Op, FValue* A, FValue* B, FValue* C, const FPrimitiveType* ResultType)
+static FValue* TryFoldOperator(FEmitter& Emitter, EOperator Op, FValue* A, FValue* B, FValue* C, FType ResultType)
 {
 	FConstant* AComponents[4];
 	int32      ANumComponents;
 	
-	// Some operations like Length, Dot and Cross are not defined on individual scalar components.
+	// CASE 1: Some operations like Length, Dot and Cross are not defined on individual scalar components.
 	// For instance length(V) is not the same as [length(V.x), ..., length(V.z)]. These operations
 	// folding is handled here as special cases.
 	// First, try to unpack the first argument to an array of constants.
@@ -2129,15 +2121,15 @@ static FValue* TryFoldOperator(FEmitter& Emitter, EOperator Op, FValue* A, FValu
 	
 		if (Op == UO_Length)
 		{
-			if (ResultType->ScalarKind == ScalarKind_Float)
+			if (ResultType.GetPrimitive().IsFloat())
 			{
 				float Result = FMath::Sqrt(ConstantDotFloat(AComponents, AComponents, ANumComponents));
 				return Emitter.ConstantFloat(Result);
 			}
-			else if (ResultType->ScalarKind == ScalarKind_LWC)
+			else if (ResultType.GetPrimitive().IsLWC())
 			{
 				double Result = FMath::Sqrt(ConstantDotDouble(AComponents, AComponents, ANumComponents));
-				return Emitter.ConstantLWC(Result);
+				return Emitter.ConstantDouble(Result);
 			}
 			else
 			{
@@ -2151,15 +2143,15 @@ static FValue* TryFoldOperator(FEmitter& Emitter, EOperator Op, FValue* A, FValu
 
 			if (Op == BO_Dot)
 			{
-				if (ResultType->ScalarKind == ScalarKind_Float)
+				if (ResultType.GetPrimitive().IsFloat())
 				{
 					float Result = ConstantDotFloat(AComponents, BComponents, ANumComponents);
 					return Emitter.ConstantFloat(Result);
 				}
-				else if (ResultType->ScalarKind == ScalarKind_LWC)
+				else if (ResultType.GetPrimitive().IsLWC())
 				{
 					double Result = ConstantDotDouble(AComponents, BComponents, ANumComponents);
-					return Emitter.ConstantLWC(Result);
+					return Emitter.ConstantDouble(Result);
 				}
 				else
 				{
@@ -2169,13 +2161,17 @@ static FValue* TryFoldOperator(FEmitter& Emitter, EOperator Op, FValue* A, FValu
 			else
 			{
 				check(Op == BO_Cross);
-				if (ResultType->ScalarKind == ScalarKind_Float)
+				if (ResultType.GetPrimitive().IsFloat())
 				{
 					FVector3f AVector{ AComponents[0]->Float, AComponents[1]->Float, AComponents[2]->Float };
 					FVector3f BVector{ BComponents[0]->Float, BComponents[1]->Float, BComponents[2]->Float };
 					FVector3f Result = AVector.Cross(BVector);
 					return Emitter.ConstantFloat3(Result);
 				}
+				else if (ResultType.GetPrimitive().IsLWC())
+				{
+					UE_MIR_TODO();
+				}
 				else
 				{
 					UE_MIR_UNREACHABLE();
@@ -2184,14 +2180,14 @@ static FValue* TryFoldOperator(FEmitter& Emitter, EOperator Op, FValue* A, FValu
 		}
 	}
 
-	// If the operation supports componentwise application, try folding the operator componentwise.
+	// CASE 2: If the operation supports componentwise application, try folding the operator componentwise.
 	if (IsComponentwiseOperator(Op))
 	{
 		return TryFoldComponentwiseOperator(Emitter, Op, A, B, C, ResultType);
 	}
 
 	// No folding was possible, simply return null to indicate this.
-	return nullptr;
+	return {};
 }
 
 FValueRef FEmitter::Operator(EOperator Op, FValueRef A, FValueRef B, FValueRef C)
@@ -2207,17 +2203,17 @@ FValueRef FEmitter::Operator(EOperator Op, FValueRef A, FValueRef B, FValueRef C
 	bool bCastLWCOutputToFloat = false;
 	if (Op == BO_Subtract && B)
 	{
-		const FPrimitiveType* AType = A->Type->AsPrimitive();
-		const FPrimitiveType* BType = B->Type->AsPrimitive();
+		TOptional<FPrimitive> AType = A->Type.AsPrimitive();
+		TOptional<FPrimitive> BType = B->Type.AsPrimitive();
 	
-		if (AType && AType->ScalarKind == EScalarKind::ScalarKind_LWC && BType && BType->ScalarKind == EScalarKind::ScalarKind_LWC && UE::MaterialTranslatorUtils::GetLWCTruncateMode() == 2)
+		if (AType && AType->ScalarKind == EScalarKind::Double && BType && BType->ScalarKind == EScalarKind::Double && UE::MaterialTranslatorUtils::GetLWCTruncateMode() == 2)
 		{
 			bCastLWCOutputToFloat = true;
 		}
 	}
 
 	// Validate the operation and retrieve the result type.
-	const FPrimitiveType* ResultType = ValidateOperatorAndGetResultType(*this, Op, A, B, C);
+	FType ResultType = ValidateOperatorAndGetResultType(*this, Op, A, B, C);
 	if (!ResultType)
 	{
 		return Poison();
@@ -2252,8 +2248,9 @@ FValueRef FEmitter::Operator(EOperator Op, FValueRef A, FValueRef B, FValueRef C
 	{
 		Result = CastToFloatKind(Result);
 	}
+
 	return Result;
-}
+} 
 
 FValueRef FEmitter::Branch(FValueRef Condition, FValueRef True, FValueRef False)
 {
@@ -2263,7 +2260,7 @@ FValueRef FEmitter::Branch(FValueRef Condition, FValueRef True, FValueRef False)
 	}
 
 	// Condition must be of type bool
-	Condition = Cast(Condition, FPrimitiveType::GetBool());
+	Condition = Cast(Condition, FType::MakeBool());
 	if (!Condition)
 	{
 		return Poison();
@@ -2277,7 +2274,7 @@ FValueRef FEmitter::Branch(FValueRef Condition, FValueRef True, FValueRef False)
 
 	// If the condition is not static, make both true and false arguments have the same type,
 	// by casting false argument into the true's type.
-	const FType* CommonType = GetCommonType(True->Type, False->Type);
+	FType CommonType = GetCommonType(True->Type, False->Type);
 	if (!CommonType)
 	{
 		return Poison();
@@ -2314,22 +2311,22 @@ FValueRef FEmitter::Subscript(FValueRef Value, int32 Index)
 	}
 	
 	// Other operations supported if value is primitive. Check it first.
-	const FPrimitiveType* PrimitiveType = Value->Type->AsPrimitive();
+	TOptional<FPrimitive> PrimitiveType = Value->Type.AsPrimitive();
 	if (!PrimitiveType)
 	{
-		Errorf(Value, TEXT("Value of type '%s' cannot be subscripted."), *Value->Type->GetSpelling());
+		Errorf(Value, TEXT("Value of type '%s' cannot be subscripted."), *Value->Type.GetSpelling());
 		return Value.ToPoison();
 	}
 
 	// Getting first component and Value is already a scalar, just return itself.
-	if (Index == 0 && Value->Type->AsScalar())
+	if (Index == 0 && PrimitiveType->IsScalar())
 	{
 		return Value;
 	}
 
 	if (Index >= PrimitiveType->NumComponents())
 	{
-		Errorf(Value, TEXT("Value of type '%s' has fewer dimensions than subscript index `%d`."), *Value->Type->GetSpelling(), Index);
+		Errorf(Value, TEXT("Value of type '%s' has fewer dimensions than subscript index `%d`."), *Value->Type.GetSpelling(), Index);
 		return Value.ToPoison();
 	}
 
@@ -2362,26 +2359,28 @@ FValueRef FEmitter::Swizzle(FValueRef Value, FSwizzleMask Mask)
 	}
 
 	// We can only swizzle on non-matrix primitive types.
-	const FPrimitiveType* PrimitiveType = Value->Type->AsVector();
-	if (!PrimitiveType || PrimitiveType->IsMatrix())
+	if (!Value->Type.AsPrimitive() || Value->Type.GetPrimitive().IsMatrix())
 	{
-		Errorf(Value, TEXT("Cannot swizzle a '%s' value."), *Value->Type->GetSpelling());
+		Errorf(Value, TEXT("Cannot swizzle a '%s' value."), *Value->Type.GetSpelling());
 		return Value.ToPoison();
 	}
+	
+	// For brevity
+	FPrimitive PrimitiveType = Value->Type.GetPrimitive();
 
 	// Make sure each component in the mask fits the number of components in Value.
 	for (EVectorComponent Component : Mask)
 	{
-		if ((int32)Component >= PrimitiveType->NumRows)
+		if ((int32)Component >= PrimitiveType.NumRows)
 		{
-			Errorf(Value, TEXT("Value of type '%s' has no component '%s'."), PrimitiveType->Spelling.GetData(), VectorComponentToString(Component));
+			Errorf(Value, TEXT("Value of type '%s' has no component '%s'."), *Value->Type.GetSpelling(), VectorComponentToString(Component));
 			return Value.ToPoison();
 		}
 	}
 
 	// If the requested number of components is the same as Value and the order in which the components
 	// are specified in the mask is sequential (e.g. x, y, z) then this is a no op, simply return Value as is.
-	if (Mask.NumComponents == PrimitiveType->NumComponents())
+	if (Mask.NumComponents == PrimitiveType.NumComponents())
 	{
 		bool InOrder = true;
 		for (int32 i = 0; i < Mask.NumComponents; ++i)
@@ -2406,8 +2405,8 @@ FValueRef FEmitter::Swizzle(FValueRef Value, FSwizzleMask Mask)
 	}
 
 	// Make the result vector type.
-	const FPrimitiveType* ResultType = FPrimitiveType::GetVector(PrimitiveType->ScalarKind, Mask.NumComponents);
-	FComposite* Result = NewCompositeValue(*this, ResultType, ResultType->NumComponents());
+	FType ResultType = FType::MakeVector(PrimitiveType.ScalarKind, Mask.NumComponents);
+	FComposite* Result = NewCompositeValue(*this, ResultType, ResultType.GetPrimitive().NumComponents());
 	TArrayView<FValue*> ResultComponents = Result->GetMutableComponents();
 
 	for (int32 i = 0; i < Mask.NumComponents; ++i)
@@ -2427,46 +2426,46 @@ static FValue* CastConstant(FEmitter& Emitter, FConstant* Constant, EScalarKind
 
 	switch (ConstantScalarKind)
 	{
-		case ScalarKind_Bool:
+		case EScalarKind::Bool:
 		{
 			switch (TargetKind)
 			{
-				case ScalarKind_Int: return Emitter.ConstantInt(Constant->Boolean ? 1 : 0);
-				case ScalarKind_Float: return Emitter.ConstantFloat(Constant->Boolean ? 1.0f : 0.0f);
-				case ScalarKind_LWC: return Emitter.ConstantLWC(Constant->Boolean ? 1.0 : 0.0);
+				case EScalarKind::Int: return Emitter.ConstantInt(Constant->Boolean ? 1 : 0);
+				case EScalarKind::Float: return Emitter.ConstantFloat(Constant->Boolean ? 1.0f : 0.0f);
+				case EScalarKind::Double: return Emitter.ConstantDouble(Constant->Boolean ? 1.0 : 0.0);
 				default: UE_MIR_UNREACHABLE();
 			}
 		}
 
-		case ScalarKind_Int:
+		case EScalarKind::Int:
 		{
 			switch (TargetKind)
 			{
-				case ScalarKind_Bool: return Emitter.ConstantBool(Constant->Integer != 0);
-				case ScalarKind_Float: return Emitter.ConstantFloat((TFloat)Constant->Integer);
-				case ScalarKind_LWC: return Emitter.ConstantLWC(Constant->Integer);
+				case EScalarKind::Bool: return Emitter.ConstantBool(Constant->Integer != 0);
+				case EScalarKind::Float: return Emitter.ConstantFloat((TFloat)Constant->Integer);
+				case EScalarKind::Double: return Emitter.ConstantDouble(Constant->Integer);
 				default: UE_MIR_UNREACHABLE();
 			}
 		}
 
-		case ScalarKind_Float:
+		case EScalarKind::Float:
 		{
 			switch (TargetKind)
 			{
-				case ScalarKind_Bool: return Emitter.ConstantBool(Constant->Float != 0.0f);
-				case ScalarKind_Int: return Emitter.ConstantInt((int32)Constant->Float);
-				case ScalarKind_LWC: return Emitter.ConstantLWC(Constant->Float);
+				case EScalarKind::Bool: return Emitter.ConstantBool(Constant->Float != 0.0f);
+				case EScalarKind::Int: return Emitter.ConstantInt((int32)Constant->Float);
+				case EScalarKind::Double: return Emitter.ConstantDouble(Constant->Float);
 				default: UE_MIR_UNREACHABLE();
 			}
 		}
 
-		case ScalarKind_LWC:
+		case EScalarKind::Double:
 		{
 			switch (TargetKind)
 			{
-				case ScalarKind_Bool: return Emitter.ConstantBool(Constant->LWCValue != 0.0);
-				case ScalarKind_Int: return Emitter.ConstantInt((int32)Constant->LWCValue);
-				case ScalarKind_Float: return Emitter.ConstantFloat((float)Constant->LWCValue);
+				case EScalarKind::Bool: return Emitter.ConstantBool(Constant->Double != 0.0);
+				case EScalarKind::Int: return Emitter.ConstantInt((int32)Constant->Double);
+				case EScalarKind::Float: return Emitter.ConstantFloat((float)Constant->Double);
 				default: UE_MIR_UNREACHABLE();
 			}
 		}
@@ -2477,37 +2476,39 @@ static FValue* CastConstant(FEmitter& Emitter, FConstant* Constant, EScalarKind
 	UE_MIR_UNREACHABLE();
 }
 
-static FValue* CastValueToPrimitiveType(FEmitter& Emitter, FValueRef Value, const FPrimitiveType* TargetPrimitiveType)
+static FValue* CastValueToPrimitiveType(FEmitter& Emitter, FValueRef Value, FType TargetType)
 {
-	const FPrimitiveType* ValuePrimitiveType = Value->Type->AsPrimitive();
-	if (!ValuePrimitiveType)
+	if (!Value->Type.AsPrimitive())
 	{
-		Emitter.Errorf(Value, TEXT("Cannot construct a '%s' from non primitive type '%s'."), TargetPrimitiveType->Spelling.GetData(), *Value->Type->GetSpelling());
+		Emitter.Errorf(Value, TEXT("Cannot construct a '%s' from non primitive type '%s'."), *Value->Type.GetSpelling(), *TargetType.GetSpelling());
 		return FPoison::Get();
 	}
 
+	FPrimitive ValuePrimitiveType = Value->Type.GetPrimitive();
+	FPrimitive TargetPrimitiveType = TargetType.GetPrimitive();
+
 	// Construct a scalar from another scalar.
-	if (TargetPrimitiveType->IsScalar())
+	if (TargetPrimitiveType.IsScalar())
 	{
 		// Get the first component of value. We already know value's type is primitive, so this will return a scalar.
 		Value = Emitter.Subscript(Value, 0);
-		ValuePrimitiveType = Value->Type->AsPrimitive();
-		check(ValuePrimitiveType->IsScalar());
+		ValuePrimitiveType = Value->Type.GetPrimitive();
+		check(ValuePrimitiveType.IsScalar());
 		
-		// If types are identical, return the component value as is.
 		if (ValuePrimitiveType == TargetPrimitiveType)
 		{
+			// If types are identical, return the component value as is.
 			return Value;
 		}
-		// If value is a constant, we can cast the constant now.
 		else if (FConstant* ConstantInitializer = As<FConstant>(Value))
 		{
-			return CastConstant(Emitter, ConstantInitializer, ValuePrimitiveType->ScalarKind, TargetPrimitiveType->ScalarKind);
+			// If value is a constant, we can cast the constant now.
+			return CastConstant(Emitter, ConstantInitializer, ValuePrimitiveType.ScalarKind, TargetPrimitiveType.ScalarKind);
 		}
 		else
 		{
 			// Otherwise emit a cast instruction of the subscript value to the target type.
-			FScalar Prototype = MakePrototype<FScalar>(TargetPrimitiveType);
+			FScalar Prototype = MakePrototype<FScalar>(TargetType);
 			Prototype.Arg = Value;
 			return EmitPrototype(Emitter, Prototype);
 		}
@@ -2515,20 +2516,20 @@ static FValue* CastValueToPrimitiveType(FEmitter& Emitter, FValueRef Value, cons
 
 	// Construct a vector or matrix from a scalar. E.g. 3.14f -> float3(3.14f, 3.14f, 3.14f)
 	// Note: we know target isn't scalar as it's been handled above.
-	if (ValuePrimitiveType->IsScalar())
+	if (ValuePrimitiveType.IsScalar())
 	{
 		// Create the result composite value.
-		FComposite* Result = NewCompositeValue(Emitter, TargetPrimitiveType, TargetPrimitiveType->NumComponents());
+		FComposite* Result = NewCompositeValue(Emitter, TargetType, TargetPrimitiveType.NumComponents());
 
 		// Create a composite and initialize each of its components to the conversion
 		// of initializer value to the single component type.
-		FValue* Component = Emitter.Cast(Value, TargetPrimitiveType->ToScalar());
+		FValue* Component = Emitter.Cast(Value, TargetPrimitiveType.ToScalar());
 
 		// Get the mutable array of components.
 		TArrayView<FValue*> ResultComponents = Result->GetMutableComponents();
 
 		// Initialize all result components to the same scalar.
-		for (int32 i = 0; i < TargetPrimitiveType->NumComponents(); ++i)
+		for (int32 i = 0; i < TargetPrimitiveType.NumComponents(); ++i)
 		{
 			ResultComponents[i] = Component;
 		}
@@ -2539,19 +2540,19 @@ static FValue* CastValueToPrimitiveType(FEmitter& Emitter, FValueRef Value, cons
 	// Construct a vector from another vector. If constructed vector is larger, initialize
 	// remaining components to zero. If it's smaller, truncate initializer vector and only use
 	// the necessary components.
-	if (TargetPrimitiveType->IsVector() && ValuePrimitiveType->IsVector())
+	if (TargetPrimitiveType.IsVector() && ValuePrimitiveType.IsVector())
 	{
 		// #todo-massimo.tristano Use swizzle when scalartypes are the same, and target num components is less than initializer's.
 
-		int32 TargetNumComponents = TargetPrimitiveType->NumComponents();
-		int32 InitializerNumComponents = ValuePrimitiveType->NumComponents();
+		int32 TargetNumComponents = TargetPrimitiveType.NumComponents();
+		int32 InitializerNumComponents = ValuePrimitiveType.NumComponents();
 
 		// Create the result composite value.
-		FComposite* Result = NewCompositeValue(Emitter, TargetPrimitiveType, TargetPrimitiveType->NumComponents());
+		FComposite* Result = NewCompositeValue(Emitter, TargetType, TargetPrimitiveType.NumComponents());
 		TArrayView<FValue*> ResultComponents = Result->GetMutableComponents();
 
 		// Determine the result component type (scalar).
-		const FPrimitiveType* ResultComponentType = TargetPrimitiveType->ToScalar();
+		FType ResultComponentType = TargetPrimitiveType.ToScalar();
 
 		// For iterating over the components of the result composite value.
 		int32 Index = 0;
@@ -2566,27 +2567,27 @@ static FValue* CastValueToPrimitiveType(FEmitter& Emitter, FValueRef Value, cons
 		// Initialize remaining result composite components to zero.
 		for (; Index < TargetNumComponents; ++Index)
 		{
-			ResultComponents[Index] = Emitter.ConstantZero(ResultComponentType->ScalarKind);
+			ResultComponents[Index] = Emitter.ConstantZero(ResultComponentType.GetPrimitive().ScalarKind);
 		}
 
 		return EmitNew(Emitter, Result);
 	}
 	
 	// The two primitive types are identical matrices that differ only by their scalar type.
-	if (TargetPrimitiveType->NumRows    == ValuePrimitiveType->NumRows &&
-		TargetPrimitiveType->NumColumns == ValuePrimitiveType->NumColumns)
+	if (TargetPrimitiveType.NumRows    == ValuePrimitiveType.NumRows &&
+		TargetPrimitiveType.NumColumns == ValuePrimitiveType.NumColumns)
 	{
-		check(TargetPrimitiveType->IsMatrix());
+		check(TargetPrimitiveType.IsMatrix());
 
 		// Create the result composite value.
-		FComposite* Result = NewCompositeValue(Emitter, TargetPrimitiveType, TargetPrimitiveType->NumComponents());
+		FComposite* Result = NewCompositeValue(Emitter, TargetType, TargetPrimitiveType.NumComponents());
 		TArrayView<FValue*> ResultComponents = Result->GetMutableComponents();
 			
 		// Determine the result component type (scalar).
-		const FPrimitiveType* ResultComponentType = TargetPrimitiveType->ToScalar();
+		FType ResultComponentType = TargetPrimitiveType.ToScalar();
 
 		// Convert components from the initializer vector.
-		for (int32 Index = 0, Num = TargetPrimitiveType->NumComponents(); Index < Num; ++Index)
+		for (int32 Index = 0, Num = ResultComponents.Num(); Index < Num; ++Index)
 		{
 			ResultComponents[Index] = Emitter.Cast(Emitter.Subscript(Value, Index), ResultComponentType);
 		}
@@ -2598,7 +2599,7 @@ static FValue* CastValueToPrimitiveType(FEmitter& Emitter, FValueRef Value, cons
 	return Emitter.Poison();
 }
 
-FValueRef FEmitter::Cast(FValueRef Value, const FType* TargetType)
+FValueRef FEmitter::Cast(FValueRef Value, FType TargetType)
 {
 	if (!Value.IsValid())
 	{
@@ -2606,21 +2607,21 @@ FValueRef FEmitter::Cast(FValueRef Value, const FType* TargetType)
 	}
 
 	// If target type matches initializer's, simply return the same value.
-	const FType* InitializerType = Value->Type;
+	FType InitializerType = Value->Type;
 	if (InitializerType == TargetType)
 	{
 		return Value;
 	}
 	
 	FValue* Result;
-	if (const FPrimitiveType* PrimitiveType = TargetType->AsPrimitive())
+	if (TargetType.AsPrimitive())
 	{
-		Result = CastValueToPrimitiveType(*this, Value, PrimitiveType);
+		Result = CastValueToPrimitiveType(*this, Value, TargetType);
 	}
 	else
 	{
 		// No other legal conversions applicable. Report error if we haven't converted the value.
-		Errorf(Value, TEXT("Cannot construct a '%s' from a '%s'."), *TargetType->GetSpelling(), *Value->Type->GetSpelling());
+		Errorf(Value, TEXT("Cannot construct a '%s' from a '%s'."), *TargetType.GetSpelling(), *Value->Type.GetSpelling());
 		return Poison();
 	}
 	
@@ -2631,7 +2632,7 @@ FValueRef FEmitter::CastToScalar(FValueRef Value)
 {
 	Value = CheckIsPrimitive(Value);
 	return Value.IsValid()
-		? Cast(Value, FPrimitiveType::GetScalar(Value->Type->AsPrimitive()->ScalarKind))
+		? Cast(Value, FType::MakeScalar(Value->Type.GetPrimitive().ScalarKind))
 		: Value;
 }
 
@@ -2639,41 +2640,41 @@ FValueRef FEmitter::CastToScalarKind(FValueRef Value, EScalarKind ToScalarKind)
 {
 	Value = CheckIsPrimitive(Value);
 	return Value.IsValid()
-		? Cast(Value, Value->Type->AsPrimitive()->WithScalarKind(ToScalarKind))
+		? Cast(Value, Value->Type.GetPrimitive().ToScalarKind(ToScalarKind))
 		: Value;
 }
 
 FValueRef FEmitter::CastToBoolKind(FValueRef Value)
 {
-	return CastToScalarKind(Value, ScalarKind_Bool);
+	return CastToScalarKind(Value, EScalarKind::Bool);
 }
 
 FValueRef FEmitter::CastToIntKind(FValueRef Value)
 {
-	return CastToScalarKind(Value, ScalarKind_Int);
+	return CastToScalarKind(Value, EScalarKind::Int);
 }
 
 FValueRef FEmitter::CastToFloatKind(FValueRef Value)
 {
-	return CastToScalarKind(Value, ScalarKind_Float);
+	return CastToScalarKind(Value, EScalarKind::Float);
 }
 
 FValueRef FEmitter::CastToBool(FValueRef Value, int NumRows)
 {
-	return Cast(Value, FPrimitiveType::GetVector(ScalarKind_Bool, NumRows));
+	return Cast(Value, FType::MakeVector(EScalarKind::Bool, NumRows));
 }
 
 FValueRef FEmitter::CastToInt(FValueRef Value, int NumRows)
 {
-	return Cast(Value, FPrimitiveType::GetVector(ScalarKind_Int, NumRows));
+	return Cast(Value, FType::MakeVector(EScalarKind::Int, NumRows));
 }
 
 FValueRef FEmitter::CastToFloat(FValueRef Value, int NumRows)
 {
-	return Cast(Value, FPrimitiveType::GetVector(ScalarKind_Float, NumRows));
+	return Cast(Value, FType::MakeVector(EScalarKind::Float, NumRows));
 }
 
-FValueRef FEmitter::StageSwitch(const FType* Type, TConstArrayView<FValueRef> ValuePerStage)
+FValueRef FEmitter::StageSwitch(FType Type, TConstArrayView<FValueRef> ValuePerStage)
 {
 	check(ValuePerStage.Num() <= NumStages);
 	FStageSwitch Prototype = MakePrototype<FStageSwitch>(Type);
@@ -2686,7 +2687,7 @@ FValueRef FEmitter::StageSwitch(const FType* Type, TConstArrayView<FValueRef> Va
 
 static FValueRef VTPageTableLoadFromSamplerSource(FEmitter& Em, FValueRef Texture, ESamplerSourceMode SamplerSourceMode, FValueRef TexCoord, FValueRef TexCoordDdx, FValueRef TexCoordDdy)
 {
-	MIR::FTextureObject* TextureObject = Texture->As<MIR::FTextureObject>();
+	FTextureObject* TextureObject = Texture->As<FTextureObject>();
 	check(TextureObject);
 
 	TextureAddress StaticAddressX = TA_Wrap;
@@ -2706,7 +2707,7 @@ FValueRef FEmitter::TextureGather(FValueRef Texture, FValueRef TexCoord, ETextur
 		return Poison();
 	}
 
-	FTextureRead Prototype = MakePrototype<FTextureRead>(FPrimitiveType::GetFloat4());
+	FTextureRead Prototype = MakePrototype<FTextureRead>(FType::MakeFloat(4));
 	Prototype.TextureObject = Texture;
 	Prototype.TexCoord = TexCoord;
 	Prototype.Mode = GatherMode;
@@ -2728,7 +2729,7 @@ FValueRef FEmitter::TextureSample(FValueRef Texture, FValueRef TexCoord, bool bA
 		return Poison();
 	}
 
-	FTextureRead PrototypeHw = MakePrototype<FTextureRead>(FPrimitiveType::GetFloat4());
+	FTextureRead PrototypeHw = MakePrototype<FTextureRead>(FType::MakeFloat(4));
 	PrototypeHw.TextureObject = Texture;
 	PrototypeHw.TexCoord = TexCoord;
 	PrototypeHw.Mode = ETextureReadMode::MipAuto;
@@ -2770,7 +2771,7 @@ FValueRef FEmitter::TextureSampleLevel(FValueRef Texture, FValueRef TexCoord, FV
 		return Poison();
 	}
 
-	FTextureRead Prototype = MakePrototype<FTextureRead>(FPrimitiveType::GetFloat4());
+	FTextureRead Prototype = MakePrototype<FTextureRead>(FType::MakeFloat(4));
 	Prototype.TextureObject = Texture;
 	Prototype.TexCoord = TexCoord;
 	Prototype.MipValue = MipLevel;
@@ -2803,7 +2804,7 @@ FValueRef FEmitter::TextureSampleBias(FValueRef Texture, FValueRef TexCoord, FVa
 		MipBias = Operator(BO_Add, MipBias, ExternalInput(EExternalInput::ViewMaterialTextureMipBias));
 	}
 
-	FTextureRead PrototypeHw = MakePrototype<FTextureRead>(FPrimitiveType::GetFloat4());
+	FTextureRead PrototypeHw = MakePrototype<FTextureRead>(FType::MakeFloat(4));
 	PrototypeHw.TextureObject = Texture;
 	PrototypeHw.TexCoord = TexCoord;
 	PrototypeHw.MipValue = MipBias;
@@ -2837,7 +2838,7 @@ FValueRef FEmitter::TextureSampleGrad(FValueRef Texture, FValueRef TexCoord, FVa
 		return Poison();
 	}
 
-	FTextureRead Prototype = MakePrototype<MIR::FTextureRead>(FPrimitiveType::GetFloat4());
+	FTextureRead Prototype = MakePrototype<FTextureRead>(FType::MakeFloat(4));
 	Prototype.TextureObject = Texture;
 	Prototype.TexCoord = TexCoord;
 	Prototype.TexCoordDdx = TexCoordDdx;
@@ -2864,7 +2865,7 @@ FValueRef FEmitter::TextureSampleGrad(FValueRef Texture, FValueRef TexCoord, FVa
 FValueRef FEmitter::VTPageTableLoad(FValueRef Texture, TextureAddress AddressU, TextureAddress AddressV, FValueRef TexCoord, FValueRef TexCoordDdx, FValueRef TexCoordDdy, bool bEnableFeedback)
 {
 	// VT stack and layer indices are initializes during IR analysis
-	FVTPageTableRead Prototype = MakePrototype<MIR::FVTPageTableRead>(FObjectType::GetVTPageTableResult());
+	FVTPageTableRead Prototype = MakePrototype<FVTPageTableRead>(FType::MakeVTPageTableResult());
 	Prototype.TextureObject = Texture;
 	Prototype.TexCoord = TexCoord;
 	Prototype.TexCoordDdx = TexCoordDdx;
@@ -2884,10 +2885,9 @@ FValueRef FEmitter::PartialDerivative(FValueRef Value, EDerivativeAxis Axis)
 	}
 
 	// Differentiation is only valid on primitive types.
-	const MIR::FPrimitiveType* ValuePrimitiveType = Value->Type->AsPrimitive();
-	if (!ValuePrimitiveType || !ScalarKindIsAnyFloat(ValuePrimitiveType->ScalarKind))
+	if (!Value->Type.IsAnyFloat())
 	{
-		Errorf(Value, TEXT("Trying to differentiate a value of type `%s` is invalid. Expected a float type."), *Value->Type->GetSpelling());
+		Errorf(Value, TEXT("Trying to differentiate a value of type `%s` is invalid. Expected a float type."), *Value->Type.GetSpelling());
 		return Poison();
 	}
 
@@ -2919,30 +2919,26 @@ static FValue* DifferentiateExternalInput(FEmitter& Emitter, FExternalInput* Ext
 	}
 
 	// All other inputs are assumed constant.
-	const MIR::FPrimitiveType* PrimType = static_cast<const MIR::FPrimitiveType*>(ExternalInput->Type);
-	return Emitter.Cast(Emitter.ConstantZero(PrimType->ScalarKind), PrimType);
+	return Emitter.Cast(Emitter.ConstantZero(ExternalInput->Type.GetPrimitive().ScalarKind), ExternalInput->Type);
 }
 
 static FValue* DifferentiateOperator(FEmitter& E, FOperator* Op, EDerivativeAxis Axis)
 {
-	const MIR::FPrimitiveType* PrimType = static_cast<const MIR::FPrimitiveType*>(Op->Type);
+	EScalarKind OpScalarKind = Op->Type.GetPrimitive().ScalarKind;
 
 	// Considering an operator acting on f(x), g(x) and h(x) arguments (e.g. "f(x) + g(x)"),
-	// calculate base terms and 
+	// calculate base terms and their partial derivatives.
 	FValue* F = Op->AArg;
 	FValue* G = Op->BArg;
 	FValue* H = Op->CArg;
-	FValue* dF = F && !F->Type->IsBoolean() ? E.AnalyticalPartialDerivative(F, Axis) : nullptr; // Note: select's first argument is a boolean, avoid making the derivative then
-	FValue* dG = G ? E.AnalyticalPartialDerivative(G, Axis) : nullptr;
-	FValue* dH = H ? E.AnalyticalPartialDerivative(H, Axis) : nullptr;
+	FValue* dF = F && !F->Type.IsBoolean() ? E.AnalyticalPartialDerivative(F, Axis) : nullptr; // Note: select's first argument is a boolean, avoid making the derivative then
+	FValue* dG = E.AnalyticalPartialDerivative(G, Axis);
+	FValue* dH = E.AnalyticalPartialDerivative(H, Axis);
 
 	// Convenience local functions as multiplications and division operations are common in derivatives.
-	auto Zero = [&E, PrimType] () { return E.ConstantZero(PrimType->ScalarKind); };
-	auto One = [&E, PrimType] () { return E.ConstantOne(PrimType->ScalarKind); };
-	auto Constant = [&E, PrimType] (TDouble FromDouble)
-	{ 
-		return PrimType->ScalarKind == ScalarKind_LWC ? E.ConstantLWC(FromDouble) : E.ConstantScalar(PrimType->ScalarKind, (float)FromDouble);
-	};
+	auto Zero = [&E, OpScalarKind] () { return E.ConstantZero(OpScalarKind); };
+	auto One = [&E, OpScalarKind] () { return E.ConstantOne(OpScalarKind); };
+	auto Constant = [&E, OpScalarKind] (TDouble Scalar) { return E.ConstantScalar(OpScalarKind, Scalar); };
 
 	// Some constants
 	constexpr TDouble Ln2 = 0.69314718055994530941723212145818;
@@ -3186,22 +3182,21 @@ FValueRef FEmitter::AnalyticalPartialDerivative(FValueRef Value, EDerivativeAxis
 	}
 
 	// Differentiation is only valid on primitive types.
-	const MIR::FPrimitiveType* ValuePrimitiveType = Value->Type->AsPrimitive();
-	if (!ValuePrimitiveType || !ScalarKindIsAnyFloat(ValuePrimitiveType->ScalarKind))
+	if (!Value->Type.IsAnyFloat())
 	{
-		Errorf(Value, TEXT("Trying to differentiate a value of type `%s` is invalid. Expected a float type."), *Value->Type->GetSpelling());
+		Errorf(Value, TEXT("Trying to differentiate a value of type `%s` is invalid. Expected a float type."), *Value->Type.GetSpelling());
 		return Poison();
 	}
 
 	switch (Value->Kind)
 	{
-		case MIR::VK_ExternalInput:
+		case VK_ExternalInput:
 			return DifferentiateExternalInput(*this, Value->As<FExternalInput>(), Axis);
 
-		case MIR::VK_Composite:
+		case VK_Composite:
 		{
 			TConstArrayView<FValue*> ValueComponents = Value->As<FComposite>()->GetComponents();
-			FComposite* Derivative = NewCompositeValue(*this, ValuePrimitiveType, ValuePrimitiveType->NumComponents());
+			FComposite* Derivative = NewCompositeValue(*this, Value->Type, Value->Type.GetPrimitive().NumComponents());
 			TArrayView<FValue*> DerivativeComponents = Derivative->GetMutableComponents();
 			for (int32 i = 0; i < ValueComponents.Num(); ++i)
 			{
@@ -3210,10 +3205,10 @@ FValueRef FEmitter::AnalyticalPartialDerivative(FValueRef Value, EDerivativeAxis
 			return Derivative;
 		}
 
-		case MIR::VK_Operator:
+		case VK_Operator:
 			return DifferentiateOperator(*this, Value->As<FOperator>(), Axis);
 		
-		case MIR::VK_Branch:
+		case VK_Branch:
 		{
 			FBranch* AsBranch = Value->As<FBranch>();
 			return Branch(AsBranch->ConditionArg,
@@ -3221,43 +3216,43 @@ FValueRef FEmitter::AnalyticalPartialDerivative(FValueRef Value, EDerivativeAxis
 				AnalyticalPartialDerivative(AsBranch->FalseArg, Axis));
 		}
 
-		case MIR::VK_Subscript:
+		case VK_Subscript:
 		{
 			FSubscript* AsSubscript = Value->As<FSubscript>();
 			return Subscript(AnalyticalPartialDerivative(AsSubscript->Arg, Axis), AsSubscript->Index);
 		}
 
-		case MIR::VK_Scalar:
+		case VK_Scalar:
 		{
 			FScalar* AsCast = Value->As<FScalar>();
 			return Cast(AnalyticalPartialDerivative(AsCast->Arg, Axis), AsCast->Type);
 		}
 
-		case MIR::VK_InlineHLSL:
+		case VK_InlineHLSL:
 		{
 			FInlineHLSL* AsInlineHLSL = Value->As<FInlineHLSL>();
-			if (!EnumHasAnyFlags(AsInlineHLSL->Flags, MIR::EValueFlags::HasDynamicHLSLCode) && AsInlineHLSL->ExternalCodeDeclaration->Derivative == EDerivativeStatus::Valid)
+			if (!EnumHasAnyFlags(AsInlineHLSL->Flags, EValueFlags::HasDynamicHLSLCode) && AsInlineHLSL->ExternalCodeDeclaration->Derivative == EDerivativeStatus::Valid)
 			{
 				TArray<FValueRef> Arguments(TArrayView<FValue*>(AsInlineHLSL->Arguments, AsInlineHLSL->NumArguments));
-				return InlineHLSL(AsInlineHLSL->ExternalCodeDeclaration, Arguments, Axis == EDerivativeAxis::X ? MIR::EValueFlags::DerivativeDDX : MIR::EValueFlags::DerivativeDDY);
+				return InlineHLSL(AsInlineHLSL->ExternalCodeDeclaration, Arguments, Axis == EDerivativeAxis::X ? EValueFlags::DerivativeDDX : EValueFlags::DerivativeDDY);
 			}
-			return Cast(ConstantZero(ValuePrimitiveType->ScalarKind), ValuePrimitiveType);
+			return Cast(ConstantZero(Value->Type.GetPrimitive().ScalarKind), Value->Type);
 		}
 
 		// These values are uniform (constant), thus their value is always zero.
-		case MIR::VK_Constant:
-		case MIR::VK_UniformParameter:
-		case MIR::VK_TextureRead:
-			return Cast(ConstantZero(ValuePrimitiveType->ScalarKind), ValuePrimitiveType);
+		case VK_Constant:
+		case VK_UniformParameter:
+		case VK_TextureRead:
+			return Cast(ConstantZero(Value->Type.GetPrimitive().ScalarKind), Value->Type);
 
 		default:
 			UE_MIR_UNREACHABLE();
 	}
 }
 
-static FValueRef EmitInlineHLSL(FEmitter& Emitter, const FType* Type, const FMaterialExternalCodeDeclaration* InExternalCodeDeclaration, const TCHAR* Code, TConstArrayView<FValueRef> InputValues, EValueFlags ValueFlags, EGraphProperties UsedGraphProperties)
+static FValueRef EmitInlineHLSL(FEmitter& Emitter, FType Type, const FMaterialExternalCodeDeclaration* InExternalCodeDeclaration, const TCHAR* Code, TConstArrayView<FValueRef> InputValues, EValueFlags ValueFlags, EGraphProperties UsedGraphProperties)
 {
-	MIR::FInlineHLSL Prototype = MakePrototype<MIR::FInlineHLSL>(Type);
+	FInlineHLSL Prototype = MakePrototype<FInlineHLSL>(Type);
 	Prototype.Type = Type;
 	Prototype.Flags = ValueFlags;
 	Prototype.GraphProperties = UsedGraphProperties;
@@ -3274,7 +3269,7 @@ static FValueRef EmitInlineHLSL(FEmitter& Emitter, const FType* Type, const FMat
 
 	if (!InputValues.IsEmpty())
 	{
-		checkf(InputValues.Num() < MIR::FInlineHLSL::MaxNumArguments, TEXT("Number of arguments for inline-HLSL out of bounds: %d was specified, but upper bound is %d"), InputValues.Num(), MIR::FInlineHLSL::MaxNumArguments);
+		checkf(InputValues.Num() < FInlineHLSL::MaxNumArguments, TEXT("Number of arguments for inline-HLSL out of bounds: %d was specified, but upper bound is %d"), InputValues.Num(), FInlineHLSL::MaxNumArguments);
 		Prototype.NumArguments = InputValues.Num();
 		for (int32 i = 0; i < InputValues.Num(); ++i)
 		{
@@ -3286,7 +3281,7 @@ static FValueRef EmitInlineHLSL(FEmitter& Emitter, const FType* Type, const FMat
 	return EmitPrototype(Emitter, Prototype);
 }
 
-FValueRef FEmitter::InlineHLSL(const FType* Type, FString Code, TConstArrayView<FValueRef> InputValues, EValueFlags ValueFlags, EGraphProperties UsedGraphProperties)
+FValueRef FEmitter::InlineHLSL(FType Type, FString Code, TConstArrayView<FValueRef> InputValues, EValueFlags ValueFlags, EGraphProperties UsedGraphProperties)
 {
 	if (IsAnyNotValid(InputValues))
 	{
@@ -3304,17 +3299,17 @@ FValueRef FEmitter::InlineHLSL(const FMaterialExternalCodeDeclaration* InExterna
 	}
 
 	check(InExternalCodeDeclaration != nullptr);
-	const FType* ReturnType = FType::FromMaterialValueType(InExternalCodeDeclaration->GetReturnTypeValue());
+	FType ReturnType = FType::FromMaterialValueType(InExternalCodeDeclaration->GetReturnTypeValue());
 	return EmitInlineHLSL(*this, ReturnType, InExternalCodeDeclaration, nullptr, InputValues, ValueFlags, UsedGraphProperties);
 }
 
 FValueRef FEmitter::PromoteSubstrateParameter()
 {
-	MIR::FPromoteSubstrateParameter Prototype = MakePrototype<MIR::FPromoteSubstrateParameter>(FObjectType::GetSubstrateData());
+	FPromoteSubstrateParameter Prototype = MakePrototype<FPromoteSubstrateParameter>(FType::MakeSubstrateData());
 	return EmitPrototype(*this, Prototype);
 }
 
-const FType* FEmitter::TryGetCommonType(const FType* A, const FType* B)
+FType FEmitter::TryGetCommonType(FType A, FType B)
 {
 	// Trivial case: types are equal
 	if (A == B)
@@ -3322,36 +3317,46 @@ const FType* FEmitter::TryGetCommonType(const FType* A, const FType* B)
 		return A;
 	}
 
-	const FPrimitiveType* PrimitiveA = A->AsPrimitive();
-	const FPrimitiveType* PrimitiveB = B->AsPrimitive();
+	TOptional<FPrimitive> PrimitiveA = A.AsPrimitive();
+	TOptional<FPrimitive> PrimitiveB = B.AsPrimitive();
 	if (!PrimitiveA || !PrimitiveB)
 	{
-		return nullptr;
+		return FType::MakePoison();
 	}
 	
 	// If both A and B are matrices, their dimensions must match (equality check above didn't trigger).
 	if (PrimitiveA->IsMatrix() || PrimitiveB->IsMatrix())
 	{
-		return nullptr;
+		return FType::MakePoison();
 	}
 
 	// Neither A nor B are matrices, but single scalar or vector. Return the largest.
 	check(PrimitiveA->NumColumns == 1 && PrimitiveB->NumColumns == 1);
 	EScalarKind ScalarKind = FMath::Max(PrimitiveA->ScalarKind, PrimitiveB->ScalarKind);
 	int32 NumRows = FMath::Max(PrimitiveA->NumRows, PrimitiveB->NumRows);
-	return FPrimitiveType::Get(ScalarKind, NumRows, 1);
+	return FType::MakePrimitive(ScalarKind, NumRows, 1);
 }
 
 void FEmitter::Initialize()
 {
 	// Create and reference the true/false constants.
-	FConstant Temp = MakePrototype<FConstant>(FPrimitiveType::GetBool());
+	FConstant Temp = MakePrototype<FConstant>(FType::MakeBool());
 
 	Temp.Boolean = true;
 	TrueConstant = EmitPrototype(*this, Temp);
 
 	Temp.Boolean = false;
 	FalseConstant = EmitPrototype(*this, Temp);
+
+	// This check ensures that FValue caching returns identical results for repeated values.
+	// If this assert fails, it's likely due to uninitialized padding bytes in FValue introduced by compiler-added padding.
+	// Review any recent changes to FValue to ensure no additional padding was introduced (i.e., the struct remains tightly packed).
+	if (FPlatformMisc::IsDebuggerPresent())
+	{
+		FValue* A = ConstantZero(MIR::EScalarKind::Float);
+		FValue* B = ConstantZero(MIR::EScalarKind::Float);
+		check(A == B);
+	}
 }
 
 bool FEmitter::FValueKeyFuncs::Matches(KeyInitType A, KeyInitType B)
diff --git a/Engine/Source/Runtime/Engine/Private/Materials/MaterialIRModule.cpp b/Engine/Source/Runtime/Engine/Private/Materials/MaterialIRModule.cpp
--- a/Engine/Source/Runtime/Engine/Private/Materials/MaterialIRModule.cpp
+++ b/Engine/Source/Runtime/Engine/Private/Materials/MaterialIRModule.cpp
@@ -8,27 +8,19 @@
 
 FMaterialIRModule::FMaterialIRModule()
 {
-	for (int32 i = 0; i < MIR::NumStages; ++i)
-	{
-		RootBlock[i] = new MIR::FBlock;
-	}
 }
 
 FMaterialIRModule::~FMaterialIRModule()
 {
 	Empty();
-
-	for (int32 i = 0; i < MIR::NumStages; ++i)
-	{
-		delete RootBlock[i];
-	}
 }
 
 void FMaterialIRModule::Empty()
 {
+	FMemory::Memset(RootBlocks, 0, sizeof(RootBlocks));
+
 	for (int32 i = 0; i < MIR::NumStages; ++i)
 	{
-		RootBlock[i]->Instructions = nullptr;
 		Outputs[i].Empty();
 	}
 
@@ -38,12 +30,6 @@ void FMaterialIRModule::Empty()
 	}
 	Values.Empty();
 
-	for (MIR::FType* Type : Types)
-	{
-		FMemory::Free(Type);
-	}
-	Types.Empty();
-
 	// Reset module statistics.
 	for (int i = 0; i < MIR::NumStages; ++i)
 	{
diff --git a/Engine/Source/Runtime/Engine/Private/Materials/MaterialIRModuleBuilder.cpp b/Engine/Source/Runtime/Engine/Private/Materials/MaterialIRModuleBuilder.cpp
--- a/Engine/Source/Runtime/Engine/Private/Materials/MaterialIRModuleBuilder.cpp
+++ b/Engine/Source/Runtime/Engine/Private/Materials/MaterialIRModuleBuilder.cpp
@@ -276,7 +276,7 @@ struct FMaterialIRModuleBuilderImpl
 			MIR::FValue* Value = Emitter->Input(Call->GetInput(i));
 			if (Value)
 			{
-				const MIR::FType* Type = MIR::FType::FromMaterialValueType(FunctionInput.ExpressionInput->GetInputValueType(0));
+				MIR::FType Type = MIR::FType::FromMaterialValueType(FunctionInput.ExpressionInput->GetInputValueType(0));
 				CallInputValues[i] = Emitter->Cast(Value, Type);
 			}
 		}
@@ -333,7 +333,7 @@ struct FMaterialIRModuleBuilderImpl
 			MIR::FValue* Value = Emitter->Input(FunctionOutput.ExpressionOutput->GetInput(0));
 
 			// Get the function output type.
-			const MIR::FType* OutputType = MIR::FType::FromMaterialValueType(FunctionOutput.ExpressionOutput->GetOutputValueType(0));
+			MIR::FType OutputType = MIR::FType::FromMaterialValueType(FunctionOutput.ExpressionOutput->GetOutputValueType(0));
 
 			// Cast the value to the expected output type. This may fail (value will be poison). 
 			Value = Emitter->Cast(Value, OutputType);
@@ -372,13 +372,13 @@ struct FMaterialIRModuleBuilderImpl
 					// If this property is the emissive color and we're previewing the material, apply gamma correction to the previewed value.
 					if (Output->Property == MP_EmissiveColor && PreviewInput.IsConnected())
 					{
-						MIR::FValue* Zero = Emitter->ConstantZero(MIR::EScalarKind::ScalarKind_Float);
+						MIR::FValue* Zero = Emitter->ConstantZero(MIR::EScalarKind::Float);
 
 						// Get preview expression back into gamma corrected space, as DrawTile does not do this adjustment.
 						Value = Emitter->Pow(Emitter->Max(Value, Zero), Emitter->ConstantFloat(1.0f / 2.2f));
 
 						// Preview should display scalars as red, so if this is a scalar, create a vector padded with zeroes
-						if (Value->IsScalar())
+						if (Value->Type.IsScalar())
 						{
 							Value = Emitter->Vector3(Value, Zero, Zero);
 						}
@@ -388,7 +388,7 @@ struct FMaterialIRModuleBuilderImpl
 					else
 					{
 						// If a value is flowing in through the connection, cast it to this material attribute type and assign it.
-						const MIR::FType* OutputArgType = MIR::FType::FromShaderType(Input.Type);
+						MIR::FType OutputArgType = MIR::FType::FromShaderType(Input.Type);
 						Output->Arg = Emitter->Cast(Value, OutputArgType);
 					}
 				}
@@ -593,7 +593,7 @@ struct FMaterialIRModuleBuilderImpl
 			// Initialize each output's block in this stage and add it to the instruction stack.
 			for (MIR::FSetMaterialOutput* Output : Module->Outputs[StageIndex])
 			{
-				Output->Block[StageIndex] = Module->RootBlock[StageIndex];
+				Output->Block[StageIndex] = &Module->RootBlocks[StageIndex];
 
 				InstructionStack.Add(Output);
 			}
@@ -673,9 +673,9 @@ struct FMaterialIRModuleBuilderImpl
 	}
 	
 	//
-	void PushConnectionInsight(const UObject* InputObject, int InputIndex, const UMaterialExpression* OutputExpression, int OutputIndex, const MIR::FType* Type)
+	void PushConnectionInsight(const UObject* InputObject, int InputIndex, const UMaterialExpression* OutputExpression, int OutputIndex, MIR::FType Type)
 	{
-		if (!Builder->TargetInsights || !Type || Type->IsPoison())
+		if (!Builder->TargetInsights || Type.IsPoison())
 		{
 			return;
 		}
@@ -685,7 +685,7 @@ struct FMaterialIRModuleBuilderImpl
 			.OutputExpression = OutputExpression,
 			.InputIndex = InputIndex,
 			.OutputIndex = OutputIndex,
-			.ValueType = Type->ToValueType(),
+			.ValueType = Type.ToValueType(),
 		};
 		
 		Builder->TargetInsights->ConnectionInsights.Push(Insight);
diff --git a/Engine/Source/Runtime/Engine/Private/Materials/MaterialIRToHLSLTranslator.cpp b/Engine/Source/Runtime/Engine/Private/Materials/MaterialIRToHLSLTranslator.cpp
--- a/Engine/Source/Runtime/Engine/Private/Materials/MaterialIRToHLSLTranslator.cpp
+++ b/Engine/Source/Runtime/Engine/Private/Materials/MaterialIRToHLSLTranslator.cpp
@@ -1,23 +1,22 @@
 // Copyright Epic Games, Inc. All Rights Reserved.
 
 #include "Materials/MaterialIRToHLSLTranslator.h"
-#include "Materials/MaterialIRModule.h"
-#include "Materials/MaterialIRTypes.h"
-#include "Materials/MaterialIR.h"
-#include "Materials/MaterialExternalCodeRegistry.h"
+#include "DataDrivenShaderPlatformInfo.h"
+#include "Engine/Texture.h"
+#include "MaterialDomain.h"
 #include "MaterialIRInternal.h"
-#include "MaterialSharedPrivate.h"
-
-#include "ShaderCore.h"
-#include "MaterialShared.h"
-#include "Materials/MaterialAttributeDefinitionMap.h"
 #include "Materials/Material.h"
-#include "MaterialDomain.h"
-#include "DataDrivenShaderPlatformInfo.h"
+#include "Materials/MaterialAttributeDefinitionMap.h"
 #include "Materials/MaterialExpressionVolumetricAdvancedMaterialOutput.h"
+#include "Materials/MaterialExternalCodeRegistry.h"
+#include "Materials/MaterialIR.h"
+#include "Materials/MaterialIRModule.h"
+#include "Materials/MaterialIRTypes.h"
+#include "MaterialShared.h"
+#include "MaterialSharedPrivate.h"
 #include "Misc/LargeWorldRenderPosition.h"
 #include "RenderUtils.h"
-#include "Engine/Texture.h"
+#include "ShaderCore.h"
 
 #include <inttypes.h>
 
@@ -223,8 +222,7 @@ static bool InstructionUsesPhiValue(const MIR::FInstruction* Instr, MIR::EStage
 {
 	if (auto Branch = Instr->As<MIR::FBranch>())
 	{
-		return Branch->TrueBlock[Stage].Instructions
-			|| Branch->FalseBlock[Stage].Instructions;
+		return Branch->TrueBlock[Stage].Instructions || Branch->FalseBlock[Stage].Instructions;
 	}
 	return false;
 }
@@ -240,10 +238,8 @@ static bool IsFoldable(const MIR::FInstruction* Instr, MIR::EStage Stage)
 	// Don't fold instructions that become very long to improve readability
 	switch (Instr->Kind)
 	{
-	case MIR::VK_VTPageTableRead:
-		return false;
-	default:
-		break;
+		case MIR::VK_VTPageTableRead: return false;
+		default: break;
 	}
 
 	return true;
@@ -400,9 +396,9 @@ struct FTranslator : FMaterialIRToHLSLTranslation
 		return NoOp;
 	}
 	
-	ENoOp LowerValue(MIR::FValue* InValue)
+	ENoOp LowerValue(const MIR::FValue* InValue)
 	{
-		if (MIR::FInstruction* Instr = MIR::AsInstruction(InValue))
+		if (const MIR::FInstruction* Instr = MIR::AsInstruction(InValue))
 		{
 			if (Instr->NumUsers[CurrentStage] <= 1 && IsFoldable(Instr, CurrentStage))
 			{
@@ -452,29 +448,29 @@ struct FTranslator : FMaterialIRToHLSLTranslation
 
 	void LowerConstant(const MIR::FConstant* Constant)
 	{
-		const MIR::FPrimitiveType* checkNoEntry = Constant->Type->AsPrimitive();
+		TOptional<MIR::FPrimitive> checkNoEntry = Constant->Type.AsPrimitive();
 		check(checkNoEntry && checkNoEntry->IsScalar());
 
 		switch (checkNoEntry->ScalarKind)
 		{
-			case MIR::ScalarKind_Bool:
+			case MIR::EScalarKind::Bool:
 				Printer.Buffer.Append(Constant->Boolean ? TEXT("true") : TEXT("false"));
 				break;
 				
-			case MIR::ScalarKind_Int:
+			case MIR::EScalarKind::Int:
 				Printer.Buffer.Appendf(TEXT("%") PRId64, Constant->Integer);
 				break;
 			
-			case MIR::ScalarKind_Float:
+			case MIR::EScalarKind::Float:
 				Printer << Constant->Float;
 				break;
-
-			case MIR::ScalarKind_LWC:
-				{
-					FLargeWorldRenderScalar ConstantLWC(Constant->LWCValue);
-					Printer << "MakeLWCScalar(" << ConstantLWC.GetTile() << ", " << ConstantLWC.GetOffset() << ")";
-				}
+			
+			case MIR::EScalarKind::Double:
+			{
+				FLargeWorldRenderScalar ConstantLWC(Constant->Double);
+				Printer << "MakeLWCScalar(" << ConstantLWC.GetTile() << ", " << ConstantLWC.GetOffset() << ")";
 				break;
+			}
 		}
 	}
 
@@ -700,7 +696,7 @@ struct FTranslator : FMaterialIRToHLSLTranslation
 
 	void LowerUniformParameter(const MIR::FUniformParameter* UniformParameter)
 	{
-		if (UniformParameter->Type->IsObjectOfKind(MIR::ObjectKind_Texture2D))
+		if (UniformParameter->Type.IsTexture2D())
 		{
 		 	EMaterialValueType TextureType = Module->GetParameterMetadata(UniformParameter->ParameterIdInModule).Value.Texture->GetMaterialType();
 			LowerTextureReference(TextureType, UniformParameter->Analysis_UniformParameterIndex);
@@ -713,7 +709,7 @@ struct FTranslator : FMaterialIRToHLSLTranslation
 
 	struct FSubstrateLegacyArgument
 	{
-		const MIR::FType* ParameterType;
+		MIR::FType ParameterType;
 		const TCHAR* ParameterName;
 		EMaterialProperty MaterialProperty;
 		const TCHAR* DefaultInlineHLSL;
@@ -723,31 +719,31 @@ struct FTranslator : FMaterialIRToHLSLTranslation
 	{
 		static const FSubstrateLegacyArgument LegacyArguments[] =
 		{
-			{ MIR::FPrimitiveType::GetFloat3(), TEXT("BaseColor"), MP_BaseColor, nullptr },
-			{ MIR::FPrimitiveType::GetFloat(), TEXT("Specular"), MP_SpecularColor, TEXT("0.5f") },
-			{ MIR::FPrimitiveType::GetFloat(), TEXT("Metallic"), MP_Metallic, nullptr },
-			{ MIR::FPrimitiveType::GetFloat(), TEXT("Roughness"), MP_Roughness, nullptr },
-			{ MIR::FPrimitiveType::GetFloat(), TEXT("Anisotropy"), MP_Anisotropy, nullptr },
-			{ MIR::FPrimitiveType::GetFloat3(), TEXT("SubSurfaceColor"), MP_SubsurfaceColor, TEXT("(float3)1") },
-			{ MIR::FPrimitiveType::GetFloat(), TEXT("SubSurfaceProfileId"), MP_MAX, TEXT("0") },
-			{ MIR::FPrimitiveType::GetFloat(), TEXT("ClearCoat"), MP_MAX, TEXT("1.0f") },
-			{ MIR::FPrimitiveType::GetFloat(), TEXT("ClearCoatRoughness"), MP_MAX, TEXT("0.1f") },
-			{ MIR::FPrimitiveType::GetFloat3(), TEXT("Emissive"), MP_EmissiveColor, nullptr },
-			{ MIR::FPrimitiveType::GetFloat(), TEXT("Opacity"), MP_Opacity, nullptr },
-			{ MIR::FPrimitiveType::GetFloat3(), TEXT("ThinTranslucentTransmittanceColor"), MP_MAX, TEXT("(float3)0.5f") },
-			{ MIR::FPrimitiveType::GetFloat(), TEXT("ThinTranslucentSurfaceCoverage"), MP_MAX, TEXT("1.0f") },
-			{ MIR::FPrimitiveType::GetFloat3(), TEXT("WaterScatteringCoefficients"), MP_MAX, nullptr },
-			{ MIR::FPrimitiveType::GetFloat3(), TEXT("WaterAbsorptionCoefficients"), MP_MAX, nullptr },
-			{ MIR::FPrimitiveType::GetFloat(), TEXT("WaterPhaseG"), MP_MAX, nullptr },
-			{ MIR::FPrimitiveType::GetFloat3(), TEXT("ColorScaleBehindWater"), MP_MAX, TEXT("(float3)1.0f") },
-			{ MIR::FPrimitiveType::GetInt(), TEXT("ShadingModel"), MP_MAX, TEXT("1") },
-			{ MIR::FPrimitiveType::GetFloat3(), TEXT("RawNormal"), MP_Normal, TEXT("Parameters.TangentToWorld[2]") },
-			{ MIR::FPrimitiveType::GetFloat3(), TEXT("RawTangent"), MP_Tangent, TEXT("Parameters.TangentToWorld[0]") },
-			{ MIR::FPrimitiveType::GetFloat3(), TEXT("RawClearCoatNormal"), MP_Normal, TEXT("Parameters.TangentToWorld[2]") },
-			{ MIR::FPrimitiveType::GetFloat3(), TEXT("RawCustomTangent"), MP_Tangent, TEXT("Parameters.TangentToWorld[0]") },
-			{ MIR::FPrimitiveType::GetInt(), TEXT("SharedLocalBasisIndex"), MP_MAX, nullptr },
-			{ MIR::FPrimitiveType::GetInt(), TEXT("ClearCoatBottomNormal_SharedLocalBasisIndex"), MP_MAX, nullptr },
-			{ MIR::FPrimitiveType::GetInt(), TEXT("SharedLocalBasisTypes"), MP_MAX, TEXT("Parameters.SharedLocalBases.Types") },
+			{ MIR::FType::MakeFloat(3), TEXT("BaseColor"), MP_BaseColor, nullptr },
+			{ MIR::FType::MakeFloat(),  TEXT("Specular"), MP_SpecularColor, TEXT("0.5f") },
+			{ MIR::FType::MakeFloat(),  TEXT("Metallic"), MP_Metallic, nullptr },
+			{ MIR::FType::MakeFloat(),  TEXT("Roughness"), MP_Roughness, nullptr },
+			{ MIR::FType::MakeFloat(),  TEXT("Anisotropy"), MP_Anisotropy, nullptr },
+			{ MIR::FType::MakeFloat(3), TEXT("SubSurfaceColor"), MP_SubsurfaceColor, TEXT("(float3)1") },
+			{ MIR::FType::MakeFloat(),  TEXT("SubSurfaceProfileId"), MP_MAX, TEXT("0") },
+			{ MIR::FType::MakeFloat(),  TEXT("ClearCoat"), MP_MAX, TEXT("1.0f") },
+			{ MIR::FType::MakeFloat(),  TEXT("ClearCoatRoughness"), MP_MAX, TEXT("0.1f") },
+			{ MIR::FType::MakeFloat(3), TEXT("Emissive"), MP_EmissiveColor, nullptr },
+			{ MIR::FType::MakeFloat(),  TEXT("Opacity"), MP_Opacity, nullptr },
+			{ MIR::FType::MakeFloat(3), TEXT("ThinTranslucentTransmittanceColor"), MP_MAX, TEXT("(float3)0.5f") },
+			{ MIR::FType::MakeFloat(),  TEXT("ThinTranslucentSurfaceCoverage"), MP_MAX, TEXT("1.0f") },
+			{ MIR::FType::MakeFloat(3), TEXT("WaterScatteringCoefficients"), MP_MAX, nullptr },
+			{ MIR::FType::MakeFloat(3), TEXT("WaterAbsorptionCoefficients"), MP_MAX, nullptr },
+			{ MIR::FType::MakeFloat(),  TEXT("WaterPhaseG"), MP_MAX, nullptr },
+			{ MIR::FType::MakeFloat(3), TEXT("ColorScaleBehindWater"), MP_MAX, TEXT("(float3)1.0f") },
+			{ MIR::FType::MakeInt(),    TEXT("ShadingModel"), MP_MAX, TEXT("1") },
+			{ MIR::FType::MakeFloat(3), TEXT("RawNormal"), MP_Normal, TEXT("Parameters.TangentToWorld[2]") },
+			{ MIR::FType::MakeFloat(3), TEXT("RawTangent"), MP_Tangent, TEXT("Parameters.TangentToWorld[0]") },
+			{ MIR::FType::MakeFloat(3), TEXT("RawClearCoatNormal"), MP_Normal, TEXT("Parameters.TangentToWorld[2]") },
+			{ MIR::FType::MakeFloat(3), TEXT("RawCustomTangent"), MP_Tangent, TEXT("Parameters.TangentToWorld[0]") },
+			{ MIR::FType::MakeInt(),    TEXT("SharedLocalBasisIndex"), MP_MAX, nullptr },
+			{ MIR::FType::MakeInt(),    TEXT("ClearCoatBottomNormal_SharedLocalBasisIndex"), MP_MAX, nullptr },
+			{ MIR::FType::MakeInt(),    TEXT("SharedLocalBasisTypes"), MP_MAX, TEXT("Parameters.SharedLocalBases.Types") },
 		};
 		return LegacyArguments;
 	}
@@ -756,9 +752,9 @@ struct FTranslator : FMaterialIRToHLSLTranslation
 	{
 		static const FSubstrateLegacyArgument LegacyArguments[] =
 		{
-			{ MIR::FPrimitiveType::GetFloat3(), TEXT("Emissive"), MP_EmissiveColor, nullptr },
-			{ MIR::FPrimitiveType::GetFloat(), TEXT("TransmittanceColor"), MP_MAX, TEXT("(float3)0.0f") },
-			{ MIR::FPrimitiveType::GetFloat3(), TEXT("RawNormal"), MP_Normal, TEXT("Parameters.TangentToWorld[2]") },
+			{ MIR::FType::MakeFloat(3), TEXT("Emissive"), MP_EmissiveColor, nullptr },
+			{ MIR::FType::MakeFloat(), TEXT("TransmittanceColor"), MP_MAX, TEXT("(float3)0.0f") },
+			{ MIR::FType::MakeFloat(3), TEXT("RawNormal"), MP_Normal, TEXT("Parameters.TangentToWorld[2]") },
 		};
 		return LegacyArguments;
 	}
@@ -793,36 +789,36 @@ struct FTranslator : FMaterialIRToHLSLTranslation
 
 				default:
 				{
-			if (Argument.MaterialProperty != MP_MAX && SubstrateParameter->PropertyArgs[Argument.MaterialProperty])
-			{
-				LowerValue(SubstrateParameter->PropertyArgs[Argument.MaterialProperty]);
-			}
-			else if (Argument.DefaultInlineHLSL)
-			{
-				Printer << Argument.DefaultInlineHLSL;
-			}
-			else if (const MIR::FPrimitiveType* PrimitiveType = Argument.ParameterType->AsPrimitive())
-			{
-				if (const MIR::FPrimitiveType* ScalarType = Argument.ParameterType->AsScalar())
-				{
-					Printer << TEXT("0.0f");
-				}
-				else
-				{
-					LowerType(Argument.ParameterType);
-					Printer << BeginArgs;
-					for (int32 ComponentIndex = 0; ComponentIndex < PrimitiveType->NumRows; ++ComponentIndex)
+					if (Argument.MaterialProperty != MP_MAX && SubstrateParameter->PropertyArgs[Argument.MaterialProperty])
+					{
+						LowerValue(SubstrateParameter->PropertyArgs[Argument.MaterialProperty]);
+					}
+					else if (Argument.DefaultInlineHLSL)
 					{
-						Printer << ListSeparator << TEXT("0.0f");
+						Printer << Argument.DefaultInlineHLSL;
+					}
+					else if (TOptional<MIR::FPrimitive> PrimitiveType = Argument.ParameterType.AsPrimitive())
+					{
+						if (PrimitiveType->IsScalar())
+						{
+							Printer << TEXT("0.0f");
+						}
+						else
+						{
+							LowerType(Argument.ParameterType);
+							Printer << BeginArgs;
+							for (int32 ComponentIndex = 0; ComponentIndex < PrimitiveType->NumRows; ++ComponentIndex)
+							{
+								Printer << ListSeparator << TEXT("0.0f");
+							}
+							Printer << EndArgs;
+						}
+					}
+					else
+					{
+						UE_MIR_UNREACHABLE();
 					}
-					Printer << EndArgs;
 				}
-			}
-			else
-			{
-				UE_MIR_UNREACHABLE();
-			}
-		}
 				break;
 			}
 		}
@@ -834,11 +830,16 @@ struct FTranslator : FMaterialIRToHLSLTranslation
 		Printer << NewLine << EndArgs;
 	}
 
-	void LowerPrimitiveUniformParameter(const MIR::FUniformParameter* UniformParameter)
+	ENoOp LowerPrimitiveUniformParameter(const MIR::FUniformParameter* UniformParameter)
 	{
-		const MIR::FPrimitiveType* PrimitiveType = UniformParameter->Type->AsPrimitive();
-		check(PrimitiveType);
-		check(PrimitiveType->IsScalar() || PrimitiveType->IsVector()); // no matrices yet
+		MIR::FPrimitive PrimitiveType = UniformParameter->Type.GetPrimitive();
+		check(PrimitiveType.IsScalar() || PrimitiveType.IsVector()); // no matrices yet
+
+		// LWC parameters are handled differently, they have their own dedicated function.
+		if (PrimitiveType.IsLWC())
+		{
+			return LowerPrimitiveUniformParameterLWC(UniformParameter, PrimitiveType.NumComponents());
+		}
 
 		const FUniformExpressionSet& UniformExpressionSet = Module->GetCompilationOutput().UniformExpressionSet;
 		
@@ -850,49 +851,39 @@ struct FTranslator : FMaterialIRToHLSLTranslation
 
 		// Starting component of the float4 slot
 		uint32 BufferSlotOffset = GlobalComponentOffset % 4;
-
-		if (PrimitiveType->ScalarKind == MIR::EScalarKind::ScalarKind_LWC)
+		
+		if (PrimitiveType.IsInteger())
 		{
-			// Currently preshaders store values as DF, but eventually it would be more efficient to store them as LWC to avoid conversion.
-			// LWC parameters are always four element vectors, so we assume the slot offset is zero.
-			check(BufferSlotOffset == 0);
-			Printer << "DFToWS(MakeDFVector4(Material.PreshaderBuffer[" << BufferSlotIndex << "], Material.PreshaderBuffer[" << BufferSlotIndex + 1 << "]))";
+			Printer << TEXT("asint(");
 		}
-		else
-		{
-			if (PrimitiveType->ScalarKind == MIR::EScalarKind::ScalarKind_Int)
-			{
-				Printer << TEXT("asint(");
-			}
-
-			Printer << TEXT("Material.PreshaderBuffer[") << BufferSlotIndex << TEXT("]");
 
-			if (PrimitiveType->NumComponents() < 4)
-			{
-				Printer << TEXT(".");
+		Printer << TEXT("Material.PreshaderBuffer[") << BufferSlotIndex << TEXT("]");
 
-				static const TCHAR* Components = TEXT("xyzw");
-				for (int32 i = 0; i < PrimitiveType->NumComponents(); ++i)
-				{
-					check(BufferSlotOffset + i < sizeof(Components));
-					Printer.Buffer.AppendChar(Components[BufferSlotOffset + i]);
-				}
-			}
+		if (PrimitiveType.NumComponents() < 4)
+		{
+			Printer << TEXT(".");
 
-			if (PrimitiveType->ScalarKind == MIR::EScalarKind::ScalarKind_Int)
+			static const TCHAR* Components = TEXT("xyzw");
+			for (int32 i = 0; i < PrimitiveType.NumComponents(); ++i)
 			{
-				Printer << ")"; // close the "asint(" bracket
+				check(BufferSlotOffset + i < sizeof(Components));
+				Printer.Buffer.AppendChar(Components[BufferSlotOffset + i]);
 			}
 		}
+
+		if (PrimitiveType.IsInteger())
+		{
+			Printer << ")"; // close the "asint(" bracket
+		}
+
+		return NoOp;
 	}
 
-	// Variation of LowerPrimitiveUniformParameter specific to LWC, to support fetching a subset of an LWC vector
-	void LowerPrimitiveUniformParameterLWC(const MIR::FUniformParameter* UniformParameter, int32 NumComponents)
+	// Variation of LowerPrimitiveUniformParameter specific to LWC, to support fetching a subset of an LWC vector.
+	ENoOp LowerPrimitiveUniformParameterLWC(const MIR::FUniformParameter* UniformParameter, int32 NumComponents)
 	{
 		check(NumComponents >= 2 && NumComponents <= 4);
-
-		const MIR::FPrimitiveType* PrimitiveType = UniformParameter->Type->AsPrimitive();
-		check(PrimitiveType->IsLWC());
+		check(UniformParameter->Type.IsLWC());
 
 		const FUniformExpressionSet& UniformExpressionSet = Module->GetCompilationOutput().UniformExpressionSet;
 
@@ -905,37 +896,43 @@ struct FTranslator : FMaterialIRToHLSLTranslation
 		Printer << "DFToWS(MakeDFVector" << NumComponents;
 		Printer << "(Material.PreshaderBuffer[" << BufferSlotIndex << "]" << GVector4SwizzleSubset[NumComponents - 1];
 		Printer << ", Material.PreshaderBuffer[" << BufferSlotIndex + 1 << "]" << GVector4SwizzleSubset[NumComponents - 1] << "))";
+	
+		return NoOp;
 	}
 
 	bool HasMatchingScalarComponentCastChain(MIR::FValue* FirstComponent, MIR::FValue* CurrentComponent, int32 Index, int32 NumComponents)
 	{
-		if (FirstComponent->Kind == MIR::VK_Subscript && CurrentComponent->Kind == MIR::VK_Subscript)
+		const MIR::FSubscript* FirstComponentAsSubscript = FirstComponent->As<MIR::FSubscript>();
+		const MIR::FSubscript* CurrentComponentAsSubscript = CurrentComponent->As<MIR::FSubscript>();
+		if (FirstComponentAsSubscript && CurrentComponentAsSubscript)
 		{
-			if (((MIR::FSubscript*)FirstComponent)->Arg == ((MIR::FSubscript*)CurrentComponent)->Arg)
+			if (FirstComponentAsSubscript->Arg == CurrentComponentAsSubscript->Arg // Both subscripts have the same argument
+				&& FirstComponentAsSubscript->Arg->Type.IsVector()
+				&& FirstComponentAsSubscript->Index == 0 && CurrentComponentAsSubscript->Index == Index)
 			{
-				const MIR::FPrimitiveType* PrimitiveType = ((MIR::FSubscript*)FirstComponent)->Arg->Type->AsVector();
-				if (PrimitiveType && ((MIR::FSubscript*)FirstComponent)->Index == 0 && ((MIR::FSubscript*)CurrentComponent)->Index == Index)
+				MIR::FPrimitive ArgPrimitiveType = FirstComponentAsSubscript->Arg->Type.GetPrimitive();
+
+				// Exact match?
+				if (ArgPrimitiveType.NumRows == NumComponents)
 				{
-					// Exact match?
-					if (PrimitiveType->NumRows == NumComponents)
-					{
-						return true;
-					}
+					return true;
+				}
 
-					// Subset of components?  See if we can generate a swizzled vector.
-					if (PrimitiveType->NumRows >= NumComponents)
-					{
-						// Allow swizzle for non-LWC values or uniform parameters (the latter having a special case for LWC)
-						return PrimitiveType->ScalarKind != MIR::ScalarKind_LWC || ((MIR::FSubscript*)FirstComponent)->Arg->Kind == MIR::VK_UniformParameter;
-					}
+				// Subset of components? See if we can generate a swizzled vector.
+				if (ArgPrimitiveType.NumRows >= NumComponents)
+				{
+					// Allow swizzle for non-LWC values or uniform parameters (the latter having a special case for LWC)
+					return !ArgPrimitiveType.IsLWC() || FirstComponentAsSubscript->Arg->As<MIR::FUniformParameter>();
 				}
 			}
 		}
-		else if (FirstComponent->Kind == MIR::VK_Scalar && CurrentComponent->Kind == MIR::VK_Scalar)
+		else
 		{
-			if (((MIR::FScalar*)FirstComponent)->Type == ((MIR::FScalar*)CurrentComponent)->Type)
+			const MIR::FScalar* FirstComponentAsScalar = FirstComponent->As<MIR::FScalar>();
+			const MIR::FScalar* CurrentComponentAsScalar = CurrentComponent->As<MIR::FScalar>();
+			if (FirstComponentAsScalar && CurrentComponentAsScalar && FirstComponentAsScalar->Type == CurrentComponentAsScalar->Type)
 			{
-				return HasMatchingScalarComponentCastChain(((MIR::FScalar*)FirstComponent)->Arg, ((MIR::FScalar*)CurrentComponent)->Arg, Index, NumComponents);
+				return HasMatchingScalarComponentCastChain(FirstComponentAsScalar->Arg, CurrentComponentAsScalar->Arg, Index, NumComponents);
 			}
 		}
 		return false;
@@ -943,25 +940,24 @@ struct FTranslator : FMaterialIRToHLSLTranslation
 
 	ENoOp LowerVectorCastChain(MIR::FValue* FirstComponent, int32 NumComponents)
 	{
-		if (FirstComponent->Kind == MIR::VK_Scalar)
+		if (const MIR::FScalar* FirstComponentAsScalar = FirstComponent->As<MIR::FScalar>())
 		{
-			MIR::FValue* Arg = ((MIR::FScalar*)FirstComponent)->Arg;
+			MIR::FValue* Arg = FirstComponentAsScalar->Arg;
+			MIR::FPrimitive PrimitiveType = FirstComponent->Type.GetPrimitive();
+			MIR::FPrimitive ArgPrimitiveType = Arg->Type.GetPrimitive();
 
-			const MIR::EScalarKind ScalarKind = ((MIR::FPrimitiveType*)FirstComponent->Type)->ScalarKind;
-			const MIR::EScalarKind ArgScalarKind = ((MIR::FPrimitiveType*)Arg->Type)->ScalarKind;
-
-			if (ArgScalarKind == MIR::ScalarKind_LWC)
+			if (ArgPrimitiveType.IsLWC())
 			{
 				// Cast from LWC
-				if (ScalarKind == MIR::ScalarKind_Bool)
+				if (PrimitiveType.IsBoolean())
 				{
 					// Cast to bool requires a comparison with zero, outside the WSDemote
 					Printer << "(WSDemote(" << LowerVectorCastChain(Arg, NumComponents) << ") != 0)";
 				}
-				else if (ScalarKind != MIR::ScalarKind_Float)
+				else if (!PrimitiveType.IsFloat())
 				{
 					// Cast to non-float (integer) requires a cast to the type, outside the WSDemote
-					const MIR::FPrimitiveType* VectorType = MIR::FPrimitiveType::GetVector(ScalarKind, NumComponents);
+					MIR::FType VectorType = MIR::FType::MakeVector(PrimitiveType.ScalarKind, NumComponents);
 					Printer << LowerType(VectorType) << "(WSDemote(" << LowerVectorCastChain(Arg, NumComponents) << "))";
 				}
 				else
@@ -970,15 +966,15 @@ struct FTranslator : FMaterialIRToHLSLTranslation
 					Printer << "WSDemote(" << LowerVectorCastChain(Arg, NumComponents) << ")";
 				}
 			}
-			else if (ScalarKind == MIR::ScalarKind_LWC)
+			else if (PrimitiveType.IsLWC())
 			{
 				// Cast to LWC
-				if (ArgScalarKind == MIR::ScalarKind_Bool)
+				if (ArgPrimitiveType.IsBoolean())
 				{
 					// Cast from bool requires a select between 1.0f and 0.0f, inside the WSPromote
 					Printer << "WSPromote(select(" << LowerVectorCastChain(Arg, NumComponents) << ", (float" << NumComponents << ")1.0f, (float" << NumComponents << ")0.0f))";
 				}
-				else if (ArgScalarKind != MIR::ScalarKind_Float)
+				else if (!ArgPrimitiveType.IsFloat())
 				{
 					// Cast from non-float (integer) requires a cast to float, inside the WSPromote
 					Printer << "WSPromote(float" << NumComponents << "(" << LowerVectorCastChain(Arg, NumComponents) << "))";
@@ -992,74 +988,71 @@ struct FTranslator : FMaterialIRToHLSLTranslation
 			else
 			{
 				// Cast between intrinsic types
-				if (ScalarKind == MIR::ScalarKind_Bool)
+				if (PrimitiveType.IsBoolean())
 				{
 					// Cast to bool requires a comparison with zero
 					Printer << "(" << LowerVectorCastChain(Arg, NumComponents) << " != 0)";
 				}
-				else if (ArgScalarKind == MIR::ScalarKind_Bool)
+				else if (ArgPrimitiveType.IsBoolean())
 				{
 					// Cast from bool requires a select between 1 and 0
-					const MIR::FPrimitiveType* VectorType = MIR::FPrimitiveType::GetVector(ScalarKind, NumComponents);
+					MIR::FType VectorType = MIR::FType::MakeVector(PrimitiveType.ScalarKind, NumComponents);
 					Printer << "select(" << LowerVectorCastChain(Arg, NumComponents) << ", (" << LowerType(VectorType) << ")1, (" << LowerType(VectorType) << ")0)";
 				}
 				else
 				{
 					// Cast between arithmetic types
-					const MIR::FPrimitiveType* VectorType = MIR::FPrimitiveType::GetVector(ScalarKind, NumComponents);
+					MIR::FType VectorType = MIR::FType::MakeVector(PrimitiveType.ScalarKind, NumComponents);
 					Printer << LowerType(VectorType) << "(" << LowerVectorCastChain(Arg, NumComponents) << ")";
 				}
 			}
 		}
-		else
+		else if (auto FirstComponentAsSubscript = FirstComponent->As<MIR::FSubscript>())
 		{
 			// Finally reached the inner subscript, print its vector argument
-			check(FirstComponent->Kind == MIR::VK_Subscript);
-			MIR::FValue* InnerArg = ((MIR::FSubscript*)FirstComponent)->Arg;
-			const MIR::FPrimitiveType* PrimitiveType = InnerArg->Type->AsVector();
-
-			if (PrimitiveType->IsLWC() && InnerArg->Kind == MIR::VK_UniformParameter)
+			const MIR::FValue* Arg = FirstComponentAsSubscript->Arg;
+			if (Arg->Type.IsVector() && Arg->Type.IsLWC() && Arg->As<MIR::FUniformParameter>())
 			{
-				// LWC types support special case swizzling logic for uniform parameters on initial fetch.  See HasMatchingScalarComponentCastChain above.
-				LowerPrimitiveUniformParameterLWC((const MIR::FUniformParameter*)InnerArg, NumComponents);
+				// LWC types support special case swizzling logic for uniform parameters on initial fetch. See HasMatchingScalarComponentCastChain above.
+				LowerPrimitiveUniformParameterLWC(static_cast<const MIR::FUniformParameter*>(Arg), NumComponents);
 			}
 			else
 			{
-				Printer << LowerValue(InnerArg);
-
-				if (PrimitiveType->NumRows > NumComponents)
+				Printer << LowerValue(Arg);
+				if (Arg->Type.GetPrimitive().NumRows > NumComponents)
 				{
 					Printer << GVector4SwizzleSubset[NumComponents - 1];
 				}
 			}
 		}
+		else
+		{
+			UE_MIR_UNREACHABLE();
+		}
 		return NoOp;
 	}
 
 	// Check if we can tidy up casts from non-LWC scalar to LWC vectors (assuming we already know all components are the same and LWC)
-	bool IsScalarToLWCVectorCast(const MIR::FValue* FirstComponent)
+	static bool IsScalarToLWCVectorCast(const MIR::FValue* FirstComponent)
 	{
-		if (FirstComponent->Kind == MIR::VK_Scalar)
+		if (const MIR::FScalar* AsScalar = FirstComponent->As<MIR::FScalar>())
 		{
-			// Check if this is a cast from an arithmetic non-LWC scalar type.  If so, we can cast the scalar to a float vector, and then to LWC.
-			const MIR::FScalar* Scalar = (const MIR::FScalar*)FirstComponent;
-			const MIR::EScalarKind ArgScalarKind = Scalar->Arg->Type->AsPrimitive()->ScalarKind;
-
-			return ArgScalarKind != MIR::ScalarKind_LWC && ArgScalarKind != MIR::ScalarKind_Bool;
+			// Check if this is a cast from an arithmetic non-LWC scalar type. If so, we can cast the scalar to a float vector, and then to LWC.
+			return !AsScalar->Arg->Type.IsBoolean() && !AsScalar->Arg->Type.IsLWC();
 		}
-		else if (FirstComponent->Kind == MIR::VK_Constant)
+		else if (const MIR::FConstant* AsConstant = FirstComponent->As<MIR::FConstant>())
 		{
 			// Check if this is a cast from a constant representable exactly as a non-LWC float
-			return ((const MIR::FConstant*)FirstComponent)->IsLWCRepresentableAsFloat();
+			return AsConstant->Type.IsLWC() && FLargeWorldRenderScalar(AsConstant->Double).GetTile() == 0.0f;
 		}
 		return false;
 	}
 
-	void LowerComposite(const MIR::FComposite* Dimensional) 
+	void LowerComposite(const MIR::FComposite* Composite) 
 	{
-		TArrayView<MIR::FValue* const> Components = Dimensional->GetComponents();
-		const MIR::FPrimitiveType* ArithmeticType = Dimensional->Type->AsPrimitive();
-		check(ArithmeticType && !ArithmeticType->IsScalar());
+		TConstArrayView<MIR::FValue*> Components = Composite->GetComponents();
+		MIR::FPrimitive PrimitiveType = Composite->Type.GetPrimitive();
+		check(!PrimitiveType.IsScalar());
 		
 		// In order to generate smaller and tidier HLSL, first check whether all components
 		// of this dimensional are actually the same. If so, we can simply emit the 
@@ -1068,41 +1061,43 @@ struct FTranslator : FMaterialIRToHLSLTranslation
 		bool bSameComponents = true;
 
 		// Track if all components are part of a constant vector
-		bool bConstantVector = Components[0]->Kind == MIR::VK_Constant && ArithmeticType->IsVector();
-
+		bool bConstantVector = Components[0]->As<MIR::FConstant>() && PrimitiveType.IsVector();
+		
 		// We can also generate tidier HLSL for cases where casts are done for whole vectors.  This is a frequent case for LWC, where casts happen
 		// in both directions (LWC to float and back), due to operations automatically casting their inputs or outputs.  For example, BO_Sub
 		// automatically downcasts LWC to float, but if you feed that into a BO_Add that has LWC for its other input, it will immediately cast it
 		// back to LWC -- two consecutive casts.  These then get expanded into individual scalar casts, which we would like to do as whole vector
 		// casts for readability.  We only need to consider this special case if the first component is a Scalar, and the type is a Vector.
-		bool bWholeVectorCast = ArithmeticType->IsVector() && (Components[0]->Kind == MIR::VK_Scalar || Components[0]->Kind == MIR::VK_Subscript);
+		bool bWholeVectorCast = PrimitiveType.IsVector() && (Components[0]->As<MIR::FScalar>() || Components[0]->As<MIR::FSubscript>());
 
 		for (int32 i = 1; i < Components.Num(); ++i)
 		{
 			bSameComponents &= (Components[i] == Components[0]);
-			bConstantVector &= Components[i]->Kind == MIR::VK_Constant;
+			bConstantVector &= Components[i]->As<MIR::FConstant>() != nullptr;
 			bWholeVectorCast = bWholeVectorCast && HasMatchingScalarComponentCastChain(Components[0], Components[i], i, Components.Num());
 		}
 		
-		if (bSameComponents && ArithmeticType->ScalarKind != MIR::ScalarKind_LWC)
+		if (bSameComponents && !PrimitiveType.IsLWC())
 		{
-			Printer << TEXT("(") << LowerType(ArithmeticType) << TEXT(")") << LowerValue(Components[0]);
+			Printer << TEXT("(") << LowerType(Composite->Type) << TEXT(")") << LowerValue(Components[0]);
 		}
-		else if (bSameComponents && ArithmeticType->IsVector() && IsScalarToLWCVectorCast(Components[0]))
+		else if (bSameComponents && PrimitiveType.IsVector() && IsScalarToLWCVectorCast(Components[0]))
 		{
 			// Cast scalar to float vector, then promote, for example:  "WSPromote((float3)1.0f)"
-			Printer << "WSPromote" << "((float" << ArithmeticType->NumRows << ")";
-			if (Components[0]->Kind == MIR::VK_Scalar)
+			Printer << "WSPromote" << "((float" << PrimitiveType.NumRows << ")";
+			if (const MIR::FScalar* Scalar = Components[0]->As<MIR::FScalar>())
 			{
 				// Print the inner non-LWC scalar value
-				const MIR::FScalar* Scalar = (const MIR::FScalar*)Components[0];
 				LowerValue(Scalar->Arg);
 			}
-			else
+			else if (const MIR::FConstant* Constant = Components[0]->As<MIR::FConstant>())
 			{
 				// Print the inner double constant in its non-LWC form
-				check(Components[0]->Kind == MIR::VK_Constant)
-				Printer << (float)((const MIR::FConstant*)Components[0])->LWCValue;
+				Printer << (float)Constant->Double;
+			}
+			else
+			{
+				UE_MIR_UNREACHABLE();
 			}
 			Printer << ")";
 		}
@@ -1110,7 +1105,7 @@ struct FTranslator : FMaterialIRToHLSLTranslation
 		{
 			LowerVectorCastChain(Components[0], Components.Num());
 		}
-		else if (bConstantVector && ArithmeticType->ScalarKind == MIR::ScalarKind_LWC)
+		else if (bConstantVector && PrimitiveType.IsLWC())
 		{
 			// Special case for LWC constant vectors
 			bool bAllTileValuesZero = true;
@@ -1118,7 +1113,7 @@ struct FTranslator : FMaterialIRToHLSLTranslation
 			float OffsetValues[4];
 			for (int32 Index = 0; Index < Components.Num(); Index++)
 			{
-				FLargeWorldRenderScalar ConstantLWC(((const MIR::FConstant*)Components[Index])->LWCValue);
+				FLargeWorldRenderScalar ConstantLWC(Components[Index]->As<MIR::FConstant>()->Double);
 				TileValues[Index] = ConstantLWC.GetTile();
 				OffsetValues[Index] = ConstantLWC.GetOffset();
 				bAllTileValuesZero &= TileValues[Index] == 0.0f;
@@ -1150,11 +1145,11 @@ struct FTranslator : FMaterialIRToHLSLTranslation
 		}
 		else
 		{
-			if (ArithmeticType->ScalarKind == MIR::ScalarKind_LWC)
+			if (PrimitiveType.IsLWC())
 			{
-				if (ArithmeticType->NumColumns == 1)
+				if (PrimitiveType.IsVector())
 				{
-					Printer << (ArithmeticType->NumRows == 1 ? TEXT("MakeWSScalar") : TEXT("MakeWSVector")) << BeginArgs;
+					Printer << "MakeWSVector" << BeginArgs;
 				}
 				else
 				{
@@ -1165,15 +1160,16 @@ struct FTranslator : FMaterialIRToHLSLTranslation
 			}
 			else
 			{
-				Printer << LowerType(ArithmeticType) << BeginArgs;
+				Printer << LowerType(Composite->Type) << BeginArgs;
 			}
-			for (MIR::FValue* Component : Components)
+
+			for (const MIR::FValue* Component : Components)
 			{
 				Printer << ListSeparator << LowerValue(Component);
 			}
+
 			Printer << EndArgs;
 		}
-
 	}
 
 	void LowerSetMaterialOutput(const MIR::FSetMaterialOutput* Output)
@@ -1223,7 +1219,7 @@ struct FTranslator : FMaterialIRToHLSLTranslation
 
 	void LowerOperator(const MIR::FOperator* Operator)
 	{
-		bool bIsLWC = Operator->AArg->Type->IsLWC() || (Operator->Op == MIR::TO_Select && Operator->BArg->Type->IsLWC());
+		bool bIsLWC = Operator->AArg->Type.IsLWC() || (Operator->Op == MIR::TO_Select && Operator->BArg->Type.IsLWC());
 
 		if (IsOperatorInfix(Operator->Op) && !bIsLWC)
 		{
@@ -1386,7 +1382,7 @@ struct FTranslator : FMaterialIRToHLSLTranslation
 	{
 		if (IsFoldable(Branch, CurrentStage))
 		{
-			if (Branch->TrueArg->Type->IsLWC())
+			if (Branch->TrueArg->Type.IsLWC())
 			{
 				Printer << TEXT("WSSelect(") << LowerValue(Branch->ConditionArg)
 					<< TEXT(", ") << LowerValue(Branch->TrueArg)
@@ -1416,12 +1412,11 @@ struct FTranslator : FMaterialIRToHLSLTranslation
 			
 	void LowerSubscript(const MIR::FSubscript* Subscript)
 	{
-		const MIR::FPrimitiveType* ArgArithmeticType = Subscript->Arg->Type->AsVector();
-
-		if (ArgArithmeticType && ArgArithmeticType->ScalarKind == MIR::EScalarKind::ScalarKind_LWC)
+		TOptional<MIR::FPrimitive> ArgVectorType = Subscript->Arg->Type.AsVector();
+		if (ArgVectorType && ArgVectorType->IsLWC())
 		{
 			const TCHAR* LwcComponentsStr[] = { TEXT("WSGetX("), TEXT("WSGetY("), TEXT("WSGetZ("), TEXT("WSGetW(") };
-			check(Subscript->Index <= ArgArithmeticType->NumComponents());
+			check(Subscript->Index <= ArgVectorType->NumComponents());
 
 			Printer << LwcComponentsStr[Subscript->Index] << LowerValue(Subscript->Arg) << TEXT(")");
 		}
@@ -1429,10 +1424,10 @@ struct FTranslator : FMaterialIRToHLSLTranslation
 		{
 			LowerValue(Subscript->Arg);
 
-			if (ArgArithmeticType)
+			if (ArgVectorType)
 			{
 				const TCHAR* ComponentsStr[] = { TEXT(".x"), TEXT(".y"), TEXT(".z"), TEXT(".w") };
-				check(Subscript->Index <= ArgArithmeticType->NumComponents());
+				check(Subscript->Index <= ArgVectorType->NumComponents());
 
 				Printer << ComponentsStr[Subscript->Index];
 			}
@@ -1441,18 +1436,18 @@ struct FTranslator : FMaterialIRToHLSLTranslation
 
 	void LowerScalar(const MIR::FScalar* Scalar)
 	{
-		const MIR::EScalarKind ScalarKind = Scalar->Type->AsPrimitive()->ScalarKind;
-		const MIR::EScalarKind ArgScalarKind = Scalar->Arg->Type->AsPrimitive()->ScalarKind;
+		MIR::FPrimitive PrimitiveType = Scalar->Type.GetPrimitive();
+		MIR::FPrimitive ArgPrimitiveType = Scalar->Arg->Type.GetPrimitive();
 
-		if (ArgScalarKind == MIR::ScalarKind_LWC)
+		if (ArgPrimitiveType.IsLWC())
 		{
 			// Cast from LWC
-			if (ScalarKind == MIR::ScalarKind_Bool)
+			if (PrimitiveType.IsBoolean())
 			{
 				// Cast to bool requires a comparison with zero, outside the WSDemote
 				Printer << "(WSDemote(" << LowerValue(Scalar->Arg) << ") != 0)";
 			}
-			else if (ScalarKind != MIR::ScalarKind_Float)
+			else if (!PrimitiveType.IsFloat())
 			{
 				// Cast to non-float (integer) requires a cast to the type, outside the WSDemote
 				Printer << LowerType(Scalar->Type) << "(WSDemote(" << LowerValue(Scalar->Arg) << "))";
@@ -1463,15 +1458,15 @@ struct FTranslator : FMaterialIRToHLSLTranslation
 				Printer << "WSDemote(" << LowerValue(Scalar->Arg) << ")";
 			}
 		}
-		else if (ScalarKind == MIR::ScalarKind_LWC)
+		else if (PrimitiveType.IsLWC())
 		{
 			// Cast to LWC
-			if (ArgScalarKind == MIR::ScalarKind_Bool)
+			if (ArgPrimitiveType.IsBoolean())
 			{
 				// Cast from bool requires a select between 1.0f and 0.0f, inside the WSPromote
 				Printer << "WSPromote(" << LowerValue(Scalar->Arg) << " ? 1.0f : 0.0f)";
 			}
-			else if (ArgScalarKind != MIR::ScalarKind_Float)
+			else if (!ArgPrimitiveType.IsFloat())
 			{
 				// Cast from non-float (integer) requires a cast to float, inside the WSPromote
 				Printer << "WSPromote(float(" << LowerValue(Scalar->Arg) << "))";
@@ -1485,12 +1480,12 @@ struct FTranslator : FMaterialIRToHLSLTranslation
 		else
 		{
 			// Cast between intrinsic types
-			if (ScalarKind == MIR::ScalarKind_Bool)
+			if (PrimitiveType.IsBoolean())
 			{
 				// Cast to bool requires a comparison with zero
 				Printer << "(" << LowerValue(Scalar->Arg) << " != 0)";
 			}
-			else if (ArgScalarKind == MIR::ScalarKind_Bool)
+			else if (ArgPrimitiveType.IsBoolean())
 			{
 				// Cast from bool requires a select between 1 and 0
 				Printer << LowerType(Scalar->Type) << "(" << LowerValue(Scalar->Arg) << " ? 1 : 0)";
@@ -1505,9 +1500,9 @@ struct FTranslator : FMaterialIRToHLSLTranslation
 
 	void LowerStandardTextureRead(const MIR::FTextureRead* TextureRead)
 	{
-		switch (TextureRead->TextureObject->Type->AsObject()->ObjectKind)
+		switch (TextureRead->TextureObject->Type.GetKind())
 		{
-			case MIR::ObjectKind_Texture2D: Printer << TEXT("Texture2D"); break;
+			case MIR::ETypeKind::Texture2D: Printer << TEXT("Texture2D"); break;
 			default: UE_MIR_UNREACHABLE();
 		}
 
@@ -1543,7 +1538,7 @@ struct FTranslator : FMaterialIRToHLSLTranslation
 	void LowerVirtualTextureRead(const MIR::FTextureRead* TextureRead)
 	{
 		checkf(
-			TextureRead->TextureObject->Type->AsObject()->ObjectKind == MIR::ObjectKind_Texture2D,
+			TextureRead->TextureObject->Type.IsTexture2D(),
 			TEXT("Texture object type must be a 2D texture for virtual texture sampling in new material translator")
 		);
 
@@ -1773,7 +1768,7 @@ struct FTranslator : FMaterialIRToHLSLTranslation
 
 	void LowerHardwarePartialDerivative(const MIR::FHardwarePartialDerivative* StageSwitch)
 	{
-		if (StageSwitch->Arg->Type->IsLWC())
+		if (StageSwitch->Arg->Type.IsLWC())
 		{
 			// Expression emitter assumes these to be LWC for analytic derivative evaluation, so we promote them on load
 			Printer << (StageSwitch->Axis == MIR::EDerivativeAxis::X ? "WSPromote(WSDdxDemote(" : "WSPromote(WSDdyDemote(") << LowerValue(StageSwitch->Arg) << "))";
@@ -1888,21 +1883,21 @@ struct FTranslator : FMaterialIRToHLSLTranslation
 		Params.Add(TEXT("uniform_material_expressions"), MoveTemp(UniformMaterialExpressions));
 	}
 	
-	ENoOp LowerType(const MIR::FType* Type)
+	ENoOp LowerType(const MIR::FType& Type)
 	{
-		if (const MIR::FPrimitiveType* ArithmeticType = Type->AsPrimitive())
+		if (TOptional<MIR::FPrimitive> PrimitiveType = Type.AsPrimitive())
 		{
-			if (ArithmeticType->ScalarKind == MIR::ScalarKind_LWC)
+			if (PrimitiveType->IsLWC())
 			{
-				if (ArithmeticType->NumRows == 1)
+				if (PrimitiveType->NumRows == 1)
 				{
 					Printer << TEXT("FWSScalar");
 				}
-				else if (ArithmeticType->NumColumns == 1)
+				else if (PrimitiveType->NumColumns == 1)
 				{
-					Printer << TEXT("FWSVector") << ArithmeticType->NumRows;
+					Printer << TEXT("FWSVector") << PrimitiveType->NumRows;
 				}
-				else if (ArithmeticType->bIsLWCInverseMatrix)
+				else if (PrimitiveType->bIsLWCInverseMatrix)
 				{
 					Printer << TEXT("FWSInverseMatrix");
 				}
@@ -1913,32 +1908,31 @@ struct FTranslator : FMaterialIRToHLSLTranslation
 			}
 			else
 			{
-				switch (ArithmeticType->ScalarKind)
+				switch (PrimitiveType->ScalarKind)
 				{
-					case MIR::ScalarKind_Bool:	Printer << TEXT("bool"); break;
-					case MIR::ScalarKind_Int: 	Printer << TEXT("int"); break;
-					case MIR::ScalarKind_Float:	Printer << TEXT("MaterialFloat"); break;
+					case MIR::EScalarKind::Bool:  Printer << TEXT("bool"); break;
+					case MIR::EScalarKind::Int:   Printer << TEXT("int"); break;
+					case MIR::EScalarKind::Float: Printer << TEXT("MaterialFloat"); break;
 				}
 
-				if (ArithmeticType->NumRows > 1)
+				if (PrimitiveType->NumRows > 1)
 				{
-					Printer << ArithmeticType->NumRows;
+					Printer << PrimitiveType->NumRows;
 				}
 			
-				if (ArithmeticType->NumColumns > 1)
+				if (PrimitiveType->NumColumns > 1)
 				{
-					Printer << TEXT("x") << ArithmeticType->NumColumns;
+					Printer << TEXT("x") << PrimitiveType->NumColumns;
 				}
 			}
 		}
-		else if (const MIR::FObjectType* ObjectType = Type->AsObject())
+		else if (Type.IsSubstrateData())
 		{
-			switch (ObjectType->ObjectKind)
-			{
-				case MIR::ObjectKind_SubstrateData: Printer << TEXT("FSubstrateData"); break;
-				case MIR::ObjectKind_VTPageTableResult: Printer << TEXT("VTPageTableResult"); break;
-				default: UE_MIR_UNREACHABLE(); break;
-			}
+			Printer << TEXT("FSubstrateData");
+		}
+		else if (Type.IsVTPageTableResult())
+		{
+			Printer << TEXT("VTPageTableResult");
 		}
 		else
 		{
diff --git a/Engine/Source/Runtime/Engine/Private/Materials/MaterialIRTypes.cpp b/Engine/Source/Runtime/Engine/Private/Materials/MaterialIRTypes.cpp
--- a/Engine/Source/Runtime/Engine/Private/Materials/MaterialIRTypes.cpp
+++ b/Engine/Source/Runtime/Engine/Private/Materials/MaterialIRTypes.cpp
@@ -12,14 +12,40 @@ const TCHAR* TypeKindToString(ETypeKind Kind)
 {
 	switch (Kind)
 	{
-		case TypeKind_Void: return TEXT("void");
-		case TypeKind_Primitive: return TEXT("primitive");
-		case TypeKind_Object: return TEXT("object");
+		case ETypeKind::Poison: return TEXT("Poison");
+		case ETypeKind::Void: return TEXT("Void");
+		case ETypeKind::Primitive: return TEXT("Primitive");
+		case ETypeKind::Aggregate: return TEXT("Aggregate");
+		case ETypeKind::Texture2D: return TEXT("Texture2D");
+		case ETypeKind::SubstrateData: return TEXT("SubstrateData");
+		case ETypeKind::VTPageTableResult: return TEXT("VTPageTableResult");
 		default: UE_MIR_UNREACHABLE();
 	}
 }
 
-const FType* FType::FromShaderType(const UE::Shader::FType& InShaderType)
+bool ScalarKindIsArithmetic(EScalarKind Kind)
+{
+	return Kind != EScalarKind::Bool;
+}
+
+bool ScalarKindIsAnyFloat(EScalarKind Kind)
+{
+	return Kind == EScalarKind::Float || Kind == EScalarKind::Double;
+}
+
+const TCHAR* ScalarKindToString(EScalarKind Kind)
+{
+	switch (Kind)
+	{
+		case EScalarKind::Bool: return TEXT("bool");
+		case EScalarKind::Int: return TEXT("int");
+		case EScalarKind::Float: return TEXT("MaterialFloat");
+		case EScalarKind::Double: return TEXT("FLWCScalar");
+		default: UE_MIR_UNREACHABLE();
+	}
+}
+
+FType FType::FromShaderType(const UE::Shader::FType& InShaderType)
 {
 	check(!InShaderType.IsStruct());
 	check(!InShaderType.IsObject());
@@ -27,487 +53,320 @@ const FType* FType::FromShaderType(const UE::Shader::FType& InShaderType)
 	switch (InShaderType.ValueType)
 	{
 		case UE::Shader::EValueType::Void:
-			return FType::GetVoid();
+			return FType::MakeVoid();
 
 		case UE::Shader::EValueType::Float1:
 		case UE::Shader::EValueType::Float2:
 		case UE::Shader::EValueType::Float3:
 		case UE::Shader::EValueType::Float4:
-			return FPrimitiveType::GetVector(ScalarKind_Float, (int)InShaderType.ValueType - (int)UE::Shader::EValueType::Float1 + 1);
+			return FType::MakeVector(EScalarKind::Float, (int)InShaderType.ValueType - (int)UE::Shader::EValueType::Float1 + 1);
 
 		case UE::Shader::EValueType::Int1:
 		case UE::Shader::EValueType::Int2:
 		case UE::Shader::EValueType::Int3:
 		case UE::Shader::EValueType::Int4:
-			return FPrimitiveType::GetVector(ScalarKind_Int, (int)InShaderType.ValueType - (int)UE::Shader::EValueType::Int1 + 1);
+			return FType::MakeVector(EScalarKind::Int, (int)InShaderType.ValueType - (int)UE::Shader::EValueType::Int1 + 1);
 
 		case UE::Shader::EValueType::Bool1:
 		case UE::Shader::EValueType::Bool2:
 		case UE::Shader::EValueType::Bool3:
 		case UE::Shader::EValueType::Bool4:
-			return FPrimitiveType::GetVector(ScalarKind_Bool, (int)InShaderType.ValueType - (int)UE::Shader::EValueType::Bool1 + 1);
+			return FType::MakeVector(EScalarKind::Bool, (int)InShaderType.ValueType - (int)UE::Shader::EValueType::Bool1 + 1);
 
 		case UE::Shader::EValueType::Double1:
 		case UE::Shader::EValueType::Double2:
 		case UE::Shader::EValueType::Double3:
 		case UE::Shader::EValueType::Double4:
-			return FPrimitiveType::GetVector(ScalarKind_LWC, (int)InShaderType.ValueType - (int)UE::Shader::EValueType::Double1 + 1);
+			return FType::MakeVector(EScalarKind::Double, (int)InShaderType.ValueType - (int)UE::Shader::EValueType::Double1 + 1);
 
 		default:
 			UE_MIR_UNREACHABLE();
 	}
 }
 
-const FType* FType::FromMaterialValueType(EMaterialValueType Type)
+FType FType::FromMaterialValueType(EMaterialValueType Type)
 {
 	switch (Type)
 	{
 		case MCT_VoidStatement:
-			return FType::GetVoid();
+			return FType::MakeVoid();
 
 		case MCT_Float1:
-			return FPrimitiveType::GetVector(ScalarKind_Float, 1);
+			return FType::MakeVector(EScalarKind::Float, 1);
 		case MCT_Float2:
-			return FPrimitiveType::GetVector(ScalarKind_Float, 2);
+			return FType::MakeVector(EScalarKind::Float, 2);
 		case MCT_Float3:
-			return FPrimitiveType::GetVector(ScalarKind_Float, 3);
+			return FType::MakeVector(EScalarKind::Float, 3);
 		case MCT_Float4:
-			return FPrimitiveType::GetVector(ScalarKind_Float, 4);
+			return FType::MakeVector(EScalarKind::Float, 4);
 
 		case MCT_Float:
-			return FPrimitiveType::GetVector(ScalarKind_Float, 4);
+			return FType::MakeVector(EScalarKind::Float, 4);
 
 		case MCT_Texture2D:
-			return FObjectType::GetTexture2D();
+			return FType::MakeTexture2D();
 
 		case MCT_StaticBool:
-			return FPrimitiveType::GetVector(ScalarKind_Bool, 1);
+			return FType::MakeVector(EScalarKind::Bool, 1);
 
 		case MCT_UInt1:
-			return FPrimitiveType::GetVector(ScalarKind_Int, 1);
+			return FType::MakeVector(EScalarKind::Int, 1);
 		case MCT_UInt2:
-			return FPrimitiveType::GetVector(ScalarKind_Int, 2);
+			return FType::MakeVector(EScalarKind::Int, 2);
 		case MCT_UInt3:
-			return FPrimitiveType::GetVector(ScalarKind_Int, 3);
+			return FType::MakeVector(EScalarKind::Int, 3);
 		case MCT_UInt4:
-			return FPrimitiveType::GetVector(ScalarKind_Int, 4);
+			return FType::MakeVector(EScalarKind::Int, 4);
 
 		case MCT_Bool:
-			return FPrimitiveType::GetVector(ScalarKind_Bool, 1);
+			return FType::MakeVector(EScalarKind::Bool, 1);
 
 		case MCT_LWCScalar:
-			return FPrimitiveType::GetVector(ScalarKind_LWC, 1);
+			return FType::MakeVector(EScalarKind::Double, 1);
 		case MCT_LWCVector2:
-			return FPrimitiveType::GetVector(ScalarKind_LWC, 2);
+			return FType::MakeVector(EScalarKind::Double, 2);
 		case MCT_LWCVector3:
-			return FPrimitiveType::GetVector(ScalarKind_LWC, 3);
+			return FType::MakeVector(EScalarKind::Double, 3);
 		case MCT_LWCVector4:
-			return FPrimitiveType::GetVector(ScalarKind_LWC, 4);
+			return FType::MakeVector(EScalarKind::Double, 4);
 
 		default:
 			UE_MIR_UNREACHABLE();
 	}
 }
 
-const FType* FType::FromMaterialParameterType(EMaterialParameterType Type)
+FType FType::FromMaterialParameterType(EMaterialParameterType Type)
 {
 	switch (Type)
 	{
-		case EMaterialParameterType::Scalar: return FPrimitiveType::GetFloat();
-		case EMaterialParameterType::Vector: return FPrimitiveType::GetFloat4();
+		case EMaterialParameterType::Scalar: return FType::MakeFloat();
+		case EMaterialParameterType::Vector: return FType::MakeFloat(4);
 		case EMaterialParameterType::DoubleVector: UE_MIR_TODO();
-		case EMaterialParameterType::Texture: return FObjectType::GetTexture2D();
+		case EMaterialParameterType::Texture: return FType::MakeTexture2D();
 		case EMaterialParameterType::TextureCollection: UE_MIR_TODO();
 		case EMaterialParameterType::ParameterCollection: UE_MIR_TODO();
 		case EMaterialParameterType::Font: UE_MIR_TODO();
 		case EMaterialParameterType::RuntimeVirtualTexture: UE_MIR_TODO();
 		case EMaterialParameterType::SparseVolumeTexture: UE_MIR_TODO();
-		case EMaterialParameterType::StaticSwitch: return FPrimitiveType::GetBool();
+		case EMaterialParameterType::StaticSwitch: return FType::MakeBool();
 		default: UE_MIR_UNREACHABLE();
 	}
 }
 
-const FType* FType::GetVoid()
+FType FType::MakePoison()
 {
-	static FType Type{ TypeKind_Void };
-	return &Type;
+	return {};
 }
 
-FString FType::GetSpelling() const
+FType FType::MakeVoid()
 {
-	if (IsPoison())
-	{
-		return TEXT("Poison");
-	}
-	else if (auto PrimitiveType = AsPrimitive())
-	{
-		return PrimitiveType->Spelling.GetData();
-	}
-	else if (auto ObjectType = AsObject())
-	{
-		switch (ObjectType->ObjectKind)
-		{
-			case ObjectKind_Texture2D: return TEXT("Texture2D");
-			case ObjectKind_SubstrateData: return TEXT("SubstrateData");
-			case ObjectKind_VTPageTableResult: return TEXT("VTPageTableResult");
-		}
-	}
-	else if (auto AggregateType = AsAggregate())
-	{
-		return AggregateType->Aggregate->GetName();
-	}
-	UE_MIR_UNREACHABLE();
+	FType Type;
+	Type.Kind = ETypeKind::Void;
+	return Type;
 }
 
-UE::Shader::EValueType FType::ToValueType() const
+FType FType::MakeBool(int NumRows, int NumColumns)
 {
-	using namespace UE::Shader;
-
-	if (const FPrimitiveType* PrimitiveType = AsPrimitive())
-	{
-		if (PrimitiveType->IsMatrix())
-		{
-			if (PrimitiveType->NumRows == 4 && PrimitiveType->NumColumns == 4)
-			{
-				if (PrimitiveType->ScalarKind == ScalarKind_Float)
-				{
-					return EValueType::Float4x4;
-				}
-				else if (PrimitiveType->ScalarKind == ScalarKind_LWC)
-				{
-					return PrimitiveType->bIsLWCInverseMatrix ? EValueType::DoubleInverse4x4 : EValueType::Double4x4;
-				}
-				else
-				{
-					return EValueType::Numeric4x4;
-				}
-			}
-
-			return EValueType::Any;
-		}
-
-		check(PrimitiveType->NumColumns == 1 && PrimitiveType->NumRows <= 4);
-
-		switch (PrimitiveType->ScalarKind)
-		{
-			case ScalarKind_Bool: 	return (EValueType)((int)EValueType::Bool1 + PrimitiveType->NumRows - 1);
-			case ScalarKind_Int: 	return (EValueType)((int)EValueType::Int1 + PrimitiveType->NumRows - 1);
-			case ScalarKind_Float: 	return (EValueType)((int)EValueType::Float1 + PrimitiveType->NumRows - 1);
-			case ScalarKind_LWC:	return (EValueType)((int)EValueType::Double1 + PrimitiveType->NumRows - 1);
-			default: UE_MIR_UNREACHABLE();
-		}
-	}
-	else if (const FObjectType* ObjectType = AsObject())
-	{
-		return EValueType::Object;
-	}
-	else if (const FAggregateType* AggregateType = AsAggregate())
-	{
-		return EValueType::Struct;
-	}
-	
-	UE_MIR_UNREACHABLE();
-}
-
-const FType* FType::GetPoison()
-{
-	static FType Type{ TypeKind_Poison };
-	return &Type;
-}
-
-bool FType::IsPoison() const
-{
-	return Kind == TypeKind_Poison;
-}
-
-bool FType::IsBoolean() const
-{
-	const FPrimitiveType* Ptr = AsPrimitive();
-	return Ptr && Ptr->ScalarKind == ScalarKind_Bool;
-}
-
-bool FType::IsInteger() const
-{
-	const FPrimitiveType* Ptr = AsPrimitive();
-	return Ptr && Ptr->ScalarKind == ScalarKind_Int;
-}
-
-bool FType::IsAnyFloat() const
-{
-	const FPrimitiveType* Ptr = AsPrimitive();
-	return Ptr && ScalarKindIsAnyFloat(Ptr->ScalarKind);
-}
-
-bool FType::IsLWC() const
-{
-	const FPrimitiveType* Ptr = AsPrimitive();
-	return Ptr && Ptr->ScalarKind == ScalarKind_LWC;
+	return MakePrimitive(EScalarKind::Bool, NumRows, NumColumns);
 }
 
-bool FType::IsBoolScalar() const
+FType FType::MakeInt(int NumRows, int NumColumns)
 {
-	return this == FPrimitiveType::GetBool();
+	return MakePrimitive(EScalarKind::Int, NumRows, NumColumns);
 }
 
-bool FType::IsObjectOfKind(EObjectKind ObjectKind) const
+FType FType::MakeFloat(int NumRows, int NumColumns)
 {
-	if (const FObjectType* ObjectPtr = AsObject())
-	{
-		return ObjectPtr->ObjectKind == ObjectKind;
-	}
-	return false;
+	return MakePrimitive(EScalarKind::Float, NumRows, NumColumns);
 }
 
-bool FType::IsTexture() const
+FType FType::MakeLWC(int NumRows, int NumColumns)
 {
-	if (const FObjectType* ObjectPtr = AsObject())
-	{
-		return ObjectPtr->ObjectKind == ObjectKind_Texture2D;
-	}
-	return false; 
+	return MakePrimitive(EScalarKind::Double, NumRows, NumColumns);
 }
 
-const FPrimitiveType* FType::AsPrimitive() const
+FType FType::MakeScalar(EScalarKind InScalarKind)
 {
-	return Kind == TypeKind_Primitive ? static_cast<const FPrimitiveType*>(this) : nullptr; 
+	return MakePrimitive(InScalarKind, 1, 1);
 }
 
-const FPrimitiveType* FType::AsArithmetic() const
+FType FType::MakeVector(EScalarKind InScalarKind, int NumRows)
 {
-	return Kind == TypeKind_Primitive && static_cast<const FPrimitiveType*>(this)->IsArithmetic()
-		? static_cast<const FPrimitiveType*>(this)
-		: nullptr; 
+	return MakePrimitive(InScalarKind, NumRows, 1);
 }
 
-const FObjectType* FType::AsObject() const
+FType FType::MakePrimitive(EScalarKind InScalarKind, int NumRows, int NumColumns, bool bIsLWCInverseMatrix)
 {
-	return Kind == TypeKind_Object ? static_cast<const FObjectType*>(this) : nullptr; 
-}
+	check(NumRows >= 1 && NumRows <= 4);
+	check(NumColumns >= 1 && NumColumns <= 4);
 
-const FPrimitiveType* FType::AsScalar() const
-{
-	const FPrimitiveType* Type = AsPrimitive();
-	return Type && Type->IsScalar() ? Type : nullptr;
+	// LWC inverse matrix flag only supported for EScalarKind::LWC and 4x4 dimensions.
+	check(!bIsLWCInverseMatrix || (InScalarKind == EScalarKind::Double && NumRows == 4 && NumColumns == 4));
+	
+	FType Type;
+	Type.Kind = ETypeKind::Primitive;
+	Type.Primitive.ScalarKind = InScalarKind;
+	Type.Primitive.NumRows = NumRows;
+	Type.Primitive.NumColumns = NumColumns;
+	Type.Primitive.bIsLWCInverseMatrix = bIsLWCInverseMatrix;
+	return Type;
 }
 
-const FPrimitiveType* FType::AsVector() const
+FType FType::MakeAggregate(const UMaterialAggregate* InAggregate)
 {
-	const FPrimitiveType* Type = AsPrimitive();
-	return Type && Type->IsVector() ? Type : nullptr;
+	FType Type;
+	Type.Kind = ETypeKind::Aggregate;
+	Type.Aggregate = InAggregate;
+	return Type;
 }
 
-const FPrimitiveType* FType::AsMatrix() const
+FType FType::MakeTexture2D()
 {
-	const FPrimitiveType* Type = AsPrimitive();
-	return Type && Type->IsMatrix() ? Type : nullptr;
+	FType Type;
+	Type.Kind = ETypeKind::Texture2D;
+	return Type;
 }
 
-const FAggregateType* FType::AsAggregate() const
+FType FType::MakeSubstrateData()
 {
-	return Kind == TypeKind_Aggregate ? static_cast<const FAggregateType*>(this) : nullptr; 
+	FType Type;
+	Type.Kind = ETypeKind::SubstrateData;
+	return Type;
 }
 
-bool ScalarKindIsArithmetic(EScalarKind Kind)
+FType FType::MakeVTPageTableResult()
 {
-	return Kind != ScalarKind_Bool;
-}
+	FType Type;
+	Type.Kind = ETypeKind::VTPageTableResult;
+	return Type;}
 
-bool ScalarKindIsAnyFloat(EScalarKind Kind)
-{
-	return Kind == ScalarKind_Float || Kind == ScalarKind_LWC;
-}
-
-const TCHAR* ScalarKindToString(EScalarKind Kind)
+FString FType::GetSpelling() const
 {
-	switch (Kind)
+	if (AsPrimitive())
 	{
-		case ScalarKind_Bool: return TEXT("bool");
-		case ScalarKind_Int: return TEXT("int");
-		case ScalarKind_Float: return TEXT("MaterialFloat");
-		case ScalarKind_LWC: return TEXT("FLWCScalar");
-		default: UE_MIR_UNREACHABLE();
+		FString Spelling;
+		switch (Primitive.ScalarKind)
+		{
+			case EScalarKind::Bool:Spelling.Append(TEXT("bool")); break;
+			case EScalarKind::Int: Spelling.Append(TEXT("int")); break;
+			case EScalarKind::Float: Spelling.Append(TEXT("float")); break;
+			case EScalarKind::Double:
+				if (Primitive.IsScalar())
+				{
+					Spelling.Append(TEXT("FWSScalar"));
+				}
+				else if (Primitive.IsVector())
+				{
+					Spelling.Appendf(TEXT("FWSVector%d"), Primitive.NumRows);
+				}
+				else if (Primitive.bIsLWCInverseMatrix)
+				{
+					Spelling.Append(TEXT("FWSInverseMatrix"));
+				}
+				else
+				{
+					Spelling.Append(TEXT("FWSMatrix"));
+				}
+				return Spelling;
+			default: UE_MIR_UNREACHABLE();
+		}
+		if (Primitive.IsScalar())
+		{
+			// append nothing
+		}
+		else if (Primitive.IsVector())
+		{
+			Spelling.Appendf(TEXT("%d"), Primitive.NumRows);
+		}
+		else
+		{
+			Spelling.Appendf(TEXT("%dx%d"), Primitive.NumRows, Primitive.NumColumns);
+		}
+		return Spelling;
+	}
+	else if (const UMaterialAggregate* MatAggregate = AsAggregate())
+	{
+		return MatAggregate->GetName();
 	}
-}
-
-const FPrimitiveType* FPrimitiveType::GetBool()
-{
-	return GetScalar(ScalarKind_Bool);
-}
-
-const FPrimitiveType* FPrimitiveType::GetInt()
-{
-	return GetScalar(ScalarKind_Int);
-}
-
-const FPrimitiveType* FPrimitiveType::GetFloat()
-{
-	return GetScalar(ScalarKind_Float);
-}
 
-const FPrimitiveType* FPrimitiveType::GetFloat2()
-{
-	return GetVector(ScalarKind_Float, 2);
+	return TypeKindToString(Kind);
 }
 
-const FPrimitiveType* FPrimitiveType::GetFloat3()
+FType FPrimitive::ToScalarKind(EScalarKind InScalarKind) const
 {
-	return GetVector(ScalarKind_Float, 3);
+	return FType::MakePrimitive(InScalarKind, NumRows, NumColumns);
 }
 
-const FPrimitiveType* FPrimitiveType::GetFloat4()
+FType FPrimitive::ToScalar() const
 {
-	return GetVector(ScalarKind_Float, 4);
+	return FType::MakeScalar(ScalarKind);
 }
 
-const FPrimitiveType* FPrimitiveType::GetLWC()
+FType FPrimitive::ToVector(int InNumRows) const
 {
-	return GetScalar(ScalarKind_LWC);
+	return FType::MakeVector(ScalarKind, InNumRows);
 }
 
-const FPrimitiveType* FPrimitiveType::GetLWC2()
-{
-	return GetVector(ScalarKind_LWC, 2);
-}
-
-const FPrimitiveType* FPrimitiveType::GetLWC3()
+UE::Shader::EValueType FType::ToValueType() const
 {
-	return GetVector(ScalarKind_LWC, 3);
-}
+	using namespace UE::Shader;
 
-const FPrimitiveType* FPrimitiveType::GetLWC4()
-{
-	return GetVector(ScalarKind_LWC, 4);
-}
+	if (AsPrimitive())
+	{
+		if (Primitive.IsMatrix())
+		{
+			if (Primitive.NumRows == 4 && Primitive.NumColumns == 4)
+			{
+				if (Primitive.ScalarKind == EScalarKind::Float)
+				{
+					return EValueType::Float4x4;
+				}
+				else if (Primitive.ScalarKind == EScalarKind::Double)
+				{
+					return Primitive.bIsLWCInverseMatrix ? EValueType::DoubleInverse4x4 : EValueType::Double4x4;
+				}
+				else
+				{
+					return EValueType::Numeric4x4;
+				}
+			}
 
-const FPrimitiveType* FPrimitiveType::GetScalar(EScalarKind InScalarKind)
-{
-	return Get(InScalarKind, 1, 1);
-}
+			return EValueType::Any;
+		}
 
-const FPrimitiveType* FPrimitiveType::GetVector(EScalarKind InScalarKind, int NumRows)
-{
-	check(NumRows >= 1 && NumRows <= 4);
-	return Get(InScalarKind, NumRows, 1);
-}
+		check(Primitive.NumColumns == 1 && Primitive.NumRows <= 4);
 
-const FPrimitiveType* FPrimitiveType::Get(EScalarKind InScalarKind, int NumRows, int NumColumns, bool bIsLWCInverseMatrix)
-{
-	check(InScalarKind >= 0 && InScalarKind <= ScalarKind_LWC);
-	
-	static const FStringView Invalid = TEXT("invalid");
-
-	static const FPrimitiveType Types[] {
-		{ { TypeKind_Primitive }, { TEXT("bool") }, 		ScalarKind_Bool, 1, 1 },
-		{ { TypeKind_Primitive }, Invalid, 				ScalarKind_Bool, 1, 2 }, 
-		{ { TypeKind_Primitive }, Invalid, 				ScalarKind_Bool, 1, 3 },
-		{ { TypeKind_Primitive }, Invalid, 				ScalarKind_Bool, 1, 4 },
-		{ { TypeKind_Primitive }, { TEXT("bool2") },   	ScalarKind_Bool, 2, 1 },
-		{ { TypeKind_Primitive }, { TEXT("bool2x2") }, 	ScalarKind_Bool, 2, 2 },
-		{ { TypeKind_Primitive }, { TEXT("bool2x3") }, 	ScalarKind_Bool, 2, 3 },
-		{ { TypeKind_Primitive }, { TEXT("bool2x4") }, 	ScalarKind_Bool, 2, 4 },
-		{ { TypeKind_Primitive }, { TEXT("bool3") },   	ScalarKind_Bool, 3, 1 },
-		{ { TypeKind_Primitive }, { TEXT("bool3x2") }, 	ScalarKind_Bool, 3, 2 },
-		{ { TypeKind_Primitive }, { TEXT("bool3x3") }, 	ScalarKind_Bool, 3, 3 },
-		{ { TypeKind_Primitive }, { TEXT("bool3x4") }, 	ScalarKind_Bool, 3, 4 },
-		{ { TypeKind_Primitive }, { TEXT("bool4") },   	ScalarKind_Bool, 4, 1 },
-		{ { TypeKind_Primitive }, { TEXT("bool4x2") }, 	ScalarKind_Bool, 4, 2 },
-		{ { TypeKind_Primitive }, { TEXT("bool4x3") }, 	ScalarKind_Bool, 4, 3 },
-		{ { TypeKind_Primitive }, { TEXT("bool4x4") }, 	ScalarKind_Bool, 4, 4 },
-		{ { TypeKind_Primitive }, { TEXT("int") }, 		ScalarKind_Int, 1, 1 },
-		{ { TypeKind_Primitive }, Invalid, 				ScalarKind_Int, 1, 2 },
-		{ { TypeKind_Primitive }, Invalid, 				ScalarKind_Int, 1, 3 },
-		{ { TypeKind_Primitive }, Invalid, 				ScalarKind_Int, 1, 4 },
-		{ { TypeKind_Primitive }, { TEXT("int2") },   	ScalarKind_Int, 2, 1 },
-		{ { TypeKind_Primitive }, { TEXT("int2x2") }, 	ScalarKind_Int, 2, 2 },
-		{ { TypeKind_Primitive }, { TEXT("int2x3") }, 	ScalarKind_Int, 2, 3 },
-		{ { TypeKind_Primitive }, { TEXT("int2x4") }, 	ScalarKind_Int, 2, 4 },
-		{ { TypeKind_Primitive }, { TEXT("int3") },   	ScalarKind_Int, 3, 1 },
-		{ { TypeKind_Primitive }, { TEXT("int3x2") }, 	ScalarKind_Int, 3, 2 },
-		{ { TypeKind_Primitive }, { TEXT("int3x3") }, 	ScalarKind_Int, 3, 3 },
-		{ { TypeKind_Primitive }, { TEXT("int3x4") }, 	ScalarKind_Int, 3, 4 },
-		{ { TypeKind_Primitive }, { TEXT("int4") },   	ScalarKind_Int, 4, 1 },
-		{ { TypeKind_Primitive }, { TEXT("int4x2") }, 	ScalarKind_Int, 4, 2 },
-		{ { TypeKind_Primitive }, { TEXT("int4x3") }, 	ScalarKind_Int, 4, 3 },
-		{ { TypeKind_Primitive }, { TEXT("int4x4") }, 	ScalarKind_Int, 4, 4 },
-		{ { TypeKind_Primitive }, { TEXT("float") }, 		ScalarKind_Float, 1, 1 },
-		{ { TypeKind_Primitive }, Invalid, 					ScalarKind_Float, 1, 2 },
-		{ { TypeKind_Primitive }, Invalid, 					ScalarKind_Float, 1, 3 },
-		{ { TypeKind_Primitive }, Invalid, 					ScalarKind_Float, 1, 4 },
-		{ { TypeKind_Primitive }, { TEXT("float2") },   	ScalarKind_Float, 2, 1 },
-		{ { TypeKind_Primitive }, { TEXT("float2x2") }, 	ScalarKind_Float, 2, 2 },
-		{ { TypeKind_Primitive }, { TEXT("float2x3") }, 	ScalarKind_Float, 2, 3 },
-		{ { TypeKind_Primitive }, { TEXT("float2x4") }, 	ScalarKind_Float, 2, 4 },
-		{ { TypeKind_Primitive }, { TEXT("float3") },   	ScalarKind_Float, 3, 1 },
-		{ { TypeKind_Primitive }, { TEXT("float3x2") }, 	ScalarKind_Float, 3, 2 },
-		{ { TypeKind_Primitive }, { TEXT("float3x3") }, 	ScalarKind_Float, 3, 3 },
-		{ { TypeKind_Primitive }, { TEXT("float3x4") }, 	ScalarKind_Float, 3, 4 },
-		{ { TypeKind_Primitive }, { TEXT("float4") },   	ScalarKind_Float, 4, 1 },
-		{ { TypeKind_Primitive }, { TEXT("float4x2") }, 	ScalarKind_Float, 4, 2 },
-		{ { TypeKind_Primitive }, { TEXT("float4x3") }, 	ScalarKind_Float, 4, 3 },
-		{ { TypeKind_Primitive }, { TEXT("float4x4") }, 	ScalarKind_Float, 4, 4 },
-		{ { TypeKind_Primitive }, { TEXT("FWSScalar") },		ScalarKind_LWC, 1, 1 },
-		{ { TypeKind_Primitive }, Invalid, 						ScalarKind_LWC, 1, 2 },
-		{ { TypeKind_Primitive }, Invalid, 						ScalarKind_LWC, 1, 3 },
-		{ { TypeKind_Primitive }, Invalid, 						ScalarKind_LWC, 1, 4 },
-		{ { TypeKind_Primitive }, { TEXT("FWSVector2") },		ScalarKind_LWC, 2, 1 },
-		{ { TypeKind_Primitive }, Invalid,						ScalarKind_LWC, 2, 2 },
-		{ { TypeKind_Primitive }, Invalid,						ScalarKind_LWC, 2, 3 },
-		{ { TypeKind_Primitive }, Invalid,						ScalarKind_LWC, 2, 4 },
-		{ { TypeKind_Primitive }, { TEXT("FWSVector3") },		ScalarKind_LWC, 3, 1 },
-		{ { TypeKind_Primitive }, Invalid,						ScalarKind_LWC, 3, 2 },
-		{ { TypeKind_Primitive }, Invalid,						ScalarKind_LWC, 3, 3 },
-		{ { TypeKind_Primitive }, Invalid,						ScalarKind_LWC, 3, 4 },
-		{ { TypeKind_Primitive }, { TEXT("FWSVector4") },		ScalarKind_LWC, 4, 1 },
-		{ { TypeKind_Primitive }, Invalid,						ScalarKind_LWC, 4, 2 },
-		{ { TypeKind_Primitive }, Invalid,						ScalarKind_LWC, 4, 3 },
-		{ { TypeKind_Primitive }, { TEXT("FWSMatrix") }, 		ScalarKind_LWC, 4, 4 },
-	};
-
-	static const FPrimitiveType LWCInverseMatrixType = { { TypeKind_Primitive }, { TEXT("FWSInverseMatrix") }, ScalarKind_LWC, 4, 4, true };
-
-	if (bIsLWCInverseMatrix)
+		switch (Primitive.ScalarKind)
+		{
+			case EScalarKind::Bool: return (EValueType)((int)EValueType::Bool1 + Primitive.NumRows - 1);
+			case EScalarKind::Int: 	return (EValueType)((int)EValueType::Int1 + Primitive.NumRows - 1);
+			case EScalarKind::Float: return (EValueType)((int)EValueType::Float1 + Primitive.NumRows - 1);
+			case EScalarKind::Double:	return (EValueType)((int)EValueType::Double1 + Primitive.NumRows - 1);
+			default: UE_MIR_UNREACHABLE();
+		}
+	}
+	else if (Is(ETypeKind::Aggregate))
 	{
-		// LWC inverse matrix flag only supported for ScalarKind_LWC and 4x4 dimensions.
-		check(InScalarKind == ScalarKind_LWC && NumRows == 4 && NumColumns == 4);
-		return &LWCInverseMatrixType;
+		return EValueType::Struct;
+	}
+	else
+	{
+		return EValueType::Object;
 	}
-
-	int Index = InScalarKind * 4 * 4 + (NumRows - 1) * 4 + (NumColumns - 1);
-	check(Index < UE_ARRAY_COUNT(Types));
-	return &Types[Index];
-}
-
-const FPrimitiveType* FPrimitiveType::ToScalar() const
-{
-	return FPrimitiveType::GetScalar(ScalarKind);
-}
-
-const FPrimitiveType* FPrimitiveType::WithScalarKind(EScalarKind InScalarKind) const
-{
-	return FPrimitiveType::Get(InScalarKind, NumRows, NumColumns);
-}
-
-const FPrimitiveType* FPrimitiveType::ToVector(int InNumRows) const
-{
-	return FPrimitiveType::GetVector(ScalarKind, InNumRows);
-}
-
-const FType* FObjectType::GetTexture2D()
-{
-	static const FObjectType Instance { { TypeKind_Object }, ObjectKind_Texture2D };
-	return &Instance;
-}
-
-const FType* FObjectType::GetSubstrateData()
-{
-	static const FObjectType Instance{ { TypeKind_Object }, ObjectKind_SubstrateData };
-	return &Instance;
 }
 
-const FType* FObjectType::GetVTPageTableResult()
+bool FType::operator==(FType Other) const
 {
-	static const FObjectType Instance{ { TypeKind_Object }, ObjectKind_VTPageTableResult };
-	return &Instance;
+	if (Kind != Other.Kind)
+	{
+		return false;
+	}
+	switch (Kind)
+	{
+		case ETypeKind::Primitive: return Primitive == Other.Primitive;
+		case ETypeKind::Aggregate: return Aggregate == Other.Aggregate;
+		default: return true;
+	}
 }
 
 } // namespace MIR
diff --git a/Engine/Source/Runtime/Engine/Private/Materials/MaterialIRValueAnalyzer.cpp b/Engine/Source/Runtime/Engine/Private/Materials/MaterialIRValueAnalyzer.cpp
--- a/Engine/Source/Runtime/Engine/Private/Materials/MaterialIRValueAnalyzer.cpp
+++ b/Engine/Source/Runtime/Engine/Private/Materials/MaterialIRValueAnalyzer.cpp
@@ -202,11 +202,11 @@ static void AnalyzePrimitiveUniformParameter(FMaterialIRValueAnalyzer& Analyzer,
 
 	Parameter->Analysis_UniformParameterIndex = UniformExpressionSet.FindOrAddNumericParameter(ParameterMetadata.Value.Type, ParameterInfo, DefaultValueOffset);
 
-	// Get the parameter type as primitive
-	const MIR::FPrimitiveType* ParameterType = Parameter->Type->AsPrimitive();
+	// Make sure the parameter type is primitive
+	MIR::FPrimitive PrimitiveType = Parameter->Type.GetPrimitive();
 
-	// Only int and float parameters supported for now (TODO)
-	check(ParameterType->ScalarKind == MIR::ScalarKind_Int || ParameterType->ScalarKind == MIR::ScalarKind_Float || ParameterType->ScalarKind == MIR::ScalarKind_LWC);
+	// Only int, float and LWC parameters supported for now
+	check(PrimitiveType.IsInteger() || PrimitiveType.IsAnyFloat());
 
 	// Following code calculates GlobalComponentOffset.
 	// GlobalComponentOffset is the i-th component in the array of float4s that make the uniform buffer.
@@ -218,8 +218,7 @@ static void AnalyzePrimitiveUniformParameter(FMaterialIRValueAnalyzer& Analyzer,
 	// in the appropriate freelist.
 
 	uint32 GlobalComponentOffset = UINT32_MAX;
-	uint32 NumComponents = ParameterType->NumComponents();
-	check(NumComponents > 0);
+	uint32 NumComponents = PrimitiveType.NumComponents();
 	check(NumComponents <= 4); // Only vectors supported for now
 
 	uint32 UsedNumComponents;
@@ -246,7 +245,8 @@ static void AnalyzePrimitiveUniformParameter(FMaterialIRValueAnalyzer& Analyzer,
 	if (UsedNumComponents >= 4)
 	{
 		check(UsedNumComponents == 4); 
-		GlobalComponentOffset = UniformExpressionSet.AllocateFromUniformBuffer(ParameterType->ScalarKind == MIR::ScalarKind_LWC ? 32 : 16) / 4;
+		const uint32 SizeInBytes = PrimitiveType.IsLWC() ? 32 : 16;
+		GlobalComponentOffset = UniformExpressionSet.AllocateFromUniformBuffer(SizeInBytes) / 4;
 	}
 
 	// Finally, check whether in the used slot there's any slack left, and if so, record it in the appropriate free-list.
@@ -269,11 +269,11 @@ static void AnalyzePrimitiveUniformParameter(FMaterialIRValueAnalyzer& Analyzer,
 		ParamInsight.ComponentsCount = NumComponents;
 		ParamInsight.ParameterName = ParameterInfo.Name;
 
-		switch (ParameterType->ScalarKind)
+		switch (PrimitiveType.ScalarKind)
 		{
-			case MIR::ScalarKind_Int: ParamInsight.ComponentType = FMaterialInsights::FUniformBufferSlotComponentType::CT_Int; break;
-			case MIR::ScalarKind_Float: ParamInsight.ComponentType = FMaterialInsights::FUniformBufferSlotComponentType::CT_Float; break;
-			case MIR::ScalarKind_LWC: ParamInsight.ComponentType = FMaterialInsights::FUniformBufferSlotComponentType::CT_LWC; break;
+			case MIR::EScalarKind::Int:   ParamInsight.ComponentType = FMaterialInsights::FUniformBufferSlotComponentType::CT_Int; break;
+			case MIR::EScalarKind::Float: ParamInsight.ComponentType = FMaterialInsights::FUniformBufferSlotComponentType::CT_Float; break;
+			case MIR::EScalarKind::Double:   ParamInsight.ComponentType = FMaterialInsights::FUniformBufferSlotComponentType::CT_LWC; break;
 			default: UE_MIR_UNREACHABLE();
 		}
 
@@ -283,13 +283,13 @@ static void AnalyzePrimitiveUniformParameter(FMaterialIRValueAnalyzer& Analyzer,
 
 static void AnalyzeUniformParameter(FMaterialIRValueAnalyzer& Analyzer, MIR::FUniformParameter* Parameter)
 {
-	if (Parameter->Type->IsObjectOfKind(MIR::ObjectKind_Texture2D))
+	if (Parameter->Type.IsTexture2D())
 	{
 		AnalyzeTextureUniformParameter(Analyzer, Parameter);
 	}
 	else
 	{
-		check(Parameter->Type->AsPrimitive());
+		check(Parameter->Type.AsPrimitive());
 		AnalyzePrimitiveUniformParameter(Analyzer, Parameter);
 	}
 }
diff --git a/Engine/Source/Runtime/Engine/Public/Materials/MaterialIR.h b/Engine/Source/Runtime/Engine/Public/Materials/MaterialIR.h
--- a/Engine/Source/Runtime/Engine/Public/Materials/MaterialIR.h
+++ b/Engine/Source/Runtime/Engine/Public/Materials/MaterialIR.h
@@ -2,6 +2,7 @@
 #pragma once
 
 #include "Materials/MaterialIRCommon.h"
+#include "Materials/MaterialIRTypes.h"
 #include "Materials/MaterialParameters.h"
 #include "Engine/TextureDefines.h"
 
@@ -139,6 +140,9 @@ const TCHAR* LexToString(EValueKind Kind);
 //   derive from FInstruction).
 struct FValue
 {
+	// The runtime type this value has.
+	FType Type;
+
 	// Used to discern the concrete C++ type of this value (e.g. Subscript)
 	EValueKind Kind : 8;
 
@@ -150,9 +154,6 @@ struct FValue
 	// See `EGraphProperties` for more information.
 	EGraphProperties GraphProperties;
 
-	// The runtime type this value has.
-	const FType* Type;
-
 	// Returns whether this value has been analyzed for specified stage.
 	bool IsAnalyzed(EStage State) const;
 
@@ -174,9 +175,6 @@ struct FValue
 	// Returns the size in bytes of this value instance.
 	uint32 GetSizeInBytes() const;
 
-	// Returns whether this value is of specified kind.
-	bool IsA(EValueKind InKind) const;
-
 	// Gets the immutable array of all this value uses.
 	// An use is another value referenced by this one (e.g. the operands of a binary expression).
 	// Returns the immutable array of uses.
@@ -187,17 +185,11 @@ struct FValue
 	// Returns the immutable array of uses.
 	TConstArrayView<FValue*> GetUsesForStage(MIR::EStage Stage) const;
 
-	// Returns whether this value is poison.
-	bool IsPoison() const;
-
-	// Returns whether this value exactly equals Other.
-	bool Equals(const FValue* Other) const;
-
-	// Returns whether this value is a scalar (its type is Primitive with exactly 1 component).
-	bool IsScalar() const;
+	// Returns whether this value is of specified kind.
+	bool IsA(EValueKind InKind) const { return Kind == InKind; }
 
-	// Returns whether this value is a vector (tis type is Primitive with 1-4 rows and exactly 1 column).
-	bool IsVector() const;
+	// Returns whether this value is poison.
+	bool IsPoison() const { return Kind == VK_Poison; }
 
 	// Returns whether this value is a constant boolean with value true.
 	bool IsTrue() const;
@@ -222,6 +214,9 @@ struct FValue
 
 	// Returns whether this value is arithmetic and all components are approximately one.
 	bool AreAllNearlyOne() const;
+	
+	// Returns whether this value exactly equals Other.
+	bool Equals(const FValue* Other) const;
 
 	// Tries to cast this value to specified type T and returns the casted pointer, if possible (nullptr otherwise).
 	template <typename T>
@@ -294,24 +289,9 @@ struct FConstant : TValue<VK_Constant>
 		bool  		Boolean;
 		TInteger	Integer;
 		TFloat 		Float;
-		TDouble		LWCValue;
+		TDouble		Double;
 	};
 
-	// Returns whether this constant is a boolean.
-	bool IsBool() const;
-
-	// Returns whether this constant is an integer.
-	bool IsInteger() const;
-
-	// Returns whether this constant is a float.
-	bool IsFloat() const;
-
-	// Returns whether this constant is LWC.
-	bool IsLWC() const;
-
-	// Returns whether this is an LWC type representable as regular float.
-	bool IsLWCRepresentableAsFloat() const;
-
 	// Returns the constant value of given type T. The type must be bool, integral or floating point.
 	template <typename T>
 	T Get() const
@@ -421,7 +401,7 @@ static constexpr int32 TexCoordMaxNum = 8;
 const TCHAR* LexToString(EExternalInput Input);
 
 // Returns the given external input type.
-const FType* GetExternalInputType(EExternalInput Id);
+FType GetExternalInputType(EExternalInput Id);
 
 // Returns whether the given external input needs to generate non-trivial derivative expressions.
 bool IsExternalInputWithDerivatives(EExternalInput Id);
diff --git a/Engine/Source/Runtime/Engine/Public/Materials/MaterialIRCommon.h b/Engine/Source/Runtime/Engine/Public/Materials/MaterialIRCommon.h
--- a/Engine/Source/Runtime/Engine/Public/Materials/MaterialIRCommon.h
+++ b/Engine/Source/Runtime/Engine/Public/Materials/MaterialIRCommon.h
@@ -41,7 +41,7 @@ namespace MIR
 
 /* Types*/
 struct FType;
-struct FPrimitiveType;
+struct FPrimitive;
 struct FObjectType;
 struct FAggregateType;
 
@@ -84,7 +84,6 @@ struct TTemporaryArray : public TArrayView<T>
 	}
 };
 
-
 }
 
 #define UE_MIR_UNREACHABLE() { check(!"Unreachable"); UE_ASSUME(false); }
diff --git a/Engine/Source/Runtime/Engine/Public/Materials/MaterialIREmitter.h b/Engine/Source/Runtime/Engine/Public/Materials/MaterialIREmitter.h
--- a/Engine/Source/Runtime/Engine/Public/Materials/MaterialIREmitter.h
+++ b/Engine/Source/Runtime/Engine/Public/Materials/MaterialIREmitter.h
@@ -179,17 +179,14 @@ public:
 	// As an example, the common type between a float3 and a bool4 will be a float4.
 	// Note; that this may cause a loss of data when casting from int to float.
 	// It returns the common type if it exists, nullptr otherwise.
-	const FType* TryGetCommonType(const FType* A, const FType* B);
+	FType TryGetCommonType(FType A, FType B);
 
 	// Gets the common type between A and B and reports an error if no such type exists.
 	// See TryGetCommonType()
-	const FType* GetCommonType(const FType* A, const FType* B);
-
-	// Returns the type associated to the given material aggregate. 
-	const FAggregateType* GetMaterialAggregateType(const UMaterialAggregate* Aggregate);
+	FType GetCommonType(FType A, FType B);
 
 	// Returns the type of the i-th attribute in the given material aggregate.
-	const FType* GetMaterialAggregateAttributeType(const UMaterialAggregate* Aggregate, int32 AttributeIndex);
+	FType GetMaterialAggregateAttributeType(const UMaterialAggregate* Aggregate, int32 AttributeIndex);
 
 	/*-------------------------------- Input management --------------------------------*/
 
@@ -233,6 +230,9 @@ public:
 	// Validates that the value is of an arithmetic type (int or float). Reports error and returns poison otherwise.
 	FValueRef CheckIsArithmetic(FValueRef Value);
 
+	// Validates that the value is of a boolean type. If not, reports an error and returns poison.
+	FValueRef CheckIsBoolean(FValueRef Value);
+
 	// Validates that the value is of a integral type. If not, reports an error and returns poison.
 	FValueRef CheckIsInteger(FValueRef Value);
 
@@ -267,7 +267,7 @@ public:
 	FValueRef ConstantFromShaderValue(const UE::Shader::FValue& InValue);
 
 	// Emits a default value of given type (zero initialized).
-	FValueRef ConstantDefault(const FType* Type);
+	FValueRef ConstantDefault(FType Type);
 
 	// Emits a scalar constant value of given kind with value 0 and returns it.
 	FValueRef ConstantZero(EScalarKind Kind);
@@ -276,7 +276,7 @@ public:
 	FValueRef ConstantOne(EScalarKind Kind);
 	
 	// Casts the given float value to the given scalar kind and returns its value.
-	FValueRef ConstantScalar(EScalarKind Kind, TFloat FromFloat);
+	FValueRef ConstantScalar(EScalarKind Kind, TDouble Value);
 
 	// Returns the constant boolean scalar `true`.
 	FValueRef ConstantTrue();
@@ -321,7 +321,7 @@ public:
 	FValueRef ConstantFloat4(UE::Math::TVector4<TFloat> InValue);
 
 	// Returns the given constant LWC scalar.
-	FValueRef ConstantLWC(TDouble InX);
+	FValueRef ConstantDouble(TDouble InX);
 
 	/*--------------------- Other non-instruction values emission ---------------------*/
 
@@ -384,7 +384,7 @@ public:
 	FValueRef Branch(FValueRef Condition, FValueRef True, FValueRef False);
 		
 	// Casts the given value to the specified target type. Returns poison if the cast is invalid.
-	FValueRef Cast(FValueRef Value, const FType* TargetType);
+	FValueRef Cast(FValueRef Value, FType TargetType);
 	
 	// Casts the given primitive value to a scalar. If this is a vector, it will convert to its first component.
 	FValueRef CastToScalar(FValueRef Value);
@@ -415,7 +415,7 @@ public:
 	// Emits a stage switch instruction. Use this instruction to select a different value
 	// in each execution stage.
 	// Note: ValuePerStage must have no more than MIR::NumStages elements.
-	FValueRef StageSwitch(const FType* Type, TConstArrayView<FValueRef> ValuePerStage);
+	FValueRef StageSwitch(FType Type, TConstArrayView<FValueRef> ValuePerStage);
 
 	// Emits a texture gather instruction.
 	// 
@@ -500,7 +500,7 @@ public:
 	// @param InputValues An array of input values required by the HLSL code.
 	// @param ValueFlags Flags that specify additional properties for this value (default is None).
 	// @param UsedGraphProperties Properties used during graph processing (default is None).
-	FValueRef InlineHLSL(const FType* Type, FString Code, TConstArrayView<FValueRef> InputValues, EValueFlags ValueFlags = EValueFlags::None, EGraphProperties UsedGraphProperties = EGraphProperties::None);
+	FValueRef InlineHLSL(FType Type, FString Code, TConstArrayView<FValueRef> InputValues, EValueFlags ValueFlags = EValueFlags::None, EGraphProperties UsedGraphProperties = EGraphProperties::None);
 
 	// Emit a new HLSL function inline within the material shader based on an external code declaration.
 	// @param InExternalCodeDeclaration The external HLSL code declaration to use for emission.
@@ -626,9 +626,6 @@ private:
 	// calculations.
 	TSet<FValue*, FValueKeyFuncs> ValueSet{};
 
-	// Maps aggregate types to their MIR type instance. Used to avoid duplicates.
-	TMap<const UMaterialAggregate*, FAggregateType*> AggregateTypeMap{};
-
 	friend FMaterialIRModuleBuilder;
 	friend FMaterialIRModuleBuilderImpl;
 };
diff --git a/Engine/Source/Runtime/Engine/Public/Materials/MaterialIRModule.h b/Engine/Source/Runtime/Engine/Public/Materials/MaterialIRModule.h
--- a/Engine/Source/Runtime/Engine/Public/Materials/MaterialIRModule.h
+++ b/Engine/Source/Runtime/Engine/Public/Materials/MaterialIRModule.h
@@ -59,7 +59,7 @@ public:
 	TArrayView<const MIR::FSetMaterialOutput* const> GetOutputs(MIR::EStage Stage) const { return Outputs[Stage]; }
 
 	// Retrieves the root block (i.e. the "main" scope) for a specific shader stage.
-	const MIR::FBlock& GetRootBlock(MIR::EStage Stage) const { return *RootBlock[Stage]; }
+	const MIR::FBlock& GetRootBlock(MIR::EStage Stage) const { return RootBlocks[Stage]; }
 
 	// Returns a list of all environment define names this module requires to be enabled for shader compilation.
 	const TSet<FName>& GetEnvironmentDefines() const { return EnvironmentDefines; }
@@ -95,9 +95,6 @@ private:
 	// Compilation output data.
 	FMaterialCompilationOutput CompilationOutput;
 
-	// List of all the dynamic types contained in this module.
-	TArray<MIR::FType*> Types;
-
 	// List of all the IR values contained in this module.
 	TArray<MIR::FValue*> Values;
 
@@ -105,7 +102,7 @@ private:
 	TArray<MIR::FSetMaterialOutput*> Outputs[MIR::NumStages];
 
 	// Root blocks per stage.
-	MIR::FBlock* RootBlock[MIR::NumStages];
+	MIR::FBlock RootBlocks[MIR::NumStages];
 
 	// Compilation statistics.
 	FStatistics Statistics;
diff --git a/Engine/Source/Runtime/Engine/Public/Materials/MaterialIRTypes.h b/Engine/Source/Runtime/Engine/Public/Materials/MaterialIRTypes.h
--- a/Engine/Source/Runtime/Engine/Public/Materials/MaterialIRTypes.h
+++ b/Engine/Source/Runtime/Engine/Public/Materials/MaterialIRTypes.h
@@ -11,214 +11,292 @@
 namespace MIR
 {
 
-enum ETypeKind
+// Represents the top-level kind of a type.
+enum class ETypeKind : uint8
 {
-	TypeKind_Poison,
-	TypeKind_Void,
-	TypeKind_Primitive,
-	TypeKind_Object,
-	TypeKind_Aggregate,
-};
+	Poison, 	// An invalid or error type.
+	Void,  		// The void type, representing no value.
+	Primitive, 	// A numeric, vector, or matrix type (e.g. float, float4, float4x4).
+	Aggregate,	// A material aggregate.
 
-enum EObjectKind 
-{
-	ObjectKind_Texture2D,
-	ObjectKind_SubstrateData,
-	ObjectKind_VTPageTableResult,
+	// Opaque Objects
+	
+	Texture2D,			// A 2D texture object.
+	SubstrateData,		// Substrate-specific data structure.
+	VTPageTableResult,	// A virtual texture page table lookup result.
 };
 
+// Returns the string representation of a ETypeKind.
 const TCHAR* TypeKindToString(ETypeKind Kind);
 
-struct FType
+// Primitive types of a single scalar.
+// Note: These are listed in precision order. Converting one to the other is then simply performed taking the max EScalarKind.
+enum class EScalarKind : uint8
 {
-	/// Identifies what derived type this is.
-	ETypeKind Kind;
-
-	/// Returns the type matching specified UE::Shader::FType.
-	static const FType* FromShaderType(const UE::Shader::FType& InShaderType);
-	
-	/// Returns the type matching specified EMaterialValueType.
-	static const FType* FromMaterialValueType(EMaterialValueType Type);
+	Bool,	// Boolean type (true/false).
+	Int,	// 32-bit signed integer type.
+	Float,	// 32-bit floating-point type.
+	Double,	// Large World Coordinates type.
+};
 
-	/// Returns the type matching specified EMaterialParameterType.
-	static const FType* FromMaterialParameterType(EMaterialParameterType Type);
+// Returns whether the specified scalar kind supports arithmetic operators (plus, minus, etc).
+bool ScalarKindIsArithmetic(EScalarKind Kind);
 
-	/// Returns the `void` type.
-	static const FType* GetVoid();
+// Returns whether the specified scalar kind is a floating point type (float, LWC).
+bool ScalarKindIsAnyFloat(EScalarKind Kind);
 
-	/// Returns the `void` type.
-	static const FType* GetPoison();
+// Returns the string representation of specified scalar kind.
+const TCHAR* ScalarKindToString(EScalarKind Kind);
 
-	/// Returns whether this is the poison type.
-	bool IsPoison() const;
+// Describes a primitive type (scalar, vector, or matrix).
+struct FPrimitive
+{
+	// The underlying scalar data type (e.g. Bool, Int, Float).
+	EScalarKind ScalarKind;
 
-	/// Returns whether this is a boolean primitive type (of any dimension).
-	bool IsBoolean() const;
+	// Number of rows. For a vector, this is the number of elements. For a scalar, this is 1.
+	int8 NumRows;
 	
-	/// Returns whether this is an integer primitive type (of any dimension).
-	bool IsInteger() const;
-
-	/// Returns whether this is a floating point primitive type (of any dimension).
-	bool IsAnyFloat() const;
+	// Number of columns. For a vector or scalar, this is 1. For a matrix, this is > 1.
+	int8 NumColumns;
 
-	/// Returns whether this is an LWC primitive type (of any dimension).
-	bool IsLWC() const;
-
-	/// Returns whether this type is a `bool` scalar.
-	bool IsBoolScalar() const;
-
-	/// Returns whether this type is of a specific object kind.
-	bool IsObjectOfKind(EObjectKind ObjectKind) const;
-
-	/// Returns whether this type is a texture.
-	bool IsTexture() const;
+	// If true, this represents an LWC inverse matrix. Only supported for EScalarKind::Double and 4x4 dimensions.
+	bool bIsLWCInverseMatrix;
 
-	/// Returns this type upcast to PrimitiveType if it is one, otherwise nullptr.
-	const FPrimitiveType* AsPrimitive() const;
+	// Returns whether this primitive type is arithmetic (it supports arithmetic operations like addition).
+	bool IsArithmetic() const { return ScalarKindIsArithmetic(ScalarKind); }
 	
-	/// Returns this type upcast to PrimitiveType if it is one, otherwise nullptr.
-	const FPrimitiveType* AsArithmetic() const;
-
-	/// Returns this type upcast to PrimitiveType if it's a scalar, otherwise nullptr.
-	const FPrimitiveType* AsScalar() const;
+	// Returns whether this type is a primitive scalar (1x1).
+	bool IsScalar() const { return NumComponents() == 1; }
+	
+	// Returns whether this primitive type is a column vector (rows > 1, columns = 1).
+	bool IsVector() const { return NumRows > 1 && NumColumns == 1; }
+	
+	// Returns whether this primitive type is a matrix (rows > 1, columns > 1).
+	bool IsMatrix() const { return NumRows > 1 && NumColumns > 1; }
+	
+	// Returns whether this is a boolean primitive type (of any dimension).
+	bool IsBoolean() const { return ScalarKind == EScalarKind::Bool; }
+	
+	// Returns whether this is a single boolean scalar.
+	bool IsBoolScalar() const { return IsBoolean() && IsScalar(); }
 
-	/// Returns this type upcast to PrimitiveType if it's a vector, otherwise nullptr.
-	const FPrimitiveType* AsVector() const;
+	// Returns whether this is an integer primitive type (of any dimension).
+	bool IsInteger() const { return ScalarKind == EScalarKind::Int; }
+	
+	// Returns whether this is a float primitive type (of any dimension).
+	bool IsFloat() const { return ScalarKind == EScalarKind::Float; }
 
-	/// Returns this type upcast to PrimitiveType if it's a matrix, otherwise nullptr.
-	const FPrimitiveType* AsMatrix() const;
+	// Returns whether this is an LWC primitive type (of any dimension).
+	bool IsLWC() const { return ScalarKind == EScalarKind::Double; }
+	
+	// Returns whether this is any floating point primitive type (e.g. float, LWC).
+	bool IsAnyFloat() const { return ScalarKindIsAnyFloat(ScalarKind); }
 
-	/// Returns this type upcast to AggregateType, otherwise nullptr.
-	const FAggregateType* AsAggregate() const;
+	// Returns the total number of scalar components in this primitive type.
+	int32 NumComponents() const { return NumRows * NumColumns; }
 
-	/// Returns the this type name spelling (e.g. float4x4).
-	FString GetSpelling() const;
+	// Returns a new FType with the same dimensions but a different scalar kind.
+	FType ToScalarKind(EScalarKind InScalarKind) const;
 	
-	/// Returns this type cast as a FObjectPtr if it is one, otherwise nullptr.
-	const FObjectType* AsObject() const;
-
-	/// Converts this type to a UE::Shader::EValueType.
-	UE::Shader::EValueType ToValueType() const;
-};
+	// Returns the scalar FType corresponding to this primitive's scalar kind.
+	FType ToScalar() const;
+	
+	// Returns a vector FType with the given number of rows and this primitive's scalar kind.
+	FType ToVector(int NumRows) const;
 
-/// Primitive types of a single scalar.
-/// Note: These are listed in precision order. Converting one to the other is then simply performed taking the max EScalarKind.
-enum EScalarKind
-{
-	ScalarKind_Bool,
-	ScalarKind_Int,
-	ScalarKind_Float,
-	ScalarKind_LWC,
+	bool operator==(const FPrimitive& Other) const = default;
+	bool operator!=(const FPrimitive& Other) const = default;
 };
 
-/// Returns whether the specified scalar kind supports arithmetic operators (plus, minus, etc).
-bool ScalarKindIsArithmetic(EScalarKind Kind);
-
-/// Returns whether the specified scalar kind is a floating point type (float, double, etc).
-bool ScalarKindIsAnyFloat(EScalarKind Kind);
-
-/// Returns the string representation of specified scalar kind.
-const TCHAR* ScalarKindToString(EScalarKind Kind);
-
-/// Represents the type of scalars, vectors and matrices. It indicates what kind of scalar
-/// data type it has and type dimensions (rows and columns).
-struct FPrimitiveType : FType
+// Describes a complete type within the material IR. This can be a primitive, an aggregate
+// (struct), an opaque object (e.g., a texture), or a special type like Poison or Void.
+// Note: This struct is lightweight and intended to be passed by value.
+struct FType
 {
-	/// String representation of this type (e.g. 'float3', "bool4x2")
-	FStringView Spelling;
-
-	/// Scalar data type kind.
-	EScalarKind ScalarKind;
-
-	/// Number of rows.
-	int NumRows;
+	// Returns the FType corresponding to a given UE::Shader::FType.
+	static FType FromShaderType(const UE::Shader::FType& InShaderType);
 	
-	/// Number of columns (rows > 1 and columns == 1, this is a vector).
-	int NumColumns;
-
-	/// LWC inverse matrix type -- only supported for ScalarKind_LWC and 4x4 dimensions
-	bool bIsLWCInverseMatrix;
+	// Returns the FType corresponding to a given EMaterialValueType.
+	static FType FromMaterialValueType(EMaterialValueType Type);
 
-	// Returns the boolean scalar type.
-	static const FPrimitiveType* GetBool();
-	
-	// Returns the integer scalar type.
-	static const FPrimitiveType* GetInt();
+	// Returns the FType corresponding to a given EMaterialParameterType.
+	static FType FromMaterialParameterType(EMaterialParameterType Type);
 	
-	// Returns the floating point scalar type.
-	static const FPrimitiveType* GetFloat();
+	// Returns the poison type, which represents an error or invalid type.
+	static FType MakePoison();
 	
-	// Returns the floating point 2D vector type.
-	static const FPrimitiveType* GetFloat2();
+	// Returns the void type, which represents the absence of a value.
+	static FType MakeVoid();
+
+	// Returns a boolean type with the specified dimensions (scalar, vector, or matrix).
+	static FType MakeBool(int NumRows = 1, int NumColumns = 1);
 	
-	// Returns the floating point 3D vector type.
-	static const FPrimitiveType* GetFloat3();
+	// Returns an integer type with the specified dimensions (scalar, vector, or matrix).
+	static FType MakeInt(int NumRows = 1, int NumColumns = 1);
 	
-	// Returns the floating point 4D vector type.
-	static const FPrimitiveType* GetFloat4();
-
-	// Returns the LWC scalar type.
-	static const FPrimitiveType* GetLWC();
+	// Returns a float type with the specified dimensions (scalar, vector, or matrix).
+	static FType MakeFloat(int NumRows = 1, int NumColumns = 1);
 
-	// Returns the LWC 2D vector type.
-	static const FPrimitiveType* GetLWC2();
+	// Returns an LWC type with the specified dimensions (scalar, vector, or matrix).
+	static FType MakeLWC(int NumRows = 1, int NumColumns = 1);
 
-	// Returns the LWC 3D vector type.
-	static const FPrimitiveType* GetLWC3();
+	// Returns a scalar type of the given scalar kind.
+	static FType MakeScalar(EScalarKind InScalarKind);
 
-	// Returns the LWC 4D vector type.
-	static const FPrimitiveType* GetLWC4();
+	// Returns a vector type of the given scalar kind and number of rows.
+	static FType MakeVector(EScalarKind InScalarKind, int NumRows);
 
-	// Returns the scalar type with given kind.
-	static const FPrimitiveType* GetScalar(EScalarKind InScalarKind);
-
-	// Returns the column vector type with given kind and number of rows.
-	static const FPrimitiveType* GetVector(EScalarKind InScalarKind, int NumRows);
+	// Returns a primitive type with the given kind and dimensions.
+	// LWC inverse matrix flag is only supported for EScalarKind::LWC and 4x4 dimensions.
+	static FType MakePrimitive(EScalarKind InScalarKind, int NumRows = 1, int NumColumns = 1, bool bIsLWCInverseMatrix = false);
 	
-	// Returns the primitive type with given kind and number of columns and rows.  LWC inverse matrix flag only supported for ScalarKind_LWC and 4x4 dimensions.
-	static const FPrimitiveType* Get(EScalarKind InScalarKind, int NumRows, int NumColumns, bool bIsLWCInverseMatrix = false);
-
-	// Returns the number of components in this primitive type.
-	int  NumComponents() const { return NumRows * NumColumns; }
+	// Returns the aggregate type associated with the given material aggregate definition.
+	static FType MakeAggregate(const UMaterialAggregate* Aggregate);
 
-	// Whether this primitive type is scalar.
-	bool IsScalar() const { return NumComponents() == 1; }
+	// Returns the Texture2D object type.
+	static FType MakeTexture2D();
 	
-	// Whether this primitive type is a column vector.
-	bool IsVector() const { return NumRows > 1 && NumColumns == 1; }
+	// Returns the SubstrateData object type.
+	static FType MakeSubstrateData();
 	
-	// Whether this primitive type is a matrix.
-	bool IsMatrix() const { return NumRows > 1 && NumColumns > 1; }
+	// Returns the VTPageTableResult object type.
+	static FType MakeVTPageTableResult();
+
+	// Returns the this type name spelling (e.g. float4x4).
+	FString GetSpelling() const;
 	
-	// Whether this primitive type is arithmetic (it supports arithemtic operations like addition).
-	bool IsArithmetic() const { return ScalarKindIsArithmetic(ScalarKind); }
+	// Checks if this type is of the specified kind.
+	bool Is(ETypeKind InKind) const { return Kind == InKind; }
 
-	// Returns this primitive type with a different scalar kind.
-	const FPrimitiveType* WithScalarKind(EScalarKind InScalarKind) const;
+	// Returns the high-level kind of this type.
+	ETypeKind GetKind() const { return Kind; }
+
+	// Returns whether this is the poison type, indicating an error.
+	bool IsPoison() const { return Kind == ETypeKind::Poison; }
 	
-	// Returns the scalar type with this type scalar kind
-	const FPrimitiveType* ToScalar() const;
+	// Returns true if the given type is a primitive that supports arithmetic operations.
+	bool IsPrimitive() const
+	{
+		return Is(ETypeKind::Primitive);
+	}
+
+	// Returns true if the given type is a primitive that supports arithmetic operations.
+	bool IsArithmetic() const
+	{
+		return Is(ETypeKind::Primitive) && Primitive.IsArithmetic();
+	}
+
+	// Returns true if the given type is a primitive with a boolean scalar kind.
+	bool IsBoolean() const
+	{
+		return Is(ETypeKind::Primitive) && Primitive.IsBoolean();
+	}
+
+	// Returns true if the given type is a primitive with an integer scalar kind.
+	bool IsInteger() const
+	{
+		return Is(ETypeKind::Primitive) && Primitive.IsInteger();
+	}
+
+	// Returns true if the given type is a primitive with a float scalar kind.
+	bool IsFloat() const
+	{
+		return Is(ETypeKind::Primitive) && Primitive.IsFloat();
+	}
+
+	// Returns true if the given type is a primitive with an LWC scalar kind.
+	bool IsLWC() const
+	{
+		return IsPrimitive() && Primitive.IsLWC();
+	}
+
+	// Returns true if the given type is a primitive with any float-based scalar kind (e.g. float, LWC).
+	bool IsAnyFloat() const
+	{
+		return Is(ETypeKind::Primitive) && Primitive.IsAnyFloat();
+	}
+
+	// Returns true if the given type is a scalar primitive.
+	bool IsScalar() const
+	{
+		return Is(ETypeKind::Primitive) && Primitive.IsScalar();
+	}
+
+	// Returns true if the given type is a boolean scalar primitive.
+	bool IsBoolScalar() const
+	{
+		return Is(ETypeKind::Primitive) && Primitive.IsBoolScalar();
+	}
+
+	// Returns true if the given type is a vector primitive.
+	bool IsVector() const
+	{
+		return Is(ETypeKind::Primitive) && Primitive.IsVector();
+	}
+
+	// Returns whether this type is a Texture2D object.
+	bool IsTexture2D() const { return Kind == ETypeKind::Texture2D; }
+
+	// Returns whether this type is a SubstrateData object.
+	bool IsSubstrateData() const { return Kind == ETypeKind::SubstrateData; }
+
+	// Returns whether this type is a VTPageTableResult object.
+	bool IsVTPageTableResult() const { return Kind == ETypeKind::VTPageTableResult; }
+
+	// Returns the primitive type information. It asserts that this type is primitive.
+	FPrimitive GetPrimitive() const { check(Is(ETypeKind::Primitive)); return Primitive; }
 	
-	// Returns the vector type with this type scalar kind
-	const FPrimitiveType* ToVector(int NumRows) const;
-};
+	// If this is a primitive type, returns its primitive info. Otherwise, returns None.
+	TOptional<FPrimitive> AsPrimitive() const
+	{
+		return IsPrimitive() ? Primitive : TOptional<FPrimitive>{};
+	}
+
+	// If this is a scalar type, returns its primitive info. Otherwise, returns None.
+	TOptional<FPrimitive> AsScalar() const
+	{
+		return IsScalar() ? Primitive : TOptional<FPrimitive>{};
+	}
+
+	// If this is a vector type, returns its primitive info. Otherwise, returns None.
+	TOptional<FPrimitive> AsVector() const
+	{
+		return IsVector() ? Primitive : TOptional<FPrimitive>{};
+	}
+
+	// If this is an aggregate type, returns a pointer to the UMaterialAggregate definition. Otherwise, returns nullptr.
+	const UMaterialAggregate* AsAggregate() const { return Is(ETypeKind::Aggregate) ? Aggregate : nullptr; }
+
+	// Converts this type to a corresponding UE::Shader::EValueType.
+	UE::Shader::EValueType ToValueType() const;
 
-// It represents an opaque object-like type such as a texture.
-struct FObjectType : FType
-{
-	EObjectKind ObjectKind;
+	// Conversion to bool, returns false if the type is Poison, true otherwise.
+	operator bool() const { return Kind != ETypeKind::Poison; }
 
-	static const FType* GetTexture2D();
-	static const FType* GetSubstrateData();
-	static const FType* GetVTPageTableResult();
-};
+	bool operator==(FType Other) const;
+	bool operator!=(FType Other) const { return !operator==(Other); }
 
-// A value of this type is a material aggregate.
-struct FAggregateType : FType
-{
-	const UMaterialAggregate* Aggregate;
+private:
+	union
+	{
+		// The primitive type information.
+		FPrimitive Primitive;
+
+		// Pointer to the material aggregate.
+		const UMaterialAggregate* Aggregate;
+	};
+
+	// Identifies the high-level kind of this type.
+	ETypeKind Kind : 8;
+
+	// Manual padding to fit FType size to a multiple of alignof(void*). This storage may be used in the future.
+	uint8 	  Unused[7];
 };
 
+// To avoid the compiler adding automatic padding, we must make sure that its size is a multiple of maximum alignment.
+static_assert(sizeof(MIR::FType) % alignof(void*) == 0);
+
 } // namespace MIR
 
 #endif // #if WITH_EDITOR
