New levelinstance.ShouldAutoUpdatePackedBlueprintsOnCommit setting introduced to automatically update packed blueprints after commiting changes to a level instance.

diff --git a/Engine/Source/Editor/LevelInstanceEditor/Private/LevelInstanceEditorModule.cpp b/Engine/Source/Editor/LevelInstanceEditor/Private/LevelInstanceEditorModule.cpp
--- a/Engine/Source/Editor/LevelInstanceEditor/Private/LevelInstanceEditorModule.cpp
+++ b/Engine/Source/Editor/LevelInstanceEditor/Private/LevelInstanceEditorModule.cpp
@@ -817,17 +817,7 @@ struct FLevelInstanceMenuUtils
 		});
 		UIAction.ExecuteAction.BindLambda([WorldAsset](const FToolMenuContext& MenuContext)
 		{
-			FScopedSlowTask SlowTask(0.0f, LOCTEXT("UpdatePackedBlueprintsProgress", "Updating Packed Blueprints..."));
-			TSet<TSoftObjectPtr<UBlueprint>> BlueprintAssets;
-			FPackedLevelActorUtils::GetPackedBlueprintsForWorldAsset(TSoftObjectPtr<UWorld>(WorldAsset.GetSoftObjectPath()), BlueprintAssets, false);
-			TSharedPtr<FPackedLevelActorBuilder> Builder = FPackedLevelActorBuilder::CreateDefaultBuilder();
-			for (TSoftObjectPtr<UBlueprint> BlueprintAsset : BlueprintAssets)
-			{
-				if (UBlueprint* Blueprint = BlueprintAsset.Get())
-				{
-					Builder->UpdateBlueprint(Blueprint, false);
-				}
-			}
+			return FPackedLevelActorUtils::UpdateAllPackedBlueprintsForWorldAssetBlueprint(TSoftObjectPtr<UWorld>(WorldAsset.GetSoftObjectPath()));
 		});
 
 		Section.AddMenuEntry(
@@ -1280,4 +1270,11 @@ void FLevelInstanceEditorModule::AddReferencedObjects(FReferenceCollector& Colle
 	DefaultBehaviorSource.AddReferencedObjects(Collector);
 }
 
+void FLevelInstanceEditorModule::UpdateAllPackedLevelActorsForWorldAsset(const TSoftObjectPtr<UWorld>& InWorldAsset, bool bInLoadedOnly)
+{
+	if (FPackedLevelActorUtils::CanPack())
+	{
+		FPackedLevelActorUtils::UpdateAllPackedBlueprintsForWorldAssetBlueprint(InWorldAsset, bInLoadedOnly);
+	}
+}
 #undef LOCTEXT_NAMESPACE
diff --git a/Engine/Source/Editor/LevelInstanceEditor/Private/LevelInstanceEditorModule.h b/Engine/Source/Editor/LevelInstanceEditor/Private/LevelInstanceEditorModule.h
--- a/Engine/Source/Editor/LevelInstanceEditor/Private/LevelInstanceEditorModule.h
+++ b/Engine/Source/Editor/LevelInstanceEditor/Private/LevelInstanceEditorModule.h
@@ -47,6 +47,9 @@ public:
 
 	virtual void AddReferencedObjects(FReferenceCollector& Collector) override;
 	virtual FString GetReferencerName() const override { return TEXT("LevelInstanceEditorModule"); }
+
+	virtual void UpdateAllPackedLevelActorsForWorldAsset(const TSoftObjectPtr<UWorld>& InWorldAsset, bool bInLoadedOnly = false) override;
+
 private:
 	virtual void UpdateEditorMode(bool bActivated) override;
 	
diff --git a/Engine/Source/Editor/LevelInstanceEditor/Private/PackedLevelActorUtils.cpp b/Engine/Source/Editor/LevelInstanceEditor/Private/PackedLevelActorUtils.cpp
--- a/Engine/Source/Editor/LevelInstanceEditor/Private/PackedLevelActorUtils.cpp
+++ b/Engine/Source/Editor/LevelInstanceEditor/Private/PackedLevelActorUtils.cpp
@@ -178,6 +178,21 @@ void FPackedLevelActorUtils::UpdateBlueprint(UBlueprint* InBlueprint, bool bChec
 	FPackedLevelActorBuilder::CreateDefaultBuilder()->UpdateBlueprint(InBlueprint, bCheckoutAndSave);
 }
 
+void FPackedLevelActorUtils::UpdateAllPackedBlueprintsForWorldAssetBlueprint(const TSoftObjectPtr<UWorld>& InWorldAsset, bool bInLoadedOnly)
+{
+	FScopedSlowTask SlowTask(0.0f, LOCTEXT("UpdatePackedBlueprintsProgress", "Updating Packed Blueprints..."));
+	TSet<TSoftObjectPtr<UBlueprint>> BlueprintAssets;
+	FPackedLevelActorUtils::GetPackedBlueprintsForWorldAsset(InWorldAsset, BlueprintAssets, bInLoadedOnly);
+	TSharedPtr<FPackedLevelActorBuilder> Builder = FPackedLevelActorBuilder::CreateDefaultBuilder();
+	for (TSoftObjectPtr<UBlueprint> BlueprintAsset : BlueprintAssets)
+	{
+		if (UBlueprint* Blueprint = BlueprintAsset.Get())
+		{
+			Builder->UpdateBlueprint(Blueprint, false);
+		}
+	}
+}
+
 #undef LOCTEXT_NAMESPACE
 
 #endif
\ No newline at end of file
diff --git a/Engine/Source/Editor/LevelInstanceEditor/Public/PackedLevelActorUtils.h b/Engine/Source/Editor/LevelInstanceEditor/Public/PackedLevelActorUtils.h
--- a/Engine/Source/Editor/LevelInstanceEditor/Public/PackedLevelActorUtils.h
+++ b/Engine/Source/Editor/LevelInstanceEditor/Public/PackedLevelActorUtils.h
@@ -23,6 +23,8 @@ public:
 	static bool CreateOrUpdateBlueprint(TSoftObjectPtr<UWorld> InWorldAsset, TSoftObjectPtr<UBlueprint> InBlueprintAsset, bool bCheckoutAndSave = true, bool bPromptForSave = true);
 	
 	static void UpdateBlueprint(UBlueprint* InBlueprint, bool bCheckoutAndSave = true);
+
+	static void UpdateAllPackedBlueprintsForWorldAssetBlueprint(const TSoftObjectPtr<UWorld>& InWorldAsset, bool bInLoadedOnly = false);
 };
 
 #endif 
\ No newline at end of file
diff --git a/Engine/Source/Runtime/Engine/Private/LevelInstance/LevelInstanceSubsystem.cpp b/Engine/Source/Runtime/Engine/Private/LevelInstance/LevelInstanceSubsystem.cpp
--- a/Engine/Source/Runtime/Engine/Private/LevelInstance/LevelInstanceSubsystem.cpp
+++ b/Engine/Source/Runtime/Engine/Private/LevelInstance/LevelInstanceSubsystem.cpp
@@ -67,6 +67,12 @@ DEFINE_LOG_CATEGORY(LogLevelInstance);
 
 #if WITH_EDITOR
 bool ULevelInstanceSubsystem::bPrimitiveColorHandlerRegistered = false;
+
+int32 GLevelInstanceShouldAutoUpdatePackedBlueprintsOnCommit = 0;
+static FAutoConsoleVariableRef CVarLevelInstanceShouldAutoUpdatePackedBlueprintsOnCommit(
+	TEXT("levelinstance.ShouldAutoUpdatePackedBlueprintsOnCommit"),
+	GLevelInstanceShouldAutoUpdatePackedBlueprintsOnCommit,
+	TEXT("Enable to automatically update packed blueprints associated with the world asset when committing changes to a level instance."));
 #endif
 
 ULevelInstanceSubsystem::ULevelInstanceSubsystem()
@@ -2680,13 +2686,18 @@ bool ULevelInstanceSubsystem::CommitLevelInstanceInternal(TUniquePtr<FLevelInsta
 	// Notify (Actor might get destroyed by this call if its a packed bp)
 	LevelInstance->OnCommit(bChangesCommitted);
 
+	if (GLevelInstanceShouldAutoUpdatePackedBlueprintsOnCommit && bChangesCommitted)
+	{
+		ILevelInstanceEditorModule& EditorModule = FModuleManager::LoadModuleChecked<ILevelInstanceEditorModule>("LevelInstanceEditor");
+		EditorModule.UpdateAllPackedLevelActorsForWorldAsset(LevelInstance->GetWorldAsset());
+	}
+
 	// Update pointer since BP Compilation might have invalidated LevelInstance
 	LevelInstance = GetLevelInstance(LevelInstanceID);
 
 	// Update Registered Container Bounds
 	UActorDescContainerSubsystem::GetChecked().NotifyContainerUpdated(WorldAssetPackage);
 
-	
 	TMap<ULevelInstanceSubsystem*, TArray<FLevelInstanceID>> LevelInstancesToUpdate;
 	// Gather list to update
 	for (TObjectIterator<UWorld> It(RF_ClassDefaultObject | RF_ArchetypeObject, true); It; ++It)
diff --git a/Engine/Source/Runtime/Engine/Public/LevelInstance/ILevelInstanceEditorModule.h b/Engine/Source/Runtime/Engine/Public/LevelInstance/ILevelInstanceEditorModule.h
--- a/Engine/Source/Runtime/Engine/Public/LevelInstance/ILevelInstanceEditorModule.h
+++ b/Engine/Source/Runtime/Engine/Public/LevelInstance/ILevelInstanceEditorModule.h
@@ -31,6 +31,9 @@ public:
 
 	virtual bool IsEditInPlaceStreamingEnabled() const = 0;
 	virtual bool IsSubSelectionEnabled() const = 0;
+
+	virtual void UpdateAllPackedLevelActorsForWorldAsset(const TSoftObjectPtr<UWorld>& InWorldAsset, bool bInLoadedOnly = false) = 0;
+
 protected:
 	friend class ULevelInstanceSubsystem;
 	friend class ULevelStreamingLevelInstanceEditorPropertyOverride;
