Protect SoundClass volume property from being set negative

diff --git a/Engine/Source/Runtime/Engine/Classes/Sound/SoundClass.h b/Engine/Source/Runtime/Engine/Classes/Sound/SoundClass.h
--- a/Engine/Source/Runtime/Engine/Classes/Sound/SoundClass.h
+++ b/Engine/Source/Runtime/Engine/Classes/Sound/SoundClass.h
@@ -56,7 +56,7 @@ struct FSoundClassProperties
 	GENERATED_USTRUCT_BODY()
 
 	/** Volume multiplier. */
-	UPROPERTY(EditAnywhere, BlueprintReadOnly, Category = General)
+	UPROPERTY(EditAnywhere, BlueprintReadOnly, meta = (ClampMin = "0.0"), Category = General)
 	float Volume;
 
 	/** Pitch multiplier. */
