Harden FCanvasDashedBoxItem against potential infinite loops.

#rnx

diff --git a/Engine/Source/Runtime/Engine/Private/UserInterface/CanvasItem.cpp b/Engine/Source/Runtime/Engine/Private/UserInterface/CanvasItem.cpp
--- a/Engine/Source/Runtime/Engine/Private/UserInterface/CanvasItem.cpp
+++ b/Engine/Source/Runtime/Engine/Private/UserInterface/CanvasItem.cpp
@@ -1618,6 +1618,11 @@ void FCanvasBoxItem::SetupBox()
 void FCanvasDashedBoxItem::Draw(class FCanvas* InCanvas)
 {
 	SCOPE_CYCLE_COUNTER(STAT_Canvas_DashedBoxItemTime);
+
+	if (DashLength <= 0.0f || DashGap <= 0.0f)
+	{
+		return;
+	}
 	
 	FBatchedElements* BatchedElements = InCanvas->GetBatchedElements( FCanvas::ET_Line );
 	FHitProxyId HitProxyId = InCanvas->GetHitProxyId();
