Removing some warnings we get about SetReplicates(true) with ANoPawnPlayerController.

#rnx

diff --git a/Engine/Source/Runtime/Engine/Classes/GameFramework/PlayerController.h b/Engine/Source/Runtime/Engine/Classes/GameFramework/PlayerController.h
--- a/Engine/Source/Runtime/Engine/Classes/GameFramework/PlayerController.h
+++ b/Engine/Source/Runtime/Engine/Classes/GameFramework/PlayerController.h
@@ -2349,6 +2349,7 @@ class ANoPawnPlayerController : public APlayerController
 public:
 
 	GENERATED_BODY()
+	ANoPawnPlayerController(const FObjectInitializer& ObjectInitializer);
 
 	virtual void PostInitializeComponents() override;
 	virtual class AActor* GetViewTarget() const override;
diff --git a/Engine/Source/Runtime/Engine/Private/PlayerController.cpp b/Engine/Source/Runtime/Engine/Private/PlayerController.cpp
--- a/Engine/Source/Runtime/Engine/Private/PlayerController.cpp
+++ b/Engine/Source/Runtime/Engine/Private/PlayerController.cpp
@@ -6631,6 +6631,13 @@ void APlayerController::ClientAckTimeDilation_Implementation(float TimeDilation,
 	}
 }
 
+ANoPawnPlayerController::ANoPawnPlayerController(const FObjectInitializer& ObjectInitializer)
+	: Super(ObjectInitializer)
+{
+	bReplicates = true;
+	SetRole(ROLE_Authority);
+}
+
 void ANoPawnPlayerController::PostInitializeComponents()
 {
 	// Don't call APlayerController::PostInitializeComponents() because it will spawn an instance of APlayerState
diff --git a/Engine/Source/Runtime/Engine/Private/World.cpp b/Engine/Source/Runtime/Engine/Private/World.cpp
--- a/Engine/Source/Runtime/Engine/Private/World.cpp
+++ b/Engine/Source/Runtime/Engine/Private/World.cpp
@@ -7463,8 +7463,6 @@ void UWorld::NotifyControlMessage(UNetConnection* Connection, uint8 MessageType,
 					SpawnInfo.bDeferConstruction = true;
 					
 					Connection->PlayerController = GetWorld()->SpawnActor<APlayerController>(ANoPawnPlayerController::StaticClass(), FVector(), FRotator(), SpawnInfo);
-					Connection->PlayerController->SetRole(ROLE_Authority);
-					Connection->PlayerController->SetReplicates(true);
 					Connection->PlayerController->SetAutonomousProxy(true);
 					Connection->PlayerController->FinishSpawning(FTransform());
 					Connection->PlayerController->NetPlayerIndex = 0;
