Added another transaction to OnSyncBodies to safeguard a crash.

diff --git a/Engine/Source/Runtime/Engine/Private/PhysicsEngine/Experimental/PhysScene_Chaos.cpp b/Engine/Source/Runtime/Engine/Private/PhysicsEngine/Experimental/PhysScene_Chaos.cpp
--- a/Engine/Source/Runtime/Engine/Private/PhysicsEngine/Experimental/PhysScene_Chaos.cpp
+++ b/Engine/Source/Runtime/Engine/Private/PhysicsEngine/Experimental/PhysScene_Chaos.cpp
@@ -2515,13 +2515,23 @@ void FPhysScene_Chaos::OnSyncBodies(Chaos::FPhysicsSolverBase* Solver)
 	};
 
 	FDispatcher Dispatcher;
-	FPhysicsCommand::ExecuteWrite(this, [this, Solver, &Dispatcher](FPhysScene* PhysScene)
+	auto ExecuteCallback = [this, Solver, &Dispatcher]()
 	{
-		Dispatcher.Outer = this;
-		Dispatcher.PhysScene = PhysScene;
-		Dispatcher.Solver = Solver;
-		Solver->PullPhysicsStateForEachDirtyProxy_External(Dispatcher);
-	});
+		FPhysicsCommand::ExecuteWrite(this, [this, Solver, &Dispatcher](FPhysScene* PhysScene)
+		{
+			Dispatcher.Outer = this;
+			Dispatcher.PhysScene = PhysScene;
+			Dispatcher.Solver = Solver;
+			Solver->PullPhysicsStateForEachDirtyProxy_External(Dispatcher);
+		});
+	};
+
+#if UE_WITH_REMOTE_OBJECT_HANDLE
+	static FName TransactionalDispatcherWorkName(TEXT("FPhysScene_Chaos::OnSyncBodies-Dispatcher"));
+	UE::RemoteExecutor::ExecuteTransactional(TransactionalDispatcherWorkName, ExecuteCallback);
+#else
+	ExecuteCallback();
+#endif
 
 	for (const FPhysScenePendingComponentTransform_Chaos& ComponentTransform : Dispatcher.PendingTransforms)
 	{
