Prevent debug HUD from having the same entry multiple times in the DebugDisplay array

diff --git a/Engine/Source/Runtime/Engine/Private/HUD.cpp b/Engine/Source/Runtime/Engine/Private/HUD.cpp
--- a/Engine/Source/Runtime/Engine/Private/HUD.cpp
+++ b/Engine/Source/Runtime/Engine/Private/HUD.cpp
@@ -386,7 +386,8 @@ void AHUD::ShowDebug(FName DebugType)
 		}
 		if (!bRemoved)
 		{
-			DebugDisplay.Add(DebugType);
+			// Since we prevent removing items when bShowDebugInfo is false, add them only once.
+			DebugDisplay.AddUnique(DebugType);
 		}
 
 		bShowDebugInfo = true;
