New material translator:
- Introduced a custom arena allocator to MaterialIRModule to allocate Values and interned strings, instead of allocating these using the global allocator (for compactness and efficiency).
- Refactored the IR-to-HLSL translator and the Emitter to work with StringViews instead of raw const TCHAR*s and FStrings to improve efficiency.

#rnx

diff --git a/Engine/Source/Runtime/Engine/Private/Materials/HLSLMaterialTranslator.cpp b/Engine/Source/Runtime/Engine/Private/Materials/HLSLMaterialTranslator.cpp
--- a/Engine/Source/Runtime/Engine/Private/Materials/HLSLMaterialTranslator.cpp
+++ b/Engine/Source/Runtime/Engine/Private/Materials/HLSLMaterialTranslator.cpp
@@ -5615,17 +5615,17 @@ int32 FHLSLMaterialTranslator::ViewProperty(EMaterialExposedViewProperty Propert
 {
 	const FMaterialExposedViewPropertyMeta& PropertyMeta = MaterialExternalCodeRegistry::Get().GetExternalViewPropertyCode(Property);
 
-	FString Code = PropertyMeta.PropertyCode;
+	FString Code = FString{ PropertyMeta.PropertyCode };
 
-	if (InvProperty && PropertyMeta.InvPropertyCode)
+	if (InvProperty && !PropertyMeta.InvPropertyCode.IsEmpty())
 	{
-		Code = PropertyMeta.InvPropertyCode;
+		Code = FString{ PropertyMeta.InvPropertyCode };
 	}
 
 	// Resolved templated code
 	Code.ReplaceInline(TEXT("<PREV>"), bCompilingPreviousFrame ? TEXT("Prev") : TEXT(""));
 		
-	if (InvProperty && !PropertyMeta.InvPropertyCode)
+	if (InvProperty && PropertyMeta.InvPropertyCode.IsEmpty())
 	{
 		// fall back to compute the property's inverse from PropertyCode
 		return Div(Constant(1.f), AddCodeChunkZeroDeriv(PropertyMeta.Type, *Code));
diff --git a/Engine/Source/Runtime/Engine/Private/Materials/MaterialExpressionsIR.cpp b/Engine/Source/Runtime/Engine/Private/Materials/MaterialExpressionsIR.cpp
--- a/Engine/Source/Runtime/Engine/Private/Materials/MaterialExpressionsIR.cpp
+++ b/Engine/Source/Runtime/Engine/Private/Materials/MaterialExpressionsIR.cpp
@@ -742,7 +742,7 @@ static FValueRef EmitInlineHLSL(MIR::FEmitter& Em, const UMaterialExpressionExte
 	checkf(InExternalCodeIdentifierIndex >= 0 && InExternalCodeIdentifierIndex < InExternalCodeExpression.ExternalCodeIdentifiers.Num(),
 		TEXT("External code identifier index (%d) out of bounds; Upper bound is %d"), InExternalCodeIdentifierIndex, InExternalCodeExpression.ExternalCodeIdentifiers.Num());
 
-	const FName& ExternalCodeIdentifier = InExternalCodeExpression.ExternalCodeIdentifiers[InExternalCodeIdentifierIndex];
+	const FName ExternalCodeIdentifier = InExternalCodeExpression.ExternalCodeIdentifiers[InExternalCodeIdentifierIndex];
 	const FMaterialExternalCodeDeclaration* ExternalCodeDeclaration = MaterialExternalCodeRegistry::Get().FindExternalCode(ExternalCodeIdentifier);
 	if (!ExternalCodeDeclaration)
 	{
@@ -1139,10 +1139,10 @@ void UMaterialExpressionTime::Build(MIR::FEmitter& Em)
 	MIR::FType ScalarFloatType = MIR::FType::MakeScalar(MIR::EScalarKind::Float);
 
 	// When pausing the game is ignored for this time expression, use real-time instead of game-time.
-	const TCHAR* TimeFieldName = bIgnorePause ? TEXT("View.<PREV>RealTime") : TEXT("View.<PREV>GameTime");
 	if (!bOverride_Period)
 	{
-		Em.Output(0, Em.InlineHLSL(ScalarFloatType, TimeFieldName, {}, MIR::EValueFlags::SubstituteTagsInInlineHLSL));
+		FStringView InlinedCode = bIgnorePause ? TEXTVIEW("View.<PREV>RealTime") : TEXTVIEW("View.<PREV>GameTime");
+		Em.Output(0, Em.InlineHLSL(ScalarFloatType, InlinedCode, {}, MIR::EValueFlags::SubstituteTagsInInlineHLSL));
 	}
 	else if (Period == 0.0f)
 	{
@@ -1153,7 +1153,8 @@ void UMaterialExpressionTime::Build(MIR::FEmitter& Em)
 		// Note: Don't use IR intrinsic for Fmod() here to avoid conversion to fp16 on mobile.
 		// We want full 32 bit float precision until the fmod when using a period.
 		FValueRef PeriodValue = Em.ConstantFloat(Period);
-		Em.Output(0, Em.InlineHLSL(ScalarFloatType, FString::Printf(TEXT("fmod(%s,$0)"), TimeFieldName), { PeriodValue }, MIR::EValueFlags::SubstituteTagsInInlineHLSL));
+		FStringView InlinedCode = bIgnorePause ? TEXTVIEW("fmod(View.<PREV>RealTime, $0)") : TEXTVIEW("fmod(View.<PREV>GameTime, $0)");
+		Em.Output(0, Em.InlineHLSL(ScalarFloatType, InlinedCode, { PeriodValue }, MIR::EValueFlags::SubstituteTagsInInlineHLSL));
 	}
 }
 
@@ -1466,10 +1467,10 @@ static FValueRef BuildViewProperty(MIR::FEmitter& Em, EMaterialExposedViewProper
 	const FMaterialExposedViewPropertyMeta& PropertyMeta = MaterialExternalCodeRegistry::Get().GetExternalViewPropertyCode(InProperty);
 	const bool bHasCustomInverseCode = PropertyMeta.InvPropertyCode != nullptr;
 
-	FString HLSLCode = bInvProperty && bHasCustomInverseCode ? PropertyMeta.InvPropertyCode : PropertyMeta.PropertyCode;
+	FStringView HLSLCode = bInvProperty && bHasCustomInverseCode ? PropertyMeta.InvPropertyCode : PropertyMeta.PropertyCode;
 	MIR::FType HLSLCodeType = MIR::FType::FromMaterialValueType(PropertyMeta.Type);
 
-	FValueRef Result = Em.InlineHLSL(HLSLCodeType, MoveTemp(HLSLCode), {}, MIR::EValueFlags::SubstituteTagsInInlineHLSL);
+	FValueRef Result = Em.InlineHLSL(HLSLCodeType, HLSLCode, {}, MIR::EValueFlags::SubstituteTagsInInlineHLSL);
 
 	// CastToNonLWCIfDisabled
 	TOptional<MIR::FPrimitive> PrimitiveType = HLSLCodeType.AsPrimitive();
@@ -1508,7 +1509,7 @@ void UMaterialExpressionCameraPositionWS::Build(MIR::FEmitter& Em)
 
 void UMaterialExpressionPixelNormalWS::Build(MIR::FEmitter& Em)
 {
-	FValueRef Output = Em.InlineHLSL(MIR::FType::MakeFloat(3), "Parameters.WorldNormal", {}, MIR::EValueFlags::None, MIR::EGraphProperties::ReadsPixelNormal);
+	FValueRef Output = Em.InlineHLSL(MIR::FType::MakeFloat(3), TEXTVIEW("Parameters.WorldNormal"), {}, MIR::EValueFlags::None, MIR::EGraphProperties::ReadsPixelNormal);
 	Em.Output(0, Output);
 }
 
diff --git a/Engine/Source/Runtime/Engine/Private/Materials/MaterialExternalCodeRegistry.cpp b/Engine/Source/Runtime/Engine/Private/Materials/MaterialExternalCodeRegistry.cpp
--- a/Engine/Source/Runtime/Engine/Private/Materials/MaterialExternalCodeRegistry.cpp
+++ b/Engine/Source/Runtime/Engine/Private/Materials/MaterialExternalCodeRegistry.cpp
@@ -142,26 +142,26 @@ const FMaterialExposedViewPropertyMeta& MaterialExternalCodeRegistry::GetExterna
 	// Compile time struct storing all EMaterialExposedViewProperty's enumerations' HLSL compilation specific meta information
 	static const FMaterialExposedViewPropertyMeta ViewPropertyMetaArray[] =
 	{
-		{MEVP_BufferSize, MCT_Float2, TEXT("View.BufferSizeAndInvSize.xy"), TEXT("View.BufferSizeAndInvSize.zw")},
-		{MEVP_FieldOfView, MCT_Float2, TEXT("View.<PREV>FieldOfViewWideAngles"), nullptr},
-		{MEVP_TanHalfFieldOfView, MCT_Float2, TEXT("Get<PREV>TanHalfFieldOfView()"), TEXT("Get<PREV>CotanHalfFieldOfView()")},
-		{MEVP_ViewSize, MCT_Float2, TEXT("View.ViewSizeAndInvSize.xy"), TEXT("View.ViewSizeAndInvSize.zw")},
-		{MEVP_WorldSpaceViewPosition, MCT_LWCVector3, TEXT("Get<PREV>WorldViewOrigin(Parameters)"), nullptr},
-		{MEVP_WorldSpaceCameraPosition, MCT_LWCVector3, TEXT("Get<PREV>WorldCameraOrigin(Parameters)"), nullptr},
-		{MEVP_ViewportOffset, MCT_Float2, TEXT("View.ViewRectMin.xy"), nullptr},
-		{MEVP_TemporalSampleCount, MCT_Float1, TEXT("View.TemporalAAParams.y"), nullptr},
-		{MEVP_TemporalSampleIndex, MCT_Float1, TEXT("View.TemporalAAParams.x"), nullptr},
-		{MEVP_TemporalSampleOffset, MCT_Float2, TEXT("View.TemporalAAParams.zw"), nullptr},
-		{MEVP_RuntimeVirtualTextureOutputLevel, MCT_Float1, TEXT("GetRuntimeVirtualTextureMipLevel().x"), nullptr},
-		{MEVP_RuntimeVirtualTextureOutputDerivative, MCT_Float2, TEXT("GetRuntimeVirtualTextureMipLevel().zw"), nullptr},
-		{MEVP_PreExposure, MCT_Float1, TEXT("View.PreExposure.x"), TEXT("View.OneOverPreExposure.x")},
-		{MEVP_RuntimeVirtualTextureMaxLevel, MCT_Float1, TEXT("GetRuntimeVirtualTextureMipLevel().y"), nullptr},
-		{MEVP_ResolutionFraction, MCT_Float1, TEXT("View.ResolutionFractionAndInv.x"), TEXT("View.ResolutionFractionAndInv.y")},
-		{MEVP_PostVolumeUserFlags, MCT_Float1, TEXT("View.PostVolumeUserFlags"), nullptr},
-		{MEVP_FirstPersonFieldOfView, MCT_Float2, TEXT("View.<PREV>FirstPersonFieldOfViewWideAngles"), nullptr},
-		{MEVP_FirstPersonTanHalfFieldOfView, MCT_Float2, TEXT("View.<PREV>FirstPersonTanAndInvTanHalfFOV.xy"), TEXT("View.<PREV>FirstPersonTanAndInvTanHalfFOV.zw")},
-		{MEVP_FirstPersonScale, MCT_Float, TEXT("View.<PREV>FirstPersonScale"), nullptr},
-		{MEVP_NearPlane, MCT_Float, TEXT("View.NearPlane"), nullptr},
+		{ MEVP_BufferSize, MCT_Float2, TEXTVIEW("View.BufferSizeAndInvSize.xy"), TEXTVIEW("View.BufferSizeAndInvSize.zw")},
+		{ MEVP_FieldOfView, MCT_Float2, TEXTVIEW("View.<PREV>FieldOfViewWideAngles"), {} },
+		{ MEVP_TanHalfFieldOfView, MCT_Float2, TEXTVIEW("Get<PREV>TanHalfFieldOfView()"), TEXTVIEW("Get<PREV>CotanHalfFieldOfView()")},
+		{ MEVP_ViewSize, MCT_Float2, TEXTVIEW("View.ViewSizeAndInvSize.xy"), TEXTVIEW("View.ViewSizeAndInvSize.zw")},
+		{ MEVP_WorldSpaceViewPosition, MCT_LWCVector3, TEXTVIEW("Get<PREV>WorldViewOrigin(Parameters)"), {} },
+		{ MEVP_WorldSpaceCameraPosition, MCT_LWCVector3, TEXTVIEW("Get<PREV>WorldCameraOrigin(Parameters)"), {} },
+		{ MEVP_ViewportOffset, MCT_Float2, TEXTVIEW("View.ViewRectMin.xy"), {} },
+		{ MEVP_TemporalSampleCount, MCT_Float1, TEXTVIEW("View.TemporalAAParams.y"), {} },
+		{ MEVP_TemporalSampleIndex, MCT_Float1, TEXTVIEW("View.TemporalAAParams.x"), {} },
+		{ MEVP_TemporalSampleOffset, MCT_Float2, TEXTVIEW("View.TemporalAAParams.zw"), {} },
+		{ MEVP_RuntimeVirtualTextureOutputLevel, MCT_Float1, TEXTVIEW("GetRuntimeVirtualTextureMipLevel().x"), {} },
+		{ MEVP_RuntimeVirtualTextureOutputDerivative, MCT_Float2, TEXTVIEW("GetRuntimeVirtualTextureMipLevel().zw"), {} },
+		{ MEVP_PreExposure, MCT_Float1, TEXTVIEW("View.PreExposure.x"), TEXTVIEW("View.OneOverPreExposure.x")},
+		{ MEVP_RuntimeVirtualTextureMaxLevel, MCT_Float1, TEXTVIEW("GetRuntimeVirtualTextureMipLevel().y"), {} },
+		{ MEVP_ResolutionFraction, MCT_Float1, TEXTVIEW("View.ResolutionFractionAndInv.x"), TEXTVIEW("View.ResolutionFractionAndInv.y")},
+		{ MEVP_PostVolumeUserFlags, MCT_Float1, TEXTVIEW("View.PostVolumeUserFlags"), {} },
+		{ MEVP_FirstPersonFieldOfView, MCT_Float2, TEXTVIEW("View.<PREV>FirstPersonFieldOfViewWideAngles"), {} },
+		{ MEVP_FirstPersonTanHalfFieldOfView, MCT_Float2, TEXTVIEW("View.<PREV>FirstPersonTanAndInvTanHalfFOV.xy"), TEXTVIEW("View.<PREV>FirstPersonTanAndInvTanHalfFOV.zw")},
+		{ MEVP_FirstPersonScale, MCT_Float, TEXTVIEW("View.<PREV>FirstPersonScale"), {} },
+		{ MEVP_NearPlane, MCT_Float, TEXTVIEW("View.NearPlane"), {} },
 	};
 	static_assert((sizeof(ViewPropertyMetaArray) / sizeof(ViewPropertyMetaArray[0])) == MEVP_MAX, "incoherency between EMaterialExposedViewProperty and ViewPropertyMetaArray");
 
diff --git a/Engine/Source/Runtime/Engine/Private/Materials/MaterialExternalCodeRegistry.h b/Engine/Source/Runtime/Engine/Private/Materials/MaterialExternalCodeRegistry.h
--- a/Engine/Source/Runtime/Engine/Private/Materials/MaterialExternalCodeRegistry.h
+++ b/Engine/Source/Runtime/Engine/Private/Materials/MaterialExternalCodeRegistry.h
@@ -170,8 +170,8 @@ struct FMaterialExposedViewPropertyMeta
 {
 	EMaterialExposedViewProperty EnumValue;
 	EMaterialValueType Type;
-	const TCHAR* PropertyCode;
-	const TCHAR* InvPropertyCode;
+	FStringView PropertyCode;
+	FStringView InvPropertyCode;
 };
 
 /** Singleton class to register external HLSL function and input declarations for material IR modules. */
diff --git a/Engine/Source/Runtime/Engine/Private/Materials/MaterialIREmitter.cpp b/Engine/Source/Runtime/Engine/Private/Materials/MaterialIREmitter.cpp
--- a/Engine/Source/Runtime/Engine/Private/Materials/MaterialIREmitter.cpp
+++ b/Engine/Source/Runtime/Engine/Private/Materials/MaterialIREmitter.cpp
@@ -82,17 +82,17 @@ struct FEmitter::FPrivate
 	}
 
 	// Allocates a block of memory.
-	static uint8* Allocate(FEmitter& Emitter, int32 Size, int32 Alignment)
+	static void* Allocate(FEmitter& Emitter, int32 Size, int32 Alignment)
 	{
-		// TODO: use a linear allocator inside the module instead.
-		uint8* Bytes = (uint8*)FMemory::Malloc(Size, Alignment); 
+		void* Bytes = (void*)Emitter.Module->Allocator.PushBytes(Size, Alignment);
 		FMemory::Memzero(Bytes, Size);
 		return Bytes;
 	}
 
+	//
 	static void Deallocate(FEmitter& Emitter, void* Ptr)
 	{
-		FMemory::Free(Ptr);
+		// This is a NOOP, only there for code robustness.
 	}
 
 	// Pushes a new value to the list of values.
@@ -108,7 +108,7 @@ static FComposite* NewCompositeValue(FEmitter& Emitter, FType Type, int NumCompo
 {
 	int32 SizeInBytes = sizeof(FComposite) + sizeof(FValue*) * NumComponents;
 
-	uint8* Bytes = FEmitter::FPrivate::Allocate(Emitter, SizeInBytes, alignof(FComposite));
+	void* Bytes = FEmitter::FPrivate::Allocate(Emitter, SizeInBytes, alignof(FComposite));
 
 	FComposite* Value = new (Bytes) FComposite{};
 	Value->Kind = VK_Composite;
@@ -155,7 +155,7 @@ static FValueRef EmitPrototype(FEmitter& Emitter, const TValueType& Prototype)
 		return Existing;
 	}
 
-	uint8* Bytes = FEmitter::FPrivate::Allocate(Emitter, sizeof(TValueType), alignof(TValueType));
+	void* Bytes = FEmitter::FPrivate::Allocate(Emitter, sizeof(TValueType), alignof(TValueType));
 	TValueType* Value = new (Bytes) TValueType{ Prototype };
 
 	FEmitter::FPrivate::PushNewValue(Emitter, Value);
@@ -2476,7 +2476,7 @@ static FValue* CastConstant(FEmitter& Emitter, FConstant* Constant, EScalarKind
 	UE_MIR_UNREACHABLE();
 }
 
-static FValue* CastValueToPrimitiveType(FEmitter& Emitter, FValueRef Value, FType TargetType)
+static FValue* CastToPrimitive(FEmitter& Emitter, FValueRef Value, FType TargetType)
 {
 	if (!Value->Type.AsPrimitive())
 	{
@@ -2616,7 +2616,7 @@ FValueRef FEmitter::Cast(FValueRef Value, FType TargetType)
 	FValue* Result;
 	if (TargetType.AsPrimitive())
 	{
-		Result = CastValueToPrimitiveType(*this, Value, TargetType);
+		Result = CastToPrimitive(*this, Value, TargetType);
 	}
 	else
 	{
@@ -3250,7 +3250,7 @@ FValueRef FEmitter::AnalyticalPartialDerivative(FValueRef Value, EDerivativeAxis
 	}
 }
 
-static FValueRef EmitInlineHLSL(FEmitter& Emitter, FType Type, const FMaterialExternalCodeDeclaration* InExternalCodeDeclaration, const TCHAR* Code, TConstArrayView<FValueRef> InputValues, EValueFlags ValueFlags, EGraphProperties UsedGraphProperties)
+static FValueRef EmitInlineHLSL(FEmitter& Emitter, FType Type, const FMaterialExternalCodeDeclaration* InExternalCodeDeclaration, FStringView Code, TConstArrayView<FValueRef> InputValues, EValueFlags ValueFlags, EGraphProperties UsedGraphProperties)
 {
 	FInlineHLSL Prototype = MakePrototype<FInlineHLSL>(Type);
 	Prototype.Type = Type;
@@ -3259,12 +3259,13 @@ static FValueRef EmitInlineHLSL(FEmitter& Emitter, FType Type, const FMaterialEx
 	
 	if (InExternalCodeDeclaration)
 	{
-		check(!Code);
+		check(Code.IsEmpty());
 		Prototype.ExternalCodeDeclaration = InExternalCodeDeclaration;
 	}
 	else
 	{
-		Prototype.Code = Code;
+		Prototype.Code = Code.GetData();
+		Prototype.CodeLen = Code.Len();
 	}
 
 	if (!InputValues.IsEmpty())
@@ -3281,14 +3282,14 @@ static FValueRef EmitInlineHLSL(FEmitter& Emitter, FType Type, const FMaterialEx
 	return EmitPrototype(Emitter, Prototype);
 }
 
-FValueRef FEmitter::InlineHLSL(FType Type, FString Code, TConstArrayView<FValueRef> InputValues, EValueFlags ValueFlags, EGraphProperties UsedGraphProperties)
+FValueRef FEmitter::InlineHLSL(FType Type, FStringView Code, TConstArrayView<FValueRef> InputValues, EValueFlags ValueFlags, EGraphProperties UsedGraphProperties)
 {
 	if (IsAnyNotValid(InputValues))
 	{
 		return Poison();
 	}
 	
-	return EmitInlineHLSL(*this, Type, nullptr, Module->PushUserString(MoveTemp(Code)), InputValues, ValueFlags | EValueFlags::HasDynamicHLSLCode, UsedGraphProperties);
+	return EmitInlineHLSL(*this, Type, nullptr, Module->InternString(Code), InputValues, ValueFlags | EValueFlags::HasDynamicHLSLCode, UsedGraphProperties);
 }
 
 FValueRef FEmitter::InlineHLSL(const FMaterialExternalCodeDeclaration* InExternalCodeDeclaration, TConstArrayView<FValueRef> InputValues, EValueFlags ValueFlags, EGraphProperties UsedGraphProperties)
@@ -3300,7 +3301,7 @@ FValueRef FEmitter::InlineHLSL(const FMaterialExternalCodeDeclaration* InExterna
 
 	check(InExternalCodeDeclaration != nullptr);
 	FType ReturnType = FType::FromMaterialValueType(InExternalCodeDeclaration->GetReturnTypeValue());
-	return EmitInlineHLSL(*this, ReturnType, InExternalCodeDeclaration, nullptr, InputValues, ValueFlags, UsedGraphProperties);
+	return EmitInlineHLSL(*this, ReturnType, InExternalCodeDeclaration, {}, InputValues, ValueFlags, UsedGraphProperties);
 }
 
 FValueRef FEmitter::PromoteSubstrateParameter()
diff --git a/Engine/Source/Runtime/Engine/Private/Materials/MaterialIRModule.cpp b/Engine/Source/Runtime/Engine/Private/Materials/MaterialIRModule.cpp
--- a/Engine/Source/Runtime/Engine/Private/Materials/MaterialIRModule.cpp
+++ b/Engine/Source/Runtime/Engine/Private/Materials/MaterialIRModule.cpp
@@ -17,6 +17,7 @@ FMaterialIRModule::~FMaterialIRModule()
 
 void FMaterialIRModule::Empty()
 {
+	Allocator.Flush();
 	FMemory::Memset(RootBlocks, 0, sizeof(RootBlocks));
 
 	for (int32 i = 0; i < MIR::NumStages; ++i)
@@ -24,10 +25,6 @@ void FMaterialIRModule::Empty()
 		Outputs[i].Empty();
 	}
 
-	for (MIR::FValue* Value : Values)
-	{
-		FMemory::Free(Value);
-	}
 	Values.Empty();
 
 	// Reset module statistics.
@@ -40,10 +37,18 @@ void FMaterialIRModule::Empty()
 	Statistics.NumPixelTexCoords = 0;
 }
 
-const TCHAR* FMaterialIRModule::PushUserString(FString InString)
+FStringView FMaterialIRModule::InternString(FStringView InString)
 {
-	UserStrings.Add(MoveTemp(InString));
-	return GetData(UserStrings.Last());
+	// Allocate the buffer to contain the string
+	FStringView::ElementType* String = (FStringView::ElementType*)Allocator.PushBytes(InString.NumBytes() + sizeof(FStringView::ElementType), alignof(FStringView::ElementType));
+
+	// Copy the string over to the interned location
+	FMemory::Memcpy(String, InString.GetData(), InString.NumBytes());
+
+	// Mark the null-character at the end
+	String[InString.Len()] = 0;
+
+	return { String, InString.Len() };
 }
 
 void FMaterialIRModule::AddError(UMaterialExpression* Expression, FString Message)
diff --git a/Engine/Source/Runtime/Engine/Private/Materials/MaterialIRToHLSLTranslator.cpp b/Engine/Source/Runtime/Engine/Private/Materials/MaterialIRToHLSLTranslator.cpp
--- a/Engine/Source/Runtime/Engine/Private/Materials/MaterialIRToHLSLTranslator.cpp
+++ b/Engine/Source/Runtime/Engine/Private/Materials/MaterialIRToHLSLTranslator.cpp
@@ -34,27 +34,13 @@ enum EListSeparator { ListSeparator };
 
 #define TAB "    "
 
-static const char* GVector4SwizzleSubset[4] = { ".x", ".xy", ".xyz", "" };
-
 struct FHLSLPrinter
 {
 	FString Buffer;
 	bool bFirstListItem = false;
 	int32 Tabs = 0;
 
-	FHLSLPrinter& operator<<(const TCHAR* Text)
-	{
-		Buffer.Append(Text);
-		return *this;
-	}
-	
-	FHLSLPrinter& operator<<(const FString& Text)
-	{
-		Buffer.Append(Text);
-		return *this;
-	}
-
-	FHLSLPrinter& operator<<(const FStringView& Text)
+	FHLSLPrinter& operator<<(FStringView Text)
 	{
 		Buffer.Append(Text);
 		return *this;
@@ -76,11 +62,11 @@ struct FHLSLPrinter
 	{
 		if (FGenericPlatformMath::IsNaN(Value))
 		{
-			Buffer.Append(TEXT("(0.0f / 0.0f)"));
+			Buffer.Append(TEXTVIEW("(0.0f / 0.0f)"));
 		}
 		else if (!FGenericPlatformMath::IsFinite(Value))
 		{
-			Buffer.Append(TEXT("INFINITE_FLOAT"));
+			Buffer.Append(TEXTVIEW("INFINITE_FLOAT"));
 		}
 		else
 		{
@@ -119,7 +105,7 @@ struct FHLSLPrinter
 
     FHLSLPrinter& operator<<(EOpenBrace)
     {
-        Buffer.Append("{");
+		Buffer.AppendChar('{');
         ++Tabs;
         *this << NewLine;
         return *this;
@@ -161,7 +147,7 @@ struct FHLSLPrinter
 	{
 		if (!bFirstListItem)
 		{
-			Buffer.Append(TEXT(", "));
+			Buffer.Append(TEXTVIEW(", "));
 		}
 		bFirstListItem = false;
 	}
@@ -266,17 +252,35 @@ static bool IsOperatorInfix(MIR::EOperator Op)
 	}
 }
 
-struct FTranslator : FMaterialIRToHLSLTranslation
+static FStringView GVector4SwizzleSubset[4] = { TEXTVIEW(".x"), TEXTVIEW(".xy"), TEXTVIEW(".xyz"), {} };
+struct FPrivate : FMaterialIRToHLSLTranslation
 {
+	// Utility for emitting HLSL code
+	FHLSLPrinter Printer; 
+
+	// Number of local variables generated during translation.
 	int32 NumLocals{};
-	TMap<const MIR::FInstruction*, FString> LocalIdentifier;
+	
+	// Mapping from instructions to their corresponding local variable index
+	// Used to map instructions that have more than one use to a local identifier like "_42".
+	TMap<const MIR::FInstruction*, uint32> InstrToLocalIndex;
+
+	// Current material stage being translated (e.g., vertex, pixel)
 	MIR::EStage CurrentStage;
-	FHLSLPrinter Printer;
+
+	// Generated HLSL for pixel attribute evaluation
 	FString PixelAttributesHLSL;
+	
+	// Generated HLSL for world position offset calculation
 	FString WorldPositionOffsetHLSL;
+
+	// HLSL code for evaluating normal material attributes per stage
 	FString EvaluateNormalMaterialAttributeHLSL[MIR::NumStages];
+	
+	// HLSL code for evaluating non-normal material attributes per stage
 	FString EvaluateOtherMaterialAttributesHLSL[MIR::NumStages];
 
+
 	void GeneratePixelAttributesHLSL()
 	{
 		for (int32 PropertyIndex = 0; PropertyIndex < MP_MAX; ++PropertyIndex)
@@ -290,11 +294,11 @@ struct FTranslator : FMaterialIRToHLSLTranslation
 			check(FMaterialAttributeDefinitionMap::GetShaderFrequency(Property) == SF_Pixel);
 			
 			// Special case MP_SubsurfaceColor as the actual property is a combination of the color and the profile but we don't want to expose the profile
-			FString PropertyName = (Property == MP_SubsurfaceColor) ? "Subsurface" : FMaterialAttributeDefinitionMap::GetAttributeName(Property);
+			FStringView PropertyName = (Property == MP_SubsurfaceColor) ? TEXTVIEW("Subsurface") : FStringView{ FMaterialAttributeDefinitionMap::GetAttributeName(Property) };
 			EMaterialValueType Type = (Property == MP_SubsurfaceColor) ? MCT_Float4 : FMaterialAttributeDefinitionMap::GetValueType(Property);
 			check(PropertyName.Len() > 0);
 
-			PixelAttributesHLSL.Appendf(TEXT(TAB "%s %s;\n"), GetHLSLTypeString(Type), *PropertyName);
+			PixelAttributesHLSL.Appendf(TEXT(TAB "%s %s;\n"), GetHLSLTypeString(Type), PropertyName.GetData());
 		}
 	}
 	
@@ -315,24 +319,24 @@ struct FTranslator : FMaterialIRToHLSLTranslation
 
 		if (Stage == MIR::Stage_Pixel)
 		{
-			Printer << TEXT("PixelMaterialInputs.Subsurface = 0") << EndOfStatement;
+			Printer << TEXTVIEW("PixelMaterialInputs.Subsurface = 0") << EndOfStatement;
 
 			if (Substrate::IsSubstrateEnabled())
 			{
-				Printer << TEXT("// Substrate normal and tangent") << NewLine;
-				Printer << TEXT("Parameters.SharedLocalBases.Normals[0] = TransformTangentNormalToWorld(Parameters.TangentToWorld, PixelMaterialInputs.Normal)") << EndOfStatement;
-				Printer << TEXT("#if MATERIAL_TANGENTSPACENORMAL") << NewLine;
-				Printer << TEXT("Parameters.SharedLocalBases.Normals[0] *= Parameters.TwoSidedSign") << EndOfStatement;
-				Printer << TEXT("#endif") << NewLine;
-				Printer << TEXT("Parameters.SharedLocalBases.Tangents[0] = Parameters.TangentToWorld[0]") << EndOfStatement;
-				Printer << TEXT("#if MATERIAL_TANGENTSPACENORMAL") << NewLine;
-				Printer << TEXT("Parameters.SharedLocalBases.Tangents[0] *= Parameters.TwoSidedSign") << EndOfStatement;
-				Printer << TEXT("#endif") << NewLine;
-				Printer << TEXT("Parameters.SharedLocalBases.Count = 1") << EndOfStatement;
+				Printer << TEXTVIEW("// Substrate normal and tangent") << NewLine;
+				Printer << TEXTVIEW("Parameters.SharedLocalBases.Normals[0] = TransformTangentNormalToWorld(Parameters.TangentToWorld, PixelMaterialInputs.Normal)") << EndOfStatement;
+				Printer << TEXTVIEW("#if MATERIAL_TANGENTSPACENORMAL") << NewLine;
+				Printer << TEXTVIEW("Parameters.SharedLocalBases.Normals[0] *= Parameters.TwoSidedSign") << EndOfStatement;
+				Printer << TEXTVIEW("#endif") << NewLine;
+				Printer << TEXTVIEW("Parameters.SharedLocalBases.Tangents[0] = Parameters.TangentToWorld[0]") << EndOfStatement;
+				Printer << TEXTVIEW("#if MATERIAL_TANGENTSPACENORMAL") << NewLine;
+				Printer << TEXTVIEW("Parameters.SharedLocalBases.Tangents[0] *= Parameters.TwoSidedSign") << EndOfStatement;
+				Printer << TEXTVIEW("#endif") << NewLine;
+				Printer << TEXTVIEW("Parameters.SharedLocalBases.Count = 1") << EndOfStatement;
 			}
 			else
 			{
-				Printer << TEXT("PixelMaterialInputs.FrontMaterial = GetInitialisedSubstrateData()") << EndOfStatement;
+				Printer << TEXTVIEW("PixelMaterialInputs.FrontMaterial = GetInitialisedSubstrateData()") << EndOfStatement;
 			}
 		}
 
@@ -359,21 +363,26 @@ struct FTranslator : FMaterialIRToHLSLTranslation
             
 			if (Instr->NumUsers[CurrentStage] >= 1)
 			{
-                FString LocalStr = FString::Printf(TEXT("_%d"), NumLocals);
-                ++NumLocals;
-
-                Printer << LowerType(Instr->Type) << TEXT(" ") << LocalStr;
+				// Allocate a new local index for this instruction
+				uint32 LocalIndex = NumLocals++;
+				
+				// Remember the mapping between this instruction and its local index
+                InstrToLocalIndex.Add(Instr, LocalIndex);
 
-                LocalIdentifier.Add(Instr, MoveTemp(LocalStr));
+				// Print the local declaration "<Type> _<LocalIndex>", e.g. "float4 _3"
+                Printer << LowerType(Instr->Type) << TEXTVIEW(" _") << LocalIndex;
+				
+				// If this instruction doesn't use a phi value, we'll immediately assign its local to its result so output the "="
+				// E.g., if Instr is a branch, we will set its value inside the generate "if" {} scopes, so no need for a "=" now.
                 if (!InstructionUsesPhiValue(Instr, CurrentStage))
 				{
-                    Printer << TEXT(" = ");
+                    Printer << TEXTVIEW(" = ");
                 }
             }
 
 			LowerInstruction(Instr);
 
-			if (Printer.Buffer.EndsWith(TEXT("}")))
+			if (Printer.Buffer.EndsWith(TEXTVIEW("}")))
 			{
 				Printer << NewLine;
 			}
@@ -406,7 +415,7 @@ struct FTranslator : FMaterialIRToHLSLTranslation
 			}
 			else
 			{
-				Printer << LocalIdentifier[Instr];
+				Printer << TEXTVIEW("_") << InstrToLocalIndex[Instr];
 			}
 
 			return NoOp;
@@ -454,7 +463,7 @@ struct FTranslator : FMaterialIRToHLSLTranslation
 		switch (checkNoEntry->ScalarKind)
 		{
 			case MIR::EScalarKind::Bool:
-				Printer.Buffer.Append(Constant->Boolean ? TEXT("true") : TEXT("false"));
+				Printer << (Constant->Boolean ? TEXTVIEW("true") : TEXTVIEW("false"));
 				break;
 				
 			case MIR::EScalarKind::Int:
@@ -468,7 +477,7 @@ struct FTranslator : FMaterialIRToHLSLTranslation
 			case MIR::EScalarKind::Double:
 			{
 				FLargeWorldRenderScalar ConstantLWC(Constant->Double);
-				Printer << "MakeLWCScalar(" << ConstantLWC.GetTile() << ", " << ConstantLWC.GetOffset() << ")";
+				Printer << TEXTVIEW("MakeLWCScalar(") << ConstantLWC.GetTile() << TEXTVIEW(", ") << ConstantLWC.GetOffset() << TEXTVIEW(")");
 				break;
 			}
 		}
@@ -483,17 +492,17 @@ struct FTranslator : FMaterialIRToHLSLTranslation
 		if (MIR::IsExternalInputTexCoord(ExternalInput->Id))
 		{
 			int32 Index = ExternalInputIndex - (int32)MIR::EExternalInput::TexCoord0;
-			Printer.Buffer.Appendf(TEXT("Parameters.TexCoords[%d]"), Index);
+			Printer << TEXTVIEW("Parameters.TexCoords[") << Index << TEXTVIEW("]");
 		}
 		else if (MIR::IsExternalInputTexCoordDdx(ExternalInput->Id))
 		{
 			int32 Index = ExternalInputIndex - (int32)MIR::EExternalInput::TexCoord0_Ddx;
-			Printer.Buffer.Appendf(TEXT("Parameters.TexCoords_DDX[%d]"), Index);
+			Printer << TEXTVIEW("Parameters.TexCoords_DDX[") << Index << TEXTVIEW("]");
 		}
 		else if (MIR::IsExternalInputTexCoordDdy(ExternalInput->Id))
 		{
 			int32 Index = ExternalInputIndex - (int32)MIR::EExternalInput::TexCoord0_Ddy;
-			Printer.Buffer.Appendf(TEXT("Parameters.TexCoords_DDY[%d]"), Index);
+			Printer << TEXTVIEW("Parameters.TexCoords_DDY[") << Index << TEXTVIEW("]");
 		}
 		else if (ExternalInput->Id >= MIR::EExternalInput::WorldPosition_Absolute && ExternalInput->Id <= MIR::EExternalInput::WorldPosition_CameraRelativeNoOffsets)
 		{
@@ -503,117 +512,117 @@ struct FTranslator : FMaterialIRToHLSLTranslation
 			// 
 			// Format:  Get[Prev][Translated]WorldPosition[_NoMaterialOffsets](Parameters)
 			//
-			Printer.Buffer.Append(TEXT("Get"));
+			Printer << TEXTVIEW("Get");
 
 			// If compiling for the previous frame in the vertex shader
 			if (CurrentStage == MIR::EStage::Stage_Vertex && bCompilingPreviousFrame)
 			{
-				Printer.Buffer.Append(TEXT("Prev"));
+				Printer << TEXTVIEW("Prev");
 			}
 
 			if (ExternalInput->Id == MIR::EExternalInput::WorldPosition_CameraRelative || ExternalInput->Id == MIR::EExternalInput::WorldPosition_CameraRelativeNoOffsets)
 			{
-				Printer.Buffer.Append(TEXT("Translated"));
+				Printer << TEXTVIEW("Translated");
 			}
 
-			Printer.Buffer.Append(TEXT("WorldPosition"));
+			Printer << TEXTVIEW("WorldPosition");
 
 			if (CurrentStage == MIR::EStage::Stage_Pixel && (ExternalInput->Id == MIR::EExternalInput::WorldPosition_AbsoluteNoOffsets || ExternalInput->Id == MIR::EExternalInput::WorldPosition_CameraRelativeNoOffsets))
 			{
-				Printer.Buffer.Append(TEXT("_NoMaterialOffsets"));
+				Printer << TEXTVIEW("_NoMaterialOffsets");
 			}
 
-			Printer.Buffer.Append(TEXT("(Parameters)"));
+			Printer << TEXTVIEW("(Parameters)");
 		}
 		else if (ExternalInput->Id >= MIR::EExternalInput::WorldPosition_Absolute_Ddx && ExternalInput->Id <= MIR::EExternalInput::WorldPosition_CameraRelativeNoOffsets_Ddx)
 		{
 			// Expression emitter assumes these to be LWC for analytic derivative evaluation, so we promote them on load
-			Printer.Buffer.Append(TEXT("WSPromote(Parameters.WorldPosition_DDX)"));
+			Printer << TEXTVIEW("WSPromote(Parameters.WorldPosition_DDX)");
 		}
 		else if (ExternalInput->Id >= MIR::EExternalInput::WorldPosition_Absolute_Ddy && ExternalInput->Id <= MIR::EExternalInput::WorldPosition_CameraRelativeNoOffsets_Ddy)
 		{
 			// Expression emitter assumes these to be LWC for analytic derivative evaluation, so we promote them on load
-			Printer.Buffer.Append(TEXT("WSPromote(Parameters.WorldPosition_DDY)"));
+			Printer << TEXTVIEW("WSPromote(Parameters.WorldPosition_DDY)");
 		}
 		else if (ExternalInput->Id == MIR::EExternalInput::ActorPosition_Absolute)
 		{
 			const bool bPrevious = (CurrentStage == MIR::EStage::Stage_Vertex && bCompilingPreviousFrame);
-			Printer.Buffer.Append(bPrevious ? TEXT("GetPreviousActorWorldPosition(Parameters)") : TEXT("GetActorWorldPosition(Parameters)"));
+			Printer << (bPrevious ? TEXTVIEW("GetPreviousActorWorldPosition(Parameters)") : TEXTVIEW("GetActorWorldPosition(Parameters)"));
 		}
 		else if (ExternalInput->Id == MIR::EExternalInput::ActorPosition_CameraRelative)
 		{
 			const bool bPrevious = (CurrentStage == MIR::EStage::Stage_Vertex && bCompilingPreviousFrame);
-			Printer.Buffer.Append(bPrevious ? TEXT("GetPreviousActorTranslatedWorldPosition(Parameters)") : TEXT("GetActorTranslatedWorldPosition(Parameters)"));
+			Printer << (bPrevious ? TEXTVIEW("GetPreviousActorTranslatedWorldPosition(Parameters)") : TEXTVIEW("GetActorTranslatedWorldPosition(Parameters)"));
 		}
 		else
 		{
-			const TCHAR* Code{};
+			FStringView Code;
 			switch (ExternalInput->Id)
 			{
-				case MIR::EExternalInput::ObjectPosition_Absolute: Code = TEXT("GetObjectWorldPosition(Parameters)"); break;
-				case MIR::EExternalInput::ObjectPosition_CameraRelative: Code = TEXT("GetObjectTranslatedWorldPosition(Parameters)"); break;
-				case MIR::EExternalInput::ViewMaterialTextureMipBias: Code = TEXT("View.MaterialTextureMipBias"); break;
-				case MIR::EExternalInput::ViewMaterialTextureDerivativeMultiply: Code = TEXT("View.MaterialTextureDerivativeMultiply"); break;
+				case MIR::EExternalInput::ObjectPosition_Absolute: Code = TEXTVIEW("GetObjectWorldPosition(Parameters)"); break;
+				case MIR::EExternalInput::ObjectPosition_CameraRelative: Code = TEXTVIEW("GetObjectTranslatedWorldPosition(Parameters)"); break;
+				case MIR::EExternalInput::ViewMaterialTextureMipBias: Code = TEXTVIEW("View.MaterialTextureMipBias"); break;
+				case MIR::EExternalInput::ViewMaterialTextureDerivativeMultiply: Code = TEXTVIEW("View.MaterialTextureDerivativeMultiply"); break;
 				default: UE_MIR_UNREACHABLE();
 			}
 			Printer << Code;
 		}
 	}
 
-	void LowerInlineHLSLWithArgumentsInternal(const FStringView& Code, const TConstArrayView<MIR::FValue*>& InArguments)
+	void LowerInlineHLSLWithArgumentsInternal(FStringView Code, const TConstArrayView<MIR::FValue*>& InArguments)
 	{
 		// Substitute argument tokens with instruction arguments
 		int32 ArgumentTokenStart = 0, ArgumentTokenEnd = 0;
 
 		// Prints the pending substring of the code and moves the range forward.
 		auto FlushCodeSubstring = [this, &ArgumentTokenEnd, &Code](int32 EndIndex) -> void
+		{
+			if (EndIndex > ArgumentTokenEnd)
 			{
-				if (EndIndex > ArgumentTokenEnd)
-				{
-					Printer << Code.Mid(ArgumentTokenEnd, EndIndex - ArgumentTokenEnd);
-					ArgumentTokenEnd = EndIndex;
-				}
-			};
+				Printer << Code.Mid(ArgumentTokenEnd, EndIndex - ArgumentTokenEnd);
+				ArgumentTokenEnd = EndIndex;
+			}
+		};
 
 		auto SubstituteNextArgument = [this, &ArgumentTokenStart, &ArgumentTokenEnd, &Code, &InArguments]() -> void
-			{
-				// Scan digits for argument index
-				int32 ArgumentIndexValue = 0;
-				int32 NumDigits = 0;
+		{
+			// Scan digits for argument index
+			int32 ArgumentIndexValue = 0;
+			int32 NumDigits = 0;
 
-				while (ArgumentTokenStart + NumDigits < Code.Len())
+			while (ArgumentTokenStart + NumDigits < Code.Len())
+			{
+				TCHAR CodeCharacter = Code[ArgumentTokenStart + NumDigits];
+				if (!FChar::IsDigit(CodeCharacter))
 				{
-					TCHAR CodeCharacter = Code[ArgumentTokenStart + NumDigits];
-					if (!FChar::IsDigit(CodeCharacter))
-					{
-						break;
-					}
-					ArgumentIndexValue *= 10;
-					ArgumentIndexValue += (CodeCharacter - TEXT('0'));
-					++NumDigits;
+					break;
 				}
+				ArgumentIndexValue *= 10;
+				ArgumentIndexValue += (CodeCharacter - TEXT('0'));
+				++NumDigits;
+			}
 
-				checkf(NumDigits > 0, TEXT("Failed to scan integer in inline-HLSL after token '$':\n\"%.*s\""), Code.Len(), Code.GetData());
+			checkf(NumDigits > 0, TEXT("Failed to scan integer in inline-HLSL after token '$':\n\"%.*s\""), Code.Len(), Code.GetData());
 
-				checkf(
-					ArgumentIndexValue < InArguments.Num(), TEXT("Failed to substitute token $%d in inline-HLSL with given number of arguments (%d):\n\"%.*s\""),
-					ArgumentIndexValue, InArguments.Num(), Code.Len(), Code.GetData()
-				);
+			checkf(
+				ArgumentIndexValue < InArguments.Num(), TEXT("Failed to substitute token $%d in inline-HLSL with given number of arguments (%d):\n\"%.*s\""),
+				ArgumentIndexValue, InArguments.Num(), Code.Len(), Code.GetData()
+			);
 
-				LowerValue(InArguments[ArgumentIndexValue]);
+			LowerValue(InArguments[ArgumentIndexValue]);
 
-				ArgumentTokenEnd = ArgumentTokenStart + NumDigits;
-			};
+			ArgumentTokenEnd = ArgumentTokenStart + NumDigits;
+		};
 
 		auto MatchCharacter = [&Code](int32& Position, TCHAR InCharacter) -> bool
+		{
+			if (Position < Code.Len() && Code[Position] == InCharacter)
 			{
-				if (Position < Code.Len() && Code[Position] == InCharacter)
-				{
-					++Position;
-					return true;
-				}
-				return false;
-			};
+				++Position;
+				return true;
+			}
+			return false;
+		};
 
 		// Find all argument token characters '$'. For example "MyFunction($1, $0.xxxw)" can be subsituted with "MyFunction(MySecondArgument, MyFirstArgument.xxxw)"
 		while ((ArgumentTokenStart = Code.Find(TEXT("$"), ArgumentTokenEnd)) != INDEX_NONE)
@@ -638,7 +647,7 @@ struct FTranslator : FMaterialIRToHLSLTranslation
 
 	void LowerInlineHLSL(const MIR::FInlineHLSL* ExternalCode)
 	{
-		auto PrintHLSLCode = [this, ExternalCode](const FStringView& Code) -> void
+		auto PrintHLSLCode = [this, ExternalCode](FStringView Code) -> void
 			{
 				if (ExternalCode->NumArguments > 0)
 				{
@@ -710,40 +719,40 @@ struct FTranslator : FMaterialIRToHLSLTranslation
 	struct FSubstrateLegacyArgument
 	{
 		MIR::FType ParameterType;
-		const TCHAR* ParameterName;
+		FStringView ParameterName;
 		EMaterialProperty MaterialProperty;
-		const TCHAR* DefaultInlineHLSL;
+		FStringView DefaultInlineHLSL;
 	};
 
 	static TConstArrayView<FSubstrateLegacyArgument> GetSubstrateLegacyConversionArguments()
 	{
 		static const FSubstrateLegacyArgument LegacyArguments[] =
 		{
-			{ MIR::FType::MakeFloat(3), TEXT("BaseColor"), MP_BaseColor, nullptr },
-			{ MIR::FType::MakeFloat(),  TEXT("Specular"), MP_SpecularColor, TEXT("0.5f") },
-			{ MIR::FType::MakeFloat(),  TEXT("Metallic"), MP_Metallic, nullptr },
-			{ MIR::FType::MakeFloat(),  TEXT("Roughness"), MP_Roughness, nullptr },
-			{ MIR::FType::MakeFloat(),  TEXT("Anisotropy"), MP_Anisotropy, nullptr },
-			{ MIR::FType::MakeFloat(3), TEXT("SubSurfaceColor"), MP_SubsurfaceColor, TEXT("(float3)1") },
-			{ MIR::FType::MakeFloat(),  TEXT("SubSurfaceProfileId"), MP_MAX, TEXT("0") },
-			{ MIR::FType::MakeFloat(),  TEXT("ClearCoat"), MP_MAX, TEXT("1.0f") },
-			{ MIR::FType::MakeFloat(),  TEXT("ClearCoatRoughness"), MP_MAX, TEXT("0.1f") },
-			{ MIR::FType::MakeFloat(3), TEXT("Emissive"), MP_EmissiveColor, nullptr },
-			{ MIR::FType::MakeFloat(),  TEXT("Opacity"), MP_Opacity, nullptr },
-			{ MIR::FType::MakeFloat(3), TEXT("ThinTranslucentTransmittanceColor"), MP_MAX, TEXT("(float3)0.5f") },
-			{ MIR::FType::MakeFloat(),  TEXT("ThinTranslucentSurfaceCoverage"), MP_MAX, TEXT("1.0f") },
-			{ MIR::FType::MakeFloat(3), TEXT("WaterScatteringCoefficients"), MP_MAX, nullptr },
-			{ MIR::FType::MakeFloat(3), TEXT("WaterAbsorptionCoefficients"), MP_MAX, nullptr },
-			{ MIR::FType::MakeFloat(),  TEXT("WaterPhaseG"), MP_MAX, nullptr },
-			{ MIR::FType::MakeFloat(3), TEXT("ColorScaleBehindWater"), MP_MAX, TEXT("(float3)1.0f") },
-			{ MIR::FType::MakeInt(),    TEXT("ShadingModel"), MP_MAX, TEXT("1") },
-			{ MIR::FType::MakeFloat(3), TEXT("RawNormal"), MP_Normal, TEXT("Parameters.TangentToWorld[2]") },
-			{ MIR::FType::MakeFloat(3), TEXT("RawTangent"), MP_Tangent, TEXT("Parameters.TangentToWorld[0]") },
-			{ MIR::FType::MakeFloat(3), TEXT("RawClearCoatNormal"), MP_Normal, TEXT("Parameters.TangentToWorld[2]") },
-			{ MIR::FType::MakeFloat(3), TEXT("RawCustomTangent"), MP_Tangent, TEXT("Parameters.TangentToWorld[0]") },
-			{ MIR::FType::MakeInt(),    TEXT("SharedLocalBasisIndex"), MP_MAX, nullptr },
-			{ MIR::FType::MakeInt(),    TEXT("ClearCoatBottomNormal_SharedLocalBasisIndex"), MP_MAX, nullptr },
-			{ MIR::FType::MakeInt(),    TEXT("SharedLocalBasisTypes"), MP_MAX, TEXT("Parameters.SharedLocalBases.Types") },
+			{ MIR::FType::MakeFloat(3), TEXTVIEW("BaseColor"), MP_BaseColor, {} },
+			{ MIR::FType::MakeFloat(),  TEXTVIEW("Specular"), MP_SpecularColor, TEXTVIEW("0.5f") },
+			{ MIR::FType::MakeFloat(),  TEXTVIEW("Metallic"), MP_Metallic, {} },
+			{ MIR::FType::MakeFloat(),  TEXTVIEW("Roughness"), MP_Roughness, {} },
+			{ MIR::FType::MakeFloat(),  TEXTVIEW("Anisotropy"), MP_Anisotropy, {} },
+			{ MIR::FType::MakeFloat(3), TEXTVIEW("SubSurfaceColor"), MP_SubsurfaceColor, TEXTVIEW("(float3)1") },
+			{ MIR::FType::MakeFloat(),  TEXTVIEW("SubSurfaceProfileId"), MP_MAX, TEXTVIEW("0") },
+			{ MIR::FType::MakeFloat(),  TEXTVIEW("ClearCoat"), MP_MAX, TEXTVIEW("1.0f") },
+			{ MIR::FType::MakeFloat(),  TEXTVIEW("ClearCoatRoughness"), MP_MAX, TEXTVIEW("0.1f") },
+			{ MIR::FType::MakeFloat(3), TEXTVIEW("Emissive"), MP_EmissiveColor, {} },
+			{ MIR::FType::MakeFloat(),  TEXTVIEW("Opacity"), MP_Opacity, {} },
+			{ MIR::FType::MakeFloat(3), TEXTVIEW("ThinTranslucentTransmittanceColor"), MP_MAX, TEXTVIEW("(float3)0.5f") },
+			{ MIR::FType::MakeFloat(),  TEXTVIEW("ThinTranslucentSurfaceCoverage"), MP_MAX, TEXTVIEW("1.0f") },
+			{ MIR::FType::MakeFloat(3), TEXTVIEW("WaterScatteringCoefficients"), MP_MAX, {} },
+			{ MIR::FType::MakeFloat(3), TEXTVIEW("WaterAbsorptionCoefficients"), MP_MAX, {} },
+			{ MIR::FType::MakeFloat(),  TEXTVIEW("WaterPhaseG"), MP_MAX, {} },
+			{ MIR::FType::MakeFloat(3), TEXTVIEW("ColorScaleBehindWater"), MP_MAX, TEXTVIEW("(float3)1.0f") },
+			{ MIR::FType::MakeInt(),    TEXTVIEW("ShadingModel"), MP_MAX, TEXTVIEW("1") },
+			{ MIR::FType::MakeFloat(3), TEXTVIEW("RawNormal"), MP_Normal, TEXTVIEW("Parameters.TangentToWorld[2]") },
+			{ MIR::FType::MakeFloat(3), TEXTVIEW("RawTangent"), MP_Tangent, TEXTVIEW("Parameters.TangentToWorld[0]") },
+			{ MIR::FType::MakeFloat(3), TEXTVIEW("RawClearCoatNormal"), MP_Normal, TEXTVIEW("Parameters.TangentToWorld[2]") },
+			{ MIR::FType::MakeFloat(3), TEXTVIEW("RawCustomTangent"), MP_Tangent, TEXTVIEW("Parameters.TangentToWorld[0]") },
+			{ MIR::FType::MakeInt(),    TEXTVIEW("SharedLocalBasisIndex"), MP_MAX, {} },
+			{ MIR::FType::MakeInt(),    TEXTVIEW("ClearCoatBottomNormal_SharedLocalBasisIndex"), MP_MAX, {} },
+			{ MIR::FType::MakeInt(),    TEXTVIEW("SharedLocalBasisTypes"), MP_MAX, TEXTVIEW("Parameters.SharedLocalBases.Types") },
 		};
 		return LegacyArguments;
 	}
@@ -752,48 +761,38 @@ struct FTranslator : FMaterialIRToHLSLTranslation
 	{
 		static const FSubstrateLegacyArgument LegacyArguments[] =
 		{
-			{ MIR::FType::MakeFloat(3), TEXT("Emissive"), MP_EmissiveColor, nullptr },
-			{ MIR::FType::MakeFloat(), TEXT("TransmittanceColor"), MP_MAX, TEXT("(float3)0.0f") },
-			{ MIR::FType::MakeFloat(3), TEXT("RawNormal"), MP_Normal, TEXT("Parameters.TangentToWorld[2]") },
+			{ MIR::FType::MakeFloat(3), TEXTVIEW("Emissive"), MP_EmissiveColor, {} },
+			{ MIR::FType::MakeFloat(), TEXTVIEW("TransmittanceColor"), MP_MAX, TEXTVIEW("(float3)0.0f") },
+			{ MIR::FType::MakeFloat(3), TEXTVIEW("RawNormal"), MP_Normal, TEXTVIEW("Parameters.TangentToWorld[2]") },
 		};
 		return LegacyArguments;
 	}
 
 	void LowerPromoteSubstrateParameter(const MIR::FPromoteSubstrateParameter* SubstrateParameter)
 	{
-		Printer << TEXT("Parameters.SubstrateTree.PromoteParameterBlendedBSDFToOperator") << BeginArgs;
+		Printer << TEXTVIEW("Parameters.SubstrateTree.PromoteParameterBlendedBSDFToOperator") << BeginArgs;
 		Printer.Tabs++;
-		Printer << NewLine << (SubstrateParameter->bIsUnlit ? TEXT("GetSubstrateUnlitBSDF") : TEXT("SubstrateConvertLegacyMaterialStatic")) << BeginArgs;
+		Printer << NewLine << (SubstrateParameter->bIsUnlit ? TEXTVIEW("GetSubstrateUnlitBSDF") : TEXTVIEW("SubstrateConvertLegacyMaterialStatic")) << BeginArgs;
 		Printer.Tabs++;
-		Printer << ListSeparator << NewLine << TEXT("/*PixelFootprint:*/ Parameters.SubstratePixelFootprint");
+		Printer << ListSeparator << NewLine << TEXTVIEW("/*PixelFootprint:*/ Parameters.SubstratePixelFootprint");
 
 		TConstArrayView<FSubstrateLegacyArgument> Arguments = SubstrateParameter->bIsUnlit ? GetSubstrateUnlitArguments() : GetSubstrateLegacyConversionArguments();
 
 		for (const FSubstrateLegacyArgument& Argument : Arguments)
 		{
-			Printer << ListSeparator << NewLine << TEXT("/*") << Argument.ParameterName << TEXT("*/ ");
+			Printer << ListSeparator << NewLine << TEXTVIEW("/*") << Argument.ParameterName << TEXTVIEW("*/ ");
 
 			switch (Argument.MaterialProperty)
 			{
-				case MP_Normal:
-				{
-					LowerValue(SubstrateParameter->WorldSpaceTangentsAndNormals[0]);
-				}
-				break;
-
-				case MP_Tangent:
-				{
-					LowerValue(SubstrateParameter->WorldSpaceTangentsAndNormals[1]);
-				}
-				break;
-
+				case MP_Normal: LowerValue(SubstrateParameter->WorldSpaceTangentsAndNormals[0]); break;
+				case MP_Tangent: LowerValue(SubstrateParameter->WorldSpaceTangentsAndNormals[1]); break;
+				
 				default:
-				{
 					if (Argument.MaterialProperty != MP_MAX && SubstrateParameter->PropertyArgs[Argument.MaterialProperty])
 					{
 						LowerValue(SubstrateParameter->PropertyArgs[Argument.MaterialProperty]);
 					}
-					else if (Argument.DefaultInlineHLSL)
+					else if (!Argument.DefaultInlineHLSL.IsEmpty())
 					{
 						Printer << Argument.DefaultInlineHLSL;
 					}
@@ -801,7 +800,7 @@ struct FTranslator : FMaterialIRToHLSLTranslation
 					{
 						if (PrimitiveType->IsScalar())
 						{
-							Printer << TEXT("0.0f");
+							Printer << TEXTVIEW("0.0f");
 						}
 						else
 						{
@@ -809,7 +808,7 @@ struct FTranslator : FMaterialIRToHLSLTranslation
 							Printer << BeginArgs;
 							for (int32 ComponentIndex = 0; ComponentIndex < PrimitiveType->NumRows; ++ComponentIndex)
 							{
-								Printer << ListSeparator << TEXT("0.0f");
+								Printer << ListSeparator << TEXTVIEW("0.0f");
 							}
 							Printer << EndArgs;
 						}
@@ -818,14 +817,13 @@ struct FTranslator : FMaterialIRToHLSLTranslation
 					{
 						UE_MIR_UNREACHABLE();
 					}
-				}
-				break;
+					break;
 			}
 		}
 
 		Printer.Tabs--;
 		Printer << NewLine << EndArgs << ListSeparator;
-		Printer << NewLine << TEXT("0, 0, 0, 1");
+		Printer << NewLine << TEXTVIEW("0, 0, 0, 1");
 		Printer.Tabs--;
 		Printer << NewLine << EndArgs;
 	}
@@ -854,14 +852,14 @@ struct FTranslator : FMaterialIRToHLSLTranslation
 		
 		if (PrimitiveType.IsInteger())
 		{
-			Printer << TEXT("asint(");
+			Printer << TEXTVIEW("asint(");
 		}
 
-		Printer << TEXT("Material.PreshaderBuffer[") << BufferSlotIndex << TEXT("]");
+		Printer << TEXTVIEW("Material.PreshaderBuffer[") << BufferSlotIndex << TEXTVIEW("]");
 
 		if (PrimitiveType.NumComponents() < 4)
 		{
-			Printer << TEXT(".");
+			Printer << TEXTVIEW(".");
 
 			static const TCHAR* Components = TEXT("xyzw");
 			for (int32 i = 0; i < PrimitiveType.NumComponents(); ++i)
@@ -873,7 +871,7 @@ struct FTranslator : FMaterialIRToHLSLTranslation
 
 		if (PrimitiveType.IsInteger())
 		{
-			Printer << ")"; // close the "asint(" bracket
+			Printer << TEXTVIEW(")"); // close the "asint(" bracket
 		}
 
 		return NoOp;
@@ -893,9 +891,9 @@ struct FTranslator : FMaterialIRToHLSLTranslation
 		// Index of the float4 slot.  LWC parameters always have four elements, so there won't be an offset.
 		uint32 BufferSlotIndex = GlobalComponentOffset / 4;
 
-		Printer << "DFToWS(MakeDFVector" << NumComponents;
-		Printer << "(Material.PreshaderBuffer[" << BufferSlotIndex << "]" << GVector4SwizzleSubset[NumComponents - 1];
-		Printer << ", Material.PreshaderBuffer[" << BufferSlotIndex + 1 << "]" << GVector4SwizzleSubset[NumComponents - 1] << "))";
+		Printer << TEXTVIEW("DFToWS(MakeDFVector") << NumComponents;
+		Printer << TEXTVIEW("(Material.PreshaderBuffer[") << BufferSlotIndex << TEXTVIEW("]") << GVector4SwizzleSubset[NumComponents - 1];
+		Printer << TEXTVIEW(", Material.PreshaderBuffer[") << BufferSlotIndex + 1 << TEXTVIEW("]") << GVector4SwizzleSubset[NumComponents - 1] << TEXTVIEW("))");
 	
 		return NoOp;
 	}
@@ -952,18 +950,18 @@ struct FTranslator : FMaterialIRToHLSLTranslation
 				if (PrimitiveType.IsBoolean())
 				{
 					// Cast to bool requires a comparison with zero, outside the WSDemote
-					Printer << "(WSDemote(" << LowerVectorCastChain(Arg, NumComponents) << ") != 0)";
+					Printer << TEXTVIEW("(WSDemote(") << LowerVectorCastChain(Arg, NumComponents) << TEXTVIEW(") != 0)");
 				}
 				else if (!PrimitiveType.IsFloat())
 				{
 					// Cast to non-float (integer) requires a cast to the type, outside the WSDemote
 					MIR::FType VectorType = MIR::FType::MakeVector(PrimitiveType.ScalarKind, NumComponents);
-					Printer << LowerType(VectorType) << "(WSDemote(" << LowerVectorCastChain(Arg, NumComponents) << "))";
+					Printer << LowerType(VectorType) << TEXTVIEW("(WSDemote(") << LowerVectorCastChain(Arg, NumComponents) << TEXTVIEW("))");
 				}
 				else
 				{
 					// Cast to float
-					Printer << "WSDemote(" << LowerVectorCastChain(Arg, NumComponents) << ")";
+					Printer << TEXTVIEW("WSDemote(") << LowerVectorCastChain(Arg, NumComponents) << TEXTVIEW(")");
 				}
 			}
 			else if (PrimitiveType.IsLWC())
@@ -972,17 +970,17 @@ struct FTranslator : FMaterialIRToHLSLTranslation
 				if (ArgPrimitiveType.IsBoolean())
 				{
 					// Cast from bool requires a select between 1.0f and 0.0f, inside the WSPromote
-					Printer << "WSPromote(select(" << LowerVectorCastChain(Arg, NumComponents) << ", (float" << NumComponents << ")1.0f, (float" << NumComponents << ")0.0f))";
+					Printer << TEXTVIEW("WSPromote(select(") << LowerVectorCastChain(Arg, NumComponents) << TEXTVIEW(", (float") << NumComponents << TEXTVIEW(")1.0f, (float") << NumComponents << TEXTVIEW(")0.0f))");
 				}
 				else if (!ArgPrimitiveType.IsFloat())
 				{
 					// Cast from non-float (integer) requires a cast to float, inside the WSPromote
-					Printer << "WSPromote(float" << NumComponents << "(" << LowerVectorCastChain(Arg, NumComponents) << "))";
+					Printer << TEXTVIEW("WSPromote(float") << NumComponents << TEXTVIEW("(") << LowerVectorCastChain(Arg, NumComponents) << TEXTVIEW("))");
 				}
 				else
 				{
 					// Cast from float
-					Printer << "WSPromote(" << LowerVectorCastChain(Arg, NumComponents) << ")";
+					Printer << TEXTVIEW("WSPromote(") << LowerVectorCastChain(Arg, NumComponents) << TEXTVIEW(")");
 				}
 			}
 			else
@@ -991,19 +989,19 @@ struct FTranslator : FMaterialIRToHLSLTranslation
 				if (PrimitiveType.IsBoolean())
 				{
 					// Cast to bool requires a comparison with zero
-					Printer << "(" << LowerVectorCastChain(Arg, NumComponents) << " != 0)";
+					Printer << TEXTVIEW("(") << LowerVectorCastChain(Arg, NumComponents) << TEXTVIEW(" != 0)");
 				}
 				else if (ArgPrimitiveType.IsBoolean())
 				{
 					// Cast from bool requires a select between 1 and 0
 					MIR::FType VectorType = MIR::FType::MakeVector(PrimitiveType.ScalarKind, NumComponents);
-					Printer << "select(" << LowerVectorCastChain(Arg, NumComponents) << ", (" << LowerType(VectorType) << ")1, (" << LowerType(VectorType) << ")0)";
+					Printer << TEXTVIEW("select()") << LowerVectorCastChain(Arg, NumComponents) << TEXTVIEW(", (") << LowerType(VectorType) << TEXTVIEW(")1, (") << LowerType(VectorType) << TEXTVIEW(")0)");
 				}
 				else
 				{
 					// Cast between arithmetic types
 					MIR::FType VectorType = MIR::FType::MakeVector(PrimitiveType.ScalarKind, NumComponents);
-					Printer << LowerType(VectorType) << "(" << LowerVectorCastChain(Arg, NumComponents) << ")";
+					Printer << LowerType(VectorType) << TEXTVIEW("(") << LowerVectorCastChain(Arg, NumComponents) << TEXTVIEW(")");
 				}
 			}
 		}
@@ -1079,12 +1077,12 @@ struct FTranslator : FMaterialIRToHLSLTranslation
 		
 		if (bSameComponents && !PrimitiveType.IsLWC())
 		{
-			Printer << TEXT("(") << LowerType(Composite->Type) << TEXT(")") << LowerValue(Components[0]);
+			Printer << TEXTVIEW("(") << LowerType(Composite->Type) << TEXTVIEW(")") << LowerValue(Components[0]);
 		}
 		else if (bSameComponents && PrimitiveType.IsVector() && IsScalarToLWCVectorCast(Components[0]))
 		{
 			// Cast scalar to float vector, then promote, for example:  "WSPromote((float3)1.0f)"
-			Printer << "WSPromote" << "((float" << PrimitiveType.NumRows << ")";
+			Printer << TEXTVIEW("WSPromote") << TEXTVIEW("((float") << PrimitiveType.NumRows << TEXTVIEW(")");
 			if (const MIR::FScalar* Scalar = Components[0]->As<MIR::FScalar>())
 			{
 				// Print the inner non-LWC scalar value
@@ -1099,7 +1097,7 @@ struct FTranslator : FMaterialIRToHLSLTranslation
 			{
 				UE_MIR_UNREACHABLE();
 			}
-			Printer << ")";
+			Printer << TEXTVIEW(")");
 		}
 		else if (bWholeVectorCast)
 		{
@@ -1122,26 +1120,26 @@ struct FTranslator : FMaterialIRToHLSLTranslation
 			if (bAllTileValuesZero)
 			{
 				// Vector representable as regular floats (all tile values zero), call WSPromote on the offset vector
-				Printer << "WSPromote(";
+				Printer << TEXTVIEW("WSPromote(");
 			}
 			else
 			{
 				// Vector needs tile values, call LWC constructor and generate tile vector, before generating offset vector below
-				Printer << "MakeLWCVector" << Components.Num() << "(float" << Components.Num() << BeginArgs;
+				Printer << TEXTVIEW("MakeLWCVector") << Components.Num() << TEXTVIEW("(float") << Components.Num() << BeginArgs;
 				for (int32 Index = 0; Index < Components.Num(); Index++)
 				{
 					Printer << ListSeparator << TileValues[Index];
 				}
-				Printer << EndArgs << ", ";
+				Printer << EndArgs << TEXTVIEW(", ");
 			}
 
 			// Generate offset vector, plus extra parentheses to close WSPromote or MakeLWCVector call
-			Printer << "float" << Components.Num() << BeginArgs;
+			Printer << TEXTVIEW("float") << Components.Num() << BeginArgs;
 			for (int32 Index = 0; Index < Components.Num(); Index++)
 			{
 				Printer << ListSeparator << OffsetValues[Index];
 			}
-			Printer << EndArgs << ")";
+			Printer << EndArgs << TEXTVIEW(")");
 		}
 		else
 		{
@@ -1149,7 +1147,7 @@ struct FTranslator : FMaterialIRToHLSLTranslation
 			{
 				if (PrimitiveType.IsVector())
 				{
-					Printer << "MakeWSVector" << BeginArgs;
+					Printer << TEXTVIEW("MakeWSVector") << BeginArgs;
 				}
 				else
 				{
@@ -1175,15 +1173,15 @@ struct FTranslator : FMaterialIRToHLSLTranslation
 	void LowerSetMaterialOutput(const MIR::FSetMaterialOutput* Output)
 	{
 		// Special case MP_SubsurfaceColor as the actual property is a combination of the color and the profile but we don't want to expose the profile
-		const FString& PropertyName = (Output->Property == MP_SubsurfaceColor) ? "Subsurface" : FMaterialAttributeDefinitionMap::GetAttributeName(Output->Property);
+		FStringView PropertyName = (Output->Property == MP_SubsurfaceColor) ? TEXTVIEW("Subsurface") : FStringView{ FMaterialAttributeDefinitionMap::GetAttributeName(Output->Property) };
 
 		if (Output->Property == MP_WorldPositionOffset)
 		{
-			Printer << TEXT("return ");
+			Printer << TEXTVIEW("return ");
 		}
 		else
 		{
-			Printer << TEXT("PixelMaterialInputs.") << PropertyName << TEXT(" = ");
+			Printer << TEXTVIEW("PixelMaterialInputs.") << PropertyName << TEXTVIEW(" = ");
 		}
 
 		if (Substrate::IsSubstrateEnabled())
@@ -1191,23 +1189,23 @@ struct FTranslator : FMaterialIRToHLSLTranslation
 			// Lower to default value when Substrate is enabled
 			switch (Output->Property)
 			{
-				case MP_EmissiveColor:		Printer << TEXT("0"); break;
-				case MP_Opacity:			Printer << TEXT("1.0f"); break;
-				case MP_OpacityMask:		Printer << TEXT("1.0f"); break;
-				case MP_BaseColor:			Printer << TEXT("0"); break;
-				case MP_Metallic:			Printer << TEXT("0"); break;
-				case MP_Specular:			Printer << TEXT("0"); break;
-				case MP_Roughness:			Printer << TEXT("0"); break;
-				case MP_Anisotropy:			Printer << TEXT("0"); break;
+				case MP_EmissiveColor:		Printer << TEXTVIEW("0"); break;
+				case MP_Opacity:			Printer << TEXTVIEW("1.0f"); break;
+				case MP_OpacityMask:		Printer << TEXTVIEW("1.0f"); break;
+				case MP_BaseColor:			Printer << TEXTVIEW("0"); break;
+				case MP_Metallic:			Printer << TEXTVIEW("0"); break;
+				case MP_Specular:			Printer << TEXTVIEW("0"); break;
+				case MP_Roughness:			Printer << TEXTVIEW("0"); break;
+				case MP_Anisotropy:			Printer << TEXTVIEW("0"); break;
 				case MP_Normal:				LowerValue(Output->Arg); break;
-				case MP_Tangent:			Printer << TEXT("0"); break;
-				case MP_SubsurfaceColor:	Printer << TEXT("0"); break;
-				case MP_AmbientOcclusion:	Printer << TEXT("1.0f"); break;
-				case MP_Refraction:			Printer << TEXT("0"); break;
-				case MP_PixelDepthOffset:	Printer << TEXT("0.0f"); break;
-				case MP_ShadingModel:		Printer << TEXT("0"); break;
-				case MP_SurfaceThickness:	Printer << TEXT("0.01f"); break;
-				case MP_Displacement:		Printer << TEXT("-1.0f"); break;
+				case MP_Tangent:			Printer << TEXTVIEW("0"); break;
+				case MP_SubsurfaceColor:	Printer << TEXTVIEW("0"); break;
+				case MP_AmbientOcclusion:	Printer << TEXTVIEW("1.0f"); break;
+				case MP_Refraction:			Printer << TEXTVIEW("0"); break;
+				case MP_PixelDepthOffset:	Printer << TEXTVIEW("0.0f"); break;
+				case MP_ShadingModel:		Printer << TEXTVIEW("0"); break;
+				case MP_SurfaceThickness:	Printer << TEXTVIEW("0.01f"); break;
+				case MP_Displacement:		Printer << TEXTVIEW("-1.0f"); break;
 				default: LowerValue(Output->Arg); break;
 			}
 		}
@@ -1223,28 +1221,28 @@ struct FTranslator : FMaterialIRToHLSLTranslation
 
 		if (IsOperatorInfix(Operator->Op) && !bIsLWC)
 		{
-			const TCHAR* OpString;
+			FStringView OpString;
 			switch (Operator->Op)
 			{
-				case MIR::BO_Add: OpString = TEXT("+"); break;
-				case MIR::BO_Divide: OpString = TEXT("/"); break;
-				case MIR::BO_Modulo: OpString = TEXT("%"); break;
-				case MIR::BO_Equals: OpString = TEXT("=="); break;
-				case MIR::BO_GreaterThan: OpString = TEXT(">"); break;
-				case MIR::BO_GreaterThanOrEquals: OpString = TEXT(">="); break;
-				case MIR::BO_LessThan: OpString = TEXT("<"); break;
-				case MIR::BO_LessThanOrEquals: OpString = TEXT("<="); break;
-				case MIR::BO_Multiply: OpString = TEXT("*"); break;
-				case MIR::BO_NotEquals: OpString = TEXT("!="); break;
-				case MIR::BO_Subtract: OpString = TEXT("-"); break;
+				case MIR::BO_Add: OpString = TEXTVIEW("+"); break;
+				case MIR::BO_Divide: OpString = TEXTVIEW("/"); break;
+				case MIR::BO_Modulo: OpString = TEXTVIEW("%"); break;
+				case MIR::BO_Equals: OpString = TEXTVIEW("=="); break;
+				case MIR::BO_GreaterThan: OpString = TEXTVIEW(">"); break;
+				case MIR::BO_GreaterThanOrEquals: OpString = TEXTVIEW(">="); break;
+				case MIR::BO_LessThan: OpString = TEXTVIEW("<"); break;
+				case MIR::BO_LessThanOrEquals: OpString = TEXTVIEW("<="); break;
+				case MIR::BO_Multiply: OpString = TEXTVIEW("*"); break;
+				case MIR::BO_NotEquals: OpString = TEXTVIEW("!="); break;
+				case MIR::BO_Subtract: OpString = TEXTVIEW("-"); break;
 				default: UE_MIR_UNREACHABLE();
 			}
 
-			Printer << "(" << LowerValue(Operator->AArg) << " " << OpString << " " << LowerValue(Operator->BArg) << ")";
+			Printer << TEXTVIEW("(") << LowerValue(Operator->AArg) << TEXTVIEW(" ") << OpString << TEXTVIEW(" ") << LowerValue(Operator->BArg) << TEXTVIEW(")");
 		}
 		else
 		{
-			const TCHAR* OpString;
+			FStringView OpString;
 			if (bIsLWC)
 			{
 				// "Demotes LWC" indicates the given operator returns a non-LWC float, even if the input is LWC.  Besides that, comparison
@@ -1252,45 +1250,45 @@ struct FTranslator : FMaterialIRToHLSLTranslation
 				// (second argument of Fmod, and third argument of Lerp).
 				switch (Operator->Op)
 				{
-					case MIR::UO_Abs: OpString = TEXT("WSAbs"); break;
-					case MIR::UO_ACos: OpString = TEXT("WSACos"); break;						// Demotes LWC
-					case MIR::UO_ASin: OpString = TEXT("WSASin"); break;						// Demotes LWC
-					case MIR::UO_ATan: OpString = TEXT("WSATan"); break;						// Demotes LWC
-					case MIR::UO_Ceil: OpString = TEXT("WSCeil"); break;
-					case MIR::UO_Cos: OpString = TEXT("WSCos"); break;							// Demotes LWC
-					case MIR::UO_Floor: OpString = TEXT("WSFloor"); break;
-					case MIR::UO_Frac: OpString = TEXT("WSFracDemote"); break;					// Demotes LWC
-					case MIR::UO_Length: OpString = TEXT("WSLength"); break;
-					case MIR::UO_Negate: OpString = TEXT("WSNegate"); break;
-					case MIR::UO_Round: OpString = TEXT("WSRound"); break;
-					case MIR::UO_Saturate: OpString = TEXT("WSSaturateDemote"); break;			// Demotes LWC
-					case MIR::UO_Sign: OpString = TEXT("WSSign"); break;						// Demotes LWC
-					case MIR::UO_Sin: OpString = TEXT("WSSin"); break;							// Demotes LWC
-					case MIR::UO_Sqrt: OpString = TEXT("WSSqrtDemote"); break;					// Demotes LWC
-					case MIR::UO_Tan: OpString = TEXT("WSTan"); break;							// Demotes LWC
-					case MIR::UO_Truncate: OpString = TEXT("WSTrunc"); break;
-
-					case MIR::BO_Add: OpString = TEXT("WSAdd"); break;
-					case MIR::BO_Distance: OpString = TEXT("WSDistance"); break;
-					case MIR::BO_Divide: OpString = TEXT("WSDivide"); break;
-					case MIR::BO_Dot: OpString = TEXT("WSDot"); break;
-					case MIR::BO_Equals: OpString = TEXT("WSEquals"); break;					// Bool output
-					case MIR::BO_Fmod: OpString = TEXT("WSFmodDemote"); break;					// Demotes LWC, second input must be float (not LWC)!
-					case MIR::BO_GreaterThan: OpString = TEXT("WSGreater"); break;				// Bool output
-					case MIR::BO_GreaterThanOrEquals: OpString = TEXT("WSGreaterEqual"); break;	// Bool output
-					case MIR::BO_LessThan: OpString = TEXT("WSLess"); break;					// Bool output
-					case MIR::BO_LessThanOrEquals: OpString = TEXT("WSLessEqual"); break;		// Bool output
-					case MIR::BO_Max: OpString = TEXT("WSMax"); break;
-					case MIR::BO_Min: OpString = TEXT("WSMin"); break;
-					case MIR::BO_Multiply: OpString = TEXT("WSMultiply"); break;
-					case MIR::BO_NotEquals: OpString = TEXT("WSNotEquals"); break;				// Bool output
-					case MIR::BO_Step: OpString = TEXT("WSStep"); break;						// Demotes LWC
-					case MIR::BO_Subtract: OpString = TEXT("WSSubtract"); break;
-
-					case MIR::TO_Clamp: OpString = TEXT("WSClamp"); break;
-					case MIR::TO_Lerp: OpString = TEXT("WSLerp"); break;						// Third input must be float (not LWC)!
-					case MIR::TO_Select: OpString = TEXT("WSSelect"); break;
-					case MIR::TO_Smoothstep: OpString = TEXT("WSSmoothStepDemote"); break;		// Demotes LWC
+					case MIR::UO_Abs: OpString = TEXTVIEW("WSAbs"); break;
+					case MIR::UO_ACos: OpString = TEXTVIEW("WSACos"); break;						// Demotes LWC
+					case MIR::UO_ASin: OpString = TEXTVIEW("WSASin"); break;						// Demotes LWC
+					case MIR::UO_ATan: OpString = TEXTVIEW("WSATan"); break;						// Demotes LWC
+					case MIR::UO_Ceil: OpString = TEXTVIEW("WSCeil"); break;
+					case MIR::UO_Cos: OpString = TEXTVIEW("WSCos"); break;							// Demotes LWC
+					case MIR::UO_Floor: OpString = TEXTVIEW("WSFloor"); break;
+					case MIR::UO_Frac: OpString = TEXTVIEW("WSFracDemote"); break;					// Demotes LWC
+					case MIR::UO_Length: OpString = TEXTVIEW("WSLength"); break;
+					case MIR::UO_Negate: OpString = TEXTVIEW("WSNegate"); break;
+					case MIR::UO_Round: OpString = TEXTVIEW("WSRound"); break;
+					case MIR::UO_Saturate: OpString = TEXTVIEW("WSSaturateDemote"); break;			// Demotes LWC
+					case MIR::UO_Sign: OpString = TEXTVIEW("WSSign"); break;						// Demotes LWC
+					case MIR::UO_Sin: OpString = TEXTVIEW("WSSin"); break;							// Demotes LWC
+					case MIR::UO_Sqrt: OpString = TEXTVIEW("WSSqrtDemote"); break;					// Demotes LWC
+					case MIR::UO_Tan: OpString = TEXTVIEW("WSTan"); break;							// Demotes LWC
+					case MIR::UO_Truncate: OpString = TEXTVIEW("WSTrunc"); break;
+
+					case MIR::BO_Add: OpString = TEXTVIEW("WSAdd"); break;
+					case MIR::BO_Distance: OpString = TEXTVIEW("WSDistance"); break;
+					case MIR::BO_Divide: OpString = TEXTVIEW("WSDivide"); break;
+					case MIR::BO_Dot: OpString = TEXTVIEW("WSDot"); break;
+					case MIR::BO_Equals: OpString = TEXTVIEW("WSEquals"); break;					// Bool output
+					case MIR::BO_Fmod: OpString = TEXTVIEW("WSFmodDemote"); break;					// Demotes LWC, second input must be float (not LWC)!
+					case MIR::BO_GreaterThan: OpString = TEXTVIEW("WSGreater"); break;				// Bool output
+					case MIR::BO_GreaterThanOrEquals: OpString = TEXTVIEW("WSGreaterEqual"); break;	// Bool output
+					case MIR::BO_LessThan: OpString = TEXTVIEW("WSLess"); break;					// Bool output
+					case MIR::BO_LessThanOrEquals: OpString = TEXTVIEW("WSLessEqual"); break;		// Bool output
+					case MIR::BO_Max: OpString = TEXTVIEW("WSMax"); break;
+					case MIR::BO_Min: OpString = TEXTVIEW("WSMin"); break;
+					case MIR::BO_Multiply: OpString = TEXTVIEW("WSMultiply"); break;
+					case MIR::BO_NotEquals: OpString = TEXTVIEW("WSNotEquals"); break;				// Bool output
+					case MIR::BO_Step: OpString = TEXTVIEW("WSStep"); break;						// Demotes LWC
+					case MIR::BO_Subtract: OpString = TEXTVIEW("WSSubtract"); break;
+
+					case MIR::TO_Clamp: OpString = TEXTVIEW("WSClamp"); break;
+					case MIR::TO_Lerp: OpString = TEXTVIEW("WSLerp"); break;						// Third input must be float (not LWC)!
+					case MIR::TO_Select: OpString = TEXTVIEW("WSSelect"); break;
+					case MIR::TO_Smoothstep: OpString = TEXTVIEW("WSSmoothStepDemote"); break;		// Demotes LWC
 
 					default: UE_MIR_UNREACHABLE();
 				}
@@ -1299,82 +1297,82 @@ struct FTranslator : FMaterialIRToHLSLTranslation
 			{
 				switch (Operator->Op)
 				{
-					case MIR::UO_Abs: OpString = TEXT("abs"); break;
-					case MIR::UO_ACos: OpString = TEXT("acos"); break;
-					case MIR::UO_ACosFast: OpString = TEXT("acosFast"); break;
-					case MIR::UO_ACosh: OpString = TEXT("acosh"); break;
-					case MIR::UO_ASin: OpString = TEXT("asin"); break;
-					case MIR::UO_ASinFast: OpString = TEXT("asinFast"); break;
-					case MIR::UO_ASinh: OpString = TEXT("asinh"); break;
-					case MIR::UO_ATan: OpString = TEXT("atan"); break;
-					case MIR::UO_ATanFast: OpString = TEXT("atanFast"); break;
-					case MIR::UO_ATanh: OpString = TEXT("atanh"); break;
-					case MIR::UO_Ceil: OpString = TEXT("ceil"); break;
-					case MIR::UO_Cos: OpString = TEXT("cos"); break;
-					case MIR::UO_Cosh: OpString = TEXT("cosh"); break;
-					case MIR::UO_Exponential: OpString = TEXT("exp"); break;
-					case MIR::UO_Exponential2: OpString = TEXT("exp2"); break;
-					case MIR::UO_Floor: OpString = TEXT("floor"); break;
-					case MIR::UO_Frac: OpString = TEXT("frac"); break;
-					case MIR::UO_IsFinite: OpString = TEXT("isfinite"); break;
-					case MIR::UO_IsInf: OpString = TEXT("isinf"); break;
-					case MIR::UO_IsNan: OpString = TEXT("isnan"); break;
-					case MIR::UO_Length: OpString = TEXT("length"); break;
-					case MIR::UO_Logarithm: OpString = TEXT("log"); break;
-					case MIR::UO_Logarithm10: OpString = TEXT("log10"); break;
-					case MIR::UO_Logarithm2: OpString = TEXT("log2"); break;
-					case MIR::UO_Reciprocal: OpString = TEXT("rcp"); break;
-					case MIR::UO_Round: OpString = TEXT("round"); break;
-					case MIR::UO_Rsqrt: OpString = TEXT("rsqrt"); break;
-					case MIR::UO_Saturate: OpString = TEXT("saturate"); break;
-					case MIR::UO_Sign: OpString = TEXT("sign"); break;
-					case MIR::UO_Sin: OpString = TEXT("sin"); break;
-					case MIR::UO_Sinh: OpString = TEXT("sinh"); break;
-					case MIR::UO_Sqrt: OpString = TEXT("sqrt"); break;
-					case MIR::UO_Tan: OpString = TEXT("tan"); break;
-					case MIR::UO_Tanh: OpString = TEXT("tanh"); break;
-					case MIR::UO_Truncate: OpString = TEXT("trunc"); break;
-
-					case MIR::BO_And: OpString = TEXT("and"); break;
-					case MIR::BO_ATan2: OpString = TEXT("atan2"); break;
-					case MIR::BO_ATan2Fast: OpString = TEXT("atan2Fast"); break;
-					case MIR::BO_Cross: OpString = TEXT("cross"); break;
-					case MIR::BO_Distance: OpString = TEXT("distance"); break;
-					case MIR::BO_Dot: OpString = TEXT("dot"); break;
-					case MIR::BO_Fmod: OpString = TEXT("fmod"); break;
-					case MIR::BO_Max: OpString = TEXT("max"); break;
-					case MIR::BO_Min: OpString = TEXT("min"); break;
-					case MIR::BO_Or: OpString = TEXT("or"); break;
-					case MIR::BO_Pow: OpString = TEXT("pow"); break;
-					case MIR::BO_Step: OpString = TEXT("step"); break;
-
-					case MIR::TO_Clamp: OpString = TEXT("clamp"); break;
-					case MIR::TO_Lerp: OpString = TEXT("lerp"); break;
-					case MIR::TO_Select: OpString = TEXT("select"); break;
-					case MIR::TO_Smoothstep: OpString = TEXT("smoothstep"); break;
+					case MIR::UO_Abs: OpString = TEXTVIEW("abs"); break;
+					case MIR::UO_ACos: OpString = TEXTVIEW("acos"); break;
+					case MIR::UO_ACosFast: OpString = TEXTVIEW("acosFast"); break;
+					case MIR::UO_ACosh: OpString = TEXTVIEW("acosh"); break;
+					case MIR::UO_ASin: OpString = TEXTVIEW("asin"); break;
+					case MIR::UO_ASinFast: OpString = TEXTVIEW("asinFast"); break;
+					case MIR::UO_ASinh: OpString = TEXTVIEW("asinh"); break;
+					case MIR::UO_ATan: OpString = TEXTVIEW("atan"); break;
+					case MIR::UO_ATanFast: OpString = TEXTVIEW("atanFast"); break;
+					case MIR::UO_ATanh: OpString = TEXTVIEW("atanh"); break;
+					case MIR::UO_Ceil: OpString = TEXTVIEW("ceil"); break;
+					case MIR::UO_Cos: OpString = TEXTVIEW("cos"); break;
+					case MIR::UO_Cosh: OpString = TEXTVIEW("cosh"); break;
+					case MIR::UO_Exponential: OpString = TEXTVIEW("exp"); break;
+					case MIR::UO_Exponential2: OpString = TEXTVIEW("exp2"); break;
+					case MIR::UO_Floor: OpString = TEXTVIEW("floor"); break;
+					case MIR::UO_Frac: OpString = TEXTVIEW("frac"); break;
+					case MIR::UO_IsFinite: OpString = TEXTVIEW("isfinite"); break;
+					case MIR::UO_IsInf: OpString = TEXTVIEW("isinf"); break;
+					case MIR::UO_IsNan: OpString = TEXTVIEW("isnan"); break;
+					case MIR::UO_Length: OpString = TEXTVIEW("length"); break;
+					case MIR::UO_Logarithm: OpString = TEXTVIEW("log"); break;
+					case MIR::UO_Logarithm10: OpString = TEXTVIEW("log10"); break;
+					case MIR::UO_Logarithm2: OpString = TEXTVIEW("log2"); break;
+					case MIR::UO_Reciprocal: OpString = TEXTVIEW("rcp"); break;
+					case MIR::UO_Round: OpString = TEXTVIEW("round"); break;
+					case MIR::UO_Rsqrt: OpString = TEXTVIEW("rsqrt"); break;
+					case MIR::UO_Saturate: OpString = TEXTVIEW("saturate"); break;
+					case MIR::UO_Sign: OpString = TEXTVIEW("sign"); break;
+					case MIR::UO_Sin: OpString = TEXTVIEW("sin"); break;
+					case MIR::UO_Sinh: OpString = TEXTVIEW("sinh"); break;
+					case MIR::UO_Sqrt: OpString = TEXTVIEW("sqrt"); break;
+					case MIR::UO_Tan: OpString = TEXTVIEW("tan"); break;
+					case MIR::UO_Tanh: OpString = TEXTVIEW("tanh"); break;
+					case MIR::UO_Truncate: OpString = TEXTVIEW("trunc"); break;
+
+					case MIR::BO_And: OpString = TEXTVIEW("and"); break;
+					case MIR::BO_ATan2: OpString = TEXTVIEW("atan2"); break;
+					case MIR::BO_ATan2Fast: OpString = TEXTVIEW("atan2Fast"); break;
+					case MIR::BO_Cross: OpString = TEXTVIEW("cross"); break;
+					case MIR::BO_Distance: OpString = TEXTVIEW("distance"); break;
+					case MIR::BO_Dot: OpString = TEXTVIEW("dot"); break;
+					case MIR::BO_Fmod: OpString = TEXTVIEW("fmod"); break;
+					case MIR::BO_Max: OpString = TEXTVIEW("max"); break;
+					case MIR::BO_Min: OpString = TEXTVIEW("min"); break;
+					case MIR::BO_Or: OpString = TEXTVIEW("or"); break;
+					case MIR::BO_Pow: OpString = TEXTVIEW("pow"); break;
+					case MIR::BO_Step: OpString = TEXTVIEW("step"); break;
+
+					case MIR::TO_Clamp: OpString = TEXTVIEW("clamp"); break;
+					case MIR::TO_Lerp: OpString = TEXTVIEW("lerp"); break;
+					case MIR::TO_Select: OpString = TEXTVIEW("select"); break;
+					case MIR::TO_Smoothstep: OpString = TEXTVIEW("smoothstep"); break;
 			
 					default: UE_MIR_UNREACHABLE();
 				}
 			}
 
 			// Unary
-			Printer << OpString << "(" << LowerValue(Operator->AArg);
+			Printer << OpString << TEXTVIEW("(") << LowerValue(Operator->AArg);
 
 			// Binary
 			if (Operator->BArg)
 			{
 				check(MIR::IsBinaryOperator(Operator->Op) || MIR::IsTernaryOperator(Operator->Op));
-				Printer << ", " << LowerValue(Operator->BArg);
+				Printer << TEXTVIEW(", ") << LowerValue(Operator->BArg);
 			}
 
 			// Ternary
 			if (Operator->CArg)
 			{
 				check(MIR::IsTernaryOperator(Operator->Op));
-				Printer << ", " << LowerValue(Operator->CArg);
+				Printer << TEXTVIEW(", ") << LowerValue(Operator->CArg);
 			}
 
-			Printer << ")";
+			Printer << TEXTVIEW(")");
 		}
 	}
 
@@ -1384,28 +1382,28 @@ struct FTranslator : FMaterialIRToHLSLTranslation
 		{
 			if (Branch->TrueArg->Type.IsLWC())
 			{
-				Printer << TEXT("WSSelect(") << LowerValue(Branch->ConditionArg)
-					<< TEXT(", ") << LowerValue(Branch->TrueArg)
-					<< TEXT(", ") << LowerValue(Branch->FalseArg)
-					<< TEXT(")");
+				Printer << TEXTVIEW("WSSelect(") << LowerValue(Branch->ConditionArg)
+					<< TEXTVIEW(", ") << LowerValue(Branch->TrueArg)
+					<< TEXTVIEW(", ") << LowerValue(Branch->FalseArg)
+					<< TEXTVIEW(")");
 			}
 			else
 			{
 				Printer << LowerValue(Branch->ConditionArg)
-					<< TEXT(" ? ") << LowerValue(Branch->TrueArg)
-					<< TEXT(" : ") << LowerValue(Branch->FalseArg);
+					<< TEXTVIEW(" ? ") << LowerValue(Branch->TrueArg)
+					<< TEXTVIEW(" : ") << LowerValue(Branch->FalseArg);
 			}
 		}
 		else
 		{
 			Printer << EndOfStatement;
-			Printer << TEXT("if (") << LowerValue(Branch->ConditionArg) << TEXT(")") << NewLine << OpenBrace;
+			Printer << TEXTVIEW("if (") << LowerValue(Branch->ConditionArg) << TEXTVIEW(")") << NewLine << OpenBrace;
 			Printer << LowerBlock(Branch->TrueBlock[CurrentStage]);
-			Printer << LocalIdentifier[Branch] << " = " << LowerValue(Branch->TrueArg) << EndOfStatement;
+			Printer << InstrToLocalIndex[Branch] << TEXTVIEW(" = ") << LowerValue(Branch->TrueArg) << EndOfStatement;
 			Printer << CloseBrace << NewLine;
-			Printer << TEXT("else") << NewLine << OpenBrace;
+			Printer << TEXTVIEW("else") << NewLine << OpenBrace;
 			Printer << LowerBlock(Branch->FalseBlock[CurrentStage]);
-			Printer << LocalIdentifier[Branch] << " = " << LowerValue(Branch->FalseArg) << EndOfStatement;
+			Printer << InstrToLocalIndex[Branch] << TEXTVIEW(" = ") << LowerValue(Branch->FalseArg) << EndOfStatement;
 			Printer << CloseBrace;
 		}
 	}
@@ -1415,10 +1413,10 @@ struct FTranslator : FMaterialIRToHLSLTranslation
 		TOptional<MIR::FPrimitive> ArgVectorType = Subscript->Arg->Type.AsVector();
 		if (ArgVectorType && ArgVectorType->IsLWC())
 		{
-			const TCHAR* LwcComponentsStr[] = { TEXT("WSGetX("), TEXT("WSGetY("), TEXT("WSGetZ("), TEXT("WSGetW(") };
+			FStringView LwcComponentsStr[] = { TEXTVIEW("WSGetX("), TEXTVIEW("WSGetY("), TEXTVIEW("WSGetZ("), TEXTVIEW("WSGetW(") };
 			check(Subscript->Index <= ArgVectorType->NumComponents());
 
-			Printer << LwcComponentsStr[Subscript->Index] << LowerValue(Subscript->Arg) << TEXT(")");
+			Printer << LwcComponentsStr[Subscript->Index] << LowerValue(Subscript->Arg) << TEXTVIEW(")");
 		}
 		else
 		{
@@ -1426,7 +1424,7 @@ struct FTranslator : FMaterialIRToHLSLTranslation
 
 			if (ArgVectorType)
 			{
-				const TCHAR* ComponentsStr[] = { TEXT(".x"), TEXT(".y"), TEXT(".z"), TEXT(".w") };
+				FStringView ComponentsStr[] = { TEXTVIEW(".x"), TEXTVIEW(".y"), TEXTVIEW(".z"), TEXTVIEW(".w") };
 				check(Subscript->Index <= ArgVectorType->NumComponents());
 
 				Printer << ComponentsStr[Subscript->Index];
@@ -1445,17 +1443,17 @@ struct FTranslator : FMaterialIRToHLSLTranslation
 			if (PrimitiveType.IsBoolean())
 			{
 				// Cast to bool requires a comparison with zero, outside the WSDemote
-				Printer << "(WSDemote(" << LowerValue(Scalar->Arg) << ") != 0)";
+				Printer << TEXTVIEW("(WSDemote(") << LowerValue(Scalar->Arg) << TEXTVIEW(") != 0)");
 			}
 			else if (!PrimitiveType.IsFloat())
 			{
 				// Cast to non-float (integer) requires a cast to the type, outside the WSDemote
-				Printer << LowerType(Scalar->Type) << "(WSDemote(" << LowerValue(Scalar->Arg) << "))";
+				Printer << LowerType(Scalar->Type) << TEXTVIEW("(WSDemote(") << LowerValue(Scalar->Arg) << TEXTVIEW("))");
 			}
 			else
 			{
 				// Cast to float
-				Printer << "WSDemote(" << LowerValue(Scalar->Arg) << ")";
+				Printer << TEXTVIEW("WSDemote(") << LowerValue(Scalar->Arg) << TEXTVIEW(")");
 			}
 		}
 		else if (PrimitiveType.IsLWC())
@@ -1464,17 +1462,17 @@ struct FTranslator : FMaterialIRToHLSLTranslation
 			if (ArgPrimitiveType.IsBoolean())
 			{
 				// Cast from bool requires a select between 1.0f and 0.0f, inside the WSPromote
-				Printer << "WSPromote(" << LowerValue(Scalar->Arg) << " ? 1.0f : 0.0f)";
+				Printer << TEXTVIEW("WSPromote(") << LowerValue(Scalar->Arg) << TEXTVIEW(" ? 1.0f : 0.0f)");
 			}
 			else if (!ArgPrimitiveType.IsFloat())
 			{
 				// Cast from non-float (integer) requires a cast to float, inside the WSPromote
-				Printer << "WSPromote(float(" << LowerValue(Scalar->Arg) << "))";
+				Printer << TEXTVIEW("WSPromote(float(") << LowerValue(Scalar->Arg) << TEXTVIEW("))");
 			}
 			else
 			{
 				// Cast from float
-				Printer << "WSPromote(" << LowerValue(Scalar->Arg) << ")";
+				Printer << TEXTVIEW("WSPromote(") << LowerValue(Scalar->Arg) << TEXTVIEW(")");
 			}
 		}
 		else
@@ -1483,17 +1481,17 @@ struct FTranslator : FMaterialIRToHLSLTranslation
 			if (PrimitiveType.IsBoolean())
 			{
 				// Cast to bool requires a comparison with zero
-				Printer << "(" << LowerValue(Scalar->Arg) << " != 0)";
+				Printer << TEXTVIEW("(") << LowerValue(Scalar->Arg) << TEXTVIEW(" != 0)");
 			}
 			else if (ArgPrimitiveType.IsBoolean())
 			{
 				// Cast from bool requires a select between 1 and 0
-				Printer << LowerType(Scalar->Type) << "(" << LowerValue(Scalar->Arg) << " ? 1 : 0)";
+				Printer << LowerType(Scalar->Type) << TEXTVIEW("(") << LowerValue(Scalar->Arg) << TEXTVIEW(" ? 1 : 0)");
 			}
 			else
 			{
 				// Cast between arithmetic types
-				Printer << LowerType(Scalar->Type) << "(" << LowerValue(Scalar->Arg) << ")";
+				Printer << LowerType(Scalar->Type) << TEXTVIEW("(") << LowerValue(Scalar->Arg) << TEXTVIEW(")");
 			}
 		}
 	}
@@ -1502,20 +1500,20 @@ struct FTranslator : FMaterialIRToHLSLTranslation
 	{
 		switch (TextureRead->TextureObject->Type.GetKind())
 		{
-			case MIR::ETypeKind::Texture2D: Printer << TEXT("Texture2D"); break;
+			case MIR::ETypeKind::Texture2D: Printer << TEXTVIEW("Texture2D"); break;
 			default: UE_MIR_UNREACHABLE();
 		}
 
 		switch (TextureRead->Mode)
 		{
-			case MIR::ETextureReadMode::GatherRed: Printer << TEXT("GatherRed"); break;
-			case MIR::ETextureReadMode::GatherGreen: Printer << TEXT("GatherGreen"); break;
-			case MIR::ETextureReadMode::GatherBlue: Printer << TEXT("GatherBlue"); break;
-			case MIR::ETextureReadMode::GatherAlpha: Printer << TEXT("GatherAlpha"); break;
-			case MIR::ETextureReadMode::MipAuto: Printer << TEXT("Sample"); break;
-			case MIR::ETextureReadMode::MipLevel: Printer << TEXT("SampleLevel"); break;
-			case MIR::ETextureReadMode::MipBias: Printer << TEXT("SampleBias"); break;
-			case MIR::ETextureReadMode::Derivatives: Printer << TEXT("SampleGrad"); break;
+			case MIR::ETextureReadMode::GatherRed: Printer << TEXTVIEW("GatherRed"); break;
+			case MIR::ETextureReadMode::GatherGreen: Printer << TEXTVIEW("GatherGreen"); break;
+			case MIR::ETextureReadMode::GatherBlue: Printer << TEXTVIEW("GatherBlue"); break;
+			case MIR::ETextureReadMode::GatherAlpha: Printer << TEXTVIEW("GatherAlpha"); break;
+			case MIR::ETextureReadMode::MipAuto: Printer << TEXTVIEW("Sample"); break;
+			case MIR::ETextureReadMode::MipLevel: Printer << TEXTVIEW("SampleLevel"); break;
+			case MIR::ETextureReadMode::MipBias: Printer << TEXTVIEW("SampleBias"); break;
+			case MIR::ETextureReadMode::Derivatives: Printer << TEXTVIEW("SampleGrad"); break;
 			default: UE_MIR_UNREACHABLE();
 		}
 
@@ -1544,46 +1542,38 @@ struct FTranslator : FMaterialIRToHLSLTranslation
 
 		checkf(TextureRead->VTPageTable, TEXT("Missing page table for virtual texture read instruction"));
 		const MIR::FVTPageTableRead* VTPageTableResult = TextureRead->VTPageTable->As<MIR::FVTPageTableRead>();
-
 		const MIR::FTextureObject* TextureObject = TextureRead->TextureObject->As<MIR::FTextureObject>();
 
-		const int32 VirtualTextureIndex = TextureObject->Analysis_UniformParameterIndex;
-		//const int32 VTStackIndex = VTPageTableResult->VTStackIndex;
-		const int32 VTLayerIndex = VTPageTableResult->VTLayerIndex;
+		int32 VirtualTextureIndex = TextureObject->Analysis_UniformParameterIndex;
+		int32 VTLayerIndex = VTPageTableResult->VTLayerIndex;
 
 		// Sampling function
-		Printer << TEXT("TextureVirtualSample");
+		Printer << TEXTVIEW("TextureVirtualSample");
 
+		// 'Texture name/sampler', 'PageTableResult', 'LayerIndex', 'PackedUniform'
+		Printer
+			<< BeginArgs
+			<< ListSeparator << LowerValue(TextureRead->TextureObject)
+			<< ListSeparator;
+		
 		// Sampling function argument list
-		FString SamplerExpression;
 		if (TextureRead->SamplerSourceMode != SSM_FromTextureAsset)
 		{
 			// VT doesn't care if the shared sampler is wrap or clamp. It only cares if it is aniso or not.
 			// The wrap/clamp/mirror operation is handled in the shader explicitly.
-			const TCHAR* SharedSamplerName = TextureRead->bUseAnisoSampler ? TEXT("View.SharedBilinearAnisoClampedSampler") : TEXT("View.SharedBilinearClampedSampler");
-			SamplerExpression = FString::Printf(TEXT("GetMaterialSharedSampler(Material.VirtualTexturePhysical_%dSampler, %s)"), VirtualTextureIndex, SharedSamplerName);
+			// This generates: GetMaterialSharedSampler(Material.VirtualTexturePhysical_<VirtualTextureIndex>Sampler, <SharedSamplerName>)
+			FStringView SharedSamplerName = TextureRead->bUseAnisoSampler ? TEXTVIEW("View.SharedBilinearAnisoClampedSampler") : TEXTVIEW("View.SharedBilinearClampedSampler");
+			Printer << TEXTVIEW("GetMaterialSharedSampler(Material.VirtualTexturePhysical_") << VirtualTextureIndex << TEXTVIEW("Sampler, ") << SharedSamplerName << TEXTVIEW(")");
 		}
 		else
 		{
-			SamplerExpression = FString::Printf(TEXT("Material.VirtualTexturePhysical_%dSampler"), VirtualTextureIndex);
+			Printer << TEXTVIEW("Material.VirtualTexturePhysical_") << VirtualTextureIndex << TEXTVIEW("Sampler");
 		}
 
-		// VT page information
-		FString VTPageTableResultExpression;
-		FString VTPackedUniformName;
-
-		VTPackedUniformName = FString::Printf(TEXT("Material.VTPackedUniform[%d]"), VirtualTextureIndex);
-
-		// 'Texture name/sampler', 'PageTableResult', 'LayerIndex', 'PackedUniform'
-		Printer
-			<< BeginArgs
-			<< ListSeparator << LowerValue(TextureRead->TextureObject)
-			<< ListSeparator << SamplerExpression
-			<< ListSeparator << LowerValue(TextureRead->VTPageTable)
+		Printer << ListSeparator << LowerValue(TextureRead->VTPageTable)
 			<< ListSeparator << VTLayerIndex
-			<< ListSeparator << TEXT("VTUniform_Unpack(") << VTPackedUniformName << TEXT(")")
-			<< EndArgs
-			;
+			<< ListSeparator << TEXTVIEW("VTUniform_Unpack(") << TEXT("Material.VTPackedUniform[") << VirtualTextureIndex << TEXTVIEW("]") << TEXTVIEW(")")
+			<< EndArgs;
 	}
 
 	void LowerTextureRead(const MIR::FTextureRead* TextureRead)
@@ -1592,7 +1582,7 @@ struct FTranslator : FMaterialIRToHLSLTranslation
 		LowerSamplerType(TextureRead->SamplerType, bSamplerNeedsBrackets);
 		if (bSamplerNeedsBrackets)
 		{
-			Printer << TEXT("(");
+			Printer << TEXTVIEW("(");
 		}
 
 		if (IsVirtualSamplerType(TextureRead->SamplerType))
@@ -1606,7 +1596,7 @@ struct FTranslator : FMaterialIRToHLSLTranslation
 
 		if (bSamplerNeedsBrackets)
 		{
-			Printer << TEXT(")");
+			Printer << TEXTVIEW(")");
 		}
 	}
 
@@ -1614,50 +1604,40 @@ struct FTranslator : FMaterialIRToHLSLTranslation
 	{
 		using namespace UE::MaterialTranslatorUtils;
 
-		const TCHAR* VTPageTableReadFunctionName = TEXT("TextureLoadVirtualPageTable");
+		FStringView VTPageTableReadFunctionName = TEXTVIEW("TextureLoadVirtualPageTable");
 
 		const bool bHasDerivativeTexCoords = VTPageTableRead->TexCoordDdx != nullptr && VTPageTableRead->TexCoordDdy != nullptr;
 
+		Printer
+			<< VTPageTableReadFunctionName << (bHasDerivativeTexCoords ? TEXTVIEW("Grad") : TEXTVIEW(""))
+			<< BeginArgs
+			<< ListSeparator << TEXTVIEW("VIRTUALTEXTURE_PAGETABLE_") << VTPageTableRead->VTStackIndex
+			<< ListSeparator << TEXTVIEW("VTPageTableUniform_Unpack(VIRTUALTEXTURE_PAGETABLE_UNIFORM_") << VTPageTableRead->VTStackIndex << TEXTVIEW(")")
+			<< ListSeparator << LowerValue(VTPageTableRead->TexCoord)
+			<< ListSeparator << FStringView{ GetVTAddressMode(VTPageTableRead->AddressU) }
+			<< ListSeparator << FStringView{ GetVTAddressMode(VTPageTableRead->AddressV) };
+
 		if (bHasDerivativeTexCoords)
 		{
 			Printer
-				<< VTPageTableReadFunctionName << TEXT("Grad")
-				<< BeginArgs
-				<< ListSeparator << FString::Printf(TEXT("VIRTUALTEXTURE_PAGETABLE_%d"), VTPageTableRead->VTStackIndex)
-				<< ListSeparator << FString::Printf(TEXT("VTPageTableUniform_Unpack(VIRTUALTEXTURE_PAGETABLE_UNIFORM_%d)"), VTPageTableRead->VTStackIndex)
-				<< ListSeparator << LowerValue(VTPageTableRead->TexCoord)
-				<< ListSeparator << GetVTAddressMode(VTPageTableRead->AddressU)
-				<< ListSeparator << GetVTAddressMode(VTPageTableRead->AddressV)
 				<< ListSeparator << LowerValue(VTPageTableRead->TexCoordDdx)
-				<< ListSeparator << LowerValue(VTPageTableRead->TexCoordDdy)
-				;
+				<< ListSeparator << LowerValue(VTPageTableRead->TexCoordDdy);
 		}
 		else
 		{
 			constexpr float MipBias = 0.0f;
-			Printer
-				<< VTPageTableReadFunctionName
-				<< BeginArgs
-				<< ListSeparator << FString::Printf(TEXT("VIRTUALTEXTURE_PAGETABLE_%d"), VTPageTableRead->VTStackIndex)
-				<< ListSeparator << FString::Printf(TEXT("VTPageTableUniform_Unpack(VIRTUALTEXTURE_PAGETABLE_UNIFORM_%d)"), VTPageTableRead->VTStackIndex)
-				<< ListSeparator << LowerValue(VTPageTableRead->TexCoord)
-				<< ListSeparator << GetVTAddressMode(VTPageTableRead->AddressU)
-				<< ListSeparator << GetVTAddressMode(VTPageTableRead->AddressV)
-				<< ListSeparator << MipBias
-				;
+			Printer << ListSeparator << MipBias;
 		}
 
-		Printer << ListSeparator << TEXT("Parameters.SvPosition.xy");
+		Printer << ListSeparator << TEXTVIEW("Parameters.SvPosition.xy");
 
 		// VT feedback
 		const bool bGenerateFeedback = VTPageTableRead->VTFeedbackRequestIndex != INDEX_NONE;
-
 		if (bGenerateFeedback)
 		{
 			Printer
-				<< ListSeparator << VTPageTableRead->VTFeedbackRequestIndex << TEXT("u + LIGHTMAP_VT_ENABLED")
-				<< ListSeparator << TEXT("Parameters.VirtualTextureFeedback")
-				;
+				<< ListSeparator << VTPageTableRead->VTFeedbackRequestIndex << TEXTVIEW("u + LIGHTMAP_VT_ENABLED")
+				<< ListSeparator << TEXTVIEW("Parameters.VirtualTextureFeedback");
 		}
 
 		Printer << EndArgs;
@@ -1668,42 +1648,42 @@ struct FTranslator : FMaterialIRToHLSLTranslation
 		switch (SamplerType)
 		{
 			case SAMPLERTYPE_External:
-				Printer << TEXT("ProcessMaterialExternalTextureLookup");
+				Printer << TEXTVIEW("ProcessMaterialExternalTextureLookup");
 				break;
 
 			case SAMPLERTYPE_Color:
-				Printer << TEXT("ProcessMaterialColorTextureLookup");
+				Printer << TEXTVIEW("ProcessMaterialColorTextureLookup");
 				break;
 			case SAMPLERTYPE_VirtualColor:
 				// has a mobile specific workaround
-				Printer << TEXT("ProcessMaterialVirtualColorTextureLookup");
+				Printer << TEXTVIEW("ProcessMaterialVirtualColorTextureLookup");
 				break;
 
 			case SAMPLERTYPE_LinearColor:
 			case SAMPLERTYPE_VirtualLinearColor:
-				Printer << TEXT("ProcessMaterialLinearColorTextureLookup");
+				Printer << TEXTVIEW("ProcessMaterialLinearColorTextureLookup");
 				break;
 
 			case SAMPLERTYPE_Alpha:
 			case SAMPLERTYPE_VirtualAlpha:
 			case SAMPLERTYPE_DistanceFieldFont:
-				Printer << TEXT("ProcessMaterialAlphaTextureLookup");
+				Printer << TEXTVIEW("ProcessMaterialAlphaTextureLookup");
 				break;
 
 			case SAMPLERTYPE_Grayscale:
 			case SAMPLERTYPE_VirtualGrayscale:
-				Printer << TEXT("ProcessMaterialGreyscaleTextureLookup");
+				Printer << TEXTVIEW("ProcessMaterialGreyscaleTextureLookup");
 				break;
 
 			case SAMPLERTYPE_LinearGrayscale:
 			case SAMPLERTYPE_VirtualLinearGrayscale:
-				Printer <<TEXT("ProcessMaterialLinearGreyscaleTextureLookup");
+				Printer << TEXTVIEW("ProcessMaterialLinearGreyscaleTextureLookup");
 				break;
 
 			case SAMPLERTYPE_Normal:
 			case SAMPLERTYPE_VirtualNormal:
 				// Normal maps need to be unpacked in the pixel shader.
-				Printer << TEXT("UnpackNormalMap");
+				Printer << TEXTVIEW("UnpackNormalMap");
 				break;
 
 			case SAMPLERTYPE_Masks:
@@ -1723,18 +1703,18 @@ struct FTranslator : FMaterialIRToHLSLTranslation
 	{
 		if (SamplerSource != SSM_FromTextureAsset)
 		{
-			Printer << "GetMaterialSharedSampler(";
+			Printer << TEXTVIEW("GetMaterialSharedSampler(");
 		}
 		
-		Printer << LowerValue(TextureValue) << TEXT("Sampler");
+		Printer << LowerValue(TextureValue) << TEXTVIEW("Sampler");
 		
 		if (SamplerSource == SSM_Wrap_WorldGroupSettings)
 		{
-			Printer << ", View.MaterialTextureBilinearWrapedSampler)";
+			Printer << TEXTVIEW(", View.MaterialTextureBilinearWrapedSampler)");
 		}
 		else if (SamplerSource == SSM_Clamp_WorldGroupSettings)
 		{
-			Printer << ", View.MaterialTextureBilinearClampedSampler)";
+			Printer << TEXTVIEW(", View.MaterialTextureBilinearClampedSampler)");
 		}
 		else
 		{
@@ -1747,12 +1727,12 @@ struct FTranslator : FMaterialIRToHLSLTranslation
 
 	ENoOp LowerTextureReference(EMaterialValueType TextureType, int32 TextureParameterIndex)
 	{
-		Printer << TEXT("Material.");
+		Printer << TEXTVIEW("Material.");
 
 		switch (TextureType)
 		{
-			case MCT_Texture2D: Printer << TEXT("Texture2D_"); break;
-			case MCT_TextureVirtual: Printer << TEXT("VirtualTexturePhysical_"); break;
+			case MCT_Texture2D: Printer << TEXTVIEW("Texture2D_"); break;
+			case MCT_TextureVirtual: Printer << TEXTVIEW("VirtualTexturePhysical_"); break;
 			default: UE_MIR_UNREACHABLE();
 		}
 
@@ -1771,11 +1751,11 @@ struct FTranslator : FMaterialIRToHLSLTranslation
 		if (StageSwitch->Arg->Type.IsLWC())
 		{
 			// Expression emitter assumes these to be LWC for analytic derivative evaluation, so we promote them on load
-			Printer << (StageSwitch->Axis == MIR::EDerivativeAxis::X ? "WSPromote(WSDdxDemote(" : "WSPromote(WSDdyDemote(") << LowerValue(StageSwitch->Arg) << "))";
+			Printer << (StageSwitch->Axis == MIR::EDerivativeAxis::X ? TEXTVIEW("WSPromote(WSDdxDemote(") : TEXTVIEW("WSPromote(WSDdyDemote(")) << LowerValue(StageSwitch->Arg) << TEXTVIEW("))");
 		}
 		else
 		{
-			Printer << (StageSwitch->Axis == MIR::EDerivativeAxis::X ? "ddx(" : "ddy(") << LowerValue(StageSwitch->Arg) << ")";
+			Printer << (StageSwitch->Axis == MIR::EDerivativeAxis::X ? TEXTVIEW("ddx(") : TEXTVIEW("ddy(")) << LowerValue(StageSwitch->Arg) << TEXTVIEW(")");
 		}
 	}
 
@@ -1784,17 +1764,6 @@ struct FTranslator : FMaterialIRToHLSLTranslation
 	void GenerateTemplateStringParameters(TMap<FString, FString>& Params)
 	{
 		const FMaterialIRModule::FStatistics ModuleStatistics = Module->GetStatistics();
-
-		auto SetParamInt = [&] (const TCHAR* InParamName, int32 InValue)
-		{
-			Params.Add(InParamName, FString::Printf(TEXT("%d"), InValue));
-		};
-		
-		auto SetParamReturnFloat = [&] (const TCHAR* InParamName, float InValue)
-		{
-			Params.Add(InParamName, FString::Printf(TEXT(TAB "return %.5f"), InValue));
-		};
-
 		Params.Add(TEXT("pixel_material_inputs"), MoveTemp(PixelAttributesHLSL));
 
 		// "Normal" is treated in a special way because the rest of the attributes may lead back to reading it.
@@ -1807,24 +1776,22 @@ struct FTranslator : FMaterialIRToHLSLTranslation
 		Params.Add(TEXT("calc_pixel_material_inputs_other_inputs"), EvaluateOtherMaterialAttributesHLSL[MIR::Stage_Pixel]);
 		
 		// MaterialAttributes
-		TArray<FGuid> OrderedVisibleAttributes = FMaterialAttributeDefinitionMap::GetOrderedVisibleAttributeList();
-		
 		FString MaterialDeclarations;
-		MaterialDeclarations.Appendf(TEXT("struct FMaterialAttributes\n{\n"));
-		for (const FGuid& AttributeID : OrderedVisibleAttributes)
+		MaterialDeclarations.Append(TEXTVIEW("struct FMaterialAttributes\n{\n"));
+		for (const FGuid& AttributeID : FMaterialAttributeDefinitionMap::GetOrderedVisibleAttributeList())
 		{
 			const FString& PropertyName = FMaterialAttributeDefinitionMap::GetAttributeName(AttributeID);
 			const EMaterialValueType PropertyType = FMaterialAttributeDefinitionMap::GetValueType(AttributeID);
 			MaterialDeclarations.Appendf(TEXT(TAB "%s %s;\n"), GetHLSLTypeString(PropertyType), *PropertyName);
 		}
-		MaterialDeclarations.Appendf(TEXT("};"));
+		MaterialDeclarations.Append(TEXTVIEW("};"));
 		Params.Add(TEXT("material_declarations"), MoveTemp(MaterialDeclarations));
 		
-		SetParamInt(TEXT("num_material_texcoords_vertex"), ModuleStatistics.NumVertexTexCoords);
-		SetParamInt(TEXT("num_material_texcoords"), ModuleStatistics.NumPixelTexCoords);
-		SetParamInt(TEXT("num_custom_vertex_interpolators"), 0);
-		SetParamInt(TEXT("num_tex_coord_interpolators"), ModuleStatistics.NumPixelTexCoords);
-
+		Params.Add(TEXT("num_material_texcoords_vertex"), FString::FromInt(ModuleStatistics.NumVertexTexCoords));
+		Params.Add(TEXT("num_material_texcoords"), FString::FromInt(ModuleStatistics.NumPixelTexCoords));
+		Params.Add(TEXT("num_custom_vertex_interpolators"), FString::FromInt(0));
+		Params.Add(TEXT("num_tex_coord_interpolators"), FString::FromInt(ModuleStatistics.NumPixelTexCoords));
+		
 		FString GetMaterialCustomizedUVS;
 		for (int32 CustomUVIndex = 0; CustomUVIndex < ModuleStatistics.NumPixelTexCoords; CustomUVIndex++)
 		{
@@ -1832,20 +1799,20 @@ struct FTranslator : FMaterialIRToHLSLTranslation
 			GetMaterialCustomizedUVS.Appendf(TEXT(TAB "OutTexCoords[%u] = Parameters.MaterialAttributes.%s;\n"), CustomUVIndex, *AttributeName);
 		}
 		Params.Add(TEXT("get_material_customized_u_vs"), MoveTemp(GetMaterialCustomizedUVS));
-
-		SetParamReturnFloat(TEXT("get_material_emissive_for_cs"), 0.f);
-		SetParamReturnFloat(TEXT("get_material_translucency_directional_lighting_intensity"), Material->GetTranslucencyDirectionalLightingIntensity());
-		SetParamReturnFloat(TEXT("get_material_translucent_shadow_density_scale"), Material->GetTranslucentShadowDensityScale());
-		SetParamReturnFloat(TEXT("get_material_translucent_self_shadow_density_scale"), Material->GetTranslucentSelfShadowDensityScale());
-		SetParamReturnFloat(TEXT("get_material_translucent_self_shadow_second_density_scale"), Material->GetTranslucentSelfShadowSecondDensityScale());
-		SetParamReturnFloat(TEXT("get_material_translucent_self_shadow_second_opacity"), Material->GetTranslucentSelfShadowSecondOpacity());
-		SetParamReturnFloat(TEXT("get_material_translucent_backscattering_exponent"), Material->GetTranslucentBackscatteringExponent());
+		
+		#define SET_PARAM_RETURN_FLOAT(ParamName, Value) Params.Add(TEXT(ParamName), FString::Printf(TEXT(TAB "return %.5f"), Value))
+		SET_PARAM_RETURN_FLOAT("get_material_emissive_for_cs", 0.f);
+		SET_PARAM_RETURN_FLOAT("get_material_translucency_directional_lighting_intensity", Material->GetTranslucencyDirectionalLightingIntensity());
+		SET_PARAM_RETURN_FLOAT("get_material_translucent_shadow_density_scale", Material->GetTranslucentShadowDensityScale());
+		SET_PARAM_RETURN_FLOAT("get_material_translucent_self_shadow_density_scale", Material->GetTranslucentSelfShadowDensityScale());
+		SET_PARAM_RETURN_FLOAT("get_material_translucent_self_shadow_second_density_scale", Material->GetTranslucentSelfShadowSecondDensityScale());
+		SET_PARAM_RETURN_FLOAT("get_material_translucent_self_shadow_second_opacity", Material->GetTranslucentSelfShadowSecondOpacity());
+		SET_PARAM_RETURN_FLOAT("get_material_translucent_backscattering_exponent", Material->GetTranslucentBackscatteringExponent());
+		SET_PARAM_RETURN_FLOAT("get_material_opacity_mask_clip_value", Material->GetOpacityMaskClipValue());
 
 		FLinearColor Extinction = Material->GetTranslucentMultipleScatteringExtinction();
 		Params.Add(TEXT("get_material_translucent_multiple_scattering_extinction"), FString::Printf(TEXT(TAB "return MaterialFloat3(%.5f, %.5f, %.5f)"), Extinction.R, Extinction.G, Extinction.B));
 
-		SetParamReturnFloat(TEXT("get_material_opacity_mask_clip_value"), Material->GetOpacityMaskClipValue());
-
 		Params.Add(TEXT("get_material_world_position_offset_raw"), WorldPositionOffsetHLSL);
 		Params.Add(TEXT("get_material_previous_world_position_offset_raw"), WorldPositionOffsetHLSL);
 	
@@ -1862,22 +1829,22 @@ struct FTranslator : FMaterialIRToHLSLTranslation
 		if (Substrate::IsSubstrateEnabled())
 		{
 			// Adde default Substrate functions
-			UniformMaterialExpressions.Append(
-				TEXT("// Substrate: HiddenMaterialAssetConversion\n")
-				TEXT("#if TEMPLATE_USES_SUBSTRATE\n")
-				TEXT("void FSubstratePixelHeader::PreUpdateAllBSDFWithBottomUpOperatorVisit(float3 V) {}\n")
-				TEXT("void FSubstratePixelHeader::UpdateAllBSDFsOperatorCoverageTransmittance(FSubstrateIntegrationSettings Settings, float3 V)\n")
-				TEXT("{\n")
-				TEXT("#if SUBSTRATE_COMPILER_SUPPORTS_STRUCT_FORWARD_DECLARATION\n")
-				TEXT("\tSubstrateTree.UpdateSingleBSDFOperatorCoverageTransmittance(this, 0, Settings, V);\n")
-				TEXT("#else\n")
-				TEXT("\tUpdateSingleBSDFOperatorCoverageTransmittance(SubstrateTree, this, 0, Settings, V);\n")
-				TEXT("#endif\n")
-				TEXT("}\n")
-				TEXT("void FSubstratePixelHeader::UpdateAllOperatorsCoverageTransmittance() {}\n")
-				TEXT("void FSubstratePixelHeader::UpdateAllBSDFWithBottomUpOperatorVisit() {}\n")
-				TEXT("#endif // TEMPLATE_USES_SUBSTRATE\n")
-				TEXT("\n")
+			UniformMaterialExpressions.Append(TEXTVIEW(
+				"// Substrate: HiddenMaterialAssetConversion\n"
+				"#if TEMPLATE_USES_SUBSTRATE\n"
+				"void FSubstratePixelHeader::PreUpdateAllBSDFWithBottomUpOperatorVisit(float3 V) {}\n"
+				"void FSubstratePixelHeader::UpdateAllBSDFsOperatorCoverageTransmittance(FSubstrateIntegrationSettings Settings, float3 V)\n"
+				"{\n"
+				"#if SUBSTRATE_COMPILER_SUPPORTS_STRUCT_FORWARD_DECLARATION\n"
+				"\tSubstrateTree.UpdateSingleBSDFOperatorCoverageTransmittance(this, 0, Settings, V);\n"
+				"#else\n"
+				"\tUpdateSingleBSDFOperatorCoverageTransmittance(SubstrateTree, this, 0, Settings, V);\n"
+				"#endif\n"
+				"}\n"
+				"void FSubstratePixelHeader::UpdateAllOperatorsCoverageTransmittance() {}\n"
+				"void FSubstratePixelHeader::UpdateAllBSDFWithBottomUpOperatorVisit() {}\n"
+				"#endif // TEMPLATE_USES_SUBSTRATE\n"
+				"\n")
 			);
 		}
 		Params.Add(TEXT("uniform_material_expressions"), MoveTemp(UniformMaterialExpressions));
@@ -1891,28 +1858,28 @@ struct FTranslator : FMaterialIRToHLSLTranslation
 			{
 				if (PrimitiveType->NumRows == 1)
 				{
-					Printer << TEXT("FWSScalar");
+					Printer << TEXTVIEW("FWSScalar");
 				}
 				else if (PrimitiveType->NumColumns == 1)
 				{
-					Printer << TEXT("FWSVector") << PrimitiveType->NumRows;
+					Printer << TEXTVIEW("FWSVector") << PrimitiveType->NumRows;
 				}
 				else if (PrimitiveType->bIsLWCInverseMatrix)
 				{
-					Printer << TEXT("FWSInverseMatrix");
+					Printer << TEXTVIEW("FWSInverseMatrix");
 				}
 				else
 				{
-					Printer << TEXT("FWSMatrix");
+					Printer << TEXTVIEW("FWSMatrix");
 				}
 			}
 			else
 			{
 				switch (PrimitiveType->ScalarKind)
 				{
-					case MIR::EScalarKind::Bool:  Printer << TEXT("bool"); break;
-					case MIR::EScalarKind::Int:   Printer << TEXT("int"); break;
-					case MIR::EScalarKind::Float: Printer << TEXT("MaterialFloat"); break;
+					case MIR::EScalarKind::Bool:  Printer << TEXTVIEW("bool"); break;
+					case MIR::EScalarKind::Int:   Printer << TEXTVIEW("int"); break;
+					case MIR::EScalarKind::Float: Printer << TEXTVIEW("MaterialFloat"); break;
 				}
 
 				if (PrimitiveType->NumRows > 1)
@@ -1922,17 +1889,17 @@ struct FTranslator : FMaterialIRToHLSLTranslation
 			
 				if (PrimitiveType->NumColumns > 1)
 				{
-					Printer << TEXT("x") << PrimitiveType->NumColumns;
+					Printer << TEXTVIEW("x") << PrimitiveType->NumColumns;
 				}
 			}
 		}
 		else if (Type.IsSubstrateData())
 		{
-			Printer << TEXT("FSubstrateData");
+			Printer << TEXTVIEW("FSubstrateData");
 		}
 		else if (Type.IsVTPageTableResult())
 		{
-			Printer << TEXT("VTPageTableResult");
+			Printer << TEXTVIEW("VTPageTableResult");
 		}
 		else
 		{
@@ -2095,13 +2062,13 @@ void FMaterialIRToHLSLTranslation::Run(TMap<FString, FString>& OutParameters, FS
 {
 	OutParameters.Empty();
 
-	FTranslator Translator{ *this };
-	Translator.GeneratePixelAttributesHLSL();
-	Translator.GenerateVertexStageHLSL();
-	Translator.GenerateOtherStageHLSL(MIR::Stage_Pixel);
-	Translator.GenerateOtherStageHLSL(MIR::Stage_Compute);
-	Translator.GenerateTemplateStringParameters(OutParameters);
-	Translator.GetShaderCompilerEnvironment(OutEnvironment);
+	FPrivate Private{ *this };
+	Private.GeneratePixelAttributesHLSL();
+	Private.GenerateVertexStageHLSL();
+	Private.GenerateOtherStageHLSL(MIR::Stage_Pixel);
+	Private.GenerateOtherStageHLSL(MIR::Stage_Compute);
+	Private.GenerateTemplateStringParameters(OutParameters);
+	Private.GetShaderCompilerEnvironment(OutEnvironment);
 }
 
 #undef TAB
diff --git a/Engine/Source/Runtime/Engine/Public/Materials/MaterialIR.h b/Engine/Source/Runtime/Engine/Public/Materials/MaterialIR.h
--- a/Engine/Source/Runtime/Engine/Public/Materials/MaterialIR.h
+++ b/Engine/Source/Runtime/Engine/Public/Materials/MaterialIR.h
@@ -579,7 +579,11 @@ struct FInlineHLSL : TInstruction<VK_InlineHLSL, 0>
 		const FMaterialExternalCodeDeclaration* ExternalCodeDeclaration;
 
 		// The arbitrary inlined HLSL code snippet
-		const TCHAR* Code;
+		struct
+		{
+			const TCHAR* Code; // The string
+			uint32 CodeLen;    // The string length (excluding null byte)
+		};
 	};
 };
 
@@ -767,7 +771,7 @@ struct FSubscript : TInstruction<VK_Subscript, 1>
 	FValue* Arg;
 
 	// The subscript index, i.e. the index of the component to extract.
-	int Index;
+	int32 Index;
 };
 
 // A scalar construction instruction.
diff --git a/Engine/Source/Runtime/Engine/Public/Materials/MaterialIREmitter.h b/Engine/Source/Runtime/Engine/Public/Materials/MaterialIREmitter.h
--- a/Engine/Source/Runtime/Engine/Public/Materials/MaterialIREmitter.h
+++ b/Engine/Source/Runtime/Engine/Public/Materials/MaterialIREmitter.h
@@ -500,7 +500,7 @@ public:
 	// @param InputValues An array of input values required by the HLSL code.
 	// @param ValueFlags Flags that specify additional properties for this value (default is None).
 	// @param UsedGraphProperties Properties used during graph processing (default is None).
-	FValueRef InlineHLSL(FType Type, FString Code, TConstArrayView<FValueRef> InputValues, EValueFlags ValueFlags = EValueFlags::None, EGraphProperties UsedGraphProperties = EGraphProperties::None);
+	FValueRef InlineHLSL(FType Type, FStringView Code, TConstArrayView<FValueRef> InputValues, EValueFlags ValueFlags = EValueFlags::None, EGraphProperties UsedGraphProperties = EGraphProperties::None);
 
 	// Emit a new HLSL function inline within the material shader based on an external code declaration.
 	// @param InExternalCodeDeclaration The external HLSL code declaration to use for emission.
diff --git a/Engine/Source/Runtime/Engine/Public/Materials/MaterialIRModule.h b/Engine/Source/Runtime/Engine/Public/Materials/MaterialIRModule.h
--- a/Engine/Source/Runtime/Engine/Public/Materials/MaterialIRModule.h
+++ b/Engine/Source/Runtime/Engine/Public/Materials/MaterialIRModule.h
@@ -3,6 +3,7 @@
 
 #include "Materials/MaterialIR.h"
 #include "MaterialShared.h"
+#include "Misc/MemStack.h"
 
 #if WITH_EDITOR
 
@@ -76,8 +77,8 @@ public:
 	// Retrieves parameter metadata for a given parameter ID.
 	const FMaterialParameterMetadata& GetParameterMetadata(uint32 ParameterId) const { return ParameterIdToData[ParameterId].Value; }
 
-	// Stores a user-defined string and returns a pointer to it.
-	const TCHAR* PushUserString(FString InString);
+	// Stores a user-defined string and returns a view of it.
+	FStringView InternString(FStringView InString);
 
 	// Checks if the module is valid (i.e., contains no errors).
 	bool IsValid() const { return Errors.IsEmpty(); }
@@ -89,6 +90,9 @@ public:
 	void AddError(UMaterialExpression* Expression, FString Message);
 
 private:
+	// The chunk allocator used to contain child allocations (e.g., values, interned strings, etc)
+	FMemStackBase Allocator;
+
 	// Target shader platform.
 	EShaderPlatform ShaderPlatform;
 
diff --git a/Engine/Source/Runtime/Engine/Public/Materials/MaterialIRToHLSLTranslator.h b/Engine/Source/Runtime/Engine/Public/Materials/MaterialIRToHLSLTranslator.h
--- a/Engine/Source/Runtime/Engine/Public/Materials/MaterialIRToHLSLTranslator.h
+++ b/Engine/Source/Runtime/Engine/Public/Materials/MaterialIRToHLSLTranslator.h
@@ -6,22 +6,23 @@
 
 #if WITH_EDITOR
 
-//
+// Used to translate a generated Material MIR module into the material HLSL shader.
 struct FMaterialIRToHLSLTranslation
 {
-	//
+	// The Material IR module, which contains the intermediate representation of the material graph.
 	const FMaterialIRModule* Module{};
 
-	//
+	// The material being translated.
 	const FMaterial* Material{};
 	
-	//
+	// The set of static parameters assignements.
 	const FStaticParameterSet* StaticParameters{};
 	
-	//
+	// The translation target platform.
 	const ITargetPlatform* TargetPlatform{};
 	
-	//
+	// Executes the translation process, filling the material shader template parameters in OutParameters and configuring OutEnvironment for shader compilation.
+	// After this call, interpolate the MaterialTemplate shader with the list of parameters to get the final HLSL shader.
 	void Run(TMap<FString, FString>& OutParameters, FShaderCompilerEnvironment& OutEnvironment);
 };
 
diff --git a/Engine/Source/Runtime/Engine/Public/Materials/MaterialIRTypes.h b/Engine/Source/Runtime/Engine/Public/Materials/MaterialIRTypes.h
--- a/Engine/Source/Runtime/Engine/Public/Materials/MaterialIRTypes.h
+++ b/Engine/Source/Runtime/Engine/Public/Materials/MaterialIRTypes.h
@@ -60,7 +60,7 @@ struct FPrimitive
 	// Number of columns. For a vector or scalar, this is 1. For a matrix, this is > 1.
 	int8 NumColumns;
 
-	// If true, this represents an LWC inverse matrix. Only supported for EScalarKind::Double and 4x4 dimensions.
+	// If true, this represents an LWC inverse matrix. Only supported for EScalarKind::LWC and 4x4 dimensions.
 	bool bIsLWCInverseMatrix;
 
 	// Returns whether this primitive type is arithmetic (it supports arithmetic operations like addition).
@@ -296,6 +296,7 @@ private:
 
 // To avoid the compiler adding automatic padding, we must make sure that its size is a multiple of maximum alignment.
 static_assert(sizeof(MIR::FType) % alignof(void*) == 0);
+static_assert(sizeof(MIR::FType) == 16);
 
 } // namespace MIR
 
