[HWRT] Check all enabled skeletal mesh sections support dynamic ray tracing geometry.

Context:
- we were only checking whether first section vertex factory supported dynamic ray tracing geometry.
- this caused a crash when that first section was bDisabled.

Changes:
- Check all sections that are not bDisabled.
- Also fixes crash when first section is disabled.

diff --git a/Engine/Source/Runtime/Engine/Private/SkeletalMeshSceneProxy.cpp b/Engine/Source/Runtime/Engine/Private/SkeletalMeshSceneProxy.cpp
--- a/Engine/Source/Runtime/Engine/Private/SkeletalMeshSceneProxy.cpp
+++ b/Engine/Source/Runtime/Engine/Private/SkeletalMeshSceneProxy.cpp
@@ -998,6 +998,8 @@ void FSkeletalMeshSceneProxy::GetDynamicRayTracingInstances(FRayTracingInstanceC
 			
 	const bool bIgnorePreviewFilter = true;
 
+	bool bVFsSupportRayTracingDynamicGeometry = true;
+
 	for (FSkeletalMeshSectionIter Iter(LODIndex, *MeshObject, LODData, LODSection, bIgnorePreviewFilter); Iter; ++Iter)
 	{
 		const FSkelMeshRenderSection& Section = Iter.GetSection();
@@ -1008,18 +1010,26 @@ void FSkeletalMeshSceneProxy::GetDynamicRayTracingInstances(FRayTracingInstanceC
 		CreateBaseMeshBatch(Collector.GetReferenceView(), LODData, LODIndex, SectionIndex, SectionElementInfo, MeshBatch, ESkinVertexFactoryMode::RayTracing);
 
 		RayTracingInstance.Materials.Add(MeshBatch);
+
+		if (bAnySegmentUsesWorldPositionOffset && !Section.bDisabled)
+		{
+			const FVertexFactoryType* VertexFactoryType = MeshBatch.VertexFactory->GetType();
+
+			if (!ensureMsgf(VertexFactoryType->SupportsRayTracingDynamicGeometry(),
+				TEXT("Mesh uses world position offset, but the vertex factory does not support ray tracing dynamic geometry. ")
+				TEXT("MeshObject: %s, VertexFactory: %s."), *MeshObject->GetDebugName().ToString(), VertexFactoryType->GetName()))
+			{
+				bVFsSupportRayTracingDynamicGeometry = false;
+			}
+		}
 	}
 
 	RayTracingInstance.InstanceTransforms.Add(GetLocalToWorld());
+
 	const uint32 VertexBufferNumVertices = LODData.StaticVertexBuffers.PositionVertexBuffer.GetNumVertices();
 	const uint32 VertexBufferStride = LODData.StaticVertexBuffers.PositionVertexBuffer.GetStride();
 
-	const FVertexFactory* VertexFactory = MeshObject->GetSkinVertexFactory(Collector.GetReferenceView(), LODIndex, 0, ESkinVertexFactoryMode::RayTracing);
-	const FVertexFactoryType* VertexFactoryType = VertexFactory->GetType();
-	if (bAnySegmentUsesWorldPositionOffset
-		&& ensureMsgf(VertexFactoryType->SupportsRayTracingDynamicGeometry(),
-		              TEXT("Mesh uses world position offset, but the vertex factory does not support ray tracing dynamic geometry. ")
-		              TEXT("MeshObject: %s, VertexFactory: %s."), *MeshObject->GetDebugName().ToString(), VertexFactoryType->GetName()))
+	if (bAnySegmentUsesWorldPositionOffset && bVFsSupportRayTracingDynamicGeometry)
 	{
 		TArray<FRayTracingGeometrySegment> GeometrySections;
 		GeometrySections.Reserve(LODData.RenderSections.Num());
