Hardened the skeletal mesh updater to not crash if a scene reference is leaked (i.e. the backend is destructed prior to the updater instance itself).

diff --git a/Engine/Source/Runtime/Engine/Internal/SkeletalMeshUpdater.h b/Engine/Source/Runtime/Engine/Internal/SkeletalMeshUpdater.h
--- a/Engine/Source/Runtime/Engine/Internal/SkeletalMeshUpdater.h
+++ b/Engine/Source/Runtime/Engine/Internal/SkeletalMeshUpdater.h
@@ -567,7 +567,6 @@ class FSkeletalMeshUpdateChannel
 		{
 			TIntrusiveDoubleLinkedList<FBackend> List;
 			int32 Num = 0;
-			int32 NumPipeRefs = 0;
 		};
 
 		static FGlobalList& GetGlobalList();
@@ -579,6 +578,7 @@ class FSkeletalMeshUpdateChannel
 		virtual void Replay(FRHICommandList& RHICmdList, FSkeletalMeshUpdateChannel& Channel, FSkeletalMeshUpdatePacket& Packet) const = 0;
 
 		int32 GlobalListIndex = INDEX_NONE;
+		TArray<FSkeletalMeshUpdateChannel*> Channels;
 	};
 
 	template <typename T>
@@ -638,7 +638,7 @@ public:
 	}
 
 	static TArray<FSkeletalMeshUpdateChannel> GetChannels();
-	FSkeletalMeshUpdateChannel(const FBackend* InBackend);
+	FSkeletalMeshUpdateChannel(FBackend* InBackend);
 	ENGINE_API ~FSkeletalMeshUpdateChannel();
 	
 	FSkeletalMeshUpdateHandle Create(FSkeletalMeshObject* MeshObject);
@@ -648,11 +648,13 @@ public:
 private:
 	TUniquePtr<FSkeletalMeshUpdatePacket> CreatePacket()
 	{
+		checkf(Backend, TEXT("Backend was released but the channel is still being used."));
 		return Backend->CreatePacket();
 	}
 
 	void Replay(FRHICommandList& RHICmdList, FSkeletalMeshUpdatePacket& Packet)
 	{
+		checkf(Backend, TEXT("Backend was released but the channel is still being used."));
 		Backend->Replay(RHICmdList, *this, Packet);
 	}
 
@@ -679,7 +681,7 @@ private:
 	TUniquePtr<FOpQueue> OpQueue;
 	FOpStream OpStream;
 	FSlotRegistry SlotRegistry;
-	const FBackend* Backend = nullptr;
+	FBackend* Backend = nullptr;
 	int32 ChannelIndex;
 
 	friend FSkeletalMeshUpdater;
diff --git a/Engine/Source/Runtime/Engine/Private/SkeletalMeshUpdater.cpp b/Engine/Source/Runtime/Engine/Private/SkeletalMeshUpdater.cpp
--- a/Engine/Source/Runtime/Engine/Private/SkeletalMeshUpdater.cpp
+++ b/Engine/Source/Runtime/Engine/Private/SkeletalMeshUpdater.cpp
@@ -132,14 +132,17 @@ FSkeletalMeshUpdateChannel::FBackend::FGlobalList& FSkeletalMeshUpdateChannel::F
 FSkeletalMeshUpdateChannel::FBackend::FBackend()
 {
 	FGlobalList& GlobalList = GetGlobalList();
-	checkf(!GlobalList.NumPipeRefs, TEXT("A backend is being registered while a SkeletalMeshUpdateChannel is active. Only use DEFINE_SKELETAL_MESH_UPDATE_BACKEND to create backends statically."));
 	GlobalList.List.AddTail(this);
 	GlobalListIndex = GlobalList.Num++;
 }
 
 FSkeletalMeshUpdateChannel::FBackend::~FBackend()
 {
-	checkf(!GetGlobalList().NumPipeRefs, TEXT("A backend is being unregistered while a SkeletalMeshUpdateChannel is still alive."));
+	for (FSkeletalMeshUpdateChannel* Channel : Channels)
+	{
+		Channel->Backend = nullptr;
+	}
+	Channels.Empty();
 	Reset();
 }
 
@@ -182,16 +185,22 @@ TArray<FSkeletalMeshUpdateChannel> FSkeletalMeshUpdateChannel::GetChannels()
 	return Channels;
 }
 
-FSkeletalMeshUpdateChannel::FSkeletalMeshUpdateChannel(const FBackend* InBackend)
+FSkeletalMeshUpdateChannel::FSkeletalMeshUpdateChannel(FBackend* InBackend)
 	: OpQueue(new FOpQueue)
 	, Backend(InBackend)
 {
-	++FBackend::GetGlobalList().NumPipeRefs;
+	check(Backend);
+	Backend->Channels.Emplace(this);
 }
 
 FSkeletalMeshUpdateChannel::~FSkeletalMeshUpdateChannel()
 {
-	--FBackend::GetGlobalList().NumPipeRefs;
+	// Backend can be nulled out first.
+	if (Backend)
+	{
+		Backend->Channels.Remove(this);
+	}
+
 	check(!OpQueue);
 	check(OpStream.Ops.IsEmpty());
 
