Prevent error logging when cancelling a post load task for a static mesh that just got garbage collected. Skip the CreateCardRepresentation task.

diff --git a/Engine/Source/Runtime/Engine/Private/MeshCardRepresentation.cpp b/Engine/Source/Runtime/Engine/Private/MeshCardRepresentation.cpp
--- a/Engine/Source/Runtime/Engine/Private/MeshCardRepresentation.cpp
+++ b/Engine/Source/Runtime/Engine/Private/MeshCardRepresentation.cpp
@@ -305,6 +305,12 @@ FAsyncCardRepresentationTask* CreateCardRepresentationTask(
 			IMeshBuilderModule& MeshBuilderModule = IMeshBuilderModule::GetForPlatform(GetTargetPlatformManagerRef().GetRunningTargetPlatform());
 			if (!MeshBuilderModule.BuildMeshVertexPositions(Mesh, NewTask->SourceMeshData.TriangleIndices, NewTask->SourceMeshData.VertexPositions, NewTask->SourceMeshData.Sections))
 			{
+				if (UE::Tasks::FCancellationTokenScope::IsCurrentWorkCanceled())
+				{
+					delete NewTask;
+					return nullptr;
+				}
+
 				UE_LOG(LogStaticMesh, Error, TEXT("Failed to build static mesh. See previous line(s) for details."));
 			}
 		}
