SplineComponent - Add OnSplineSet event which is broadcast when the backing spline data is overwritten.

diff --git a/Engine/Source/Runtime/Engine/Classes/Components/SplineComponent.h b/Engine/Source/Runtime/Engine/Classes/Components/SplineComponent.h
--- a/Engine/Source/Runtime/Engine/Classes/Components/SplineComponent.h
+++ b/Engine/Source/Runtime/Engine/Classes/Components/SplineComponent.h
@@ -367,6 +367,9 @@ public:
 	ENGINE_API FSpline GetSpline() const;
 	ENGINE_API void SetSpline(const FSpline& InSpline);
 	ENGINE_API void SetSpline(const FSplineCurves& InSplineCurves);
+
+	UE_EXPERIMENTAL(5.7, TEXT("Spline change notification is experimental while the API is under active development."))
+	FSimpleMulticastDelegate& GetOnSplineSet() { return OnSplineSet; }
 	
 	/** Reset all spline data to match the current state of SplineCurves */
 	void SynchronizeSplines() { SetSpline(WarninglessSplineCurves()); }
@@ -961,6 +964,8 @@ protected:
 
 	/** Returns the length of the specified spline segment up to the parametric value given */
 	ENGINE_API float GetSegmentLength(const int32 Index, const float Param = 1.0f) const;
+
+	FSimpleMulticastDelegate OnSplineSet;
 	
 private:
 	
diff --git a/Engine/Source/Runtime/Engine/Private/Components/SplineComponent.cpp b/Engine/Source/Runtime/Engine/Private/Components/SplineComponent.cpp
--- a/Engine/Source/Runtime/Engine/Private/Components/SplineComponent.cpp
+++ b/Engine/Source/Runtime/Engine/Private/Components/SplineComponent.cpp
@@ -2590,12 +2590,16 @@ void USplineComponent::SetSpline(const FSpline& InSpline)
 {
 	WarninglessSplineCurves() = FSplineCurves::FromSplineInterface(InSpline);
 	Spline = InSpline;
+	
+	OnSplineSet.Broadcast();
 }
 
 void USplineComponent::SetSpline(const FSplineCurves& InSplineCurves)
 {
 	WarninglessSplineCurves() = InSplineCurves;
 	Spline = InSplineCurves;
+
+	OnSplineSet.Broadcast();
 }
 
 int32 USplineComponent::GetVersion() const
