SplineComponent - Revise CVar Logic - This change prevents using a disabled FSpline even if we tell the component not to use FSplineCurves.

This works by adding an IsEnabledStatic to FSpline which just checks the FSpline CVar to see whether or not any given FSpline is actualy implemented. The intent is to never use the FSpline in USplineComponent if it is not implemented, regardless of the value of SplineComponent.UseSplineCurves.

diff --git a/Engine/Source/Runtime/Engine/Classes/Curves/Spline.h b/Engine/Source/Runtime/Engine/Classes/Curves/Spline.h
--- a/Engine/Source/Runtime/Engine/Classes/Curves/Spline.h
+++ b/Engine/Source/Runtime/Engine/Classes/Curves/Spline.h
@@ -184,6 +184,11 @@ private:
 	void OnSplineImplementationChanged();
 	FDelegateHandle OnSplineImplementationChangedHandle;
 #endif
+
+public:
+
+	// True if a given FSpline will actually be implemented when instantiated.
+	static bool IsEnabledStatic();
 };
 
 template<>
diff --git a/Engine/Source/Runtime/Engine/Private/Components/SplineComponent.cpp b/Engine/Source/Runtime/Engine/Private/Components/SplineComponent.cpp
--- a/Engine/Source/Runtime/Engine/Private/Components/SplineComponent.cpp
+++ b/Engine/Source/Runtime/Engine/Private/Components/SplineComponent.cpp
@@ -40,6 +40,11 @@ FAutoConsoleVariableRef CVarUseSplineCurves(
 	GUseSplineCurves,
 	TEXT("When true, SplineCurves is the 'source of truth'.")
 );
+
+static bool ShouldUseSplineCurves()
+{
+	return GUseSplineCurves || FSpline::IsEnabledStatic() == false;
+}
 	
 }
 
@@ -55,7 +60,7 @@ void USplineComponent::PopulateFromLegacy()
 
 FName USplineComponent::GetSplinePropertyName()
 {
-	if (UE::SplineComponent::GUseSplineCurves)
+	if (UE::SplineComponent::ShouldUseSplineCurves())
 	{
 PRAGMA_DISABLE_DEPRECATION_WARNINGS
 		return GET_MEMBER_NAME_CHECKED(USplineComponent, SplineCurves);
@@ -422,7 +427,7 @@ float FSplineCurves::GetSegmentLength(const int32 Index, const float Param, bool
 
 float USplineComponent::GetSegmentLength(const int32 Index, const float Param) const
 {
-	return UE::SplineComponent::GUseSplineCurves
+	return UE::SplineComponent::ShouldUseSplineCurves()
 		? WarninglessSplineCurves().GetSegmentLength(Index, Param, bClosedLoop, GetComponentTransform().GetScale3D())
 		: Spline.GetSegmentLength(Index, Param, GetComponentTransform().GetScale3D());
 }
@@ -453,7 +458,7 @@ float USplineComponent::GetSegmentParamFromLength(const int32 Index, const float
 
 	// TODO: check if this works OK with delta InVal != 1.0f
 
-	int32 NumPoints = UE::SplineComponent::GUseSplineCurves
+	int32 NumPoints = UE::SplineComponent::ShouldUseSplineCurves()
 		? WarninglessSplineCurves().Position.Points.Num()
 		: Spline.GetNumControlPoints();
 	
@@ -467,7 +472,7 @@ float USplineComponent::GetSegmentParamFromLength(const int32 Index, const float
 	// two iterations of Newton-Raphson is enough
 	for (int32 Iteration = 0; Iteration < 2; ++Iteration)
 	{
-		float TangentMagnitude = UE::SplineComponent::GUseSplineCurves
+		float TangentMagnitude = UE::SplineComponent::ShouldUseSplineCurves()
 			? WarninglessSplineCurves().Position.EvalDerivative(Index + Param, FVector::ZeroVector).Size()
 			: Spline.EvaluateDerivative(Index + Param).Size();
 		
@@ -483,7 +488,7 @@ float USplineComponent::GetSegmentParamFromLength(const int32 Index, const float
 
 FVector USplineComponent::GetLocationAtSplineInputKey(float InKey, ESplineCoordinateSpace::Type CoordinateSpace) const
 {
-	FVector Location = UE::SplineComponent::GUseSplineCurves
+	FVector Location = UE::SplineComponent::ShouldUseSplineCurves()
 		? WarninglessSplineCurves().Position.Eval(InKey, FVector::ZeroVector)
 		: Spline.Evaluate(InKey);
 
@@ -497,7 +502,7 @@ FVector USplineComponent::GetLocationAtSplineInputKey(float InKey, ESplineCoordi
 
 FVector USplineComponent::GetTangentAtSplineInputKey(float InKey, ESplineCoordinateSpace::Type CoordinateSpace) const
 {
-	FVector Tangent = UE::SplineComponent::GUseSplineCurves
+	FVector Tangent = UE::SplineComponent::ShouldUseSplineCurves()
 		? WarninglessSplineCurves().Position.EvalDerivative(InKey, FVector::ZeroVector)
 		: Spline.EvaluateDerivative(InKey);
 
@@ -511,7 +516,7 @@ FVector USplineComponent::GetTangentAtSplineInputKey(float InKey, ESplineCoordin
 
 FVector USplineComponent::GetDirectionAtSplineInputKey(float InKey, ESplineCoordinateSpace::Type CoordinateSpace) const
 {
-	FVector Direction = UE::SplineComponent::GUseSplineCurves
+	FVector Direction = UE::SplineComponent::ShouldUseSplineCurves()
 		? WarninglessSplineCurves().Position.EvalDerivative(InKey, FVector::ZeroVector).GetSafeNormal()
 		: Spline.EvaluateDerivative(InKey).GetSafeNormal();
 
@@ -531,12 +536,12 @@ FRotator USplineComponent::GetRotationAtSplineInputKey(float InKey, ESplineCoord
 
 FQuat USplineComponent::GetQuaternionAtSplineInputKey(float InKey, ESplineCoordinateSpace::Type CoordinateSpace) const
 {
-	FQuat Quat = UE::SplineComponent::GUseSplineCurves
+	FQuat Quat = UE::SplineComponent::ShouldUseSplineCurves()
 		? WarninglessSplineCurves().Rotation.Eval(InKey, FQuat::Identity)
 		: Spline.EvaluateRotation(InKey);
 	Quat.Normalize();
 
-	FVector Direction = UE::SplineComponent::GUseSplineCurves
+	FVector Direction = UE::SplineComponent::ShouldUseSplineCurves()
 		? WarninglessSplineCurves().Position.EvalDerivative(InKey, FVector::ZeroVector).GetSafeNormal()
 		: Spline.EvaluateDerivative(InKey).GetSafeNormal();
 	
@@ -602,14 +607,14 @@ float USplineComponent::GetRollAtSplineInputKey(float InKey, ESplineCoordinateSp
 
 FVector USplineComponent::GetScaleAtSplineInputKey(float InKey) const
 {
-	return UE::SplineComponent::GUseSplineCurves
+	return UE::SplineComponent::ShouldUseSplineCurves()
 		? WarninglessSplineCurves().Scale.Eval(InKey, FVector(1.0f))
 		: Spline.EvaluateScale(InKey);
 }
 
 float USplineComponent::GetDistanceAlongSplineAtSplineInputKey(float InKey) const
 {
-	const int32 NumPoints = UE::SplineComponent::GUseSplineCurves
+	const int32 NumPoints = UE::SplineComponent::ShouldUseSplineCurves()
 		? WarninglessSplineCurves().Position.Points.Num()
 		: Spline.GetNumControlPoints();
 	
@@ -619,7 +624,7 @@ float USplineComponent::GetDistanceAlongSplineAtSplineInputKey(float InKey) cons
 	{
 		float Distance;
 		
-		if (UE::SplineComponent::GUseSplineCurves)
+		if (UE::SplineComponent::ShouldUseSplineCurves())
 		{
 			const int32 PointIndex = FMath::FloorToInt(InKey);
 			const float Fraction = InKey - PointIndex;
@@ -646,7 +651,7 @@ float USplineComponent::GetDistanceAlongSplineAtLocation(const FVector& InLocati
 	const FVector LocalLocation = (CoordinateSpace == ESplineCoordinateSpace::World) ? GetComponentTransform().InverseTransformPosition(InLocation) : InLocation;
 	float Dummy;
 	
-	const float Key = UE::SplineComponent::GUseSplineCurves
+	const float Key = UE::SplineComponent::ShouldUseSplineCurves()
 		? WarninglessSplineCurves().Position.FindNearest(LocalLocation, Dummy)
 		: Spline.FindNearest(LocalLocation, Dummy);
 	
@@ -655,7 +660,7 @@ float USplineComponent::GetDistanceAlongSplineAtLocation(const FVector& InLocati
 
 int32 USplineComponent::GetNumberOfPropertyValues(FName PropertyName) const
 {
-	if (UE::SplineComponent::GUseSplineCurves == false && Spline.SupportsAttributes())
+	if (UE::SplineComponent::ShouldUseSplineCurves() == false && Spline.SupportsAttributes())
 	{
 		return Spline.NumAttributeValues<float>(PropertyName);
 	}
@@ -665,7 +670,7 @@ int32 USplineComponent::GetNumberOfPropertyValues(FName PropertyName) const
 
 TArray<FName> USplineComponent::GetFloatPropertyChannels() const
 {
-	if (UE::SplineComponent::GUseSplineCurves == false && Spline.SupportsAttributes())
+	if (UE::SplineComponent::ShouldUseSplineCurves() == false && Spline.SupportsAttributes())
 	{
 		return Spline.GetFloatPropertyChannels();
 	}
@@ -772,7 +777,7 @@ void SetPropertyAtIndex(FSpline& Spline, int32 Index, float Value, FName Propert
 
 void USplineComponent::RemovePropertyAtIndex(int32 Index, FName PropertyName)
 {
-	if (UE::SplineComponent::GUseSplineCurves == false && Spline.SupportsAttributes())
+	if (UE::SplineComponent::ShouldUseSplineCurves() == false && Spline.SupportsAttributes())
 	{
 		Spline.RemoveAttributeValue<float>(Index, PropertyName);
 	}
@@ -780,7 +785,7 @@ void USplineComponent::RemovePropertyAtIndex(int32 Index, FName PropertyName)
 
 bool USplineComponent::SupportsAttributes() const
 {
-	return UE::SplineComponent::GUseSplineCurves == false && Spline.SupportsAttributes();
+	return UE::SplineComponent::ShouldUseSplineCurves() == false && Spline.SupportsAttributes();
 }
 
 void USplineComponent::CreateFloatPropertyChannel(FName PropertyName)
@@ -1292,7 +1297,7 @@ ESplinePointType::Type USplineComponent::GetSplinePointType(int32 PointIndex) co
 	
 	if (PointIndex >= 0 && PointIndex < NumPoints)
 	{
-		const EInterpCurveMode Mode = UE::SplineComponent::GUseSplineCurves
+		const EInterpCurveMode Mode = UE::SplineComponent::ShouldUseSplineCurves()
 			? WarninglessSplineCurves().Position.Points[PointIndex].InterpMode.GetValue()
 			: Spline.GetSplinePointType(PointIndex);
 		
@@ -1320,7 +1325,7 @@ void USplineComponent::SetSplinePointType(int32 PointIndex, ESplinePointType::Ty
 
 int32 USplineComponent::GetNumberOfSplinePoints() const
 {
-	return UE::SplineComponent::GUseSplineCurves
+	return UE::SplineComponent::ShouldUseSplineCurves()
 		? WarninglessSplineCurves().Position.Points.Num()
 		: Spline.GetNumControlPoints();
 }
@@ -1340,7 +1345,7 @@ float USplineComponent::GetInputKeyValueAtSplinePoint(int32 PointIndex) const
 
 	PointIndex = GetClampedIndex(PointIndex);
 	
-	return UE::SplineComponent::GUseSplineCurves
+	return UE::SplineComponent::ShouldUseSplineCurves()
 		? WarninglessSplineCurves().Position.Points[PointIndex].InVal
 		: Spline.GetParameterAtIndex(PointIndex);
 }
@@ -1367,7 +1372,7 @@ FVector USplineComponent::GetLocationAtSplinePoint(int32 PointIndex, ESplineCoor
 
 	PointIndex = GetClampedIndex(PointIndex);
 	
-	const FVector Location = UE::SplineComponent::GUseSplineCurves
+	const FVector Location = UE::SplineComponent::ShouldUseSplineCurves()
 		? WarninglessSplineCurves().Position.Points[PointIndex].OutVal
 		: Spline.GetLocation(PointIndex);
 	
@@ -1393,7 +1398,7 @@ FVector USplineComponent::GetArriveTangentAtSplinePoint(int32 PointIndex, ESplin
 
 	PointIndex = GetClampedIndex(PointIndex);
 
-	const FVector Tangent = UE::SplineComponent::GUseSplineCurves
+	const FVector Tangent = UE::SplineComponent::ShouldUseSplineCurves()
 		? WarninglessSplineCurves().Position.Points[PointIndex].ArriveTangent
 		: Spline.GetInTangent(PointIndex);
 
@@ -1409,7 +1414,7 @@ FVector USplineComponent::GetLeaveTangentAtSplinePoint(int32 PointIndex, ESpline
 
 	PointIndex = GetClampedIndex(PointIndex);
 
-	const FVector Tangent = UE::SplineComponent::GUseSplineCurves
+	const FVector Tangent = UE::SplineComponent::ShouldUseSplineCurves()
 		? WarninglessSplineCurves().Position.Points[PointIndex].LeaveTangent
 		: Spline.GetOutTangent(PointIndex);
 
@@ -1450,7 +1455,7 @@ FVector USplineComponent::GetScaleAtSplinePoint(int32 PointIndex) const
 
 	PointIndex = GetClampedIndex(PointIndex);
 
-	return UE::SplineComponent::GUseSplineCurves
+	return UE::SplineComponent::ShouldUseSplineCurves()
 		? WarninglessSplineCurves().Scale.Points[PointIndex].OutVal
 		: Spline.GetScale(PointIndex);
 }
@@ -1475,7 +1480,7 @@ float USplineComponent::GetDistanceAlongSplineAtSplinePoint(int32 PointIndex) co
 		return GetSplineLength();
 	}
 	
-	if (UE::SplineComponent::GUseSplineCurves)
+	if (UE::SplineComponent::ShouldUseSplineCurves())
 	{
 		const int32 NumPoints = GetNumberOfSplinePoints();
         const int32 NumSegments = bClosedLoop ? NumPoints : NumPoints - 1;
@@ -1512,7 +1517,7 @@ float FSplineCurves::GetSplineLength() const
 
 float USplineComponent::GetSplineLength() const
 {
-	return UE::SplineComponent::GUseSplineCurves
+	return UE::SplineComponent::ShouldUseSplineCurves()
 		? WarninglessSplineCurves().GetSplineLength()
 		: Spline.GetSplineLength();
 }
@@ -1557,7 +1562,7 @@ float USplineComponent::GetInputKeyValueAtDistanceAlongSpline(float Distance) co
 		return 0.0f;
 	}
 	
-	return UE::SplineComponent::GUseSplineCurves
+	return UE::SplineComponent::ShouldUseSplineCurves()
 		? WarninglessSplineCurves().ReparamTable.Eval(Distance, 0.0f)
 		: Spline.GetParameterAtDistance(Distance);
 }
@@ -1572,14 +1577,14 @@ float USplineComponent::GetTimeAtDistanceAlongSpline(float Distance) const
 	}
 
 	const float TimeMultiplier = Duration / (bClosedLoop ? NumPoints : (NumPoints - 1.0f));
-	return (UE::SplineComponent::GUseSplineCurves
+	return (UE::SplineComponent::ShouldUseSplineCurves()
 		? WarninglessSplineCurves().ReparamTable.Eval(Distance, 0.0f)
 		: Spline.GetParameterAtDistance(Distance)) * TimeMultiplier;
 }
 
 FVector USplineComponent::GetLocationAtDistanceAlongSpline(float Distance, ESplineCoordinateSpace::Type CoordinateSpace) const
 {
-	const float Param = UE::SplineComponent::GUseSplineCurves
+	const float Param = UE::SplineComponent::ShouldUseSplineCurves()
 		? WarninglessSplineCurves().ReparamTable.Eval(Distance, 0.0f)
 		: Spline.GetParameterAtDistance(Distance);
 	return GetLocationAtSplineInputKey(Param, CoordinateSpace);
@@ -1587,7 +1592,7 @@ FVector USplineComponent::GetLocationAtDistanceAlongSpline(float Distance, ESpli
 
 FVector USplineComponent::GetTangentAtDistanceAlongSpline(float Distance, ESplineCoordinateSpace::Type CoordinateSpace) const
 {
-	const float Param = UE::SplineComponent::GUseSplineCurves
+	const float Param = UE::SplineComponent::ShouldUseSplineCurves()
 		? WarninglessSplineCurves().ReparamTable.Eval(Distance, 0.0f)
 		: Spline.GetParameterAtDistance(Distance);
 	return GetTangentAtSplineInputKey(Param, CoordinateSpace);
@@ -1596,7 +1601,7 @@ FVector USplineComponent::GetTangentAtDistanceAlongSpline(float Distance, ESplin
 
 FVector USplineComponent::GetDirectionAtDistanceAlongSpline(float Distance, ESplineCoordinateSpace::Type CoordinateSpace) const
 {
-	const float Param = UE::SplineComponent::GUseSplineCurves
+	const float Param = UE::SplineComponent::ShouldUseSplineCurves()
 		? WarninglessSplineCurves().ReparamTable.Eval(Distance, 0.0f)
 		: Spline.GetParameterAtDistance(Distance);
 	return GetDirectionAtSplineInputKey(Param, CoordinateSpace);
@@ -1604,7 +1609,7 @@ FVector USplineComponent::GetDirectionAtDistanceAlongSpline(float Distance, ESpl
 
 FQuat USplineComponent::GetQuaternionAtDistanceAlongSpline(float Distance, ESplineCoordinateSpace::Type CoordinateSpace) const
 {
-	const float Param = UE::SplineComponent::GUseSplineCurves
+	const float Param = UE::SplineComponent::ShouldUseSplineCurves()
 		? WarninglessSplineCurves().ReparamTable.Eval(Distance, 0.0f)
 		: Spline.GetParameterAtDistance(Distance);
 	return GetQuaternionAtSplineInputKey(Param, CoordinateSpace);
@@ -1612,7 +1617,7 @@ FQuat USplineComponent::GetQuaternionAtDistanceAlongSpline(float Distance, ESpli
 
 FRotator USplineComponent::GetRotationAtDistanceAlongSpline(float Distance, ESplineCoordinateSpace::Type CoordinateSpace) const
 {
-	const float Param = UE::SplineComponent::GUseSplineCurves
+	const float Param = UE::SplineComponent::ShouldUseSplineCurves()
 		? WarninglessSplineCurves().ReparamTable.Eval(Distance, 0.0f)
 		: Spline.GetParameterAtDistance(Distance);
 	return GetRotationAtSplineInputKey(Param, CoordinateSpace);
@@ -1620,7 +1625,7 @@ FRotator USplineComponent::GetRotationAtDistanceAlongSpline(float Distance, ESpl
 
 FVector USplineComponent::GetUpVectorAtDistanceAlongSpline(float Distance, ESplineCoordinateSpace::Type CoordinateSpace) const
 {
-	const float Param = UE::SplineComponent::GUseSplineCurves
+	const float Param = UE::SplineComponent::ShouldUseSplineCurves()
 		? WarninglessSplineCurves().ReparamTable.Eval(Distance, 0.0f)
 		: Spline.GetParameterAtDistance(Distance);
 	return GetUpVectorAtSplineInputKey(Param, CoordinateSpace);
@@ -1628,7 +1633,7 @@ FVector USplineComponent::GetUpVectorAtDistanceAlongSpline(float Distance, ESpli
 
 FVector USplineComponent::GetRightVectorAtDistanceAlongSpline(float Distance, ESplineCoordinateSpace::Type CoordinateSpace) const
 {
-	const float Param = UE::SplineComponent::GUseSplineCurves
+	const float Param = UE::SplineComponent::ShouldUseSplineCurves()
 		? WarninglessSplineCurves().ReparamTable.Eval(Distance, 0.0f)
 		: Spline.GetParameterAtDistance(Distance);
 	return GetRightVectorAtSplineInputKey(Param, CoordinateSpace);
@@ -1636,7 +1641,7 @@ FVector USplineComponent::GetRightVectorAtDistanceAlongSpline(float Distance, ES
 
 float USplineComponent::GetRollAtDistanceAlongSpline(float Distance, ESplineCoordinateSpace::Type CoordinateSpace) const
 {
-	const float Param = UE::SplineComponent::GUseSplineCurves
+	const float Param = UE::SplineComponent::ShouldUseSplineCurves()
 		? WarninglessSplineCurves().ReparamTable.Eval(Distance, 0.0f)
 		: Spline.GetParameterAtDistance(Distance);
 	return GetRollAtSplineInputKey(Param, CoordinateSpace);
@@ -1644,7 +1649,7 @@ float USplineComponent::GetRollAtDistanceAlongSpline(float Distance, ESplineCoor
 
 FVector USplineComponent::GetScaleAtDistanceAlongSpline(float Distance) const
 {
-	const float Param = UE::SplineComponent::GUseSplineCurves
+	const float Param = UE::SplineComponent::ShouldUseSplineCurves()
 		? WarninglessSplineCurves().ReparamTable.Eval(Distance, 0.0f)
 		: Spline.GetParameterAtDistance(Distance);
 	return GetScaleAtSplineInputKey(Param);
@@ -1652,7 +1657,7 @@ FVector USplineComponent::GetScaleAtDistanceAlongSpline(float Distance) const
 
 FTransform USplineComponent::GetTransformAtDistanceAlongSpline(float Distance, ESplineCoordinateSpace::Type CoordinateSpace, bool bUseScale) const
 {
-	const float Param = UE::SplineComponent::GUseSplineCurves
+	const float Param = UE::SplineComponent::ShouldUseSplineCurves()
 		? WarninglessSplineCurves().ReparamTable.Eval(Distance, 0.0f)
 		: Spline.GetParameterAtDistance(Distance);
 	return GetTransformAtSplineInputKey(Param, CoordinateSpace, bUseScale);
@@ -1862,7 +1867,7 @@ float USplineComponent::FindInputKeyClosestToWorldLocation(const FVector& WorldL
 {
 	const FVector LocalLocation = GetComponentTransform().InverseTransformPosition(WorldLocation);
 	float Dummy;
-	return UE::SplineComponent::GUseSplineCurves
+	return UE::SplineComponent::ShouldUseSplineCurves()
 		? WarninglessSplineCurves().Position.FindNearest(LocalLocation, Dummy)
 		: Spline.FindNearest(LocalLocation, Dummy);
 }
@@ -1871,7 +1876,7 @@ float USplineComponent::FindInputKeyOnSegmentClosestToWorldLocation(const FVecto
 {
 	const FVector LocalLocation = GetComponentTransform().InverseTransformPosition(WorldLocation);
 	float Dummy;
-	return UE::SplineComponent::GUseSplineCurves
+	return UE::SplineComponent::ShouldUseSplineCurves()
 		? WarninglessSplineCurves().Position.FindNearestOnSegment(LocalLocation, Index, Dummy)
 		: Spline.FindNearestOnSegment(LocalLocation, Index, Dummy);
 }
@@ -2161,7 +2166,7 @@ bool USplineComponent::ConvertSplineToPolyline_InDistanceRange(ESplineCoordinate
 	
 	int32 SegmentStart;
 	int32 SegmentEnd;
-	if (UE::SplineComponent::GUseSplineCurves)
+	if (UE::SplineComponent::ShouldUseSplineCurves())
 	{
 		int32 StartIndex = WarninglessSplineCurves().ReparamTable.GetPointIndexForInputValue(RangeStart);
 		int32 EndIndex = WarninglessSplineCurves().ReparamTable.GetPointIndexForInputValue(RangeEnd);
@@ -2327,7 +2332,7 @@ FPrimitiveSceneProxy* USplineComponent::CreateSceneProxy()
 		FSplineSceneProxy(const USplineComponent* InComponent)
 			: FPrimitiveSceneProxy(InComponent)
 			, bDrawDebug(InComponent->bDrawDebug)
-			, SplineInfo(UE::SplineComponent::GUseSplineCurves ? InComponent->WarninglessSplineCurves().Position : InComponent->Spline.GetSplinePointsPosition())
+			, SplineInfo(UE::SplineComponent::ShouldUseSplineCurves() ? InComponent->WarninglessSplineCurves().Position : InComponent->Spline.GetSplinePointsPosition())
 #if WITH_EDITORONLY_DATA
 			, LineColor(InComponent->EditorUnselectedSplineSegmentColor)
 #else
@@ -2474,7 +2479,7 @@ void USplineComponent::Draw(FPrimitiveDrawInterface* PDI, const FSceneView* View
 
 FBoxSphereBounds USplineComponent::CalcBounds(const FTransform& LocalToWorld) const
 {
-	const auto& InterpCurve = UE::SplineComponent::GUseSplineCurves ? WarninglessSplineCurves().Position : Spline.GetSplinePointsPosition();
+	const auto& InterpCurve = UE::SplineComponent::ShouldUseSplineCurves() ? WarninglessSplineCurves().Position : Spline.GetSplinePointsPosition();
 
 #if SPLINE_FAST_BOUNDS_CALCULATION
 	FBox BoundingBox(0);
@@ -2528,7 +2533,7 @@ bool USplineComponent::GetIgnoreBoundsForEditorFocus() const
 
 TStructOnScope<FActorComponentInstanceData> USplineComponent::GetComponentInstanceData() const
 {
-	if (UE::SplineComponent::GUseSplineCurves)
+	if (UE::SplineComponent::ShouldUseSplineCurves())
 	{
 		TStructOnScope<FActorComponentInstanceData> InstanceData = MakeStructOnScope<FActorComponentInstanceData, FSplineInstanceData>(this);
 		FSplineInstanceData* SplineInstanceData = InstanceData.Cast<FSplineInstanceData>();
@@ -2567,14 +2572,14 @@ TStructOnScope<FActorComponentInstanceData> USplineComponent::GetComponentInstan
 
 FSplineCurves USplineComponent::GetSplineCurves() const
 {
-	return UE::SplineComponent::GUseSplineCurves
+	return UE::SplineComponent::ShouldUseSplineCurves()
 		? WarninglessSplineCurves()
 		: FSplineCurves::FromSplineInterface(Spline);
 }
 
 FSpline USplineComponent::GetSpline() const
 {
-	if (UE::SplineComponent::GUseSplineCurves)
+	if (UE::SplineComponent::ShouldUseSplineCurves())
 	{
 		FSpline Intermediate;
 		Intermediate = SplineCurves;
@@ -2604,28 +2609,28 @@ void USplineComponent::SetSpline(const FSplineCurves& InSplineCurves)
 
 int32 USplineComponent::GetVersion() const
 {
-	return UE::SplineComponent::GUseSplineCurves
+	return UE::SplineComponent::ShouldUseSplineCurves()
 		? WarninglessSplineCurves().Version
 		: Spline.GetVersion();
 }
 
 const FInterpCurveVector& USplineComponent::GetSplinePointsPosition() const
 {
-	return UE::SplineComponent::GUseSplineCurves
+	return UE::SplineComponent::ShouldUseSplineCurves()
 		? WarninglessSplineCurves().Position
 		: Spline.GetSplinePointsPosition();
 }
 
 const FInterpCurveQuat& USplineComponent::GetSplinePointsRotation() const
 {
-	return UE::SplineComponent::GUseSplineCurves
+	return UE::SplineComponent::ShouldUseSplineCurves()
 		? WarninglessSplineCurves().Rotation
 		: Spline.GetSplinePointsRotation();
 }
 
 const FInterpCurveVector& USplineComponent::GetSplinePointsScale() const
 {
-	return UE::SplineComponent::GUseSplineCurves
+	return UE::SplineComponent::ShouldUseSplineCurves()
 		? WarninglessSplineCurves().Scale
 		: Spline.GetSplinePointsScale();
 }
diff --git a/Engine/Source/Runtime/Engine/Private/Curves/Spline.cpp b/Engine/Source/Runtime/Engine/Private/Curves/Spline.cpp
--- a/Engine/Source/Runtime/Engine/Private/Curves/Spline.cpp
+++ b/Engine/Source/Runtime/Engine/Private/Curves/Spline.cpp
@@ -1563,6 +1563,11 @@ TSharedPtr<TSplineInterface<FVector>> FSpline::GetSplineInterface() const
 	return nullptr;
 }
 
+bool FSpline::IsEnabledStatic()
+{
+	return UE::Spline::GImplementation != 0;
+}
+
 #if WITH_EDITOR
 void FSpline::OnSplineImplementationChanged()
 {
