SplineComponent - Fix incorrect usage of UE_EXPERIMENTAL

diff --git a/Engine/Source/Runtime/Engine/Classes/Components/SplineComponent.h b/Engine/Source/Runtime/Engine/Classes/Components/SplineComponent.h
--- a/Engine/Source/Runtime/Engine/Classes/Components/SplineComponent.h
+++ b/Engine/Source/Runtime/Engine/Classes/Components/SplineComponent.h
@@ -368,7 +368,7 @@ public:
 	ENGINE_API void SetSpline(const FSpline& InSpline);
 	ENGINE_API void SetSpline(const FSplineCurves& InSplineCurves);
 
-	UE_EXPERIMENTAL(5.7, TEXT("Spline change notification is experimental while the API is under active development."))
+	UE_EXPERIMENTAL(5.7, "Spline change notification is experimental while the API is under active development.")
 	FSimpleMulticastDelegate& GetOnSplineSet() { return OnSplineSet; }
 	
 	/** Reset all spline data to match the current state of SplineCurves */
