NetEmulation
* Added netEmulation setting to simulate buffer bloat on game traffic
* Set PktBufferBloatInMS for outgoing packets and PktIncomingBufferBloatInMS for incoming packets
* Ex: 'netEmulatation.PktBufferBloatInMS 1000' will apply 1sec of buffer bloat on outgoing traffic.
* netEmulation.BufferBloatCooldownTimeInMS can also be used to set a cooldown period between each buffer bloat occurence.
* Ex: 'netEmulation.BufferBloatCooldownTimeInMS 2000' means buffer bloat is not applied for 2secs after a buffer flush.

diff --git a/Engine/Source/Runtime/Engine/Classes/Engine/NetConnection.h b/Engine/Source/Runtime/Engine/Classes/Engine/NetConnection.h
--- a/Engine/Source/Runtime/Engine/Classes/Engine/NetConnection.h
+++ b/Engine/Source/Runtime/Engine/Classes/Engine/NetConnection.h
@@ -899,6 +899,15 @@ private:
 	/** set to true when already delayed packets are reinjected */
 	bool bIsReinjectingDelayedPackets = false;
 
+	/** Time when the next flush of outgoing delayed packets will occur. */
+	double NextBufferBloatFlushOut = 0.0;
+	/** Time when the next flush of incoming delayed packets will occur. */
+	double NextBufferBloatFlushIn = 0.0;
+	
+	/** Time at which we restart the buffer bloat emulation. */
+	double NextBufferBloatActiveTimeOut = 0.0;
+	double NextBufferBloatActiveTimeIn = 0.0;
+
 	/** Process incoming packets that have been delayed for long enough */
 	void ReinjectDelayedPackets();
 
@@ -1688,7 +1697,8 @@ private:
 	*	Returns true if the packet was emulated and false when the packet must be sent via the normal path
 	*/
 #if DO_ENABLE_NET_TEST
-	bool CheckOutgoingPacketEmulation(FOutPacketTraits& Traits);
+	bool CheckOutgoingPacketEmulation(double CurrentTimeInS, FOutPacketTraits& Traits);
+	bool CheckIncomingPacketEmulation(double CurrentTimeInS, FBitReader& Reader);
 #endif
 
 	/** Write packetHeader */
diff --git a/Engine/Source/Runtime/Engine/Classes/Engine/NetDriver.h b/Engine/Source/Runtime/Engine/Classes/Engine/NetDriver.h
--- a/Engine/Source/Runtime/Engine/Classes/Engine/NetDriver.h
+++ b/Engine/Source/Runtime/Engine/Classes/Engine/NetDriver.h
@@ -572,6 +572,14 @@ struct FPacketSimulationSettings
 	UPROPERTY(EditAnywhere, Category = "Simulation Settings")
 	int32	PktIncomingFrameDelay = 0;
 
+	/** Time in milliseconds during which outgoing packets will get delayed before being sent all at once. */
+	UPROPERTY(EditAnywhere, Category = "Simulation Settings")
+	int32 PktBufferBloatInMS = 0;
+
+	/** Time in milliseconds during which incoming packets will get delayed before being processed all at once. */
+	UPROPERTY(EditAnywhere, Category = "Simulation Settings")
+	int32 PktIncomingBufferBloatInMS = 0;
+
 	/** reads in settings from the .ini file 
 	 * @note: overwrites all previous settings
 	 */
diff --git a/Engine/Source/Runtime/Engine/Private/Net/NetEmulationHelper.cpp b/Engine/Source/Runtime/Engine/Private/Net/NetEmulationHelper.cpp
--- a/Engine/Source/Runtime/Engine/Private/Net/NetEmulationHelper.cpp
+++ b/Engine/Source/Runtime/Engine/Private/Net/NetEmulationHelper.cpp
@@ -4,11 +4,14 @@
 
 #include "Engine/Engine.h"
 #include "Engine/NetworkSettings.h"
+#include "Engine/NetDriver.h"
+#include "Misc/ConfigCacheIni.h"
 
 #if DO_ENABLE_NET_TEST
 
 namespace UE::Net::Private::NetEmulationHelper
 {
+
 /** Global that stores the network emulation values outside the NetDriver lifetime */
 static TOptional<FPacketSimulationSettings> PersistentPacketSimulationSettings;
 
@@ -358,8 +361,286 @@ FConsoleCommandWithWorldArgsAndOutputDeviceDelegate::CreateStatic([](const TArra
 	BUILD_NETEMULATION_CONSOLE_COMMAND(PktIncomingLagMax, "Sets maximum incoming packet latency");
 	BUILD_NETEMULATION_CONSOLE_COMMAND(PktIncomingLoss, "Simulates incoming packet loss");
 	BUILD_NETEMULATION_CONSOLE_COMMAND(PktJitter, "Simulates outgoing packet jitter");
+	BUILD_NETEMULATION_CONSOLE_COMMAND(PktBufferBloatInMS, "Simulates outgoing buffer bloat");
+	BUILD_NETEMULATION_CONSOLE_COMMAND(PktIncomingBufferBloatInMS, "Simulates incoming buffer bloat");
 
 
 } // end namespace UE::Net::Private::NetEmulationHelper
 
+
+class FPacketSimulationConsoleCommandVisitor
+{
+public:
+	static void OnPacketSimulationConsoleCommand(const TCHAR* Name, IConsoleObject* CVar, TArray<IConsoleObject*>& Sink)
+	{
+		Sink.Add(CVar);
+	}
+};
+
+/** reads in settings from the .ini file
+ * @note: overwrites all previous settings
+ */
+void FPacketSimulationSettings::LoadConfig(const TCHAR* OptionalQualifier)
+{
+#if DO_ENABLE_NET_TEST
+	ConfigHelperInt(TEXT("PktLoss"), PktLoss, OptionalQualifier);
+
+	ConfigHelperInt(TEXT("PktLossMinSize"), PktLossMinSize, OptionalQualifier);
+	ConfigHelperInt(TEXT("PktLossMaxSize"), PktLossMaxSize, OptionalQualifier);
+
+	bool InPktOrder = !!PktOrder;
+	ConfigHelperBool(TEXT("PktOrder"), InPktOrder, OptionalQualifier);
+	PktOrder = int32(InPktOrder);
+
+	ConfigHelperInt(TEXT("PktLag"), PktLag, OptionalQualifier);
+	ConfigHelperInt(TEXT("PktLagVariance"), PktLagVariance, OptionalQualifier);
+
+	ConfigHelperInt(TEXT("PktLagMin"), PktLagMin, OptionalQualifier);
+	ConfigHelperInt(TEXT("PktLagMax"), PktLagMax, OptionalQualifier);
+
+	ConfigHelperInt(TEXT("PktDup"), PktDup, OptionalQualifier);
+
+	ConfigHelperInt(TEXT("PktIncomingLagMin"), PktIncomingLagMin, OptionalQualifier);
+	ConfigHelperInt(TEXT("PktIncomingLagMax"), PktIncomingLagMax, OptionalQualifier);
+	ConfigHelperInt(TEXT("PktIncomingLoss"), PktIncomingLoss, OptionalQualifier);
+
+	ConfigHelperInt(TEXT("PktJitter"), PktJitter, OptionalQualifier);
+	
+	ConfigHelperInt(TEXT("PktBufferBloatInMS"), PktBufferBloatInMS, OptionalQualifier);
+	ConfigHelperInt(TEXT("PktIncomingBufferBloatInMS"), PktIncomingBufferBloatInMS, OptionalQualifier);
+
+	ValidateSettings();
+#endif //DO_ENABLE_NET_TEST
+}
+
+bool FPacketSimulationSettings::LoadEmulationProfile(const TCHAR* ProfileName)
+{
+#if DO_ENABLE_NET_TEST
+	const FString SectionName = FString::Printf(TEXT("%s.%s"), TEXT("PacketSimulationProfile"), ProfileName);
+
+	TArray<FString> SectionConfigs;
+	bool bSectionExists = GConfig->GetSection(*SectionName, SectionConfigs, GEngineIni);
+	if (!bSectionExists)
+	{
+		UE_LOG(LogNet, Log, TEXT("EmulationProfile [%s] was not found in %s. Packet settings were not changed"), *SectionName, *GEngineIni);
+		return false;
+	}
+
+	ResetSettings();
+
+	UScriptStruct* ThisStruct = FPacketSimulationSettings::StaticStruct();
+
+	for (const FString& ConfigVar : SectionConfigs)
+	{
+		FString VarName;
+		FString VarValue;
+		if (ConfigVar.Split(TEXT("="), &VarName, &VarValue))
+		{
+			// If using the one line struct definition
+			if (VarName.Equals(TEXT("PacketSimulationSettings"), ESearchCase::IgnoreCase))
+			{
+				ThisStruct->ImportText(*VarValue, this, nullptr, 0, (FOutputDevice*)GWarn, TEXT("FPacketSimulationSettings"));
+			}
+			else if (FProperty* StructProperty = ThisStruct->FindPropertyByName(FName(*VarName, FNAME_Find)))
+			{
+				StructProperty->ImportText_InContainer(*VarValue, this, nullptr, 0);
+			}
+			else
+			{
+				UE_LOG(LogNet, Warning, TEXT("FPacketSimulationSettings::LoadEmulationProfile could not find property named %s"), *VarName);
+			}
+		}
+	}
+
+	ValidateSettings();
+#endif //DO_ENABLE_NET_TEST
+	return true;
+}
+
+void FPacketSimulationSettings::ResetSettings()
+{
+#if DO_ENABLE_NET_TEST
+	* this = FPacketSimulationSettings();
+#endif
+}
+
+void FPacketSimulationSettings::ValidateSettings()
+{
+#if DO_ENABLE_NET_TEST
+
+	PktLoss = FMath::Clamp<int32>(PktLoss, 0, 100);
+
+	PktOrder = FMath::Clamp<int32>(PktOrder, 0, 1);
+
+	PktLagMin = FMath::Max(PktLagMin, 0);
+	PktLagMax = FMath::Max(PktLagMin, PktLagMax);
+
+	PktDup = FMath::Clamp<int32>(PktDup, 0, 100);
+
+	PktIncomingLagMin = FMath::Max(PktIncomingLagMin, 0);
+	PktIncomingLagMax = FMath::Max(PktIncomingLagMin, PktIncomingLagMax);
+	PktIncomingLoss = FMath::Clamp<int32>(PktIncomingLoss, 0, 100);
+
+	PktBufferBloatInMS = FMath::Max(PktBufferBloatInMS, 0);
+	PktIncomingBufferBloatInMS = FMath::Max(PktIncomingBufferBloatInMS, 0);
+#endif
+}
+
+bool FPacketSimulationSettings::ConfigHelperInt(const TCHAR* Name, int32& Value, const TCHAR* OptionalQualifier)
+{
+#if DO_ENABLE_NET_TEST
+
+	if (OptionalQualifier)
+	{
+		if (GConfig->GetInt(TEXT("PacketSimulationSettings"), *FString::Printf(TEXT("%s%s"), OptionalQualifier, Name), Value, GEngineIni))
+		{
+			return true;
+		}
+	}
+
+	if (GConfig->GetInt(TEXT("PacketSimulationSettings"), Name, Value, GEngineIni))
+	{
+		return true;
+	}
+#endif
+	return false;
+}
+
+bool FPacketSimulationSettings::ConfigHelperBool(const TCHAR* Name, bool& Value, const TCHAR* OptionalQualifier)
+{
+#if DO_ENABLE_NET_TEST
+	if (OptionalQualifier)
+	{
+		if (GConfig->GetBool(TEXT("PacketSimulationSettings"), *FString::Printf(TEXT("%s%s"), OptionalQualifier, Name), Value, GEngineIni))
+		{
+			return true;
+		}
+	}
+
+	if (GConfig->GetBool(TEXT("PacketSimulationSettings"), Name, Value, GEngineIni))
+	{
+		return true;
+	}
+#endif
+	return false;
+}
+
+/**
+ * Reads the settings from a string: command line or an exec
+ *
+ * @param Stream the string to read the settings from
+ */
+bool FPacketSimulationSettings::ParseSettings(const TCHAR* Cmd, const TCHAR* OptionalQualifier)
+{
+	// note that each setting is tested.
+	// this is because the same function will be used to parse the command line as well
+	bool bParsed = false;
+
+#if DO_ENABLE_NET_TEST
+	FString EmulationProfileName;
+	if (FParse::Value(Cmd, TEXT("PktEmulationProfile="), EmulationProfileName))
+	{
+		UE_LOG(LogNet, Log, TEXT("Applying EmulationProfile %s"), *EmulationProfileName);
+		bParsed = LoadEmulationProfile(*EmulationProfileName);
+	}
+	if (ParseHelper(Cmd, TEXT("PktLoss="), PktLoss, OptionalQualifier))
+	{
+		bParsed = true;
+		UE_LOG(LogNet, Log, TEXT("PktLoss set to %d"), PktLoss);
+	}
+	if (ParseHelper(Cmd, TEXT("PktLossMinSize="), PktLossMinSize, OptionalQualifier))
+	{
+		bParsed = true;
+		UE_LOG(LogNet, Log, TEXT("PktLossMinSize set to %d"), PktLossMinSize);
+	}
+	if (ParseHelper(Cmd, TEXT("PktLossMaxSize="), PktLossMaxSize, OptionalQualifier))
+	{
+		bParsed = true;
+		UE_LOG(LogNet, Log, TEXT("PktLossMaxSize set to %d"), PktLossMaxSize);
+	}
+	if (ParseHelper(Cmd, TEXT("PktOrder="), PktOrder, OptionalQualifier))
+	{
+		bParsed = true;
+		UE_LOG(LogNet, Log, TEXT("PktOrder set to %d"), PktOrder);
+	}
+	if (ParseHelper(Cmd, TEXT("PktLag="), PktLag, OptionalQualifier))
+	{
+		bParsed = true;
+		UE_LOG(LogNet, Log, TEXT("PktLag set to %d"), PktLag);
+	}
+	if (ParseHelper(Cmd, TEXT("PktDup="), PktDup, OptionalQualifier))
+	{
+		bParsed = true;
+		UE_LOG(LogNet, Log, TEXT("PktDup set to %d"), PktDup);
+	}
+	if (ParseHelper(Cmd, TEXT("PktLagVariance="), PktLagVariance, OptionalQualifier))
+	{
+		bParsed = true;
+		UE_LOG(LogNet, Log, TEXT("PktLagVariance set to %d"), PktLagVariance);
+	}
+	if (ParseHelper(Cmd, TEXT("PktLagMin="), PktLagMin, OptionalQualifier))
+	{
+		bParsed = true;
+		UE_LOG(LogNet, Log, TEXT("PktLagMin set to %d"), PktLagMin);
+	}
+	if (ParseHelper(Cmd, TEXT("PktLagMax="), PktLagMax, OptionalQualifier))
+	{
+		bParsed = true;
+		UE_LOG(LogNet, Log, TEXT("PktLagMax set to %d"), PktLagMax);
+	}
+	if (ParseHelper(Cmd, TEXT("PktIncomingLagMin="), PktIncomingLagMin, OptionalQualifier))
+	{
+		bParsed = true;
+		UE_LOG(LogNet, Log, TEXT("PktIncomingLagMin set to %d"), PktIncomingLagMin);
+	}
+	if (ParseHelper(Cmd, TEXT("PktIncomingLagMax="), PktIncomingLagMax, OptionalQualifier))
+	{
+		bParsed = true;
+		UE_LOG(LogNet, Log, TEXT("PktIncomingLagMax set to %d"), PktIncomingLagMax);
+	}
+	if (ParseHelper(Cmd, TEXT("PktIncomingLoss="), PktIncomingLoss, OptionalQualifier))
+	{
+		bParsed = true;
+		UE_LOG(LogNet, Log, TEXT("PktIncomingLoss set to %d"), PktIncomingLoss);
+	}
+	if (ParseHelper(Cmd, TEXT("PktJitter="), PktJitter, OptionalQualifier))
+	{
+		bParsed = true;
+		UE_LOG(LogNet, Log, TEXT("PktJitter set to %d"), PktJitter);
+	}
+	if (ParseHelper(Cmd, TEXT("PktBufferBloatInMS="), PktBufferBloatInMS, OptionalQualifier))
+	{
+		bParsed = true;
+		UE_LOG(LogNet, Log, TEXT("PktBufferBloatInMS set to %d"), PktBufferBloatInMS);
+	}
+	if (ParseHelper(Cmd, TEXT("PktIncomingBufferBloatInMS="), PktIncomingBufferBloatInMS, OptionalQualifier))
+	{
+		bParsed = true;
+		UE_LOG(LogNet, Log, TEXT("PktIncomingBufferBloatInMS set to %d"), PktIncomingBufferBloatInMS);
+	}
+
+	ValidateSettings();
+#endif //DO_ENABLE_NET_TEST
+	return bParsed;
+}
+
+bool FPacketSimulationSettings::ParseHelper(const TCHAR* Cmd, const TCHAR* Name, int32& Value, const TCHAR* OptionalQualifier)
+{
+#if DO_ENABLE_NET_TEST
+	if (OptionalQualifier)
+	{
+		if (FParse::Value(Cmd, *FString::Printf(TEXT("%s%s"), OptionalQualifier, Name), Value))
+		{
+			return true;
+		}
+	}
+
+	if (FParse::Value(Cmd, Name, Value))
+	{
+		return true;
+	}
+#endif
+	return false;
+}
+
 #endif //#if DO_ENABLE_NET_TEST
\ No newline at end of file
diff --git a/Engine/Source/Runtime/Engine/Private/NetConnection.cpp b/Engine/Source/Runtime/Engine/Private/NetConnection.cpp
--- a/Engine/Source/Runtime/Engine/Private/NetConnection.cpp
+++ b/Engine/Source/Runtime/Engine/Private/NetConnection.cpp
@@ -115,6 +115,13 @@ static FAutoConsoleVariableRef CVarLogAllConnectionCleanup(TEXT("net.Connection.
 
 static TAutoConsoleVariable<bool> CVarSkipMissingLevelDisconnect( TEXT("net.SkipMissingLevelDisconnect"), false, TEXT("If true skip disconnecting a player if they have a level that doesn't exist on server"));
 
+int32 BufferBloatCooldownTimeInMS = 1000;
+static FAutoConsoleVariableRef CVarBufferBloatCooldownTimeInMS(
+	TEXT("netEmulation.BufferBloatCooldownTimeInMS"),
+	BufferBloatCooldownTimeInMS,
+	TEXT("Cooldown time in milliseconds during which buffer bloat won't be active."),
+	ECVF_Default);
+
 extern int32 GNetDormancyValidate;
 extern bool GbNetReuseReplicatorsForDormantObjects;
 
@@ -2364,14 +2371,13 @@ void UNetConnection::FlushNet(bool bIgnoreSimulation)
 
 		// Send now.
 #if DO_ENABLE_NET_TEST
-
 		bool bWasPacketEmulated(false);
 
 		// if the connection is closing/being destroyed/etc we need to send immediately regardless of settings
 		// because we won't be around to send it delayed
 		if (GetConnectionState() != USOCK_Closed && !IsGarbageCollecting() && !bIgnoreSimulation && !IsInternalAck())
 		{
-			bWasPacketEmulated = CheckOutgoingPacketEmulation(Traits);
+			bWasPacketEmulated = CheckOutgoingPacketEmulation(PacketSentTimeInS, Traits);
 		}
 
 		if (!bWasPacketEmulated)
@@ -2463,16 +2469,26 @@ void UNetConnection::FlushNet(bool bIgnoreSimulation)
 }
 
 #if DO_ENABLE_NET_TEST
-bool UNetConnection::CheckOutgoingPacketEmulation(FOutPacketTraits& Traits)
+bool UNetConnection::CheckOutgoingPacketEmulation(double CurrentTimeInS, FOutPacketTraits& Traits)
 {
 	if (ShouldDropOutgoingPacketForLossSimulation(SendBuffer.GetNumBits()))
 	{
-		UE_LOG(LogNet, VeryVerbose, TEXT("Dropping outgoing packet at %f"), FPlatformTime::Seconds());
+		UE_LOG(LogNet, VeryVerbose, TEXT("Dropping outgoing packet at %f"), CurrentTimeInS);
+		return true;
+	}
+
+	if (PacketSimulationSettings.PktFrameDelay > 0)
+	{
+		FDelayedPacket& DelayedPacket = Delayed.Emplace_GetRef(FDelayedPacket(SendBuffer.GetData(), SendBuffer.GetNumBits(), Traits));
+		DelayedPacket.DelayFrameCount = PacketSimulationSettings.PktFrameDelay;
+
+		// No more emulation applied
 		return true;
 	}
-	else if (PacketSimulationSettings.PktOrder)
+
+	if (PacketSimulationSettings.PktOrder)
 	{
-		FDelayedPacket& B = *(new(Delayed)FDelayedPacket(SendBuffer.GetData(), SendBuffer.GetNumBits(), Traits));
+		Delayed.Emplace(SendBuffer.GetData(), SendBuffer.GetNumBits(), Traits);
 
 		for (int32 i = Delayed.Num() - 1; i >= 0; i--)
 		{
@@ -2490,11 +2506,14 @@ bool UNetConnection::CheckOutgoingPacketEmulation(FOutPacketTraits& Traits)
 
 		return true;
 	}
-	else if (PacketSimulationSettings.PktJitter != 0)
+
+	FDelayedPacket* DelayedPacket = nullptr;
+
+	if (PacketSimulationSettings.PktJitter != 0)
 	{
 		bSendDelayedPacketsOutofOrder = true;
 
-		FDelayedPacket& B = *(new(Delayed)FDelayedPacket(SendBuffer.GetData(), SendBuffer.GetNumBits(), Traits));
+		DelayedPacket = &(Delayed.Emplace_GetRef(FDelayedPacket(SendBuffer.GetData(), SendBuffer.GetNumBits(), Traits)));
 
 		// In order to cause jitter, send one packet at min latency and the next packet at high latency
 		bool bIsLowLatency = (OutPacketId % 2) == 0;
@@ -2504,48 +2523,55 @@ bool UNetConnection::CheckOutgoingPacketEmulation(FOutPacketTraits& Traits)
 			// ExtraLag goes from [PktLagMin, PktLagMin+PktLagVariance]
 			const double LagVariance = 2.0f * (FMath::FRand() - 0.5f) * double(PacketSimulationSettings.PktLagVariance);
 			const double ExtraLag = (double(PacketSimulationSettings.PktLagMin) + LagVariance) / 1000.f;
-			B.SendTime = FPlatformTime::Seconds() + ExtraLag;
+			DelayedPacket->SendTime = CurrentTimeInS + ExtraLag;
 		}
 		else
 		{
 			// ExtraLag goes from [PktLagMin+PktJitter, PktLagMin+PktJitter+PktLagVariance]
 			const double LagVariance = 2.0f * (FMath::FRand() - 0.5f) * double(PacketSimulationSettings.PktLagVariance);
 			const double ExtraLag = (double(PacketSimulationSettings.PktLagMin + PacketSimulationSettings.PktJitter) + LagVariance) / 1000.f;
-			B.SendTime = FPlatformTime::Seconds() + ExtraLag;
+			DelayedPacket->SendTime = CurrentTimeInS + ExtraLag;
 		}
 	}
 	else if (PacketSimulationSettings.PktLag)
 	{
-		FDelayedPacket& B = *(new(Delayed)FDelayedPacket(SendBuffer.GetData(), SendBuffer.GetNumBits(), Traits));
+		DelayedPacket = &(Delayed.Emplace_GetRef(FDelayedPacket(SendBuffer.GetData(), SendBuffer.GetNumBits(), Traits)));
 
 		// ExtraLag goes from PktLag + [-PktLagVariance, PktLagVariance]
 		const double LagVariance = 2.0f * (FMath::FRand() - 0.5f) * double(PacketSimulationSettings.PktLagVariance);
 		const double ExtraLag = (double(PacketSimulationSettings.PktLag) + LagVariance) / 1000.f;
-		B.SendTime = FPlatformTime::Seconds() + ExtraLag;
-
-		return true;
-
+		DelayedPacket->SendTime = CurrentTimeInS + ExtraLag;
 	}
 	else if (PacketSimulationSettings.PktLagMin > 0 || PacketSimulationSettings.PktLagMax > 0)
 	{
-		FDelayedPacket& B = *(new(Delayed)FDelayedPacket(SendBuffer.GetData(), SendBuffer.GetNumBits(), Traits));
+		DelayedPacket = &(Delayed.Emplace_GetRef(FDelayedPacket(SendBuffer.GetData(), SendBuffer.GetNumBits(), Traits)));
 
 		// ExtraLag goes from [PktLagMin, PktLagMax]
 		const double LagVariance = FMath::FRand() * double(PacketSimulationSettings.PktLagMax - PacketSimulationSettings.PktLagMin);
 		const double ExtraLag = (double(PacketSimulationSettings.PktLagMin) + LagVariance) / 1000.f;
-		B.SendTime = FPlatformTime::Seconds() + ExtraLag;
-		
-		return true;
+		DelayedPacket->SendTime = CurrentTimeInS + ExtraLag;
 	}
-	else if (PacketSimulationSettings.PktFrameDelay > 0)
+
+	// Delay the packet until the next buffer flush or later
+	if (PacketSimulationSettings.PktBufferBloatInMS > 0)
 	{
-		FDelayedPacket& B = *(new(Delayed)FDelayedPacket(SendBuffer.GetData(), SendBuffer.GetNumBits(), Traits));
-		B.DelayFrameCount = PacketSimulationSettings.PktFrameDelay;
-		
-		return true;
-	}
+		if (!DelayedPacket)
+		{
+			DelayedPacket = &(Delayed.Emplace_GetRef(FDelayedPacket(SendBuffer.GetData(), SendBuffer.GetNumBits(), Traits)));
+		}
 
-	return false;
+		if (NextBufferBloatFlushOut < CurrentTimeInS && CurrentTimeInS >= NextBufferBloatActiveTimeOut)
+		{
+			// Find the time at which we flush all delayed packets, and the next time we restart the buffer bloat accumulation
+			NextBufferBloatFlushOut = CurrentTimeInS + (PacketSimulationSettings.PktBufferBloatInMS/ 1000.0);
+			NextBufferBloatActiveTimeOut = NextBufferBloatFlushOut + (BufferBloatCooldownTimeInMS / 1000.0);
+		}
+
+		// Only keep the latest delayed time
+		DelayedPacket->SendTime = FMath::Max(NextBufferBloatFlushOut, DelayedPacket->SendTime);
+	}
+	
+	return DelayedPacket != nullptr;
 }
 #endif
 
@@ -2973,6 +2999,69 @@ FNetworkGUID UNetConnection::GetActorGUIDFromOpenBunch(FInBunch& Bunch)
 	return ActorGUID;
 }
 
+#if DO_ENABLE_NET_TEST
+bool UNetConnection::CheckIncomingPacketEmulation(double CurrentTimeInS, FBitReader& Reader)
+{
+	if (PacketSimulationSettings.PktIncomingLoss)
+	{
+		if (FMath::FRand() * 100.f < PacketSimulationSettings.PktIncomingLoss)
+		{
+			UE_LOG(LogNet, VeryVerbose, TEXT("Dropped incoming packet at %f"), CurrentTimeInS);
+			return true;
+		}
+	}
+
+	FDelayedIncomingPacket* DelayedPacket = nullptr;
+
+	if (PacketSimulationSettings.PktIncomingLagMin > 0 || PacketSimulationSettings.PktIncomingLagMax > 0)
+	{
+		// ExtraLagInSec goes from [PktIncomingLagMin, PktIncomingLagMax]
+		const double LagVarianceInMS = FMath::FRand() * double(PacketSimulationSettings.PktIncomingLagMax - PacketSimulationSettings.PktIncomingLagMin);
+		const double ExtraLagInSec = (double(PacketSimulationSettings.PktIncomingLagMin) + LagVarianceInMS) / 1000.f;
+
+		DelayedPacket = &(DelayedIncomingPackets.AddDefaulted_GetRef());
+
+		DelayedPacket->PacketData = MakeUnique<FBitReader>(Reader);
+		DelayedPacket->ReinjectionTime = CurrentTimeInS + ExtraLagInSec;
+
+		UE_LOG(LogNet, VeryVerbose, TEXT("Delaying incoming packet for %f seconds"), ExtraLagInSec);
+	}
+	else if (PacketSimulationSettings.PktIncomingFrameDelay)
+	{
+		DelayedPacket = &(DelayedIncomingPackets.AddDefaulted_GetRef());
+
+		DelayedPacket->PacketData = MakeUnique<FBitReader>(Reader);
+		DelayedPacket->ReinjectionFrameCount = PacketSimulationSettings.PktIncomingFrameDelay;
+
+		UE_LOG(LogNet, VeryVerbose, TEXT("Delaying incoming packet for %u frames"), PacketSimulationSettings.PktIncomingFrameDelay);
+
+		// Not compatible with BufferBloat
+		return true;
+	}
+
+	if (PacketSimulationSettings.PktIncomingBufferBloatInMS > 0)
+	{
+		if (!DelayedPacket)
+		{
+			DelayedPacket = &(DelayedIncomingPackets.AddDefaulted_GetRef());
+			DelayedPacket->PacketData = MakeUnique<FBitReader>(Reader);
+		}
+
+		if (NextBufferBloatFlushIn < CurrentTimeInS && CurrentTimeInS >= NextBufferBloatActiveTimeIn)
+		{
+			NextBufferBloatFlushIn = CurrentTimeInS + (PacketSimulationSettings.PktIncomingBufferBloatInMS / 1000.f);
+			NextBufferBloatActiveTimeIn = NextBufferBloatFlushIn + (BufferBloatCooldownTimeInMS / 1000.0);
+		}
+
+		// Only keep the latest delayed time
+		DelayedPacket->ReinjectionTime = FMath::Max(NextBufferBloatFlushIn, DelayedPacket->ReinjectionTime);
+	}
+
+	return DelayedPacket != nullptr;
+
+}
+#endif //DO_ENABLE_NET_TEST
+
 void UNetConnection::ReceivedPacket( FBitReader& Reader, bool bIsReinjectedPacket, bool bDispatchPacket )
 {
 	using namespace UE::Net;
@@ -2987,42 +3076,17 @@ void UNetConnection::ReceivedPacket( FBitReader& Reader, bool bIsReinjectedPacke
 		return;
 	}
 
+	// Record the packet time to the histogram
+	const double CurrentReceiveTimeInS = FPlatformTime::Seconds();
+
 #if DO_ENABLE_NET_TEST
 	if (!IsInternalAck() && !bIsReinjectingDelayedPackets)
 	{
-		if (PacketSimulationSettings.PktIncomingLoss)
-		{
-			if (FMath::FRand() * 100.f < PacketSimulationSettings.PktIncomingLoss)
-			{
-				UE_LOG(LogNet, VeryVerbose, TEXT("Dropped incoming packet at %f"), FPlatformTime::Seconds());
-				return;
-			}
+		const bool bIsPackedEmulated = CheckIncomingPacketEmulation(CurrentReceiveTimeInS, Reader);
 
-		}
-		if (PacketSimulationSettings.PktIncomingLagMin > 0 || PacketSimulationSettings.PktIncomingLagMax > 0)
+		if (bIsPackedEmulated)
 		{
-			// ExtraLagInSec goes from [PktIncomingLagMin, PktIncomingLagMax]
-			const double LagVarianceInMS = FMath::FRand() * double(PacketSimulationSettings.PktIncomingLagMax - PacketSimulationSettings.PktIncomingLagMin);
-			const double ExtraLagInSec = (double(PacketSimulationSettings.PktIncomingLagMin) + LagVarianceInMS) / 1000.f;
-
-			FDelayedIncomingPacket DelayedPacket;
-			DelayedPacket.PacketData = MakeUnique<FBitReader>(Reader);
-			DelayedPacket.ReinjectionTime = FPlatformTime::Seconds() + ExtraLagInSec;
-
-			DelayedIncomingPackets.Emplace(MoveTemp(DelayedPacket));
-
-			UE_LOG(LogNet, VeryVerbose, TEXT("Delaying incoming packet for %f seconds"), ExtraLagInSec);
-			return;
-		}
-		else if (PacketSimulationSettings.PktIncomingFrameDelay)
-		{
-			FDelayedIncomingPacket DelayedPacket;
-			DelayedPacket.PacketData = MakeUnique<FBitReader>(Reader);
-			DelayedPacket.ReinjectionFrameCount = PacketSimulationSettings.PktIncomingFrameDelay;
-
-			DelayedIncomingPackets.Emplace(MoveTemp(DelayedPacket));
-
-			UE_LOG(LogNet, VeryVerbose, TEXT("Delaying incoming packet for %u frames"), PacketSimulationSettings.PktIncomingFrameDelay);
+			// Packet was delayed or dropped, nothing to do here anymore
 			return;
 		}
 	}
@@ -3033,9 +3097,6 @@ void UNetConnection::ReceivedPacket( FBitReader& Reader, bool bIsReinjectedPacke
 
 	ValidateSendBuffer();
 
-	// Record the packet time to the histogram
-	const double CurrentReceiveTimeInS = FPlatformTime::Seconds();
-
 	if (!bIsReinjectedPacket)
 	{
 		const double LastPacketTimeDiffInMs = (CurrentReceiveTimeInS - LastReceiveRealtime) * 1000.0;
diff --git a/Engine/Source/Runtime/Engine/Private/NetDriver.cpp b/Engine/Source/Runtime/Engine/Private/NetDriver.cpp
--- a/Engine/Source/Runtime/Engine/Private/NetDriver.cpp
+++ b/Engine/Source/Runtime/Engine/Private/NetDriver.cpp
@@ -4731,266 +4731,7 @@ void UNetDriver::OnPacketSimulationSettingsChanged()
 #endif
 }
 
-#endif
-
-class FPacketSimulationConsoleCommandVisitor 
-{
-public:
-	static void OnPacketSimulationConsoleCommand(const TCHAR *Name, IConsoleObject* CVar, TArray<IConsoleObject*>& Sink)
-	{
-		Sink.Add(CVar);
-	}
-};
-
-/** reads in settings from the .ini file 
- * @note: overwrites all previous settings
- */
-void FPacketSimulationSettings::LoadConfig(const TCHAR* OptionalQualifier)
-{
-#if DO_ENABLE_NET_TEST
-	ConfigHelperInt(TEXT("PktLoss"), PktLoss, OptionalQualifier);
-	
-	ConfigHelperInt(TEXT("PktLossMinSize"), PktLossMinSize, OptionalQualifier);
-	ConfigHelperInt(TEXT("PktLossMaxSize"), PktLossMaxSize, OptionalQualifier);
-
-	bool InPktOrder = !!PktOrder;
-	ConfigHelperBool(TEXT("PktOrder"), InPktOrder, OptionalQualifier);
-	PktOrder = int32(InPktOrder);
-	
-	ConfigHelperInt(TEXT("PktLag"), PktLag, OptionalQualifier);
-	ConfigHelperInt(TEXT("PktLagVariance"), PktLagVariance, OptionalQualifier);
-
-	ConfigHelperInt(TEXT("PktLagMin"), PktLagMin, OptionalQualifier);
-	ConfigHelperInt(TEXT("PktLagMax"), PktLagMax, OptionalQualifier);
-	
-	ConfigHelperInt(TEXT("PktDup"), PktDup, OptionalQualifier);
-
-	ConfigHelperInt(TEXT("PktIncomingLagMin"), PktIncomingLagMin, OptionalQualifier);
-	ConfigHelperInt(TEXT("PktIncomingLagMax"), PktIncomingLagMax, OptionalQualifier);
-	ConfigHelperInt(TEXT("PktIncomingLoss"), PktIncomingLoss, OptionalQualifier);
-
-	ConfigHelperInt(TEXT("PktJitter"), PktJitter, OptionalQualifier);
-
-	ValidateSettings();
-#endif
-}
-
-bool FPacketSimulationSettings::LoadEmulationProfile(const TCHAR* ProfileName)
-{
-#if DO_ENABLE_NET_TEST
-	const FString SectionName = FString::Printf(TEXT("%s.%s"), TEXT("PacketSimulationProfile"), ProfileName);
-
-	TArray<FString> SectionConfigs;
-	bool bSectionExists = GConfig->GetSection(*SectionName, SectionConfigs, GEngineIni);
-	if (!bSectionExists)
-	{
-		UE_LOG(LogNet, Log, TEXT("EmulationProfile [%s] was not found in %s. Packet settings were not changed"), *SectionName, *GEngineIni);
-		return false;
-	}
-
-	ResetSettings();
-
-	UScriptStruct* ThisStruct = FPacketSimulationSettings::StaticStruct();
-
-	for (const FString& ConfigVar : SectionConfigs)
-	{
-		FString VarName;
-		FString VarValue;
-		if (ConfigVar.Split(TEXT("="), &VarName, &VarValue))
-		{
-			// If using the one line struct definition
-			if (VarName.Equals(TEXT("PacketSimulationSettings"), ESearchCase::IgnoreCase))
-			{
-				ThisStruct->ImportText(*VarValue, this, nullptr, 0, (FOutputDevice*)GWarn, TEXT("FPacketSimulationSettings"));
-			}
-			else if (FProperty* StructProperty = ThisStruct->FindPropertyByName(FName(*VarName, FNAME_Find)))
-			{
-				StructProperty->ImportText_InContainer(*VarValue, this, nullptr, 0);
-			}
-			else
-			{
-				UE_LOG(LogNet, Warning, TEXT("FPacketSimulationSettings::LoadEmulationProfile could not find property named %s"), *VarName);
-			}
-		}
-	}
-	
-	ValidateSettings();
-#endif
-	return true;
-}
-
-void FPacketSimulationSettings::ResetSettings()
-{
-#if DO_ENABLE_NET_TEST
-	*this = FPacketSimulationSettings();
-#endif
-}
-
-void FPacketSimulationSettings::ValidateSettings()
-{
-#if DO_ENABLE_NET_TEST
-
-	PktLoss = FMath::Clamp<int32>(PktLoss, 0, 100);
-
-	PktOrder = FMath::Clamp<int32>(PktOrder, 0, 1);
-
-	PktLagMin = FMath::Max(PktLagMin, 0);
-	PktLagMax = FMath::Max(PktLagMin, PktLagMax);
-
-	PktDup = FMath::Clamp<int32>(PktDup, 0, 100);
-
-	PktIncomingLagMin = FMath::Max(PktIncomingLagMin, 0);
-	PktIncomingLagMax = FMath::Max(PktIncomingLagMin, PktIncomingLagMax);
-	PktIncomingLoss = FMath::Clamp<int32>(PktIncomingLoss, 0, 100);
-#endif
-}
-
-bool FPacketSimulationSettings::ConfigHelperInt(const TCHAR* Name, int32& Value, const TCHAR* OptionalQualifier)
-{
-#if DO_ENABLE_NET_TEST
-
-	if (OptionalQualifier)
-	{
-		if (GConfig->GetInt(TEXT("PacketSimulationSettings"), *FString::Printf(TEXT("%s%s"), OptionalQualifier, Name), Value, GEngineIni))
-		{
-			return true;
-		}
-	}
-
-	if (GConfig->GetInt(TEXT("PacketSimulationSettings"), Name, Value, GEngineIni))
-	{
-		return true;
-	}
-#endif
-	return false;
-}
-
-bool FPacketSimulationSettings::ConfigHelperBool(const TCHAR* Name, bool& Value, const TCHAR* OptionalQualifier)
-{
-#if DO_ENABLE_NET_TEST
-	if (OptionalQualifier)
-	{
-		if (GConfig->GetBool(TEXT("PacketSimulationSettings"), *FString::Printf(TEXT("%s%s"), OptionalQualifier, Name), Value, GEngineIni))
-		{
-			return true;
-		}
-	}
-
-	if (GConfig->GetBool(TEXT("PacketSimulationSettings"), Name, Value, GEngineIni))
-	{
-		return true;
-	}
-#endif
-	return false;
-}
-
-/**
- * Reads the settings from a string: command line or an exec
- *
- * @param Stream the string to read the settings from
- */
-bool FPacketSimulationSettings::ParseSettings(const TCHAR* Cmd, const TCHAR* OptionalQualifier)
-{
-	// note that each setting is tested.
-	// this is because the same function will be used to parse the command line as well
-	bool bParsed = false;
-
-#if DO_ENABLE_NET_TEST
-	FString EmulationProfileName;
-	if (FParse::Value(Cmd, TEXT("PktEmulationProfile="), EmulationProfileName))
-	{
-		UE_LOG(LogNet, Log, TEXT("Applying EmulationProfile %s"), *EmulationProfileName);
-		bParsed = LoadEmulationProfile(*EmulationProfileName);
-	}
-	if( ParseHelper(Cmd, TEXT("PktLoss="), PktLoss, OptionalQualifier) )
-	{
-		bParsed = true;
-		UE_LOG(LogNet, Log, TEXT("PktLoss set to %d"), PktLoss);
-	}
-	if (ParseHelper(Cmd, TEXT("PktLossMinSize="), PktLossMinSize, OptionalQualifier))
-	{
-		bParsed = true;
-		UE_LOG(LogNet, Log, TEXT("PktLossMinSize set to %d"), PktLossMinSize);
-	}
-	if (ParseHelper(Cmd, TEXT("PktLossMaxSize="), PktLossMaxSize, OptionalQualifier))
-	{
-		bParsed = true;
-		UE_LOG(LogNet, Log, TEXT("PktLossMaxSize set to %d"), PktLossMaxSize);
-	}
-	if( ParseHelper(Cmd, TEXT("PktOrder="), PktOrder, OptionalQualifier) )
-	{
-		bParsed = true;
-		UE_LOG(LogNet, Log, TEXT("PktOrder set to %d"), PktOrder);
-	}
-	if( ParseHelper(Cmd, TEXT("PktLag="), PktLag, OptionalQualifier) )
-	{
-		bParsed = true;
-		UE_LOG(LogNet, Log, TEXT("PktLag set to %d"), PktLag);
-	}
-	if( ParseHelper(Cmd, TEXT("PktDup="), PktDup, OptionalQualifier) )
-	{
-		bParsed = true;
-		UE_LOG(LogNet, Log, TEXT("PktDup set to %d"), PktDup);
-	}	
-	if (ParseHelper(Cmd, TEXT("PktLagVariance="), PktLagVariance, OptionalQualifier))
-	{
-		bParsed = true;
-		UE_LOG(LogNet, Log, TEXT("PktLagVariance set to %d"), PktLagVariance);
-	}
-	if (ParseHelper(Cmd, TEXT("PktLagMin="), PktLagMin, OptionalQualifier))
-	{
-		bParsed = true;
-		UE_LOG(LogNet, Log, TEXT("PktLagMin set to %d"), PktLagMin);
-	}
-	if (ParseHelper(Cmd, TEXT("PktLagMax="), PktLagMax, OptionalQualifier))
-	{
-		bParsed = true;
-		UE_LOG(LogNet, Log, TEXT("PktLagMax set to %d"), PktLagMax);
-	}
-	if (ParseHelper(Cmd, TEXT("PktIncomingLagMin="), PktIncomingLagMin, OptionalQualifier))
-	{
-		bParsed = true;
-		UE_LOG(LogNet, Log, TEXT("PktIncomingLagMin set to %d"), PktIncomingLagMin);
-	}
-	if (ParseHelper(Cmd, TEXT("PktIncomingLagMax="), PktIncomingLagMax, OptionalQualifier))
-	{
-		bParsed = true;
-		UE_LOG(LogNet, Log, TEXT("PktIncomingLagMax set to %d"), PktIncomingLagMax);
-	}
-	if (ParseHelper(Cmd, TEXT("PktIncomingLoss="), PktIncomingLoss, OptionalQualifier))
-	{
-		bParsed = true;
-		UE_LOG(LogNet, Log, TEXT("PktIncomingLoss set to %d"), PktIncomingLoss);
-	}
-	if (ParseHelper(Cmd, TEXT("PktJitter="), PktJitter, OptionalQualifier))
-	{
-		bParsed = true;
-		UE_LOG(LogNet, Log, TEXT("PktJitter set to %d"), PktJitter);
-	}
-
-	ValidateSettings();
-#endif
-	return bParsed;
-}
-
-bool FPacketSimulationSettings::ParseHelper(const TCHAR* Cmd, const TCHAR* Name, int32& Value, const TCHAR* OptionalQualifier)
-{
-#if DO_ENABLE_NET_TEST
-	if (OptionalQualifier)
-	{
-		if (FParse::Value(Cmd, *FString::Printf(TEXT("%s%s"), OptionalQualifier, Name), Value))
-		{
-			return true;
-		}
-	}
-
-	if (FParse::Value(Cmd, Name, Value))
-	{
-		return true;
-	}
-#endif
-	return false;
-}
+#endif //DO_ENABLE_NET_TEST
 
 FNetViewer::FNetViewer(UNetConnection* InConnection, float DeltaSeconds) :
 	Connection(InConnection),
