New material translator:
- MaterialIRValueAnalyzer: fully documented, renamed a function to a better name.
- Removed requirement for FValues to be trivially constructible. This relaxation allows for FValues to contain FStringViews and FNames for instance, without requiring backflips. This requires though that all members of FValue can be zero-initialized, which can't be enforced with code.
- Added MIR::FSimpleStringView for FValue usage, even inside unions (like MIR::FInlineHLSL).

#rnx

diff --git a/Engine/Source/Runtime/Engine/Private/Materials/MaterialExpressionsIR.cpp b/Engine/Source/Runtime/Engine/Private/Materials/MaterialExpressionsIR.cpp
--- a/Engine/Source/Runtime/Engine/Private/Materials/MaterialExpressionsIR.cpp
+++ b/Engine/Source/Runtime/Engine/Private/Materials/MaterialExpressionsIR.cpp
@@ -394,7 +394,7 @@ void UMaterialExpressionVectorParameter::Build(MIR::FEmitter& Em)
 void UMaterialExpressionDoubleVectorParameter::Build(MIR::FEmitter& Em)
 {
 	FValueRef Value = BuildMaterialExpressionParameter(Em, this);
-	Em.Output(0, Em.Cast(Value, MIR::FType::MakeLWC(3)));
+	Em.Output(0, Em.Cast(Value, MIR::FType::MakeDouble(3)));
 	Em.Output(1, Em.Subscript(Value, 0));
 	Em.Output(2, Em.Subscript(Value, 1));
 	Em.Output(3, Em.Subscript(Value, 2));
diff --git a/Engine/Source/Runtime/Engine/Private/Materials/MaterialIR.cpp b/Engine/Source/Runtime/Engine/Private/Materials/MaterialIR.cpp
--- a/Engine/Source/Runtime/Engine/Private/Materials/MaterialIR.cpp
+++ b/Engine/Source/Runtime/Engine/Private/Materials/MaterialIR.cpp
@@ -449,8 +449,9 @@ FType GetExternalInputType(EExternalInput Id)
 	
 	switch (Id)
 	{
-		case EExternalInput::ViewMaterialTextureMipBias: return FType::MakeFloat();
-		case EExternalInput::ViewMaterialTextureDerivativeMultiply: return FType::MakeFloat();
+		case EExternalInput::ViewMaterialTextureMipBias:
+		case EExternalInput::ViewMaterialTextureDerivativeMultiply:
+			 return FType::MakeFloat();
 
 		case EExternalInput::ActorPosition_Absolute:
 		case EExternalInput::ObjectPosition_Absolute:
@@ -460,7 +461,7 @@ FType GetExternalInputType(EExternalInput Id)
 		case EExternalInput::WorldPosition_AbsoluteNoOffsets_Ddx:		// analytic derivative expressions.  So we cast them to LWC on load.  The shader compiler should
 		case EExternalInput::WorldPosition_Absolute_Ddy:				// be able to optimize the math (the tile offsets will all be zero, cancelling things out).
 		case EExternalInput::WorldPosition_AbsoluteNoOffsets_Ddy:
-			return FType::MakeLWC(3);
+			return FType::MakeDouble(3);
 
 		case EExternalInput::ActorPosition_CameraRelative:
 		case EExternalInput::ObjectPosition_CameraRelative:
diff --git a/Engine/Source/Runtime/Engine/Private/Materials/MaterialIRDebug.cpp b/Engine/Source/Runtime/Engine/Private/Materials/MaterialIRDebug.cpp
--- a/Engine/Source/Runtime/Engine/Private/Materials/MaterialIRDebug.cpp
+++ b/Engine/Source/Runtime/Engine/Private/Materials/MaterialIRDebug.cpp
@@ -327,7 +327,7 @@ struct FDebugDumpIRState
 		{
 			if (InlineHLSL->HasFlags(EValueFlags::HasDynamicHLSLCode))
 			{
-				Out.Appendf(TEXT(" \"%s\""), InlineHLSL->Code);
+				Out.Appendf(TEXT(" \"%.*s\""), InlineHLSL->Code.Len, InlineHLSL->Code.Ptr);
 			}
 			else
 			{
diff --git a/Engine/Source/Runtime/Engine/Private/Materials/MaterialIREmitter.cpp b/Engine/Source/Runtime/Engine/Private/Materials/MaterialIREmitter.cpp
--- a/Engine/Source/Runtime/Engine/Private/Materials/MaterialIREmitter.cpp
+++ b/Engine/Source/Runtime/Engine/Private/Materials/MaterialIREmitter.cpp
@@ -134,7 +134,7 @@ static FValueRef EmitNew(FEmitter& Emitter, FValue* Value)
 template <typename T>
 static T MakePrototype(FType InType)
 {
-	static_assert(std::is_trivially_constructible_v<T> && std::is_trivially_destructible_v<T> &&  std::is_trivially_copy_constructible_v<T> && std::is_trivially_copy_assignable_v<T>,
+	static_assert(std::is_trivially_destructible_v<T> && std::is_trivially_copy_constructible_v<T> && std::is_trivially_copy_assignable_v<T>,
 		"FValues are expected to be trivial types.");
 
 	T Value;
@@ -710,7 +710,7 @@ FValueRef FEmitter::ConstantFloat4(UE::Math::TVector4<TFloat> InValue)
 
 FValueRef FEmitter::ConstantDouble(TDouble InX)
 {
-	FConstant Scalar = MakePrototype<FConstant>(FType::MakeLWC());
+	FConstant Scalar = MakePrototype<FConstant>(FType::MakeDouble());
 	Scalar.Double = InX;
 	return EmitPrototype(*this, Scalar);
 }
@@ -3264,8 +3264,7 @@ static FValueRef EmitInlineHLSL(FEmitter& Emitter, FType Type, const FMaterialEx
 	}
 	else
 	{
-		Prototype.Code = Code.GetData();
-		Prototype.CodeLen = Code.Len();
+		Prototype.Code = Code;
 	}
 
 	if (!InputValues.IsEmpty())
diff --git a/Engine/Source/Runtime/Engine/Private/Materials/MaterialIRModuleBuilder.cpp b/Engine/Source/Runtime/Engine/Private/Materials/MaterialIRModuleBuilder.cpp
--- a/Engine/Source/Runtime/Engine/Private/Materials/MaterialIRModuleBuilder.cpp
+++ b/Engine/Source/Runtime/Engine/Private/Materials/MaterialIRModuleBuilder.cpp
@@ -514,7 +514,7 @@ struct FMaterialIRModuleBuilderImpl
 					ValueAnalyzer.Analyze(Value);
 				}
 
-				ValueAnalyzer.PropagateStateInStage(Value, CurrentStage);
+				ValueAnalyzer.AnalyzeInStage(Value, CurrentStage);
 
 				// Mark the used instruction as analyzed for this stage.
 				Value->Flags |= MIR::EValueFlags(1 << StageIndex);
diff --git a/Engine/Source/Runtime/Engine/Private/Materials/MaterialIRToHLSLTranslator.cpp b/Engine/Source/Runtime/Engine/Private/Materials/MaterialIRToHLSLTranslator.cpp
--- a/Engine/Source/Runtime/Engine/Private/Materials/MaterialIRToHLSLTranslator.cpp
+++ b/Engine/Source/Runtime/Engine/Private/Materials/MaterialIRToHLSLTranslator.cpp
@@ -662,12 +662,12 @@ struct FPrivate : FMaterialIRToHLSLTranslation
 
 		if (ExternalCode->HasFlags(MIR::EValueFlags::HasDynamicHLSLCode))
 		{
-			check(ExternalCode->Code != nullptr);
+			check(ExternalCode->Code.Ptr);
 
 			// Substitute placeholder tokens now unless disabled for custom nodes
 			if (ExternalCode->HasFlags(MIR::EValueFlags::SubstituteTagsInInlineHLSL))
 			{
-				FString Code = ExternalCode->Code;
+				FString Code = ExternalCode->Code.ToString();
 
 				// @todo-laura.hermanns - Must be replaced as soon as this state is available in MIR
 				constexpr bool bCompilingPreviousFrame = false;
diff --git a/Engine/Source/Runtime/Engine/Private/Materials/MaterialIRTypes.cpp b/Engine/Source/Runtime/Engine/Private/Materials/MaterialIRTypes.cpp
--- a/Engine/Source/Runtime/Engine/Private/Materials/MaterialIRTypes.cpp
+++ b/Engine/Source/Runtime/Engine/Private/Materials/MaterialIRTypes.cpp
@@ -180,7 +180,7 @@ FType FType::MakeFloat(int NumRows, int NumColumns)
 	return MakePrimitive(EScalarKind::Float, NumRows, NumColumns);
 }
 
-FType FType::MakeLWC(int NumRows, int NumColumns)
+FType FType::MakeDouble(int NumRows, int NumColumns)
 {
 	return MakePrimitive(EScalarKind::Double, NumRows, NumColumns);
 }
diff --git a/Engine/Source/Runtime/Engine/Private/Materials/MaterialIRValueAnalyzer.cpp b/Engine/Source/Runtime/Engine/Private/Materials/MaterialIRValueAnalyzer.cpp
--- a/Engine/Source/Runtime/Engine/Private/Materials/MaterialIRValueAnalyzer.cpp
+++ b/Engine/Source/Runtime/Engine/Private/Materials/MaterialIRValueAnalyzer.cpp
@@ -43,7 +43,7 @@ static EMaterialShaderFrequency MapToMaterialShaderFrequencyOrAny(MIR::EStage St
 	return EMaterialShaderFrequency::Any;
 }
 
-static void PropagateStateInStageInlineHLSL(FMaterialIRValueAnalyzer& Analyzer, MIR::FInlineHLSL* InlineHLSL, MIR::EStage Stage)
+static void AnalyzeInStageInlineHLSL(FMaterialIRValueAnalyzer& Analyzer, MIR::FInlineHLSL* InlineHLSL, MIR::EStage Stage)
 {
 	if (!InlineHLSL->HasFlags(MIR::EValueFlags::HasDynamicHLSLCode))
 	{
@@ -102,7 +102,7 @@ static int32 AcquireVTStackIndex(FMaterialIRValueAnalyzer& Analyzer, MIR::FVTPag
 	return Analyzer.CompilationOutput->UniformExpressionSet.AddVTStack(INDEX_NONE);
 }
 
-static void PropagateStateInStageVTPageTableRead(FMaterialIRValueAnalyzer& Analyzer, MIR::FVTPageTableRead* VTPageTableRead, MIR::EStage Stage)
+static void AnalyzeInStageVTPageTableRead(FMaterialIRValueAnalyzer& Analyzer, MIR::FVTPageTableRead* VTPageTableRead, MIR::EStage Stage)
 {
 	check(VTPageTableRead->TextureObject);
 	MIR::FTextureObject* Texture = VTPageTableRead->TextureObject->As<MIR::FTextureObject>();
@@ -305,13 +305,22 @@ static void AnalyzeSetMaterialOutput(FMaterialIRValueAnalyzer& Analyzer, MIR::FS
 	}
 }
 
+//////////////////////////////////////////////////////////////////////////////////////////
+// FMaterialIRValueAnalyzer
+//////////////////////////////////////////////////////////////////////////////////////////
 
 void FMaterialIRValueAnalyzer::Setup(UMaterial* InMaterial, FMaterialIRModule* InModule, FMaterialCompilationOutput* InCompilationOutput, FMaterialInsights* InInsights)
 {
 	Material = InMaterial;
 	Module = InModule;
-	CompilationOutput = InCompilationOutput;
 	Insights = InInsights;
+	CompilationOutput = InCompilationOutput;
+	UniformDefaultValueOffsets = {};
+	FreeOffsetsPerNumComponents[0] = {};
+	FreeOffsetsPerNumComponents[1] = {};
+	FreeOffsetsPerNumComponents[2] = {};
+	EnvironmentDefines = {};
+	VTStacks = {};
 }
 
 void FMaterialIRValueAnalyzer::Analyze(MIR::FValue* Value)
@@ -331,6 +340,7 @@ void FMaterialIRValueAnalyzer::Analyze(MIR::FValue* Value)
 		break; \
 	}
 
+	// IMPORTANT: Before adding a case here, read the FMaterialIRValueAnalyzer documentation.
 	switch (Value->Kind)
 	{
 		EXPAND_CASE(ExternalInput);
@@ -344,15 +354,16 @@ void FMaterialIRValueAnalyzer::Analyze(MIR::FValue* Value)
 	#undef EXPAND_CASE
 }
 
-void FMaterialIRValueAnalyzer::PropagateStateInStage(MIR::FValue* Value, MIR::EStage Stage)
+void FMaterialIRValueAnalyzer::AnalyzeInStage(MIR::FValue* Value, MIR::EStage Stage)
 {
 	#define EXPAND_CASE(ValueType) \
 	case MIR::VK_##ValueType: \
 	{ \
-		PropagateStateInStage##ValueType(*this, static_cast<MIR::F##ValueType*>(Value), Stage); \
+		AnalyzeInStage##ValueType(*this, static_cast<MIR::F##ValueType*>(Value), Stage); \
 		break; \
 	}
 
+	// IMPORTANT: Before adding a case here, read the FMaterialIRValueAnalyzer documentation.
 	switch (Value->Kind)
 	{
 		EXPAND_CASE(InlineHLSL);
diff --git a/Engine/Source/Runtime/Engine/Private/Materials/MaterialIRValueAnalyzer.h b/Engine/Source/Runtime/Engine/Private/Materials/MaterialIRValueAnalyzer.h
--- a/Engine/Source/Runtime/Engine/Private/Materials/MaterialIRValueAnalyzer.h
+++ b/Engine/Source/Runtime/Engine/Private/Materials/MaterialIRValueAnalyzer.h
@@ -13,6 +13,33 @@
 #include "MaterialShared.h"
 #include "Engine/Texture.h"
 
+// It analyzes a value or instruction, performing semantic validation and side-effects execution.
+//
+// These analyze functions allow a value or instruction to perform non-trivial custom operations  
+// and validation. These operations may have side effects, such as setting state in  
+// CompilationOutput, allocating resources, etc. The reason these are done here (rather than at  
+// FValue emission time) is that only values which are *actually* used (not pruned as unused or  
+// optimized out) are analyzed. In other words, while implementing a value's analyze function, you  
+// are guaranteed that the value is needed in the final material.
+//
+// Only values that require this kind of post-emission validation or that generate sidecar  
+// resources should define Analyze functions. Otherwise, they can omit them.
+
+// There are two types of analyze functions: Analyze() and AnalyzeInStage().
+// - Analyze() is invoked once, regardless of which stage (vertex, pixel, or both) the value is  
+//   scheduled for. It runs before any per-stage logic.
+// - Some values require all or part of their analysis logic to run per stage. These should  
+//   implement AnalyzeInStage() with that stage-specific logic.
+//
+// IMPORTANT: Only give a value AnalyzeInStage() if it needs per-stage analysis. If the logic is  
+// stage-agnostic, it belongs in Analyze(). A value can implement both Analyze() and  
+// AnalyzeInStage() if it needs general logic once and other logic per stage.
+//
+// # GraphProperties  
+// Analyze functions can set and read GraphProperty flags. These are per-value bitflags that are  
+// automatically propagated as the IR graph is analyzed. When a value receives its Analyze() call,  
+// all its dependencies have already been analyzed, so it can freely inspect their GraphProperties.  
+// For example, it can verify that no other value sets a specific graph property upstream.
 struct FMaterialIRValueAnalyzer
 {
 	// Translator specific structure during generation of VT stacks. This is later converted to FMaterialVirtualTextureStack entries in FUniformExpressionSet.
@@ -25,35 +52,35 @@ struct FMaterialIRValueAnalyzer
 		TextureAddress AddressV{ TA_MAX };
 	};
 
-	//
+	// Resets the analyzer to process a new translation run.
 	void Setup(UMaterial* InMaterial, FMaterialIRModule* InModule, FMaterialCompilationOutput* InCompilationOutput, FMaterialInsights* InInsights = nullptr);
 		
-	//
+	// Performs stage-agnostic analysis on the given value.
 	void Analyze(MIR::FValue* Value);
 
-	//
-	void PropagateStateInStage(MIR::FValue* Value, MIR::EStage Stage);
+	// Performs stage-specific analysis on the given value, if the value kind requires some.
+	void AnalyzeInStage(MIR::FValue* Value, MIR::EStage Stage);
 
-	//
+	// The material being built.
 	UMaterial* Material{};
 	
-	//
+	// The destination module that will contain the result of material translation.
 	FMaterialIRModule* Module{};
 	
-	//
+	// Optional. Specifies the target MaterialInsights to populate, if provided.
+	FMaterialInsights* Insights{};
+
+	// The destination compilation output whose state to populate, based on the values being analyzed.
 	FMaterialCompilationOutput* CompilationOutput{};
 		
-	//
+	// Maps default values to their default-value offset, as used in UniformExpressionSet (e.g. UniformExpressionSet.AddDefaultParameterValue() and UniformExpressionSet.FindOrAddNumericParameter()).
 	TMap<UE::Shader::FValue, uint32> UniformDefaultValueOffsets{};
 		
-	//
+	// Stores free uniform buffer offsets for 1, 2, and 3 leftover components to optimize float4 packing
 	TArray<uint32, TInlineAllocator<8>> FreeOffsetsPerNumComponents[3];
 
 	// List of enabled shader environment defines.
 	TSet<FName> EnvironmentDefines;
-		
-	//
-	FMaterialInsights* Insights{}; // Optional
 
 	// VT stack entries during analysis. This is later converted to FMaterialVirtualTextureStack entries in FUniformExpressionSet.
 	TArray<FVTStackEntry> VTStacks;
diff --git a/Engine/Source/Runtime/Engine/Public/Materials/MaterialIR.h b/Engine/Source/Runtime/Engine/Public/Materials/MaterialIR.h
--- a/Engine/Source/Runtime/Engine/Public/Materials/MaterialIR.h
+++ b/Engine/Source/Runtime/Engine/Public/Materials/MaterialIR.h
@@ -13,6 +13,18 @@ struct FMaterialExternalCodeDeclaration;
 namespace MIR
 {
 
+// A simple string view without constructors, to be used inside FValues.
+struct FSimpleStringView
+{
+	const TCHAR* Ptr; // The string characters
+	int32        Len; // The string length (excluding null byte)
+
+	FStringView ToView() const { return { Ptr, Len }; }
+	FString ToString() const { return FString{ ToView() }; }
+	operator FStringView() const { return ToView(); }
+	FSimpleStringView& operator=(FStringView View) { Ptr = View.GetData(); Len = View.Len(); return *this; }
+};
+
 // Identifies the block of execution in which an instruction runs.
 // Note: If you introduce a new stage, make sure you update EValueFlags accordingly.
 enum EStage
@@ -579,11 +591,7 @@ struct FInlineHLSL : TInstruction<VK_InlineHLSL, 0>
 		const FMaterialExternalCodeDeclaration* ExternalCodeDeclaration;
 
 		// The arbitrary inlined HLSL code snippet
-		struct
-		{
-			const TCHAR* Code; // The string
-			uint32 CodeLen;    // The string length (excluding null byte)
-		};
+		FSimpleStringView Code;
 	};
 };
 
diff --git a/Engine/Source/Runtime/Engine/Public/Materials/MaterialIRTypes.h b/Engine/Source/Runtime/Engine/Public/Materials/MaterialIRTypes.h
--- a/Engine/Source/Runtime/Engine/Public/Materials/MaterialIRTypes.h
+++ b/Engine/Source/Runtime/Engine/Public/Materials/MaterialIRTypes.h
@@ -138,8 +138,8 @@ struct FType
 	// Returns a float type with the specified dimensions (scalar, vector, or matrix).
 	static FType MakeFloat(int NumRows = 1, int NumColumns = 1);
 
-	// Returns an LWC type with the specified dimensions (scalar, vector, or matrix).
-	static FType MakeLWC(int NumRows = 1, int NumColumns = 1);
+	// Returns an double precision type with the specified dimensions (scalar, vector, or matrix).
+	static FType MakeDouble(int NumRows = 1, int NumColumns = 1);
 
 	// Returns a scalar type of the given scalar kind.
 	static FType MakeScalar(EScalarKind InScalarKind);
@@ -176,64 +176,34 @@ struct FType
 	bool IsPoison() const { return Kind == ETypeKind::Poison; }
 	
 	// Returns true if the given type is a primitive that supports arithmetic operations.
-	bool IsPrimitive() const
-	{
-		return Is(ETypeKind::Primitive);
-	}
+	bool IsPrimitive() const { return Is(ETypeKind::Primitive); }
 
 	// Returns true if the given type is a primitive that supports arithmetic operations.
-	bool IsArithmetic() const
-	{
-		return Is(ETypeKind::Primitive) && Primitive.IsArithmetic();
-	}
+	bool IsArithmetic() const { return Is(ETypeKind::Primitive) && Primitive.IsArithmetic(); }
 
 	// Returns true if the given type is a primitive with a boolean scalar kind.
-	bool IsBoolean() const
-	{
-		return Is(ETypeKind::Primitive) && Primitive.IsBoolean();
-	}
+	bool IsBoolean() const { return Is(ETypeKind::Primitive) && Primitive.IsBoolean(); }
 
 	// Returns true if the given type is a primitive with an integer scalar kind.
-	bool IsInteger() const
-	{
-		return Is(ETypeKind::Primitive) && Primitive.IsInteger();
-	}
+	bool IsInteger() const { return Is(ETypeKind::Primitive) && Primitive.IsInteger(); }
 
 	// Returns true if the given type is a primitive with a float scalar kind.
-	bool IsFloat() const
-	{
-		return Is(ETypeKind::Primitive) && Primitive.IsFloat();
-	}
+	bool IsFloat() const { return Is(ETypeKind::Primitive) && Primitive.IsFloat(); }
 
 	// Returns true if the given type is a primitive with an LWC scalar kind.
-	bool IsLWC() const
-	{
-		return IsPrimitive() && Primitive.IsLWC();
-	}
+	bool IsLWC() const { return IsPrimitive() && Primitive.IsLWC(); }
 
 	// Returns true if the given type is a primitive with any float-based scalar kind (e.g. float, LWC).
-	bool IsAnyFloat() const
-	{
-		return Is(ETypeKind::Primitive) && Primitive.IsAnyFloat();
-	}
+	bool IsAnyFloat() const { return Is(ETypeKind::Primitive) && Primitive.IsAnyFloat(); }
 
 	// Returns true if the given type is a scalar primitive.
-	bool IsScalar() const
-	{
-		return Is(ETypeKind::Primitive) && Primitive.IsScalar();
-	}
+	bool IsScalar() const { return Is(ETypeKind::Primitive) && Primitive.IsScalar(); }
 
 	// Returns true if the given type is a boolean scalar primitive.
-	bool IsBoolScalar() const
-	{
-		return Is(ETypeKind::Primitive) && Primitive.IsBoolScalar();
-	}
+	bool IsBoolScalar() const { return Is(ETypeKind::Primitive) && Primitive.IsBoolScalar(); }
 
 	// Returns true if the given type is a vector primitive.
-	bool IsVector() const
-	{
-		return Is(ETypeKind::Primitive) && Primitive.IsVector();
-	}
+	bool IsVector() const { return Is(ETypeKind::Primitive) && Primitive.IsVector(); }
 
 	// Returns whether this type is a Texture2D object.
 	bool IsTexture2D() const { return Kind == ETypeKind::Texture2D; }
