[MaterialInstance] Early out if Parent is null inside UMaterialInstance::IsRedundant which can happen in rare cases when material instance is in invalid state.

#rnx

diff --git a/Engine/Source/Runtime/Engine/Private/Materials/MaterialInstance.cpp b/Engine/Source/Runtime/Engine/Private/Materials/MaterialInstance.cpp
--- a/Engine/Source/Runtime/Engine/Private/Materials/MaterialInstance.cpp
+++ b/Engine/Source/Runtime/Engine/Private/Materials/MaterialInstance.cpp
@@ -5441,6 +5441,11 @@ bool UMaterialInstance::Equivalent(const UMaterialInstance* CompareTo) const
 
 bool UMaterialInstance::IsRedundant() const
 {
+	if (!Parent)
+	{
+		return false;
+	}
+		 
 	if (NaniteOverrideMaterial.bEnableOverride)
 	{
 		// Check if we resolve to a different material to our parent 
