Transactional support for NotifyControlMessage.

diff --git a/Engine/Source/Runtime/Engine/Private/DataChannel.cpp b/Engine/Source/Runtime/Engine/Private/DataChannel.cpp
--- a/Engine/Source/Runtime/Engine/Private/DataChannel.cpp
+++ b/Engine/Source/Runtime/Engine/Private/DataChannel.cpp
@@ -2015,7 +2015,20 @@ void UControlChannel::ReceivedBunch( FInBunch& Bunch )
 		else if (Connection->Driver->Notify != nullptr)
 		{
 			// Process control message on client/server connection
-			Connection->Driver->Notify->NotifyControlMessage(Connection, MessageType, Bunch);
+#if UE_WITH_REMOTE_OBJECT_HANDLE
+			if (Connection->Driver->ShouldReplicateTransactionally())
+			{
+				static FName TransactionalWorkName("UControlChannel::ReceivedBunch NotifyControlMessage");
+				UE::RemoteExecutor::ExecuteTransactional(TransactionalWorkName, [this, MessageType, &Bunch]()
+					{
+						Connection->Driver->Notify->NotifyControlMessage(Connection, MessageType, Bunch);
+					});
+			}
+			else
+#endif
+			{
+				Connection->Driver->Notify->NotifyControlMessage(Connection, MessageType, Bunch);
+			}
 		}
 
 		// if the message was not handled, eat it ourselves
