Dont allow VRS on materials using the Pixel Depth Output. When using pixel depth offset, coarse shading causes the whole block to output the same depth value resulting in noticeable artifacts.

#tests Perf replayruns A/B on consoles (no perf regression)

diff --git a/Engine/Source/Runtime/Engine/Private/Materials/MaterialShared.cpp b/Engine/Source/Runtime/Engine/Private/Materials/MaterialShared.cpp
--- a/Engine/Source/Runtime/Engine/Private/Materials/MaterialShared.cpp
+++ b/Engine/Source/Runtime/Engine/Private/Materials/MaterialShared.cpp
@@ -1634,7 +1634,8 @@ bool FMaterialResource::IsVariableRateShadingAllowed() const {
 
 	return Material->bAllowVariableRateShading
 		&& !IsMasked()								// When using pixel discard, coarse shading causes the whole block to get discarded resulting in noticeable artifacts
-		&& !bDisableForCustomPrimitiveData;			// When using custom primitive data, Nanite can't determine which primitive ID to use for shading clusters, leading to edge artifacts
+		&& !bDisableForCustomPrimitiveData			// When using custom primitive data, Nanite can't determine which primitive ID to use for shading clusters, leading to edge artifacts
+		&& !HasPixelDepthOffsetConnected();			// When using pixel depth offset, coarse shading causes the whole block to output the same depth value resulting in noticeable artifacts
 }
 FString FMaterialResource::GetBaseMaterialPathName() const { return Material->GetPathName(); }
 FString FMaterialResource::GetDebugName() const
