Added helper to obtain the Generated Package list per World/per-plugin

Changed FStreamingDescriptor to always use GetLevelPackageName() as it's way to report the Cell package (used to differ wether you were a RuntimeHashSet or a RuntimeSpatialHash)
Made FWorldPartitionUtils::FSimulateCookedSession accurately simulate EDLs

diff --git a/Engine/Source/Editor/UnrealEd/Private/WorldPartition/WorldPartitionStaticLightingBuilder.cpp b/Engine/Source/Editor/UnrealEd/Private/WorldPartition/WorldPartitionStaticLightingBuilder.cpp
--- a/Engine/Source/Editor/UnrealEd/Private/WorldPartition/WorldPartitionStaticLightingBuilder.cpp
+++ b/Engine/Source/Editor/UnrealEd/Private/WorldPartition/WorldPartitionStaticLightingBuilder.cpp
@@ -628,8 +628,19 @@ void FStaticLightingDescriptors::InitializeFromWorld(UWorld* InWorld)
 	{
 		for (FStreamingDescriptor::FStreamingCell& Cell : Grid.StreamingCells)
 		{
-			FLightingCellDesc& CellDesc = LightingCellsDescs.Add(Cell.CellPackage);
-			CellDesc.CellLevelPackage = Cell.CellPackage;
+			FString CellPackageNameStr = Cell.CellPackage.ToString();
+			
+			// Trim the path only keep the name
+			int32 LastSlash = 0;
+			if (CellPackageNameStr.FindLastChar('/', LastSlash))
+			{
+				CellPackageNameStr.RightChopInline(LastSlash + 1);
+			}
+
+			FName CellPackageName(CellPackageNameStr);
+
+			FLightingCellDesc& CellDesc = LightingCellsDescs.Add(CellPackageName);
+			CellDesc.CellLevelPackage = CellPackageName;
 			CellDesc.Bounds = Cell.Bounds;
 			CellDesc.DataLayers = Cell.DataLayers;
 			CellDesc.RuntimeGrid = Grid.Name;
@@ -638,7 +649,7 @@ void FStaticLightingDescriptors::InitializeFromWorld(UWorld* InWorld)
 			{
 				check(!ActorGuidsToDesc.Find(Actor.ActorGuid));
 
-				UE_LOG(LogWorldPartitionStaticLightingBuilder, Verbose, TEXT("Discovered ActorInstance Guid %s, Cell %s / Grid %s"), *Actor.ActorGuid.ToString(), *Cell.CellPackage.ToString(), *Grid.Name.ToString());
+				UE_LOG(LogWorldPartitionStaticLightingBuilder, Verbose, TEXT("Discovered ActorInstance Guid %s, Cell %s / Grid %s"), *Actor.ActorGuid.ToString(), *CellPackageNameStr, *Grid.Name.ToString());
 
 				FLightingActorDesc& LightingActorDesc = ActorGuidsToDesc.Add(Actor.ActorGuid);
 				LightingActorDesc.ActorGuid = Actor.ActorGuid;
@@ -777,14 +788,9 @@ bool FStaticLightingDescriptors::CreateAndUpdateActors()
 			if (!DataActor)
 			{
 				// Generate Actor Name/Label
-
-				// Trim path from CellLevelPackage 
 				FString CellPackageName = CellDesc.CellLevelPackage.ToString();
 				int32 LastSlash = 0;
-				if (CellPackageName.FindLastChar('/', LastSlash))
-				{
-					CellPackageName.RightChopInline(LastSlash+1);
-				}
+				check(!CellPackageName.FindLastChar('/', LastSlash))
 				
 				FString DataActorName = FString::Printf(TEXT("%s_MapBuildData_%s"), *World->GetName(), *CellPackageName);
 
diff --git a/Engine/Source/Runtime/Engine/Private/WorldPartition/DataLayer/ExternalDataLayerManager.cpp b/Engine/Source/Runtime/Engine/Private/WorldPartition/DataLayer/ExternalDataLayerManager.cpp
--- a/Engine/Source/Runtime/Engine/Private/WorldPartition/DataLayer/ExternalDataLayerManager.cpp
+++ b/Engine/Source/Runtime/Engine/Private/WorldPartition/DataLayer/ExternalDataLayerManager.cpp
@@ -824,6 +824,11 @@ bool UExternalDataLayerManager::UnregisterExternalDataLayerInstance(UExternalDat
 	return true;
 }
 
+URuntimeHashExternalStreamingObjectBase* UExternalDataLayerManager::GetExternalStreamingObjectForDataLayerAsset(const UExternalDataLayerAsset* InAsset) const
+{
+	return ExternalStreamingObjects.FindChecked(InAsset);
+}
+
 void UExternalDataLayerManager::ForEachExternalStreamingObjects(TFunctionRef<bool(URuntimeHashExternalStreamingObjectBase*)> Func) const
 {
 	for (auto& [ExternalDataLayerAsset, ExternalStreamingObject] : ExternalStreamingObjects)
diff --git a/Engine/Source/Runtime/Engine/Private/WorldPartition/WorldPartitionStreamingDescriptor.cpp b/Engine/Source/Runtime/Engine/Private/WorldPartition/WorldPartitionStreamingDescriptor.cpp
--- a/Engine/Source/Runtime/Engine/Private/WorldPartition/WorldPartitionStreamingDescriptor.cpp
+++ b/Engine/Source/Runtime/Engine/Private/WorldPartition/WorldPartitionStreamingDescriptor.cpp
@@ -231,7 +231,7 @@ bool FStreamingDescriptor::GenerateStreamingDescriptor(UWorld* InWorld, FStreami
 						StreamingCell.bIsAlwaysLoaded = Cell->IsAlwaysLoaded();
 						StreamingCell.bIsSpatiallyLoaded = Cell->IsSpatiallyLoaded();
 						StreamingCell.DataLayers = Cell->GetDataLayers();
-						StreamingCell.CellPackage = FName(Cell->GetName());
+						StreamingCell.CellPackage = Cell->GetLevelPackageName();
 
 						if (DataLayerManager)
 						{
@@ -278,6 +278,28 @@ bool FStreamingDescriptor::GenerateStreamingDescriptor(UWorld* InWorld, FStreami
 	return false;
 }
 
+TMap<FName, TArray<FName>> FStreamingDescriptor::GenerateWorldPluginPackagesList(UWorld* InWorld)
+{
+	TMap<FName, TArray<FName>> PluginPackages;
+
+	FStreamingDescriptor Descriptor;
+	FStreamingDescriptor::GenerateStreamingDescriptor(InWorld, Descriptor);
+
+	for (FStreamingDescriptor::FStreamingGrid& Grid : Descriptor.StreamingGrids)
+	{
+		for (FStreamingDescriptor::FStreamingCell& Cell : Grid.StreamingCells)
+		{
+			FName PluginName = FName(FPathViews::GetMountPointNameFromPath(Cell.CellPackage.ToString()));
+			FName PackageName = Cell.CellPackage;
+
+			PluginPackages.FindOrAdd(PluginName).Add(PackageName);
+		}
+	}
+
+	return PluginPackages;
+}
+
+
 };
 
 #endif
\ No newline at end of file
diff --git a/Engine/Source/Runtime/Engine/Private/WorldPartition/WorldPartitionUtils.cpp b/Engine/Source/Runtime/Engine/Private/WorldPartition/WorldPartitionUtils.cpp
--- a/Engine/Source/Runtime/Engine/Private/WorldPartition/WorldPartitionUtils.cpp
+++ b/Engine/Source/Runtime/Engine/Private/WorldPartition/WorldPartitionUtils.cpp
@@ -8,6 +8,8 @@
 #include "WorldPartition/Cook/WorldPartitionCookPackageContext.h"
 #include "WorldPartition/ContentBundle/ContentBundleEditor.h"
 #include "WorldPartition/ContentBundle/ContentBundleWorldSubsystem.h"
+#include "WorldPartition/DataLayer/DataLayerManager.h"
+#include "WorldPartition/DataLayer/ExternalDataLayerManager.h"
 #include "WorldPartition/WorldPartitionLog.h"
 #include "Engine/World.h"
 
@@ -62,7 +64,21 @@ bool FWorldPartitionUtils::FSimulateCookedSession::SimulateCook(const FSimulateC
 	{
 		if (ContentBundle->HasCookedContent())
 		{
-			WorldPartition->RuntimeHash->InjectExternalStreamingObject(ContentBundle->GetStreamingObject());
+			URuntimeHashExternalStreamingObjectBase* ExternalStreamingObject = ContentBundle->GetStreamingObject();
+			WorldPartition->RuntimeHash->InjectExternalStreamingObject(ExternalStreamingObject);
+			InjectedStreamingObjects.Add(ExternalStreamingObject);
+		}
+	}
+
+	// Simulate EDL injection
+	UExternalDataLayerManager* EDLManager = WorldPartition->GetExternalDataLayerManager();
+	for (const UDataLayerInstance* Instance : WorldPartition->GetWorld()->GetDataLayerManager()->GetDataLayerInstances())
+	{	
+		if (const UExternalDataLayerInstance* ExternalInstance = Cast<const UExternalDataLayerInstance>(Instance))
+		{
+			URuntimeHashExternalStreamingObjectBase* EDLStreamingObject = EDLManager->GetExternalStreamingObjectForDataLayerAsset(CastChecked<const UExternalDataLayerAsset>(ExternalInstance->GetAsset()));
+			WorldPartition->RuntimeHash->InjectExternalStreamingObject(EDLStreamingObject);
+			InjectedStreamingObjects.Add(EDLStreamingObject);
 		}
 	}
 
@@ -73,19 +89,13 @@ FWorldPartitionUtils::FSimulateCookedSession::~FSimulateCookedSession()
 {
 	if (CookContext)
 	{
-		// Simulate Content Bundles injection
-		TArray<TSharedPtr<FContentBundleEditor>> ContentBundles;
-		WorldPartition->GetWorld()->ContentBundleManager->GetEditorContentBundle(ContentBundles);
-		for (TSharedPtr<FContentBundleEditor>& ContentBundle : ContentBundles)
+		// Clean-up External Streaming Objects injection
+		for(URuntimeHashExternalStreamingObjectBase* InjectedStreamingObject : InjectedStreamingObjects)
 		{
-			if (ContentBundle->HasCookedContent())
-			{
-				URuntimeHashExternalStreamingObjectBase* ExternalStreamingObject = ContentBundle->GetStreamingObject();
-				WorldPartition->RuntimeHash->RemoveExternalStreamingObject(ExternalStreamingObject);
-				// Trash external streaming object
-				ExternalStreamingObject->Rename(*MakeUniqueObjectName(GetTransientPackage(), URuntimeHashExternalStreamingObjectBase::StaticClass()).ToString(), GetTransientPackage(), REN_NonTransactional | REN_DontCreateRedirectors | REN_DoNotDirty);
-			}
-		}
+			WorldPartition->RuntimeHash->RemoveExternalStreamingObject(InjectedStreamingObject);
+			// Trash external streaming object
+			InjectedStreamingObject->Rename(*MakeUniqueObjectName(GetTransientPackage(), URuntimeHashExternalStreamingObjectBase::StaticClass()).ToString(), GetTransientPackage(), REN_NonTransactional | REN_DontCreateRedirectors | REN_DoNotDirty);
+		}		
 
 		WorldPartition->EndCook(*CookContext);
 		delete CookContext;
diff --git a/Engine/Source/Runtime/Engine/Public/WorldPartition/DataLayer/ExternalDataLayerManager.h b/Engine/Source/Runtime/Engine/Public/WorldPartition/DataLayer/ExternalDataLayerManager.h
--- a/Engine/Source/Runtime/Engine/Public/WorldPartition/DataLayer/ExternalDataLayerManager.h
+++ b/Engine/Source/Runtime/Engine/Public/WorldPartition/DataLayer/ExternalDataLayerManager.h
@@ -86,6 +86,7 @@ private:
 	UE_API FString GetActorPackageName(const UExternalDataLayerAsset* InExternalDataLayerAsset, const ULevel* InDestinationLevel, const FString& InActorPath) const;
 	UE_API URuntimeHashExternalStreamingObjectBase* CreateExternalStreamingObjectUsingStreamingGeneration(const UExternalDataLayerAsset* InExternalDataLayerAsset);
 	bool HasInjectedExternalDataLayerAssets() const { return InjectedExternalDataLayerAssets.Num() > 0; }
+	UE_API URuntimeHashExternalStreamingObjectBase* GetExternalStreamingObjectForDataLayerAsset(const UExternalDataLayerAsset* InAsset) const;
 
 	// Used for PIE/-game
 	UE_API void PrepareEditorGameWorld();
@@ -130,6 +131,7 @@ private:
 	friend class UContentBundleEditingSubmodule;
 	friend class UWorldPartitionRuntimeLevelStreamingCell;
 	friend class UGameFeatureActionConvertContentBundleWorldPartitionBuilder;
+	friend struct FWorldPartitionUtils;
 };
 
 #undef UE_API
diff --git a/Engine/Source/Runtime/Engine/Public/WorldPartition/WorldPartitionStreamingDescriptor.h b/Engine/Source/Runtime/Engine/Public/WorldPartition/WorldPartitionStreamingDescriptor.h
--- a/Engine/Source/Runtime/Engine/Public/WorldPartition/WorldPartitionStreamingDescriptor.h
+++ b/Engine/Source/Runtime/Engine/Public/WorldPartition/WorldPartitionStreamingDescriptor.h
@@ -96,6 +96,7 @@ namespace UE::Private::WorldPartition
 
 		ENGINE_API void DumpStateLog(FHierarchicalLogArchive& Ar);
 		ENGINE_API static bool GenerateStreamingDescriptor(UWorld* InWorld, FStreamingDescriptor& OutStreamingDescriptor, const FStreamingDescriptorParams& Params = FStreamingDescriptorParams());
+		ENGINE_API static TMap<FName, TArray<FName>> GenerateWorldPluginPackagesList(UWorld* InWorld);
 	};
 };
 
diff --git a/Engine/Source/Runtime/Engine/Public/WorldPartition/WorldPartitionUtils.h b/Engine/Source/Runtime/Engine/Public/WorldPartition/WorldPartitionUtils.h
--- a/Engine/Source/Runtime/Engine/Public/WorldPartition/WorldPartitionUtils.h
+++ b/Engine/Source/Runtime/Engine/Public/WorldPartition/WorldPartitionUtils.h
@@ -16,6 +16,7 @@ class IWorldPartitionCell;
 class FWorldPartitionCookPackageContext;
 struct FWorldPartitionStreamingQuerySource;
 class AActor;
+class URuntimeHashExternalStreamingObjectBase;
 
 struct FWorldPartitionUtils
 {
@@ -38,6 +39,7 @@ struct FWorldPartitionUtils
 	private:
 		ENGINE_API bool SimulateCook(const FSimulateCookSessionParams& Params);
 
+		TArray<TObjectPtr<URuntimeHashExternalStreamingObjectBase>> InjectedStreamingObjects;
 		FWorldPartitionCookPackageContext* CookContext;
 		TWeakObjectPtr<UWorldPartition> WorldPartition;
 	};
