Fix SkinWeightPorfileManager RHI crash

#rnx

diff --git a/Engine/Source/Runtime/Engine/Private/Animation/SkinWeightProfileManager.cpp b/Engine/Source/Runtime/Engine/Private/Animation/SkinWeightProfileManager.cpp
--- a/Engine/Source/Runtime/Engine/Private/Animation/SkinWeightProfileManager.cpp
+++ b/Engine/Source/Runtime/Engine/Private/Animation/SkinWeightProfileManager.cpp
@@ -256,8 +256,9 @@ void FSkinWeightProfileManager::DoTick(float DeltaTime, ENamedThreads::Type Curr
 
 			if (!bRemove)
 			{	
+				check(SkeletalMesh);
 				bool bAllBuffersReady = true;
-				int32 LODIndexToWaitStreaming = INDEX_NONE;
+				int32 MinResidentLODsToWaitFor = INDEX_NONE;
 				for (int32 LODProfileArrayIndex = 0; LODProfileArrayIndex < Request.LODIndices.Num(); ++LODProfileArrayIndex)
 				{
 					int32 LODRenderDataArrayIndex = Request.LODIndices[LODProfileArrayIndex];
@@ -265,12 +266,14 @@ void FSkinWeightProfileManager::DoTick(float DeltaTime, ENamedThreads::Type Curr
 					if (!SkinWeightProfilesData.ContainsOverrideBuffer(Request.ProfileStack))
 					{
 						bAllBuffersReady = false;
-						if (ShouldIgnoreLOD(SkinWeightProfilesData, Request.WeakSkeletalMesh))
+						int32 NumResidentLODsRequired = ConvertLODIndexToCount(*SkeletalMesh, LODRenderDataArrayIndex);
+						if (ShouldIgnoreLOD(SkinWeightProfilesData, *SkeletalMesh, NumResidentLODsRequired))
 						{
 							// This LOD isn't streamed in and has no requests for that. There may be several such LODs for the current model.
-							// We will create a record in StreamingBlockedMeshes to track when any of these LODs will be streamed in,
-							// and for that we need to know the index of the least detailed LOD among them.
-							LODIndexToWaitStreaming = FMath::Max(LODIndexToWaitStreaming, LODRenderDataArrayIndex);
+							// We will create a record in StreamingBlockedMeshes to track when any of these LODs will be streamed in
+							MinResidentLODsToWaitFor = (MinResidentLODsToWaitFor == INDEX_NONE) ?
+								NumResidentLODsRequired :
+								FMath::Min(MinResidentLODsToWaitFor, NumResidentLODsRequired);
 						}
 						else
 						{
@@ -283,11 +286,7 @@ void FSkinWeightProfileManager::DoTick(float DeltaTime, ENamedThreads::Type Curr
 					}
 				}
 
-				if (LODIndexToWaitStreaming != INDEX_NONE)
-				{
-					const FStreamableRenderResourceState& SRRState = SkeletalMesh->GetStreamableResourceState();
-					Request.ExpectedResidentLODs = SRRState.MaxNumLODs - (LODIndexToWaitStreaming - SRRState.AssetLODBias);
-				}
+				Request.ExpectedResidentLODs = MinResidentLODsToWaitFor;
 
 				if (bAllBuffersReady)
 				{
@@ -377,23 +376,22 @@ void FSkinWeightProfileManager::CleanupRequest(const FSetProfileRequest& Request
 	
 }
 
-bool FSkinWeightProfileManager::ShouldIgnoreLOD(const FSkinWeightProfilesData& SkinWeightProfilesData, const TWeakObjectPtr<USkeletalMesh>& WeakSkeletalMesh) const
+bool FSkinWeightProfileManager::ShouldIgnoreLOD(const FSkinWeightProfilesData& SkinWeightProfilesData, const USkeletalMesh& SkeletalMesh, int32 NumResidentLODsRequired) const
 {
 	if (!GSkipLODWithoutDataOrStreamingRequest)
 	{
 		return false;
 	}
 
-	// Are LOD's vertices loaded? 
-	FSkinWeightVertexBuffer* BaseBuffer = SkinWeightProfilesData.BaseBuffer;
-	if (BaseBuffer && BaseBuffer->GetDataVertexBuffer()->VertexBufferRHI->GetSize())
+	// Do we have enough LODs in memory?
+	const FStreamableRenderResourceState& SRRState = SkeletalMesh.GetStreamableResourceState();
+	if (SRRState.NumResidentLODs >= NumResidentLODsRequired)
 	{
 		return false;
 	}
 
 	// Are they scheduled to be streamed in?
-	const USkeletalMesh* SkeletalMesh = WeakSkeletalMesh.Get();
-	if (SkeletalMesh && SkeletalMesh->HasPendingInitOrStreaming())
+	if (SkeletalMesh.HasPendingInitOrStreaming())
 	{
 		return false;
 	}
@@ -427,6 +425,12 @@ bool FSkinWeightProfileManager::ShouldSkipTick() const
 	return true; // All current requests are still waiting for streaming, we can end current tick early
 }
 
+int32 FSkinWeightProfileManager::ConvertLODIndexToCount(const USkeletalMesh& SkeletalMesh, int32 Index) const
+{
+	const FStreamableRenderResourceState& SRRState = SkeletalMesh.GetStreamableResourceState();
+	return (SRRState.MaxNumLODs - (Index - SRRState.AssetLODBias));
+}
+
 bool FSkinWeightProfileManager::IsTickableWhenPaused() const
 {
 	return false;
diff --git a/Engine/Source/Runtime/Engine/Public/Animation/SkinWeightProfileManager.h b/Engine/Source/Runtime/Engine/Public/Animation/SkinWeightProfileManager.h
--- a/Engine/Source/Runtime/Engine/Public/Animation/SkinWeightProfileManager.h
+++ b/Engine/Source/Runtime/Engine/Public/Animation/SkinWeightProfileManager.h
@@ -120,8 +120,9 @@ public:
 	void DoTick(float DeltaTime, ENamedThreads::Type CurrentThread, const FGraphEventRef& MyCompletionGraphEvent);
 protected:	
 	void CleanupRequest(const FSetProfileRequest& Request);
-	bool ShouldIgnoreLOD(const FSkinWeightProfilesData& SkinWeightProfilesData, const TWeakObjectPtr<USkeletalMesh>& WeakSkeletalMesh) const;
+	bool ShouldIgnoreLOD(const FSkinWeightProfilesData& SkinWeightProfilesData, const USkeletalMesh& SkeletalMesh, int32 NumResidentLODsRequired) const;
 	bool ShouldSkipTick() const;
+	int32 ConvertLODIndexToCount(const USkeletalMesh& SkeletalMesh, int32 Index) const;
 protected:
 	TArray<FSetProfileRequest, TInlineAllocator<4>> CanceledRequest;
 	TArray<FSetProfileRequest> PendingSetProfileRequests;
