[Input] Deprecate the bEnableDynamicComponentInputBinding setting and always treat it as true.

This setting was originally added in 5.1 and has been defaulted to true ever since, and is the correct behavior for how component's which get added via BP graphs or at runtime should behave.

Also, remove the bEnableInputDeviceSubsystem setting which was deprecated in 5.5.

diff --git a/Engine/Source/Runtime/Engine/Classes/GameFramework/InputDeviceSubsystem.h b/Engine/Source/Runtime/Engine/Classes/GameFramework/InputDeviceSubsystem.h
--- a/Engine/Source/Runtime/Engine/Classes/GameFramework/InputDeviceSubsystem.h
+++ b/Engine/Source/Runtime/Engine/Classes/GameFramework/InputDeviceSubsystem.h
@@ -126,9 +126,6 @@ public:
 
 	/**
 	 * Returns a pointer to the Input Device Engine Subsystem if it is available.
-	 * 
-	 * NOTE: This may be null if the bEnableInputDeviceSubsystem flag in UInputSettings
-	 * is set to false!
 	 */
 	static ENGINE_API UInputDeviceSubsystem* Get();
 
diff --git a/Engine/Source/Runtime/Engine/Classes/GameFramework/InputSettings.h b/Engine/Source/Runtime/Engine/Classes/GameFramework/InputSettings.h
--- a/Engine/Source/Runtime/Engine/Classes/GameFramework/InputSettings.h
+++ b/Engine/Source/Runtime/Engine/Classes/GameFramework/InputSettings.h
@@ -79,16 +79,6 @@ class UInputSettings
 	UPROPERTY(config, EditAnywhere, Category = "Input")
 	uint8 bFilterInputByPlatformUser:1;
 
-	/**
-	 * If true, then the input device subsystem will be allowed to Initalize when the engine boots.
-	 * NOTE: For this setting to take effect, and editor restart is required.
-	 * 
-	 * @see UInputDeviceSubsystem
-	 */
-	UE_DEPRECATED(5.5, "bEnableInputDeviceSubsystem is deprecated, ths input device subsystem will now always be created.")
-	UPROPERTY(config, meta = (DeprecatedProperty, DeprecationMessage = "bEnableInputDeviceSubsystem is deprecated, ths input device subsystem will now always be created."))
-	uint8 bEnableInputDeviceSubsystem:1;
-
 	/**
 	 * If true, then the Player Controller will have it's Pressed Keys flushed when the input mode is changed
 	 * to Game and UI mode or the game viewport loses focus. The default behavior is true.
@@ -103,7 +93,8 @@ class UInputSettings
 	 * Should components that are dynamically added via the 'AddComponent' function at runtime have input delegates bound to them?
 	 * @see AActor::FinishAddComponent
 	 */
-	UPROPERTY(config, EditAnywhere, Category = "Input")
+	UE_DEPRECATED(5.7, "bEnableDynamicComponentInputBinding is deprecated.")
+	UPROPERTY(config, meta = (DeprecatedProperty, DeprecationMessage = "bEnableDynamicComponentInputBinding is deprecated, it will be treated as always being true."))
 	uint8 bEnableDynamicComponentInputBinding:1;
 	
 	/** Should the touch input interface be shown always, or only when the platform has a touch screen? */
diff --git a/Engine/Source/Runtime/Engine/Private/Components/ActorComponent.cpp b/Engine/Source/Runtime/Engine/Private/Components/ActorComponent.cpp
--- a/Engine/Source/Runtime/Engine/Private/Components/ActorComponent.cpp
+++ b/Engine/Source/Runtime/Engine/Private/Components/ActorComponent.cpp
@@ -1962,8 +1962,8 @@ void UActorComponent::RegisterComponentWithWorld(UWorld* InWorld, FRegisterCompo
 		}
 
 	}
-
-	if (MyOwner && MyOwner->InputComponent && GetDefault<UInputSettings>()->bEnableDynamicComponentInputBinding)
+	
+	if (MyOwner && MyOwner->InputComponent)
 	{
 		UInputDelegateBinding::BindInputDelegates(GetClass(), MyOwner->InputComponent, this);
 	}
diff --git a/Engine/Source/Runtime/Engine/Private/GameFramework/InputDeviceSubsystem.cpp b/Engine/Source/Runtime/Engine/Private/GameFramework/InputDeviceSubsystem.cpp
--- a/Engine/Source/Runtime/Engine/Private/GameFramework/InputDeviceSubsystem.cpp
+++ b/Engine/Source/Runtime/Engine/Private/GameFramework/InputDeviceSubsystem.cpp
@@ -221,17 +221,6 @@ bool UInputDeviceSubsystem::ShouldCreateSubsystem(UObject* Outer) const
 	{
 		return false;
 	}
-
-	// As of 5.5 this option is disabled. Some input interfaces depend on this subsystem to get info about
-	// the currently connected devices, so we are enforcing it as a requirement from now on.
-PRAGMA_DISABLE_DEPRECATION_WARNINGS
-	// There is a setting to turn off this subsystem entirely.
-	const bool bShouldCreate = GetDefault<UInputSettings>()->bEnableInputDeviceSubsystem;
-	if (!bShouldCreate)
-	{
-		UE_LOG(LogInputDeviceProperties, Error, TEXT("UInputSettings::bEnableInputDeviceSubsystem is false, but it is required to be true as of 5.5. Ignoring..."));
-	}
-PRAGMA_ENABLE_DEPRECATION_WARNINGS
 	
 	return Super::ShouldCreateSubsystem(Outer);
 }
diff --git a/Engine/Source/Runtime/Engine/Private/UserInterface/InputSettings.cpp b/Engine/Source/Runtime/Engine/Private/UserInterface/InputSettings.cpp
--- a/Engine/Source/Runtime/Engine/Private/UserInterface/InputSettings.cpp
+++ b/Engine/Source/Runtime/Engine/Private/UserInterface/InputSettings.cpp
@@ -41,11 +41,10 @@ UInputSettings::UInputSettings(const FObjectInitializer& ObjectInitializer)
 	, bEnableLegacyInputScales(true)
 	, bEnableMotionControls(true)
 	, bFilterInputByPlatformUser(false)
-PRAGMA_DISABLE_DEPRECATION_WARNINGS
-	, bEnableInputDeviceSubsystem(true)
-PRAGMA_ENABLE_DEPRECATION_WARNINGS
 	, bShouldFlushPressedKeysOnViewportFocusLost(true)
+PRAGMA_DISABLE_DEPRECATION_WARNINGS
 	, bEnableDynamicComponentInputBinding(true)
+PRAGMA_ENABLE_DEPRECATION_WARNINGS
 	, DefaultViewportMouseCaptureMode(EMouseCaptureMode::CapturePermanently_IncludingInitialMouseDown)
 	, DefaultViewportMouseLockMode(EMouseLockMode::LockOnCapture)
 	, DefaultPlayerInputClass(UPlayerInput::StaticClass())
@@ -203,14 +202,6 @@ void UInputSettings::PostReloadConfig(FProperty* PropertyThatWasLoaded)
 	PopulateAxisConfigs();
 #endif
 	AddInternationalConsoleKey();
-
-PRAGMA_DISABLE_DEPRECATION_WARNINGS
-	if (!bEnableInputDeviceSubsystem)
-	{
-		bEnableInputDeviceSubsystem = true;
-		UE_LOG(LogInput, Warning, TEXT("UInputSettings::bEnableInputDeviceSubsystem was detected as false from a config, but it is required to be true as of 5.5. Updating the value to be true."));
-	}
-PRAGMA_ENABLE_DEPRECATION_WARNINGS
 }
 
 #if WITH_EDITOR
