Disallow direct cycles of composite textures.

#rnx

diff --git a/Engine/Source/Runtime/Engine/Private/Texture.cpp b/Engine/Source/Runtime/Engine/Private/Texture.cpp
--- a/Engine/Source/Runtime/Engine/Private/Texture.cpp
+++ b/Engine/Source/Runtime/Engine/Private/Texture.cpp
@@ -935,6 +935,13 @@ PRAGMA_ENABLE_DEPRECATION_WARNINGS
 		}
 		else if ( PropertyName == CompositeTextureName )
 		{
+			PRAGMA_DISABLE_DEPRECATION_WARNINGS
+			if (CompositeTexture && CompositeTexture->CompositeTexture == this)
+			{
+				UE_LOG(LogTexture, Warning, TEXT("Cycle detected: Requested composite texture has this texture as its composite texture!"));
+				CompositeTexture = nullptr;
+			}
+			PRAGMA_ENABLE_DEPRECATION_WARNINGS
 			NotifyIfCompositeTextureChanged();
 		}
 #endif //WITH_EDITORONLY_DATA
