New material translator: Fix self-inflicted static analyzer error

#tests local compile

diff --git a/Engine/Source/Runtime/Engine/Private/Materials/MaterialIREmitter.cpp b/Engine/Source/Runtime/Engine/Private/Materials/MaterialIREmitter.cpp
--- a/Engine/Source/Runtime/Engine/Private/Materials/MaterialIREmitter.cpp
+++ b/Engine/Source/Runtime/Engine/Private/Materials/MaterialIREmitter.cpp
@@ -3352,12 +3352,14 @@ void FEmitter::Initialize()
 	// If this assert fails, it's likely due to uninitialized padding bytes in FValue introduced by compiler-added padding.
 	// Review any recent changes to FValue to ensure no additional padding was introduced (i.e., the struct remains tightly packed).
 	// Temporary disable this check because it's causing a 100% assert on Mac
-	if (false && FPlatformMisc::IsDebuggerPresent())
+	/*
+	if (FPlatformMisc::IsDebuggerPresent())
 	{
 		FValue* A = ConstantZero(MIR::EScalarKind::Float);
 		FValue* B = ConstantZero(MIR::EScalarKind::Float);
 		check(A == B);
 	}
+	*/
 }
 
 bool FEmitter::FValueKeyFuncs::Matches(KeyInitType A, KeyInitType B)
