Expose Max Draw Distance and Max Distance Fade Range to BP

diff --git a/Engine/Source/Runtime/Engine/Classes/Components/LightComponent.h b/Engine/Source/Runtime/Engine/Classes/Components/LightComponent.h
--- a/Engine/Source/Runtime/Engine/Classes/Components/LightComponent.h
+++ b/Engine/Source/Runtime/Engine/Classes/Components/LightComponent.h
@@ -53,10 +53,12 @@ class ULightComponent : public ULightComponentBase
 	UPROPERTY(EditAnywhere, BlueprintReadOnly, interp, Category = Light, meta = (UIMin = "1700.0", UIMax = "12000.0", ShouldShowInViewport = true, DisplayAfter ="bUseTemperature"))
 	float Temperature;
 	
-	UPROPERTY(EditAnywhere, Category = Performance)
+	/** World distance to where the light stops rendering (lighting & shadows). 0.0 means no max draw distance. This can be controlled via a global scalar scalability setting r.LightMaxDrawDistanceScale as well */
+	UPROPERTY(EditAnywhere, BlueprintReadOnly, Category = Performance)
 	float MaxDrawDistance;
 
-	UPROPERTY(EditAnywhere, Category = Performance)
+	/** Fade the light out as it approaches MaxDrawDistance. The fading starts when the light is MaxDistanceFadeRange distance away from MaxDrawDistance, where it ends up at 0 for a smooth fade-out without popping. Requires MaxDrawDistance to be set. */
+	UPROPERTY(EditAnywhere, BlueprintReadOnly, Category = Performance, meta = (EditCondition = "MaxDrawDistance > 0.0"))
 	float MaxDistanceFadeRange;
 
 	/** false: use white (D65) as illuminant. */
@@ -357,6 +359,12 @@ public:
 	UFUNCTION(BlueprintCallable, Category="Rendering|Components|Light")
 	ENGINE_API void SetUseRayTracedDistanceFieldShadows(bool bNewValue);
 
+	UFUNCTION(BlueprintCallable, Category = "Rendering|Components|Light")
+	ENGINE_API void SetMaxDrawDistance(float NewValue);
+
+	UFUNCTION(BlueprintCallable, Category = "Rendering|Components|Light")
+	ENGINE_API void SetMaxDistanceFadeRange(float NewValue);
+
 public:
 	/** The light's scene info. */
 	FLightSceneProxy* SceneProxy;
diff --git a/Engine/Source/Runtime/Engine/Private/Components/LightComponent.cpp b/Engine/Source/Runtime/Engine/Private/Components/LightComponent.cpp
--- a/Engine/Source/Runtime/Engine/Private/Components/LightComponent.cpp
+++ b/Engine/Source/Runtime/Engine/Private/Components/LightComponent.cpp
@@ -1350,6 +1350,26 @@ void ULightComponent::SetUseRayTracedDistanceFieldShadows(bool bNewValue)
 	}
 }
 
+void ULightComponent::SetMaxDrawDistance(float bNewValue)
+{
+	if (AreDynamicDataChangesAllowed() &&
+		bNewValue != MaxDrawDistance)
+	{
+		MaxDrawDistance = bNewValue;
+		MarkRenderStateDirty();
+	}
+}
+
+void ULightComponent::SetMaxDistanceFadeRange(float bNewValue)
+{
+	if (AreDynamicDataChangesAllowed() &&
+		bNewValue != MaxDistanceFadeRange)
+	{
+		MaxDistanceFadeRange = bNewValue;
+		MarkRenderStateDirty();
+	}
+}
+
 // GetDirection
 FVector ULightComponent::GetDirection() const 
 { 
