Fix possible UB when reading 32 bits from the FDwordBitReader implementation in MorphTargetVertexCodec.

#rnx

diff --git a/Engine/Source/Runtime/Engine/Private/Rendering/MorphTargetVertexCodec.cpp b/Engine/Source/Runtime/Engine/Private/Rendering/MorphTargetVertexCodec.cpp
--- a/Engine/Source/Runtime/Engine/Private/Rendering/MorphTargetVertexCodec.cpp
+++ b/Engine/Source/Runtime/Engine/Private/Rendering/MorphTargetVertexCodec.cpp
@@ -91,7 +91,7 @@ public:
 		}
 		else
 		{
-			const uint32 BitMask = (1 << NumBits) - 1;
+			const uint32 BitMask = (1ull << NumBits) - 1;
 			return (Buffer[BaseIndex] >> BitOffset) & BitMask;
 		}
 	}
