Fix comparing ActorInstanceHandles pointing to the same manager

- Problem: Currently, if two ActorInstanceHandles resolve to their own managers, they will be evaluated as equal no matter
their InstanceIndex, which seems wrong considering they represent two different instances;

- Solution: Check if the ActorInstanceHandles point to their own managers before comparing, and if so, compare their InstanceIndex;

diff --git a/Engine/Source/Runtime/Engine/Private/Engine/ActorInstanceHandle.cpp b/Engine/Source/Runtime/Engine/Private/Engine/ActorInstanceHandle.cpp
--- a/Engine/Source/Runtime/Engine/Private/Engine/ActorInstanceHandle.cpp
+++ b/Engine/Source/Runtime/Engine/Private/Engine/ActorInstanceHandle.cpp
@@ -381,14 +381,27 @@ bool FActorInstanceHandle::operator==(const FActorInstanceHandle& Other) const
 		return false;
 	}
 
-	// Try to compare the cached actors if available
-	const AActor* MyActor = GetCachedActor();
-	const AActor* OtherActor = Other.GetCachedActor();
+	// If both handles point to the same Manager/InstanceIndex pair, we can assume they'll point
+	// to the same ReferenceObject.
+	if (ManagerInterface == Other.ManagerInterface && InstanceIndex == Other.InstanceIndex)
+	{
+		return true;
+	}
 
 	// If both Actors are set, we can compare the pointers.
 	// But if one Actor is missing, there's a possibility that the Actor was spawned after the handle was initially setup.
 	// In that case we defer to operator==(AActor*) which will try to update the cached actor by calling FindActor on the manager interface.
 	// If one of the actor is set and FindActor returns null, we know both handles can't point to the same Actor.
+	const AActor* MyActor = GetCachedActor();
+	const AActor* OtherActor = Other.GetCachedActor();
+	if (MyActor == ManagerInterface.GetObject() && OtherActor == Other.ManagerInterface.GetObject())
+	{
+		// Both actors point to their own managers, so we can compare their Manager/InstanceIndex instead,
+		// since two handles may live inside the same manager. In this case, we'll just return false, since
+		// we already verified above that their Manager/InstanceIndex are different.
+		return false;
+	}
+	
 	if (MyActor)
 	{
 		if (OtherActor)
@@ -405,9 +418,7 @@ bool FActorInstanceHandle::operator==(const FActorInstanceHandle& Other) const
 		return operator==(OtherActor);
 	}
 
-	// The actors are null, try to compare managers and indices first if we have them.
-	// If the ManagerInterface and InstanceIndex of both handles are invalid, we consider them null and equal to each other.
-	return ManagerInterface == Other.ManagerInterface && InstanceIndex == Other.InstanceIndex;
+	return false;
 }
 
 bool FActorInstanceHandle::operator!=(const FActorInstanceHandle& Other) const
