Added default initialization of FMaterialExpressionAggregateEntry::AttributeIndex.

#rnx

diff --git a/Engine/Source/Runtime/Engine/Public/Materials/MaterialExpressionAggregate.h b/Engine/Source/Runtime/Engine/Public/Materials/MaterialExpressionAggregate.h
--- a/Engine/Source/Runtime/Engine/Public/Materials/MaterialExpressionAggregate.h
+++ b/Engine/Source/Runtime/Engine/Public/Materials/MaterialExpressionAggregate.h
@@ -22,7 +22,7 @@ struct FMaterialExpressionAggregateEntry
 	GENERATED_BODY()
 
 	UPROPERTY()
-	int32 AttributeIndex;
+	int32 AttributeIndex{};
 
 	UPROPERTY()
 	FExpressionInput Input;
