Fixup function argument name.
Clamping max distance input value as negative values have no meaning.
Clarified comment that 0 or less has no effect for max distance.

diff --git a/Engine/Source/Runtime/Engine/Classes/Components/LightComponent.h b/Engine/Source/Runtime/Engine/Classes/Components/LightComponent.h
--- a/Engine/Source/Runtime/Engine/Classes/Components/LightComponent.h
+++ b/Engine/Source/Runtime/Engine/Classes/Components/LightComponent.h
@@ -53,7 +53,7 @@ class ULightComponent : public ULightComponentBase
 	UPROPERTY(EditAnywhere, BlueprintReadOnly, interp, Category = Light, meta = (UIMin = "1700.0", UIMax = "12000.0", ShouldShowInViewport = true, DisplayAfter ="bUseTemperature"))
 	float Temperature;
 	
-	/** World distance to where the light stops rendering (lighting & shadows). 0.0 means no max draw distance. This can be controlled via a global scalar scalability setting r.LightMaxDrawDistanceScale as well */
+	/** World distance to where the light stops rendering (lighting & shadows). 0.0 or less means no max draw distance. This can be controlled via a global scalar scalability setting r.LightMaxDrawDistanceScale as well */
 	UPROPERTY(EditAnywhere, BlueprintReadOnly, Category = Performance)
 	float MaxDrawDistance;
 
diff --git a/Engine/Source/Runtime/Engine/Private/Components/LightComponent.cpp b/Engine/Source/Runtime/Engine/Private/Components/LightComponent.cpp
--- a/Engine/Source/Runtime/Engine/Private/Components/LightComponent.cpp
+++ b/Engine/Source/Runtime/Engine/Private/Components/LightComponent.cpp
@@ -1350,22 +1350,25 @@ void ULightComponent::SetUseRayTracedDistanceFieldShadows(bool bNewValue)
 	}
 }
 
-void ULightComponent::SetMaxDrawDistance(float bNewValue)
+void ULightComponent::SetMaxDrawDistance(float NewValue)
 {
+	// Negative max draw distance has no meaning. We only process positive distances in the lighting code anyway. Clamping here to avoid marking render state dirty in case of negative input value.
+	const float ClampedNewValue = FMath::Max(NewValue, 0.0f);
+
 	if (AreDynamicDataChangesAllowed() &&
-		bNewValue != MaxDrawDistance)
+		ClampedNewValue != MaxDrawDistance)
 	{
-		MaxDrawDistance = bNewValue;
+		MaxDrawDistance = ClampedNewValue;
 		MarkRenderStateDirty();
 	}
 }
 
-void ULightComponent::SetMaxDistanceFadeRange(float bNewValue)
+void ULightComponent::SetMaxDistanceFadeRange(float NewValue)
 {
 	if (AreDynamicDataChangesAllowed() &&
-		bNewValue != MaxDistanceFadeRange)
+		NewValue != MaxDistanceFadeRange)
 	{
-		MaxDistanceFadeRange = bNewValue;
+		MaxDistanceFadeRange = NewValue;
 		MarkRenderStateDirty();
 	}
 }
