Fix for layer stack crash error message resolving when incorrect material attribute output nodes are present.

diff --git a/Engine/Source/Runtime/Engine/Private/Materials/MaterialExpressionLayerStack.cpp b/Engine/Source/Runtime/Engine/Private/Materials/MaterialExpressionLayerStack.cpp
--- a/Engine/Source/Runtime/Engine/Private/Materials/MaterialExpressionLayerStack.cpp
+++ b/Engine/Source/Runtime/Engine/Private/Materials/MaterialExpressionLayerStack.cpp
@@ -332,7 +332,7 @@ bool UMaterialExpressionLayerStack::PollFunctionExpressionsForLayerUsage(FMateri
 
 	if (Tracker.MAOutputCount != AcceptableNumLayerBlendMAOutputs)
 	{
-		LogError(Compiler, TEXT("Function %i: %s, must have exactly %d Material Attributes type output node. It has %d"), *CurrentFunction->GetName(), AcceptableNumLayerBlendMAOutputs, Tracker.MAOutputCount);
+		LogError(Compiler, TEXT("Function %s, must have exactly %d Material Attributes type output node. It has %d"), *CurrentFunction->GetName(), AcceptableNumLayerBlendMAOutputs, Tracker.MAOutputCount);
 		bResult = false;
 	}
 
diff --git a/Engine/Source/Runtime/Engine/Public/Materials/MaterialExpressionMaterialAttributeLayers.h b/Engine/Source/Runtime/Engine/Public/Materials/MaterialExpressionMaterialAttributeLayers.h
--- a/Engine/Source/Runtime/Engine/Public/Materials/MaterialExpressionMaterialAttributeLayers.h
+++ b/Engine/Source/Runtime/Engine/Public/Materials/MaterialExpressionMaterialAttributeLayers.h
@@ -18,9 +18,9 @@ struct FMaterialParameterInfo;
 * These defines are for tracking the required number of inputs / outputs for blends rather than the values being
 * hard coded throughout the layer functions and expressions code
 */
-constexpr uint8 AcceptableNumLayerMAInputs = 1;
-constexpr uint8 AcceptableNumBlendMAInputs = 2;
-constexpr uint8 AcceptableNumLayerBlendMAOutputs = 1;
+constexpr int8 AcceptableNumLayerMAInputs = 1;
+constexpr int8 AcceptableNumBlendMAInputs = 2;
+constexpr int8 AcceptableNumLayerBlendMAOutputs = 1;
 
 UCLASS(hidecategories=Object, MinimalAPI)
 class UMaterialExpressionMaterialAttributeLayers : public UMaterialExpression
