Restore ClampedLODIndex for skeletal meshes to reduce vertex format mismatches during PSO precaching

FSkinnedMeshComponentHelper::ComputeMinLOD respects bOverrideMinLod/MinLodModel on the component, but the proxy can't use that because they are dynamic and can potentially decrease after ClampedLODIndex is locked in. Skip over that part of the MinLOD calculation.

diff --git a/Engine/Source/Runtime/Engine/Private/SkeletalMeshSceneProxy.cpp b/Engine/Source/Runtime/Engine/Private/SkeletalMeshSceneProxy.cpp
--- a/Engine/Source/Runtime/Engine/Private/SkeletalMeshSceneProxy.cpp
+++ b/Engine/Source/Runtime/Engine/Private/SkeletalMeshSceneProxy.cpp
@@ -84,19 +84,19 @@ bool AllowLumenCardGenerationForSkeletalMeshes(EShaderPlatform Platform)
  * @param	Component - skeletal mesh primitive being added
  */
 FSkeletalMeshSceneProxy::FSkeletalMeshSceneProxy(const USkinnedMeshComponent* Component, FSkeletalMeshRenderData* InSkelMeshRenderData)
-	: FSkeletalMeshSceneProxy(FSkinnedMeshSceneProxyDesc(Component), InSkelMeshRenderData, Component->ComputeMinLOD())
+	: FSkeletalMeshSceneProxy(FSkinnedMeshSceneProxyDesc(Component), InSkelMeshRenderData, Component->GetValidMinLOD(Component->GetSkinnedAsset()->GetMinLodIdx()))
 { }
 
 FSkeletalMeshSceneProxy::FSkeletalMeshSceneProxy(const FSkinnedMeshSceneProxyDesc& InMeshDesc, FSkeletalMeshRenderData* InSkelMeshRenderData)
 	: FSkeletalMeshSceneProxy(InMeshDesc, InSkelMeshRenderData, 0)
 { }
 
-FSkeletalMeshSceneProxy::FSkeletalMeshSceneProxy(const FSkinnedMeshSceneProxyDesc& InMeshDesc, FSkeletalMeshRenderData* InSkelMeshRenderData, int32 MinLODIndex)
+FSkeletalMeshSceneProxy::FSkeletalMeshSceneProxy(const FSkinnedMeshSceneProxyDesc& InMeshDesc, FSkeletalMeshRenderData* InSkelMeshRenderData, int32 InClampedLODIndex)
 		:	FPrimitiveSceneProxy(InMeshDesc, InMeshDesc.GetSkinnedAsset()->GetFName())
 		,	Owner(Cast<AActor>(InMeshDesc.GetOwner()))
 		,	MeshObject(InMeshDesc.MeshObject)
 		,	SkeletalMeshRenderData(InSkelMeshRenderData)
-		,	ClampedLODIndex(MinLODIndex)
+		,	ClampedLODIndex(InClampedLODIndex)
 		,	SkeletalMeshForDebug(InMeshDesc.GetSkinnedAsset())
 		,	PhysicsAssetForDebug(InMeshDesc.GetPhysicsAsset())
 		,	OverlayMaterial(InMeshDesc.GetOverlayMaterial())
@@ -488,8 +488,7 @@ void FSkeletalMeshSceneProxy::DrawStaticElements(FStaticPrimitiveDrawInterface*
 
 		int32 NumLODs = SkeletalMeshRenderData->LODRenderData.Num();
 
-		// TODO: Use ClampedLODIndex but without meshes disappearing
-		for (int32 LODIndex = 0; LODIndex < NumLODs; ++LODIndex)
+		for (int32 LODIndex = ClampedLODIndex; LODIndex < NumLODs; ++LODIndex)
 		{
 			const FSkeletalMeshLODRenderData& LODData = SkeletalMeshRenderData->LODRenderData[LODIndex];
 
diff --git a/Engine/Source/Runtime/Engine/Private/SkinnedMeshSceneProxyDesc.cpp b/Engine/Source/Runtime/Engine/Private/SkinnedMeshSceneProxyDesc.cpp
--- a/Engine/Source/Runtime/Engine/Private/SkinnedMeshSceneProxyDesc.cpp
+++ b/Engine/Source/Runtime/Engine/Private/SkinnedMeshSceneProxyDesc.cpp
@@ -287,7 +287,8 @@ FPrimitiveSceneProxy* FSkinnedMeshSceneProxyDesc::CreateSceneProxy(const FSkinne
 			}
 			else
 			{
-				Result = ::new FSkeletalMeshSceneProxy(Desc, SkelMeshRenderData, MinLODIndex);
+				// TODO: MinLODIndex could work for InClampedLODIndex if it's the true minimum. Otherwise we could pass another parameter. Using 0 for now to ensure meshes render.
+				Result = ::new FSkeletalMeshSceneProxy(Desc, SkelMeshRenderData, 0);
 			}
 		}
 	}
diff --git a/Engine/Source/Runtime/Engine/Public/SkeletalMeshSceneProxy.h b/Engine/Source/Runtime/Engine/Public/SkeletalMeshSceneProxy.h
--- a/Engine/Source/Runtime/Engine/Public/SkeletalMeshSceneProxy.h
+++ b/Engine/Source/Runtime/Engine/Public/SkeletalMeshSceneProxy.h
@@ -26,9 +26,9 @@ public:
 	 * @param	Component - skeletal mesh primitive being added
 	 */
 	ENGINE_API FSkeletalMeshSceneProxy(const USkinnedMeshComponent* Component, FSkeletalMeshRenderData* InSkelMeshRenderData);
-	UE_DEPRECATED(5.7, "Pass MinLODIndex to FSkeletalMeshSceneProxy constructor.")
+	UE_DEPRECATED(5.7, "Pass InClampedLODIndex to FSkeletalMeshSceneProxy constructor.")
 	ENGINE_API FSkeletalMeshSceneProxy(const FSkinnedMeshSceneProxyDesc& InMeshDesc, FSkeletalMeshRenderData* InSkelMeshRenderData);
-	ENGINE_API FSkeletalMeshSceneProxy(const FSkinnedMeshSceneProxyDesc& InMeshDesc, FSkeletalMeshRenderData* InSkelMeshRenderData, int32 MinLODIndex);
+	ENGINE_API FSkeletalMeshSceneProxy(const FSkinnedMeshSceneProxyDesc& InMeshDesc, FSkeletalMeshRenderData* InSkelMeshRenderData, int32 InClampedLODIndex);
 
 	ENGINE_API virtual ~FSkeletalMeshSceneProxy();
 
