New Material Translator:  Support for Scene Texture, User Scene Texture, Scene Color, Scene Depth, Scene Depth Without Water, DBuffer Texture, and Scene Texel Size expressions.  Added "FScreenTexture" value type to handle validation, and adding defines and compilation output settings in analyzer.  Fixed issue where WorldPosition needs to mark PPI_SceneDepth as used when in a Post Process Material.

#rnx

diff --git a/Engine/Source/Runtime/Engine/Private/Materials/HLSLMaterialTranslator.cpp b/Engine/Source/Runtime/Engine/Private/Materials/HLSLMaterialTranslator.cpp
--- a/Engine/Source/Runtime/Engine/Private/Materials/HLSLMaterialTranslator.cpp
+++ b/Engine/Source/Runtime/Engine/Private/Materials/HLSLMaterialTranslator.cpp
@@ -204,10 +204,6 @@ UE::DerivedData::FValueId MaterialCompilationOutputId = UE::DerivedData::FValueI
 UE::DerivedData::FValueId MaterialResultsOutputId = UE::DerivedData::FValueId::FromName("FHLSLMaterialTranslator_Results");
 UE::DerivedData::FValueId EnvironmentDefinesId = UE::DerivedData::FValueId::FromName("FHLSLMaterialTranslator_EnvironmentDefines");
 
-/** HLSL generating utility functions */
-extern FString GenerateUserSceneTextureRemapHLSLDefines(FMaterialCompilationOutput& CompilationOutput);
-extern FString SceneTextureIdToHLSLString(ESceneTextureId TexId);
-
 /** Data structure used to cache a part of material translation results. It contains all the generated
  *  defines that will be declared during the compilation of the generated material shader.
  *  Note: if you make a change to this structure, remember to bump FDevSystemGuids::MaterialTranslationDDCVersion.
@@ -8361,7 +8357,7 @@ int32 FHLSLMaterialTranslator::SceneTextureLookup(int32 ViewportUV, uint32 InSce
 		return Constant4(0.0f, 0.0f, 0.0f, 0.0f);
 	}
 
-	FString SceneTextureIdString = SceneTextureIdToHLSLString(SceneTextureId);
+	FString SceneTextureIdString = UE::MaterialTranslatorUtils::SceneTextureIdToHLSLString(SceneTextureId);
 
 	int32 BufferUV;
 	if (ViewportUV != INDEX_NONE)
@@ -8421,7 +8417,7 @@ int32 FHLSLMaterialTranslator::SceneTextureLookup(int32 ViewportUV, uint32 InSce
 
 int32 FHLSLMaterialTranslator::GetSceneTextureViewSize(int32 SceneTextureId, bool InvProperty)
 {
-	FString SceneTextureIdString = SceneTextureIdToHLSLString((ESceneTextureId)SceneTextureId);
+	FString SceneTextureIdString = UE::MaterialTranslatorUtils::SceneTextureIdToHLSLString((ESceneTextureId)SceneTextureId);
 
 	if (InvProperty)
 	{
@@ -16956,7 +16952,7 @@ void FHLSLMaterialTranslator::PrepareMaterialSourceStringParameters()
 		MaterialSourceTemplateParams.Add({ TEXT("calc_pixel_material_inputs_analytic_derivatives_other_inputs"), PixelMembersSetupAndAssignments[CompiledPDV_Analytic] });
 	}
 
-	MaterialSourceTemplateParams.Add({ TEXT("user_scene_texture_remap"), GenerateUserSceneTextureRemapHLSLDefines(MaterialCompilationOutput) });
+	MaterialSourceTemplateParams.Add({ TEXT("user_scene_texture_remap"), UE::MaterialTranslatorUtils::GenerateUserSceneTextureRemapHLSLDefines(MaterialCompilationOutput) });
 }
 
 int32 FHLSLMaterialTranslator::SparseVolumeTexture(USparseVolumeTexture* Texture, int32& TextureReferenceIndex, EMaterialSamplerType SamplerType)
diff --git a/Engine/Source/Runtime/Engine/Private/Materials/Material.cpp b/Engine/Source/Runtime/Engine/Private/Materials/Material.cpp
--- a/Engine/Source/Runtime/Engine/Private/Materials/Material.cpp
+++ b/Engine/Source/Runtime/Engine/Private/Materials/Material.cpp
@@ -7521,6 +7521,11 @@ bool UMaterial::IsCompatibleWithLumenCardSharing() const
 	return bCompatibleWithLumenCardSharing;
 }
 
+bool UMaterial::IsMobileSeparateTranslucencyEnabled() const
+{
+	return bEnableMobileSeparateTranslucency && !IsUIMaterial() && !IsDeferredDecal();
+}
+
 void UMaterial::SetShadingModel(EMaterialShadingModel NewModel)
 {
 	ensure(ShadingModel < MSM_NUM);
diff --git a/Engine/Source/Runtime/Engine/Private/Materials/MaterialExpressionsIR.cpp b/Engine/Source/Runtime/Engine/Private/Materials/MaterialExpressionsIR.cpp
--- a/Engine/Source/Runtime/Engine/Private/Materials/MaterialExpressionsIR.cpp
+++ b/Engine/Source/Runtime/Engine/Private/Materials/MaterialExpressionsIR.cpp
@@ -224,6 +224,7 @@
 #include "Materials/MaterialExpressionTruncate.h"
 #include "Materials/MaterialExpressionTruncateLWC.h"
 #include "Materials/MaterialExpressionTwoSidedSign.h"
+#include "Materials/MaterialExpressionUserSceneTexture.h"
 #include "Materials/MaterialExpressionVectorNoise.h"
 #include "Materials/MaterialExpressionVectorParameter.h"
 #include "Materials/MaterialExpressionVertexColor.h"
@@ -960,6 +961,98 @@ void UMaterialExpressionIf::Build(MIR::FEmitter& Em)
 	Em.Output(0, OutputValue);
 }
 
+// If DefaultOffset is not null, Coordinates are treated as an offset (or DefaultOffset if unset), rather than absolute coordinates.  Clamping is
+// automatically applied for custom or offset fetches -- the "bClamped" parameter only controls clamping for default texture coordinate fetches,
+// and is only needed when fetching from lower resolution User Scene Textures.  A zero constant can be passed in for "SceneTextureInput" for
+// cases where the default view rect should be used for UV calculations.
+static FValueRef SceneTextureExpressionTexCoords(MIR::FEmitter& Em, FValueRef SceneTextureInput, const FExpressionInput& Coordinates, const FVector2D* DefaultOffset = nullptr, bool bClamped = false)
+{
+	FValueRef TexCoords;
+	if (DefaultOffset)
+	{
+		TexCoords = Em.InputDefaultFloat2(&Coordinates, { (float)DefaultOffset->X, (float)DefaultOffset->Y });
+		TexCoords = Em.InlineHLSL(MIR::FType::MakeFloat(2), TEXTVIEW("CalcScreenUVFromOffsetFraction(GetScreenPosition(Parameters), $0)"), { TexCoords });
+	}
+	else
+	{
+		TexCoords = Em.TryInput(&Coordinates);
+		if (TexCoords)
+		{
+			// Convert raw TexCoords expression to Float2, then convert from viewport to scene texture space
+			TexCoords = Em.Cast(TexCoords, MIR::FType::MakeFloat(2));
+			TexCoords = Em.InlineHLSL(MIR::FType::MakeFloat(2), TEXTVIEW("ClampSceneTextureUV(ViewportUVToSceneTextureUV($0, $1), $1)"), { TexCoords, SceneTextureInput });
+		}
+		else
+		{
+			FStringView TexCoordsCode = bClamped ?
+				TEXTVIEW("ClampSceneTextureUV(GetDefaultSceneTextureUV(Parameters, $0), $0)") :
+				TEXTVIEW("GetDefaultSceneTextureUV(Parameters, $0)");
+
+			TexCoords = Em.InlineHLSL(MIR::FType::MakeFloat(2), TexCoordsCode, { SceneTextureInput });
+		}
+	}
+	return TexCoords;
+}
+
+static void SceneTextureExpressionBuild(MIR::FEmitter& Em, FValueRef SceneTextureInput, const FExpressionInput& Coordinates, const FVector2D* DefaultOffset, bool bClamped, bool bFiltered)
+{
+	FValueRef TexCoords = SceneTextureExpressionTexCoords(Em, SceneTextureInput, Coordinates, DefaultOffset, bClamped);
+
+	FStringView SceneTextureLookupCode = bFiltered ?
+		TEXTVIEW("SceneTextureLookup(Parameters, $0, $1, true)") :
+		TEXTVIEW("SceneTextureLookup(Parameters, $0, $1, false)");
+
+	Em.Output(0, Em.InlineHLSL(MIR::FType::MakeFloat(4), SceneTextureLookupCode, { TexCoords, SceneTextureInput }));
+	Em.Output(1, Em.InlineHLSL(MIR::FType::MakeFloat(2), TEXTVIEW("GetSceneTextureViewSize($0).xy"), { SceneTextureInput }));
+	Em.Output(2, Em.InlineHLSL(MIR::FType::MakeFloat(2), TEXTVIEW("GetSceneTextureViewSize($0).zw"), { SceneTextureInput }));
+}
+
+void UMaterialExpressionSceneTexture::Build(MIR::FEmitter& Em)
+{
+	const bool bClamped = false;
+	SceneTextureExpressionBuild(Em, Em.SceneTexture(SceneTextureId), Coordinates, nullptr, bClamped, bFiltered);
+}
+
+void UMaterialExpressionUserSceneTexture::Build(MIR::FEmitter& Em)
+{
+	SceneTextureExpressionBuild(Em, Em.UserSceneTexture(UserSceneTexture), Coordinates, nullptr, bClamped, bFiltered);
+}
+
+void UMaterialExpressionSceneColor::Build(MIR::FEmitter& Em)
+{
+	FValueRef TexCoords = SceneTextureExpressionTexCoords(Em, Em.ConstantInt(0), Input, this->InputMode == EMaterialSceneAttributeInputMode::OffsetFraction ? &ConstInput : nullptr);
+
+	// We need a dependency on ScreenTexture as a second argument, so the value analyzer can see it, even though it's technically not used in the code.
+	FValueRef ScreenTexture = Em.ScreenTexture(MIR::EScreenTexture::SceneColor);
+	Em.Output(0, Em.InlineHLSL(MIR::FType::MakeFloat(4), TEXTVIEW("DecodeSceneColorAndAlpharForMaterialNode($0)"), { TexCoords, ScreenTexture }));;
+}
+
+void UMaterialExpressionSceneDepth::Build(MIR::FEmitter& Em)
+{
+	FValueRef TexCoords = SceneTextureExpressionTexCoords(Em, Em.ConstantInt(0), Input, this->InputMode == EMaterialSceneAttributeInputMode::OffsetFraction ? &ConstInput : nullptr);
+
+	// We need a dependency on ScreenTexture as a second argument, so the value analyzer can see it, even though it's technically not used in the code.
+	FValueRef ScreenTexture = Em.ScreenTexture(MIR::EScreenTexture::SceneDepth);
+	Em.Output(0, Em.InlineHLSL(MIR::FType::MakeFloat(1), TEXTVIEW("CalcSceneDepth($0)"), { TexCoords, ScreenTexture }));;
+}
+
+void UMaterialExpressionSceneDepthWithoutWater::Build(MIR::FEmitter& Em)
+{
+	FValueRef TexCoords = SceneTextureExpressionTexCoords(Em, Em.ConstantInt(0), Input, this->InputMode == EMaterialSceneAttributeInputMode::OffsetFraction ? &ConstInput : nullptr);
+
+	// We need a dependency on ScreenTexture as a third argument, so the value analyzer can see it, even though it's technically not used in the code.
+	FValueRef ScreenTexture = Em.ScreenTexture(MIR::EScreenTexture::SceneDepthWithoutWater);
+	Em.Output(0, Em.InlineHLSL(MIR::FType::MakeFloat(1), TEXTVIEW("MaterialExpressionSceneDepthWithoutWater($0, $1)"), { TexCoords, Em.ConstantFloat(FallbackDepth), ScreenTexture }));;
+}
+
+void UMaterialExpressionDBufferTexture::Build(MIR::FEmitter& Em)
+{
+	FValueRef TexCoords = SceneTextureExpressionTexCoords(Em, Em.ConstantInt(0), Coordinates);
+
+	FValueRef ScreenTexture = Em.DBufferTexture(DBufferTextureId);
+	Em.Output(0, Em.InlineHLSL(MIR::FType::MakeFloat(4), TEXTVIEW("MaterialExpressionDBufferTextureLookup(Parameters, $0, $1)"), { TexCoords, ScreenTexture }));;
+}
+
 void UMaterialExpressionSphericalParticleOpacity::Build(MIR::FEmitter& Em)
 {
 	FValueRef DensityValue = Em.InputDefaultFloat(&Density, ConstantDensity);
@@ -1502,6 +1595,13 @@ void UMaterialExpressionViewSize::Build(MIR::FEmitter& Em)
 	Em.Output(0, BuildViewProperty(Em, MEVP_ViewSize));
 }
 
+void UMaterialExpressionSceneTexelSize::Build(MIR::FEmitter& Em)
+{
+	// To make sure any material that were correctly handling BufferUV != ViewportUV, we just lie to material
+	// to make it believe ViewSize == BufferSize, so they are still compatible with SceneTextureLookup().
+	Em.Output(0, BuildViewProperty(Em, MEVP_ViewSize, true));
+}
+
 void UMaterialExpressionCameraPositionWS::Build(MIR::FEmitter& Em)
 {
 	Em.Output(0, BuildViewProperty(Em, MEVP_WorldSpaceCameraPosition));
diff --git a/Engine/Source/Runtime/Engine/Private/Materials/MaterialIR.cpp b/Engine/Source/Runtime/Engine/Private/Materials/MaterialIR.cpp
--- a/Engine/Source/Runtime/Engine/Private/Materials/MaterialIR.cpp
+++ b/Engine/Source/Runtime/Engine/Private/Materials/MaterialIR.cpp
@@ -42,6 +42,7 @@ const TCHAR* LexToString(EValueKind Kind)
 		case VK_Poison: return TEXT("Poison");
 		case VK_Constant: return TEXT("Constant");
 		case VK_ExternalInput: return TEXT("ExternalInput");
+		case VK_ScreenTexture: return TEXT("ScreenTexture");
 		case VK_TextureObject: return TEXT("TextureObject");
 		case VK_UniformParameter: return TEXT("UniformParameter");
 		case VK_Composite: return TEXT("Composite");
@@ -101,6 +102,7 @@ uint32 FValue::GetSizeInBytes() const
 		case VK_Poison: return sizeof(FPoison);
 		case VK_Constant: return sizeof(FConstant);
 		case VK_ExternalInput: return sizeof(FExternalInput);
+		case VK_ScreenTexture: return sizeof(FScreenTexture);
 		case VK_TextureObject: return sizeof(FTextureObject);
 		case VK_UniformParameter: return sizeof(FUniformParameter);
 		case VK_Composite: return sizeof(FComposite) + sizeof(FValue*) * static_cast<const FComposite*>(this)->GetComponents().Num();
@@ -530,6 +532,14 @@ bool IsExternalInputTexCoordOrPartialDerivative(EExternalInput Id)
 	return IsExternalInputTexCoord(Id) || IsExternalInputTexCoordDdx(Id) || IsExternalInputTexCoordDdy(Id);
 }
 
+bool IsExternalInputWorldPosition(EExternalInput Id)
+{
+	return
+		(Id >= EExternalInput::WorldPosition_Absolute && Id <= EExternalInput::WorldPosition_CameraRelativeNoOffsets) ||
+		(Id >= EExternalInput::WorldPosition_Absolute_Ddx && Id <= EExternalInput::WorldPosition_CameraRelativeNoOffsets_Ddx) ||
+		(Id >= EExternalInput::WorldPosition_Absolute_Ddy && Id <= EExternalInput::WorldPosition_CameraRelativeNoOffsets_Ddy);
+}
+
 FBlock* FBlock::FindCommonParentWith(MIR::FBlock* Other)
 {
 	FBlock* A = this;
diff --git a/Engine/Source/Runtime/Engine/Private/Materials/MaterialIREmitter.cpp b/Engine/Source/Runtime/Engine/Private/Materials/MaterialIREmitter.cpp
--- a/Engine/Source/Runtime/Engine/Private/Materials/MaterialIREmitter.cpp
+++ b/Engine/Source/Runtime/Engine/Private/Materials/MaterialIREmitter.cpp
@@ -879,6 +879,43 @@ FValueRef FEmitter::CustomPrimitiveData(uint32 PrimitiveDataIndex)
 	return {};
 }
 
+// Emit a scene texture value with the given Id
+FValueRef FEmitter::SceneTexture(ESceneTextureId SceneTextureId)
+{
+	FScreenTexture Proto = MakePrototype<FScreenTexture>(MIR::FType::MakeInt(1));
+	Proto.TextureKind = EScreenTexture::SceneTexture;
+	Proto.Id = SceneTextureId;
+	return EmitPrototype(*this, Proto);
+}
+
+// Emit a user scene texture value with the given name
+FValueRef FEmitter::UserSceneTexture(FName UserSceneTexture)
+{
+	FScreenTexture Proto = MakePrototype<FScreenTexture>(MIR::FType::MakeInt(1));
+	Proto.TextureKind = EScreenTexture::UserSceneTexture;
+	Proto.Id = PPI_UserSceneTexture0;				// Actual ID determined later during HLSL generation
+	Proto.UserSceneTexture = UserSceneTexture;
+	return EmitPrototype(*this, Proto);
+}
+
+// Emit a DBuffer texture value with the given Id
+FValueRef FEmitter::DBufferTexture(EDBufferTextureId DBufferId)
+{
+	FScreenTexture Proto = MakePrototype<FScreenTexture>(MIR::FType::MakeInt(1));
+	Proto.TextureKind = EScreenTexture::DBufferTexture;
+	Proto.DBufferId = DBufferId;
+	return EmitPrototype(*this, Proto);
+}
+
+FValueRef FEmitter::ScreenTexture(EScreenTexture TextureKind)
+{
+	check(TextureKind != EScreenTexture::SceneTexture && TextureKind != EScreenTexture::UserSceneTexture && TextureKind != EScreenTexture::DBufferTexture);
+	FScreenTexture Proto = MakePrototype<FScreenTexture>(MIR::FType::MakeInt(1));
+	Proto.TextureKind = TextureKind;
+	Proto.Id = PPI_SceneColor;					// Ignored
+	return EmitPrototype(*this, Proto);
+}
+
 /*------------------------------ Instruction emission ------------------------------*/
 
 FSetMaterialOutput* FEmitter::SetMaterialOutput(EMaterialProperty InProperty, FValue* Arg)
diff --git a/Engine/Source/Runtime/Engine/Private/Materials/MaterialIRModuleBuilder.cpp b/Engine/Source/Runtime/Engine/Private/Materials/MaterialIRModuleBuilder.cpp
--- a/Engine/Source/Runtime/Engine/Private/Materials/MaterialIRModuleBuilder.cpp
+++ b/Engine/Source/Runtime/Engine/Private/Materials/MaterialIRModuleBuilder.cpp
@@ -24,6 +24,7 @@
 #include "Materials/MaterialInsights.h"
 #include "Engine/Texture.h"
 #include "Misc/FileHelper.h"
+#include "PostProcess/PostProcessMaterialInputs.h"
 #include "RenderUtils.h"
 
 static TAutoConsoleVariable<int32> CVarMaterialIRDebugDumpLevel(
@@ -91,6 +92,7 @@ struct FMaterialIRModuleBuilderImpl
 	{
 		Module->Empty();
 		Module->ShaderPlatform = Builder->ShaderPlatform;
+		Module->FeatureLevel = Builder->FeatureLevel;
 		
 		Emitter->Initialize();
 		AnalysisContextStack.Emplace();
@@ -650,6 +652,12 @@ struct FMaterialIRModuleBuilderImpl
 	{
 		FMaterialCompilationOutput& CompilationOutput = Module->CompilationOutput;
 		CompilationOutput.NumUsedUVScalars = Module->Statistics.NumPixelTexCoords * 2;
+
+		int32 NumPostProcessInputs = CompilationOutput.GetNumPostProcessInputsUsed();
+		if (NumPostProcessInputs > kPostProcessMaterialInputCountMax)
+		{
+			Module->AddError(nullptr, FString::Printf(TEXT("Maximum Scene Texture post process inputs exceeded (%d > %d), between SceneTexture nodes with PostProcessInputs or UserSceneTexture nodes."), NumPostProcessInputs, kPostProcessMaterialInputCountMax));
+		}
 	}
 	
 	/* Other functions */
diff --git a/Engine/Source/Runtime/Engine/Private/Materials/MaterialIRToHLSLTranslator.cpp b/Engine/Source/Runtime/Engine/Private/Materials/MaterialIRToHLSLTranslator.cpp
--- a/Engine/Source/Runtime/Engine/Private/Materials/MaterialIRToHLSLTranslator.cpp
+++ b/Engine/Source/Runtime/Engine/Private/Materials/MaterialIRToHLSLTranslator.cpp
@@ -425,6 +425,7 @@ struct FPrivate : FMaterialIRToHLSLTranslation
 		{
 			case MIR::VK_Constant: LowerConstant(static_cast<const MIR::FConstant*>(InValue)); break;
 			case MIR::VK_ExternalInput: LowerExternalInput(static_cast<const MIR::FExternalInput*>(InValue)); break;
+			case MIR::VK_ScreenTexture: LowerScreenTexture(static_cast<const MIR::FScreenTexture*>(InValue)); break;
 			case MIR::VK_TextureObject: LowerTextureObject(static_cast<const MIR::FTextureObject*>(InValue)); break;
 			case MIR::VK_UniformParameter: LowerUniformParameter(static_cast<const MIR::FUniformParameter*>(InValue)); break;
 			default: UE_MIR_UNREACHABLE();
@@ -698,6 +699,23 @@ struct FPrivate : FMaterialIRToHLSLTranslation
 		}
 	}
 
+	void LowerScreenTexture(const MIR::FScreenTexture* ScreenTexture)
+	{
+		// Types besides the following aren't directly printed in HLSL, and don't need to do anything here
+		switch (ScreenTexture->TextureKind)
+		{
+		case MIR::EScreenTexture::SceneTexture:
+			Printer << UE::MaterialTranslatorUtils::SceneTextureIdToHLSLString(ScreenTexture->Id);
+			break;
+		case MIR::EScreenTexture::UserSceneTexture:
+			Printer << UE::MaterialTranslatorUtils::SceneTextureIdToHLSLString((ESceneTextureId)Module->GetCompilationOutput().FindUserSceneTexture(ScreenTexture->UserSceneTexture));
+			break;
+		case MIR::EScreenTexture::DBufferTexture:
+			Printer << (int32)ScreenTexture->DBufferId;
+			break;
+		}
+	}
+
 	void LowerTextureObject(const MIR::FTextureObject* TextureObject)
 	{
 		LowerTextureReference(TextureObject->Texture->GetMaterialType(), TextureObject->Analysis_UniformParameterIndex);
@@ -1848,6 +1866,8 @@ struct FPrivate : FMaterialIRToHLSLTranslation
 			);
 		}
 		Params.Add(TEXT("uniform_material_expressions"), MoveTemp(UniformMaterialExpressions));
+
+		Params.Add(TEXT("user_scene_texture_remap"), UE::MaterialTranslatorUtils::GenerateUserSceneTextureRemapHLSLDefines(Module->GetCompilationOutput()));
 	}
 	
 	ENoOp LowerType(const MIR::FType& Type)
diff --git a/Engine/Source/Runtime/Engine/Private/Materials/MaterialIRValueAnalyzer.cpp b/Engine/Source/Runtime/Engine/Private/Materials/MaterialIRValueAnalyzer.cpp
--- a/Engine/Source/Runtime/Engine/Private/Materials/MaterialIRValueAnalyzer.cpp
+++ b/Engine/Source/Runtime/Engine/Private/Materials/MaterialIRValueAnalyzer.cpp
@@ -5,6 +5,7 @@
 #if WITH_EDITOR
 
 #include "Materials/MaterialExternalCodeRegistry.h"
+#include "Materials/MaterialExpressionDBufferTexture.h"
 #include "Materials/MaterialExpressionUtils.h"
 #include "Materials/MaterialIRModule.h"
 #include "Materials/MaterialIRDebug.h"
@@ -14,9 +15,23 @@
 #include "Materials/MaterialInsights.h"
 #include "Materials/Material.h"
 #include "Materials/MaterialEnumeration.h"
+#include "MaterialCachedData.h"
+#include "MaterialDomain.h"
 #include "MaterialShared.h"
 #include "Engine/Texture.h"
 #include "UObject/Package.h"
+#include "PostProcess/PostProcessMaterialInputs.h"
+
+static void ErrorUnlessFeatureLevelSupported(FMaterialIRModule* Module, ERHIFeatureLevel::Type RequiredFeatureLevel, const TCHAR* Message)
+{
+	ERHIFeatureLevel::Type FeatureLevel = Module->GetFeatureLevel();
+	if (FeatureLevel < RequiredFeatureLevel)
+	{
+		FString FeatureLevelName;
+		GetFeatureLevelName(FeatureLevel, FeatureLevelName);
+		Module->AddError(nullptr, FString::Printf(TEXT("%s  Current feature level is %s."), Message, *FeatureLevelName));
+	}
+}
 
 static void AnalyzeExternalInput(FMaterialIRValueAnalyzer& Analyzer, MIR::FExternalInput* ExternalInput)
 {
@@ -30,6 +45,11 @@ static void AnalyzeExternalInput(FMaterialIRValueAnalyzer& Analyzer, MIR::FExter
 		Statistics.NumVertexTexCoords = FMath::Max(Statistics.NumVertexTexCoords, TexCoordIndex + 1);
 		Statistics.NumPixelTexCoords = FMath::Max(Statistics.NumPixelTexCoords, TexCoordIndex + 1);
 	}
+
+	if (Analyzer.Material->IsPostProcessMaterial() && MIR::IsExternalInputWorldPosition(ExternalInput->Id))
+	{
+		Analyzer.CompilationOutput->SetIsSceneTextureUsed(PPI_SceneDepth);
+	}
 }
 
 static EMaterialShaderFrequency MapToMaterialShaderFrequencyOrAny(MIR::EStage Stage)
@@ -144,6 +164,282 @@ static void AnalyzeInStageVTPageTableRead(FMaterialIRValueAnalyzer& Analyzer, MI
 	}
 }
 
+static void AnalyzeScreenTexture(FMaterialIRValueAnalyzer& Analyzer, MIR::FScreenTexture* ScreenTexture)
+{
+	const EMaterialDomain MaterialDomain = Analyzer.Material->MaterialDomain;
+
+	switch (ScreenTexture->TextureKind)
+	{
+		case MIR::EScreenTexture::SceneTexture:
+		case MIR::EScreenTexture::UserSceneTexture:
+			{
+				// TODO:  If referenced from Custom HLSL, this can be false.  Revisit when custom HLSL support is added to new translator.
+				const bool bTextureLookup = true;
+
+				// Add defines and compilation outputs
+				Analyzer.EnvironmentDefines.Add(FName(TEXT("NEEDS_SCENE_TEXTURES")));
+
+				Analyzer.CompilationOutput->bNeedsSceneTextures = true;
+				Analyzer.CompilationOutput->SetIsSceneTextureUsed(ScreenTexture->Id);
+
+				if (ScreenTexture->TextureKind == MIR::EScreenTexture::UserSceneTexture)
+				{
+					if (ScreenTexture->UserSceneTexture.IsNone())
+					{
+						Analyzer.Module->AddError(nullptr, TEXT("UserSceneTexture missing name -- value must be set to something other than None"));
+					}
+					else
+					{
+						// Allocate value during analyze, and check for failure.  FindUserSceneTexture is later called during HLSL generation to retrieve the allocated Id.
+						int32 SceneTextureId = Analyzer.CompilationOutput->FindOrAddUserSceneTexture(ScreenTexture->UserSceneTexture);
+						if (SceneTextureId == INDEX_NONE)
+						{
+							Analyzer.Module->AddError(nullptr, FString::Printf(TEXT("Too many unique UserSceneTexture inputs in the post process material -- max allowed is %d"), kPostProcessMaterialInputCountMax));
+						}
+					}
+				}
+
+				// Substrate TODO:
+				// Substrate only
+				// When SceneTexture lookup node is used, single/simple paths are disabled to ensure texture decoding is properly handled.
+				// Reading SceneTexture, when Substrate is enabled, implies unpacking material buffer data. The unpacking function exists in different 'flavor' 
+				// for optimization purpose (simple/single/complex). To avoid compiling out single or complex unpacking paths (due to defines set by analyzing 
+				// the current shader, vs. scene texture pixels), we force Simple/Single versions to be disabled
+				// FSubstrateCompilationContext& SubstrateCtx = SubstrateCompilationContext[CurrentSubstrateCompilationContext];
+				// SubstrateCtx.SubstrateMaterialComplexity.bIsSimple = false;
+				// SubstrateCtx.SubstrateMaterialComplexity.bIsSingle = false;
+
+				// Error checking
+				EShaderPlatform Platform = Analyzer.Module->GetShaderPlatform();
+
+				// Error checking from FHLSLMaterialTranslator::SceneTextureLookup
+				// Guard against using unsupported textures with SLW
+				const bool bHasSingleLayerWaterSM = Analyzer.Material->GetShadingModels().HasShadingModel(MSM_SingleLayerWater);
+				if (bHasSingleLayerWaterSM && ScreenTexture->Id != PPI_CustomDepth && ScreenTexture->Id != PPI_CustomStencil)
+				{
+					Analyzer.Module->AddError(nullptr, TEXT("Only custom depth and custom stencil can be sampled with SceneTexture when used with the Single Layer Water shading model."));
+				}
+
+				if (ScreenTexture->Id == PPI_DecalMask)
+				{
+					Analyzer.Module->AddError(nullptr, TEXT("Decal Mask bit was moved from GBuffer to the Stencil Buffer for performance optimisation so therefore no longer available."));
+				}
+
+				// Error checking from FHLSLMaterialTranslator::UseSceneTextureId
+				if (MaterialDomain == MD_DeferredDecal)
+				{
+					const bool bSceneTextureSupportsDecal = ScreenTexture->Id == PPI_SceneDepth || ScreenTexture->Id == PPI_WorldNormal || ScreenTexture->Id == PPI_CustomDepth || ScreenTexture->Id == PPI_CustomStencil;
+					if (!bSceneTextureSupportsDecal)
+					{
+						// Note: For DBuffer decals CustomDepth and CustomStencil are not available if r.CustomDepth.Order = 1
+						Analyzer.Module->AddError(nullptr, TEXT("Decals can only access SceneDepth, CustomDepth, CustomStencil, and WorldNormal."));
+					}
+
+					const bool bSceneTextureRequiresSM5 = ScreenTexture->Id == PPI_WorldNormal;
+					if (bSceneTextureRequiresSM5)
+					{
+						ErrorUnlessFeatureLevelSupported(Analyzer.Module, ERHIFeatureLevel::SM5, TEXT("Deferred decals require SM5 for World Normal access."));
+					}
+
+					if (ScreenTexture->Id == PPI_WorldNormal && !IsUsingDBuffers(Platform))
+					{
+						bool bHasNormalConnected =
+							Substrate::IsSubstrateEnabled() ?
+							FSubstrateMaterialInfo::HasPropertyConnected(Analyzer.Material->GetCachedExpressionData().PropertyConnectedMask, MP_Normal) :
+							Analyzer.Material->HasNormalConnected();
+
+						if (bHasNormalConnected)
+						{
+							// GBuffer decals can't bind Normal for read and write.
+							// Note: DBuffer decals can support this but only if the sampled WorldNormal isn't connected to the output normal.
+							Analyzer.Module->AddError(nullptr, TEXT("Decals that read WorldNormal cannot output to normal at the same time. Enable DBuffer to support this."));
+						}
+					}
+				}
+
+				if (ScreenTexture->Id == PPI_SceneColor && MaterialDomain != MD_Surface)
+				{
+					if (MaterialDomain == MD_PostProcess)
+					{
+						Analyzer.Module->AddError(nullptr, TEXT("SceneColor lookups are only available when MaterialDomain = Surface. PostProcessMaterials should use the SceneTexture PostProcessInput0."));
+					}
+					else
+					{
+						Analyzer.Module->AddError(nullptr, TEXT("SceneColor lookups are only available when MaterialDomain = Surface."));
+					}
+				}
+
+				bool bNeedsSceneTexturePostProcessInputs = false;
+				if (bTextureLookup)
+				{
+					bNeedsSceneTexturePostProcessInputs =
+						((ScreenTexture->Id >= PPI_PostProcessInput0 && ScreenTexture->Id <= PPI_PostProcessInput6) ||
+						 (ScreenTexture->Id >= PPI_UserSceneTexture0 && ScreenTexture->Id <= PPI_UserSceneTexture6) ||
+						 ScreenTexture->Id == PPI_Velocity ||
+						 ScreenTexture->Id == PPI_SceneColor);
+				}
+
+				if ((ScreenTexture->TextureKind == MIR::EScreenTexture::SceneTexture) && ((1u << ScreenTexture->Id) & FMaterialCompilationOutput::GetGBufferMask()))
+				{
+					if (IsForwardShadingEnabled(Platform) || (IsMobilePlatform(Platform) && !IsMobileDeferredShadingEnabled(Platform)))
+					{
+						Analyzer.Module->AddError(nullptr, FString::Printf(TEXT("GBuffer scene textures not available with forward shading (platform id %d)."), Platform));
+					}
+
+					// Post-process can't access memoryless GBuffer on mobile
+					if (IsMobilePlatform(Platform))
+					{
+						if (MaterialDomain == MD_PostProcess)
+						{
+							Analyzer.Module->AddError(nullptr, FString::Printf(TEXT("GBuffer scene textures not available in post-processing with mobile shading (platform id %d)."), Platform));
+						}
+
+						if (Analyzer.Material->IsMobileSeparateTranslucencyEnabled())
+						{
+							Analyzer.Module->AddError(nullptr, FString::Printf(TEXT("GBuffer scene textures not available for separate translucency with mobile shading (platform id %d)."), Platform));
+						}
+					}
+				}
+
+				if (ScreenTexture->Id == PPI_Velocity)
+				{
+					if (MaterialDomain != MD_PostProcess)
+					{
+						Analyzer.Module->AddError(nullptr, TEXT("Velocity scene textures are only available in post process materials."));
+					}
+				}
+
+				// Error checking from FHLSLMaterialTranslator::TranslateMaterial
+				if (MaterialDomain != MD_DeferredDecal && MaterialDomain != MD_PostProcess)
+				{
+					if (!Analyzer.Material->GetShadingModels().HasShadingModel(MSM_SingleLayerWater) && IsOpaqueOrMaskedBlendMode(Analyzer.Material->BlendMode))
+					{
+						// In opaque pass, none of the textures are available
+						Analyzer.Module->AddError(nullptr, TEXT("SceneTexture expressions cannot be used in opaque materials except if used with the Single Layer Water shading model."));
+					}
+					else if (bNeedsSceneTexturePostProcessInputs)
+					{
+						Analyzer.Module->AddError(nullptr, TEXT("SceneTexture expressions cannot use post process inputs or scene color in non post process domain materials"));
+					}
+				}
+
+				if (ScreenTexture->Id == PPI_SceneDepth && bTextureLookup)
+				{
+					// Don't allow opaque and masked materials to access scene depth as the results are undefined
+					if (MaterialDomain != MD_PostProcess && !IsTranslucentBlendMode(Analyzer.Material->BlendMode))
+					{
+						Analyzer.Module->AddError(nullptr, TEXT("Only transparent or postprocess materials can read from scene depth."));
+					}
+				}
+			}
+			break;
+		case MIR::EScreenTexture::SceneColor:
+			{
+				// Add defines and compilation outputs
+				Analyzer.CompilationOutput->SetIsSceneTextureUsed(PPI_SceneColor);
+
+				// Error checking
+				if (MaterialDomain != MD_Surface)
+				{
+					Analyzer.Module->AddError(nullptr, TEXT("SceneColor lookups are only available when MaterialDomain = Surface."));
+				}
+
+				ErrorUnlessFeatureLevelSupported(Analyzer.Module, ERHIFeatureLevel::SM5, TEXT("Scene Color access require SM5."));
+			}
+			break;
+		case MIR::EScreenTexture::SceneDepth:
+			{
+				// Add defines and compilation outputs
+				Analyzer.CompilationOutput->SetIsSceneTextureUsed(PPI_SceneDepth);
+
+				// Error checking
+				if (Analyzer.Material->IsTranslucencyWritingVelocity())
+				{
+					Analyzer.Module->AddError(nullptr, TEXT("Translucenct material with 'Output Velocity' enabled will write to depth buffer, therefore cannot read from depth buffer at the same time."));
+				}
+
+				// Don't allow opaque and masked materials to access scene depth as the results are undefined
+				if (MaterialDomain != MD_PostProcess && !IsTranslucentBlendMode(Analyzer.Material->BlendMode))
+				{
+					Analyzer.Module->AddError(nullptr, TEXT("Only transparent or postprocess materials can read from scene depth."));
+				}
+			}
+			break;
+		case MIR::EScreenTexture::SceneDepthWithoutWater:
+			{
+				// Error checking (no defines or compilation outputs needed for SceneDepthWithoutWater)
+				if (MaterialDomain != MD_PostProcess)
+				{
+					if (!Analyzer.Material->GetShadingModels().HasShadingModel(MSM_SingleLayerWater))
+					{
+						Analyzer.Module->AddError(nullptr, TEXT("Can only read scene depth below water when material Shading Model is Single Layer Water or when material Domain is PostProcess."));
+					}
+
+					if (MaterialDomain != MD_Surface)
+					{
+						Analyzer.Module->AddError(nullptr, TEXT("Can only read scene depth below water when material Domain is set to Surface or PostProcess."));
+					}
+
+					if (IsTranslucentBlendMode(Analyzer.Material->GetBlendMode()))
+					{
+						Analyzer.Module->AddError(nullptr, TEXT("Can only read scene depth below water when material Blend Mode isn't translucent."));
+					}
+				}
+			}
+			break;
+		case MIR::EScreenTexture::DBufferTexture:
+			{
+				// Add defines and compilation outputs
+				Analyzer.EnvironmentDefines.Add(FName(TEXT("MATERIAL_USES_DECAL_LOOKUP")));
+
+				Analyzer.CompilationOutput->SetIsDBufferTextureUsed(ScreenTexture->DBufferId);
+				// set separate flag to indicate that material uses DBuffer lookup specifically
+				// can't rely on UsedDBufferTextures because those bits are also set depending on the default decal response behavior.
+				Analyzer.CompilationOutput->SetIsDBufferTextureLookupUsed(true);
+
+				// Error checking
+				if (MaterialDomain != MD_Surface || IsTranslucentBlendMode(Analyzer.Material->GetBlendMode()))
+				{
+					Analyzer.Module->AddError(nullptr, TEXT("DBuffer scene textures are only available on opaque or masked surfaces."));
+				}
+			}
+			break;
+	}
+}
+
+static void AnalyzeInStageScreenTexture(FMaterialIRValueAnalyzer& Analyzer, MIR::FScreenTexture* ScreenTexture, MIR::EStage Stage)
+{
+	// This function handles stage specific error checking
+	switch (ScreenTexture->TextureKind)
+	{
+		case MIR::EScreenTexture::SceneTexture:
+		case MIR::EScreenTexture::UserSceneTexture:
+			break;
+		case MIR::EScreenTexture::SceneColor:
+			if (Stage == MIR::Stage_Vertex)
+			{
+				Analyzer.Module->AddError(nullptr, TEXT("Scene Color is only supported in pixel shader input!"));
+			}
+			break;
+		case MIR::EScreenTexture::SceneDepth:
+			if (Stage == MIR::Stage_Vertex)
+			{
+				ErrorUnlessFeatureLevelSupported(Analyzer.Module, ERHIFeatureLevel::SM5, TEXT("Reading scene depth from the vertex shader requires SM5."));
+			}
+			break;
+		case MIR::EScreenTexture::SceneDepthWithoutWater:
+			if (Stage == MIR::Stage_Vertex)
+			{
+				// Mobile currently does not support this, we need to read a separate copy of the depth, we must disable framebuffer fetch and force scene texture reads.
+				// (Texture bindings are not setup properly for any platform so we're disallowing usage in vertex shader altogether now)
+				Analyzer.Module->AddError(nullptr, TEXT("Cannot read scene depth without water from the vertex shader."));
+			}
+			break;
+		case MIR::EScreenTexture::DBufferTexture:
+			break;
+	}
+}
+
 static void AnalyzeTextureObject(FMaterialIRValueAnalyzer& Analyzer, MIR::FTextureObject* TextureObject)
 {
 	EMaterialTextureParameterType ParamType = MIR::Internal::TextureMaterialValueTypeToParameterType(TextureObject->Texture->GetMaterialType());
@@ -348,6 +644,7 @@ void FMaterialIRValueAnalyzer::Analyze(MIR::FValue* Value)
 		EXPAND_CASE(UniformParameter);
 		EXPAND_CASE(SetMaterialOutput);
 		EXPAND_CASE(InlineHLSL);
+		EXPAND_CASE(ScreenTexture);
 		default: break;
 	}
 
@@ -368,6 +665,7 @@ void FMaterialIRValueAnalyzer::AnalyzeInStage(MIR::FValue* Value, MIR::EStage St
 	{
 		EXPAND_CASE(InlineHLSL);
 		EXPAND_CASE(VTPageTableRead);
+		EXPAND_CASE(ScreenTexture);
 		default: break;
 	}
 
diff --git a/Engine/Source/Runtime/Engine/Private/Materials/MaterialShared.cpp b/Engine/Source/Runtime/Engine/Private/Materials/MaterialShared.cpp
--- a/Engine/Source/Runtime/Engine/Private/Materials/MaterialShared.cpp
+++ b/Engine/Source/Runtime/Engine/Private/Materials/MaterialShared.cpp
@@ -1827,7 +1827,7 @@ bool FMaterialResource::IsDualBlendingEnabled(EShaderPlatform Platform) const
 
 bool FMaterialResource::IsMobileSeparateTranslucencyEnabled() const
 {
-	return Material->bEnableMobileSeparateTranslucency && !IsUIMaterial() && !IsDeferredDecal();
+	return Material->IsMobileSeparateTranslucencyEnabled();
 }
 
 bool FMaterialResource::IsFullyRough() const
@@ -3353,6 +3353,7 @@ bool FMaterial::Translate_New(const FMaterialShaderMapId& InShaderMapId,
 		.Material = GetMaterialInterface()->GetMaterial(),
 		.ShaderPlatform = InShaderPlatform,
 		.TargetPlatform = InTargetPlatform,
+		.FeatureLevel = InShaderMapId.FeatureLevel,
 		.StaticParameters = InStaticParameters,
 		.TargetInsights = GetMaterialInterface()->MaterialInsight.Get(),
 		.PreviewExpression = GetMaterialGraphNodePreviewExpression()
@@ -5756,6 +5757,17 @@ int32 FMaterialCompilationOutput::FindOrAddUserSceneTexture(FName UserSceneTextu
 	return INDEX_NONE;
 }
 
+int32 FMaterialCompilationOutput::FindUserSceneTexture(FName UserSceneTexture) const
+{
+	FScriptName UserSceneTextureScript(UserSceneTexture);
+	int32 FoundIndex = UserSceneTextureInputs.Find(UserSceneTextureScript);
+	if (FoundIndex != INDEX_NONE)
+	{
+		return FoundIndex + PPI_UserSceneTexture0;
+	}
+	return INDEX_NONE;
+}
+
 int32 FMaterialCompilationOutput::GetNumPostProcessInputsUsed() const
 {
 	// Check how many post process inputs are used in the mask (SceneTexture nodes explicitly referencing PostProcessInput0-6)
@@ -5772,9 +5784,11 @@ int32 FMaterialCompilationOutput::GetNumPostProcessInputsUsed() const
 	return NumPostProcessInputs + UserSceneTextureInputs.Num();
 }
 
+namespace UE::MaterialTranslatorUtils
+{
 // Utility to generate HLSL defines to remap PPI_UserSceneTexture0-6 tokens generated by the HLSL emitter/translator to PPI_PostProcessInput0-6,
 // using unoccupied slots (UsedSceneTextures bits not set) in the order they are encountered.
-FString GenerateUserSceneTextureRemapHLSLDefines(FMaterialCompilationOutput& CompilationOutput)
+FString GenerateUserSceneTextureRemapHLSLDefines(const FMaterialCompilationOutput& CompilationOutput)
 {
 	if (CompilationOutput.UserSceneTextureInputs.IsEmpty())
 	{
@@ -5833,6 +5847,8 @@ FString SceneTextureIdToHLSLString(ESceneTextureId TexId)
 	}
 
 }
+}  // namespace UE::MaterialTranslatorUtils
+
 #endif // WITH_EDITOR
 
 #undef LOCTEXT_NAMESPACE
\ No newline at end of file
diff --git a/Engine/Source/Runtime/Engine/Private/Materials/MaterialSharedPrivate.h b/Engine/Source/Runtime/Engine/Private/Materials/MaterialSharedPrivate.h
--- a/Engine/Source/Runtime/Engine/Private/Materials/MaterialSharedPrivate.h
+++ b/Engine/Source/Runtime/Engine/Private/Materials/MaterialSharedPrivate.h
@@ -59,6 +59,9 @@ namespace UE::MaterialTranslatorUtils
 	bool IsLWCEnabled();
 	int32 GetLWCTruncateMode();
 
+	FString GenerateUserSceneTextureRemapHLSLDefines(const FMaterialCompilationOutput& CompilationOutput);
+	FString SceneTextureIdToHLSLString(ESceneTextureId TexId);
+
 } // /namespace UE::MaterialTranslatorUtils
 
 ///////////////////////////////////////////////////////
diff --git a/Engine/Source/Runtime/Engine/Public/MaterialShared.h b/Engine/Source/Runtime/Engine/Public/MaterialShared.h
--- a/Engine/Source/Runtime/Engine/Public/MaterialShared.h
+++ b/Engine/Source/Runtime/Engine/Public/MaterialShared.h
@@ -818,6 +818,9 @@ public:
 	// Returns Id for a user scene texture input, or INDEX_NONE if allocation fails.
 	int32 FindOrAddUserSceneTexture(FName UserSceneTexture);
 
+	// Returns Id for a user scene texture input, or INDEX_NONE if not found.
+	int32 FindUserSceneTexture(FName UserSceneTexture) const;
+
 	// Returns number of post process inputs used, between PostProcessInput SceneTexture nodes and UserSceneTexture nodes.  If this exceeds
 	// kPostProcessMaterialInputCountMax, compile should be failed (number is returned rather than bool for error reporting purposes).
 	int32 GetNumPostProcessInputsUsed() const;
@@ -831,29 +834,35 @@ public:
 	/** Indicates whether the material uses scene color. */
 	bool RequiresSceneColorCopy() const { return IsSceneTextureUsed(PPI_SceneColor); }
 
+	/** Returns mask for all scene textures that are from the GBuffer */
+	static uint32 GetGBufferMask()
+	{
+		return
+			(1u << PPI_DiffuseColor) |
+			(1u << PPI_SpecularColor) |
+			(1u << PPI_SubsurfaceColor) |
+			(1u << PPI_BaseColor) |
+			(1u << PPI_Specular) |
+			(1u << PPI_Metallic) |
+			(1u << PPI_WorldNormal) |
+			(1u << PPI_WorldTangent) |
+			(1u << PPI_Opacity) |
+			(1u << PPI_Roughness) |
+			(1u << PPI_Anisotropy) |
+			(1u << PPI_MaterialAO) |
+			(1u << PPI_DecalMask) |
+			(1u << PPI_ShadingModelColor) |
+			(1u << PPI_ShadingModelID) |
+			(1u << PPI_StoredBaseColor) |
+			(1u << PPI_StoredSpecular) |
+			(1u << PPI_Velocity) |
+			(1u << PPI_IsFirstPerson);
+	}
+
 	/** true if the material uses any GBuffer textures */
 	bool NeedsGBuffer() const
 	{
-		return
-			IsSceneTextureUsed(PPI_DiffuseColor) ||
-			IsSceneTextureUsed(PPI_SpecularColor) ||
-			IsSceneTextureUsed(PPI_SubsurfaceColor) ||
-			IsSceneTextureUsed(PPI_BaseColor) ||
-			IsSceneTextureUsed(PPI_Specular) ||
-			IsSceneTextureUsed(PPI_Metallic) ||
-			IsSceneTextureUsed(PPI_WorldNormal) ||
-			IsSceneTextureUsed(PPI_WorldTangent) ||
-			IsSceneTextureUsed(PPI_Opacity) ||
-			IsSceneTextureUsed(PPI_Roughness) ||
-			IsSceneTextureUsed(PPI_Anisotropy) ||
-			IsSceneTextureUsed(PPI_MaterialAO) ||
-			IsSceneTextureUsed(PPI_DecalMask) ||
-			IsSceneTextureUsed(PPI_ShadingModelColor) ||
-			IsSceneTextureUsed(PPI_ShadingModelID) ||
-			IsSceneTextureUsed(PPI_StoredBaseColor) ||
-			IsSceneTextureUsed(PPI_StoredSpecular) ||
-			IsSceneTextureUsed(PPI_Velocity) ||
-			IsSceneTextureUsed(PPI_IsFirstPerson);
+		return UsedSceneTextures & GetGBufferMask();
 	}
 
 	/** true if the material uses the SceneDepth lookup */
diff --git a/Engine/Source/Runtime/Engine/Public/Materials/Material.h b/Engine/Source/Runtime/Engine/Public/Materials/Material.h
--- a/Engine/Source/Runtime/Engine/Public/Materials/Material.h
+++ b/Engine/Source/Runtime/Engine/Public/Materials/Material.h
@@ -1292,6 +1292,8 @@ public:
 	ENGINE_API virtual bool HasSubstrateRoughnessTracking() const override;
 	ENGINE_API virtual bool IsCompatibleWithLumenCardSharing() const override;
 
+	ENGINE_API bool IsMobileSeparateTranslucencyEnabled() const;
+
 	ENGINE_API virtual FGraphEventArray PrecachePSOs(const FPSOPrecacheVertexFactoryDataList& VertexFactoryDataList, const FPSOPrecacheParams& PreCacheParams, EPSOPrecachePriority Priority, TArray<FMaterialPSOPrecacheRequestID>& OutMaterialPSORequestIDs) override;
 
 	ENGINE_API void SetShadingModel(EMaterialShadingModel NewModel);
diff --git a/Engine/Source/Runtime/Engine/Public/Materials/MaterialExpressionDBufferTexture.h b/Engine/Source/Runtime/Engine/Public/Materials/MaterialExpressionDBufferTexture.h
--- a/Engine/Source/Runtime/Engine/Public/Materials/MaterialExpressionDBufferTexture.h
+++ b/Engine/Source/Runtime/Engine/Public/Materials/MaterialExpressionDBufferTexture.h
@@ -33,6 +33,7 @@ class UMaterialExpressionDBufferTexture : public UMaterialExpression
 
 	//~ Begin UMaterialExpression Interface
 #if WITH_EDITOR
+	virtual void Build(MIR::FEmitter& Emitter) override;
 	virtual int32 Compile(class FMaterialCompiler* Compiler, int32 OutputIndex) override;
 	virtual void GetCaption(TArray<FString>& OutCaptions) const override;
 #endif
diff --git a/Engine/Source/Runtime/Engine/Public/Materials/MaterialExpressionSceneColor.h b/Engine/Source/Runtime/Engine/Public/Materials/MaterialExpressionSceneColor.h
--- a/Engine/Source/Runtime/Engine/Public/Materials/MaterialExpressionSceneColor.h
+++ b/Engine/Source/Runtime/Engine/Public/Materials/MaterialExpressionSceneColor.h
@@ -52,6 +52,7 @@ class UMaterialExpressionSceneColor : public UMaterialExpression
 
 	//~ Begin UMaterialExpression Interface
 #if WITH_EDITOR
+	virtual void Build(MIR::FEmitter& Emitter) override;
 	virtual int32 Compile(class FMaterialCompiler* Compiler, int32 OutputIndex) override;
 	virtual void GetCaption(TArray<FString>& OutCaptions) const override;
 	virtual FName GetInputName(int32 InputIndex) const override;
diff --git a/Engine/Source/Runtime/Engine/Public/Materials/MaterialExpressionSceneDepth.h b/Engine/Source/Runtime/Engine/Public/Materials/MaterialExpressionSceneDepth.h
--- a/Engine/Source/Runtime/Engine/Public/Materials/MaterialExpressionSceneDepth.h
+++ b/Engine/Source/Runtime/Engine/Public/Materials/MaterialExpressionSceneDepth.h
@@ -43,6 +43,7 @@ class UMaterialExpressionSceneDepth : public UMaterialExpression
 
 	//~ Begin UMaterialExpression Interface
 #if WITH_EDITOR
+	virtual void Build(MIR::FEmitter& Emitter) override;
 	virtual int32 Compile(class FMaterialCompiler* Compiler, int32 OutputIndex) override;
 	virtual void GetCaption(TArray<FString>& OutCaptions) const override;
 	virtual FName GetInputName(int32 InputIndex) const override;
diff --git a/Engine/Source/Runtime/Engine/Public/Materials/MaterialExpressionSceneDepthWithoutWater.h b/Engine/Source/Runtime/Engine/Public/Materials/MaterialExpressionSceneDepthWithoutWater.h
--- a/Engine/Source/Runtime/Engine/Public/Materials/MaterialExpressionSceneDepthWithoutWater.h
+++ b/Engine/Source/Runtime/Engine/Public/Materials/MaterialExpressionSceneDepthWithoutWater.h
@@ -43,6 +43,7 @@ public:
 
 	//~ Begin UMaterialExpression Interface
 #if WITH_EDITOR
+	virtual void Build(MIR::FEmitter& Emitter) override;
 	virtual int32 Compile(class FMaterialCompiler* Compiler, int32 OutputIndex) override;
 	virtual void GetCaption(TArray<FString>& OutCaptions) const override;
 	virtual FName GetInputName(int32 InputIndex) const override;
diff --git a/Engine/Source/Runtime/Engine/Public/Materials/MaterialExpressionSceneTexelSize.h b/Engine/Source/Runtime/Engine/Public/Materials/MaterialExpressionSceneTexelSize.h
--- a/Engine/Source/Runtime/Engine/Public/Materials/MaterialExpressionSceneTexelSize.h
+++ b/Engine/Source/Runtime/Engine/Public/Materials/MaterialExpressionSceneTexelSize.h
@@ -17,6 +17,7 @@ public:
 
 	//~ Begin UMaterialExpression Interface
 #if WITH_EDITOR
+	virtual void Build(MIR::FEmitter& Emitter) override;
 	virtual int32 Compile(class FMaterialCompiler* Compiler, int32 OutputIndex) override;
 	virtual void GetCaption(TArray<FString>& OutCaptions) const override;
 #endif
diff --git a/Engine/Source/Runtime/Engine/Public/Materials/MaterialExpressionSceneTexture.h b/Engine/Source/Runtime/Engine/Public/Materials/MaterialExpressionSceneTexture.h
--- a/Engine/Source/Runtime/Engine/Public/Materials/MaterialExpressionSceneTexture.h
+++ b/Engine/Source/Runtime/Engine/Public/Materials/MaterialExpressionSceneTexture.h
@@ -29,6 +29,7 @@ class UMaterialExpressionSceneTexture : public UMaterialExpression
 
 	//~ Begin UMaterialExpression Interface
 #if WITH_EDITOR
+	virtual void Build(MIR::FEmitter& Emitter) override;
 	virtual int32 Compile(class FMaterialCompiler* Compiler, int32 OutputIndex) override;
 	virtual void GetCaption(TArray<FString>& OutCaptions) const override;
 #endif
diff --git a/Engine/Source/Runtime/Engine/Public/Materials/MaterialExpressionUserSceneTexture.h b/Engine/Source/Runtime/Engine/Public/Materials/MaterialExpressionUserSceneTexture.h
--- a/Engine/Source/Runtime/Engine/Public/Materials/MaterialExpressionUserSceneTexture.h
+++ b/Engine/Source/Runtime/Engine/Public/Materials/MaterialExpressionUserSceneTexture.h
@@ -32,6 +32,7 @@ class UMaterialExpressionUserSceneTexture : public UMaterialExpression
 
 	//~ Begin UMaterialExpression Interface
 #if WITH_EDITOR
+	virtual void Build(MIR::FEmitter& Emitter) override;
 	virtual int32 Compile(class FMaterialCompiler* Compiler, int32 OutputIndex) override;
 	virtual void GetCaption(TArray<FString>& OutCaptions) const override;
 #endif
diff --git a/Engine/Source/Runtime/Engine/Public/Materials/MaterialIR.h b/Engine/Source/Runtime/Engine/Public/Materials/MaterialIR.h
--- a/Engine/Source/Runtime/Engine/Public/Materials/MaterialIR.h
+++ b/Engine/Source/Runtime/Engine/Public/Materials/MaterialIR.h
@@ -4,6 +4,8 @@
 #include "Materials/MaterialIRCommon.h"
 #include "Materials/MaterialIRTypes.h"
 #include "Materials/MaterialParameters.h"
+#include "Materials/MaterialExpressionDBufferTexture.h"
+#include "MaterialSceneTextureId.h"
 #include "Engine/TextureDefines.h"
 
 #if WITH_EDITOR
@@ -101,6 +103,7 @@ enum EValueKind
 	VK_Poison,
 	VK_Constant,
 	VK_ExternalInput,
+	VK_ScreenTexture,
 	VK_TextureObject,
 	VK_UniformParameter,
 
@@ -439,6 +442,9 @@ bool IsExternalInputTexCoordDdy(EExternalInput Id);
 // Returns whether the given external is a texture coordinate, a ddx or ddy.
 bool IsExternalInputTexCoordOrPartialDerivative(EExternalInput Id);
 
+// Returns whether the given external is a world position, including a ddx or ddy.
+bool IsExternalInputWorldPosition(EExternalInput Id);
+
 // It represents an external input (like texture coordinates) value.
 struct FExternalInput : TValue<VK_ExternalInput>
 {
@@ -473,6 +479,30 @@ struct FUniformParameter: TValue<VK_UniformParameter>
 	uint32 Analysis_UniformParameterIndex;
 };
 
+enum class EScreenTexture : uint8
+{
+	SceneTexture,
+	UserSceneTexture,
+	SceneColor,
+	SceneDepth,
+	SceneDepthWithoutWater,
+	DBufferTexture,
+};
+
+// Reference to various types of full screen textures accessible by expressions.  These
+// generally involve custom code to set data in the compilation output, extensive validation
+// logic, and HLSL generation deferred to the analyzer for the case of UserSceneTextures.
+struct FScreenTexture : TValue<VK_ScreenTexture>
+{
+	EScreenTexture TextureKind;
+	union
+	{
+		ESceneTextureId Id;
+		EDBufferTextureId DBufferId;
+	};
+	FName UserSceneTexture;
+};
+
 /*------------------------------------ Instructions ------------------------------------*/
 
 // A block of instructions.
diff --git a/Engine/Source/Runtime/Engine/Public/Materials/MaterialIREmitter.h b/Engine/Source/Runtime/Engine/Public/Materials/MaterialIREmitter.h
--- a/Engine/Source/Runtime/Engine/Public/Materials/MaterialIREmitter.h
+++ b/Engine/Source/Runtime/Engine/Public/Materials/MaterialIREmitter.h
@@ -339,7 +339,19 @@ public:
 
 	// Emits a custom primitive data value with given index.
 	FValueRef CustomPrimitiveData(uint32 PrimitiveDataIndex);
-	
+
+	// Emit a screen texture with the given Scene Texture id.
+	FValueRef SceneTexture(ESceneTextureId SceneTextureId);
+
+	// Emit a screen texture with the given User Scene Texture name.
+	FValueRef UserSceneTexture(FName UserSceneTexture);
+
+	// Emit a screen texture with the given DBuffer Texture id.
+	FValueRef DBufferTexture(EDBufferTextureId DBufferId);
+
+	// Emit another type of screen texture (not one of the three types above).
+	FValueRef ScreenTexture(EScreenTexture TextureKind);
+
 	/*----------------------------- Instructions emission -----------------------------*/
 
 	// Emits an instruction that sets the material attribute with given property to given argument.
diff --git a/Engine/Source/Runtime/Engine/Public/Materials/MaterialIRModule.h b/Engine/Source/Runtime/Engine/Public/Materials/MaterialIRModule.h
--- a/Engine/Source/Runtime/Engine/Public/Materials/MaterialIRModule.h
+++ b/Engine/Source/Runtime/Engine/Public/Materials/MaterialIRModule.h
@@ -53,6 +53,9 @@ public:
 	// Returns the shader platform associated with this module.
 	EShaderPlatform GetShaderPlatform() const { return ShaderPlatform; }
 
+	// Returns the feature level associated with this module.
+	ERHIFeatureLevel::Type GetFeatureLevel() const { return FeatureLevel; }
+
 	// Retrieves the material compilation output.
 	const FMaterialCompilationOutput& GetCompilationOutput() const { return CompilationOutput; }
 
@@ -96,6 +99,9 @@ private:
 	// Target shader platform.
 	EShaderPlatform ShaderPlatform;
 
+	// Target feature level.
+	ERHIFeatureLevel::Type FeatureLevel;
+
 	// Compilation output data.
 	FMaterialCompilationOutput CompilationOutput;
 
diff --git a/Engine/Source/Runtime/Engine/Public/Materials/MaterialIRModuleBuilder.h b/Engine/Source/Runtime/Engine/Public/Materials/MaterialIRModuleBuilder.h
--- a/Engine/Source/Runtime/Engine/Public/Materials/MaterialIRModuleBuilder.h
+++ b/Engine/Source/Runtime/Engine/Public/Materials/MaterialIRModuleBuilder.h
@@ -11,6 +11,7 @@ struct FMaterialIRModuleBuilder
 	UMaterial* Material;
 	EShaderPlatform ShaderPlatform;
 	const ITargetPlatform* TargetPlatform;
+	ERHIFeatureLevel::Type FeatureLevel;
 	const FStaticParameterSet& StaticParameters;
 	FMaterialInsights* TargetInsights{};
 	UMaterialExpression* PreviewExpression;
