FNewSpline - Temporary workaround for appending points to a closed spline

TTangentBezierSpline::AppendPoint does not currently work when the spline is closed because it does not account for the closing segment.

diff --git a/Engine/Source/Runtime/Engine/Private/Curves/Spline.cpp b/Engine/Source/Runtime/Engine/Private/Curves/Spline.cpp
--- a/Engine/Source/Runtime/Engine/Private/Curves/Spline.cpp
+++ b/Engine/Source/Runtime/Engine/Private/Curves/Spline.cpp
@@ -1870,12 +1870,21 @@ void FNewSpline::ResetScale()
 
 void FNewSpline::AddPoint(const FSplinePoint& Point)
 {
+	// todo: this is a hack to support appending points when the spline is closed. remove when appending to a closed spline is fixed.
+	bool bWasClosed = IsClosedLoop();
+	SetClosedLoop(false);
+	
 	ValidateRotScale();
 	ON_SCOPE_EXIT
 	{
 		SetAttributeChannelRange(ScaleAttrName, FInterval1f(0.f, 1.f), EMappingRangeSpace::Normalized);
 		SetAttributeChannelRange(RotationAttrName, FInterval1f(0.f, 1.f), EMappingRangeSpace::Normalized);
 
+		if (bWasClosed)
+		{
+			SetClosedLoop(true);
+		}
+		
 		ValidateRotScale();
 	};
 	
@@ -1886,8 +1895,7 @@ void FNewSpline::AddPoint(const FSplinePoint& Point)
 	int32 Index = FMath::FloorToInt(Point.InputKey);
 	float Fraction = Point.InputKey - Index;
 	
-	// Special case for empty spline
-	if (NumPoints == 0 || (Index >= NumPoints && !IsClosedLoop()))
+	if (NumPoints == 0 || Index >= NumPoints)
 	{
 		GetSpline().AppendPoint(ControlPoint);
 
@@ -1918,7 +1926,6 @@ void FNewSpline::AddPoint(const FSplinePoint& Point)
     	return;
     }
 	
-	
 	// Insert at calculated parameter
 	int32 InsertedIndex = GetSpline().InsertPointAtSegmentParam(Index, Fraction, ControlPoint);
 	float NewPointInternalParam = GetSpline().GetParameter(InsertedIndex);
@@ -1938,11 +1945,20 @@ void FNewSpline::AddPoint(const FSplinePoint& Point)
 
 void FNewSpline::InsertPoint(const FSplinePoint& Point, int32 Index)
 {
+	// todo: this is a hack to support appending points when the spline is closed. remove when appending to a closed spline is fixed.
+	bool bWasClosed = IsClosedLoop();
+	SetClosedLoop(false);
+	
 	ValidateRotScale();
 	ON_SCOPE_EXIT
 	{
 		SetAttributeChannelRange(ScaleAttrName, FInterval1f(0.f, 1.f), EMappingRangeSpace::Normalized);
 		SetAttributeChannelRange(RotationAttrName, FInterval1f(0.f, 1.f), EMappingRangeSpace::Normalized);
+
+		if (bWasClosed)
+		{
+			SetClosedLoop(true);
+		}
 		
 		ValidateRotScale();
 	};
@@ -1951,7 +1967,7 @@ void FNewSpline::InsertPoint(const FSplinePoint& Point, int32 Index)
 	const int32 NumPoints = GetSpline().GetNumPoints();
     
 	// Special case for empty spline
-	if (NumPoints == 0 || (Index >= NumPoints && !IsClosedLoop()))
+	if (NumPoints == 0 || Index >= NumPoints)
 	{
 		GetSpline().AppendPoint(ControlPoint);
 		
