Apply CoreRedirects when trying to load using the StreamableManager, this matches the existing behavior in SoftObjectPath.TryLoad and fixes issues with asset manager references to native classes

diff --git a/Engine/Source/Runtime/Engine/Private/StreamableManager.cpp b/Engine/Source/Runtime/Engine/Private/StreamableManager.cpp
--- a/Engine/Source/Runtime/Engine/Private/StreamableManager.cpp
+++ b/Engine/Source/Runtime/Engine/Private/StreamableManager.cpp
@@ -2485,7 +2485,12 @@ FSoftObjectPath FStreamableManager::ResolveRedirects(const FSoftObjectPath& Targ
 		UE_LOG(LogStreamableManager, Verbose, TEXT("Redirected %s -> %s"), *Target.ToString(), *Redir->NewPath.ToString());
 		return Redir->NewPath;
 	}
-	return Target;
+
+	// If it wasn't an asset redirect, try applying CoreRedirects
+	FSoftObjectPath ReturnPath = Target;
+	ReturnPath.FixupCoreRedirects();
+	
+	return ReturnPath;
 }
 
 FSoftObjectPath FStreamableManager::HandleLoadedRedirector(UObjectRedirector* LoadedRedirector, FSoftObjectPath RequestedPath, FStreamable* RequestedStreamable)
