Touch streaming mips we don't inline so they dont get GCd by DDC backends.

#rnx

diff --git a/Engine/Source/Runtime/Engine/Classes/Engine/Texture.h b/Engine/Source/Runtime/Engine/Classes/Engine/Texture.h
--- a/Engine/Source/Runtime/Engine/Classes/Engine/Texture.h
+++ b/Engine/Source/Runtime/Engine/Classes/Engine/Texture.h
@@ -1100,6 +1100,7 @@ public:
 	bool TryCancelCache();
 	void CancelCache();
 	ENGINE_API bool TryInlineMipData(int32 FirstMipToLoad = 0, FStringView DebugContext=FStringView());
+	void TouchStreamingMipDDCData(int32 InMipCount, FStringView DebugContext);
 	ENGINE_API TFuture<TTuple<uint64, uint64>> LaunchEstimateOnDiskSizeTask(
 		FOodleDataCompression::ECompressor InOodleCompressor,
 		FOodleDataCompression::ECompressionLevel InOodleCompressionLevel,
diff --git a/Engine/Source/Runtime/Engine/Private/TextureDerivedData.cpp b/Engine/Source/Runtime/Engine/Private/TextureDerivedData.cpp
--- a/Engine/Source/Runtime/Engine/Private/TextureDerivedData.cpp
+++ b/Engine/Source/Runtime/Engine/Private/TextureDerivedData.cpp
@@ -1964,14 +1964,16 @@ typedef TArray<uint32> FAsyncVTChunkHandles;
  * so we know what our in-use key is. This might be on the worker immediately
  * after the fetch completes.
  */
-static bool LoadDerivedStreamingMips(FTexturePlatformData& PlatformData, int32 FirstMipToLoad, FStringView DebugContext, TFunctionRef<void (int32 MipIndex, FSharedBuffer MipData)> Callback)
+static bool LoadDerivedStreamingMips(FTexturePlatformData& PlatformData, int32 FirstMipToLoad, int32 MaxMipCount, bool bTouchOnly, FStringView DebugContext, TFunctionRef<void (int32 MipIndex, FSharedBuffer MipData)> Callback)
 {
 	using namespace UE::DerivedData;
 
 	bool bMiss = false;
 
 	TIndirectArray<FTexture2DMipMap>& Mips = PlatformData.Mips;
-	const int32 ReadableMipCount = Mips.Num() - (PlatformData.GetNumMipsInTail() > 0 ? PlatformData.GetNumMipsInTail() - 1 : 0);
+	int32 ReadableMipCount = Mips.Num() - (PlatformData.GetNumMipsInTail() > 0 ? PlatformData.GetNumMipsInTail() - 1 : 0);
+	ReadableMipCount = FMath::Min(MaxMipCount, ReadableMipCount);
+	
 
 	if (PlatformData.DerivedDataKey.IsType<FString>())
 	{
@@ -1988,28 +1990,43 @@ static bool LoadDerivedStreamingMips(FTexturePlatformData& PlatformData, int32 F
 				Request.Name = MipNameBuilder;
 				Request.Key = ConvertLegacyCacheKey(PlatformData.GetDerivedDataMipKeyString(MipIndex, Mip));
 				Request.UserData = MipIndex;
+				if (bTouchOnly)
+				{
+					Request.Policy = ECachePolicy::Query | ECachePolicy::SkipData;
+				}
 			}
 		}
 
 		if (!Requests.IsEmpty())
 		{
-			COOK_STAT(auto Timer = TextureCookStats::StreamingMipUsageStats.TimeSyncWork());
-			uint64 Size = 0;
-			FRequestOwner BlockingOwner(EPriority::Blocking);
-			GetCache().GetValue(Requests, BlockingOwner, [Callback = MoveTemp(Callback), &Size, &bMiss](FCacheGetValueResponse&& Response)
+			if (bTouchOnly)
 			{
-				Size += Response.Value.GetRawSize();
-				if (Response.Status == EStatus::Ok)
-				{
-					Callback(int32(Response.UserData), Response.Value.GetData().Decompress());
-				}
-				else
-				{
-					bMiss = true;
-				}
-			});
-			BlockingOwner.Wait();
-			COOK_STAT(Timer.AddHitOrMiss(!bMiss ? FCookStats::CallStats::EHitOrMiss::Hit : FCookStats::CallStats::EHitOrMiss::Miss, int64(Size)));
+				// We want to fire off the request but we honestly don't care about the results,
+				// this is just to update server referencing timestamps.
+				FRequestOwner Owner(EPriority::Lowest);
+				Owner.KeepAlive();
+				GetCache().GetValue(Requests, Owner, [](FCacheGetValueResponse&&) {} );
+			}
+			else
+			{
+				COOK_STAT(auto Timer = TextureCookStats::StreamingMipUsageStats.TimeSyncWork());
+				uint64 Size = 0;
+				FRequestOwner BlockingOwner(EPriority::Blocking);
+				GetCache().GetValue(Requests, BlockingOwner, [Callback = MoveTemp(Callback), &Size, &bMiss](FCacheGetValueResponse&& Response)
+					{
+						Size += Response.Value.GetRawSize();
+						if (Response.Status == EStatus::Ok)
+						{
+							Callback(int32(Response.UserData), Response.Value.GetData().Decompress());
+						}
+						else
+						{
+							bMiss = true;
+						}
+					});
+				BlockingOwner.Wait();
+				COOK_STAT(Timer.AddHitOrMiss(!bMiss ? FCookStats::CallStats::EHitOrMiss::Hit : FCookStats::CallStats::EHitOrMiss::Miss, int64(Size)));
+			}
 		}
 	}
 	else if (PlatformData.DerivedDataKey.IsType<FCacheKeyProxy>())
@@ -2029,6 +2046,9 @@ static bool LoadDerivedStreamingMips(FTexturePlatformData& PlatformData, int32 F
 				Request.Key = Key;
 				Request.Id = FTexturePlatformData::MakeMipId(MipIndex);
 				Request.UserData = MipIndex;
+
+				// We don't need to handle Touch because GetChunks uses CacheRecords and all mips
+				// are touched if we touch any mip.
 			}
 		}
 
@@ -2143,7 +2163,7 @@ static bool FetchAllTextureDataSynchronous(FTexturePlatformData* PlatformData, F
 	OutMipData.Empty(MipCount);
 	OutMipData.AddDefaulted(MipCount);
 
-	if (!LoadDerivedStreamingMips(*PlatformData, 0, DebugContext, StoreMip))
+	if (!LoadDerivedStreamingMips(*PlatformData, 0, MipCount, false, DebugContext, StoreMip))
 	{
 		return false;
 	}
@@ -2415,6 +2435,11 @@ TFuture<TTuple<uint64, uint64>> FTexturePlatformData::LaunchEstimateOnDiskSizeTa
 	return ResultFuture;
 }
 
+void FTexturePlatformData::TouchStreamingMipDDCData(int32 InMipCount, FStringView DebugContext)
+{
+	LoadDerivedStreamingMips(*this, 0, InMipCount, true, DebugContext, [](int32, FSharedBuffer){});
+}
+
 bool FTexturePlatformData::TryInlineMipData(int32 FirstMipToLoad, FStringView DebugContext)
 {
 	TRACE_CPUPROFILER_EVENT_SCOPE(FTexturePlatformData::TryInlineMipData);
@@ -2428,7 +2453,7 @@ bool FTexturePlatformData::TryInlineMipData(int32 FirstMipToLoad, FStringView De
 		Mip.BulkData.Unlock();
 	};
 
-	if (!LoadDerivedStreamingMips(*this, FirstMipToLoad, DebugContext, StoreMip))
+	if (!LoadDerivedStreamingMips(*this, FirstMipToLoad, Mips.Num(), false, DebugContext, StoreMip))
 	{
 		return false;
 	}
@@ -2533,7 +2558,7 @@ bool FTexturePlatformData::TryLoadMipsWithSizes(int32 FirstMipToLoad, void** Out
 		}
 	};
 
-	if (!LoadDerivedStreamingMips(*this, FirstMipToLoad, DebugContext, StoreMip))
+	if (!LoadDerivedStreamingMips(*this, FirstMipToLoad, Mips.Num(), false, DebugContext, StoreMip))
 	{
 		return false;
 	}
diff --git a/Engine/Source/Runtime/Engine/Private/TextureDerivedDataTask.cpp b/Engine/Source/Runtime/Engine/Private/TextureDerivedDataTask.cpp
--- a/Engine/Source/Runtime/Engine/Private/TextureDerivedDataTask.cpp
+++ b/Engine/Source/Runtime/Engine/Private/TextureDerivedDataTask.cpp
@@ -1657,6 +1657,17 @@ static void DDC1_FetchAndFillDerivedData(
 			{
 				UE_LOG(LogTexture, Display, TEXT("Texture %s is missing streaming mips when loading for an inline request. The texture will be rebuilt."), *TexturePathName);
 			}
+			else
+			{
+				// The mips we didn't inline could be under different keys. Until we migrate over to ddc cache records, we have to manually
+				// touch them otherwise they could get GC'd in the ddc after a while. We expect to fail this some times for mips that have
+				// already been GC'd, and we do NOT want to force a rebuild with them at this time as that will surface any latent 
+				// determinism
+				if (BuildSettingsPerLayerFetchOrBuild[0].LODBiasWithCinematicMips)
+				{
+					DerivedData->TouchStreamingMipDDCData(BuildSettingsPerLayerFetchOrBuild[0].LODBiasWithCinematicMips, TexturePathName);
+				}
+			}
 		}
 		else
 		{
