Adding function to check if runtime cell transformers are installed.
(Unshelved from pending changelist '43479851')

diff --git a/Engine/Source/Runtime/Engine/Private/WorldPartition/WorldPartition.cpp b/Engine/Source/Runtime/Engine/Private/WorldPartition/WorldPartition.cpp
--- a/Engine/Source/Runtime/Engine/Private/WorldPartition/WorldPartition.cpp
+++ b/Engine/Source/Runtime/Engine/Private/WorldPartition/WorldPartition.cpp
@@ -2513,6 +2513,19 @@ void UWorldPartition::ApplyRuntimeCellsTransformerStack(ULevel* InLevel)
 		ApplyTransformPhase([](const FRuntimeCellTransformerInstance& TransformerInstance, ULevel* Level) { TransformerInstance.PostTransform(Level); });
 	}
 }
+
+bool UWorldPartition::HasRuntimeCellTransformerOfType(const UClass* InCellTransformerClass)
+{
+	for (const FRuntimeCellTransformerInstance& TransformerInstance : RuntimeCellsTransformerStack)
+	{
+		if (TransformerInstance.Class.Get() == InCellTransformerClass)
+		{
+			return true;
+		}
+	}
+	return false;
+}
+
 #endif
 
 bool UWorldPartition::SupportsWorldAssetStreaming(const FName& InTargetGrid)
diff --git a/Engine/Source/Runtime/Engine/Public/WorldPartition/WorldPartition.h b/Engine/Source/Runtime/Engine/Public/WorldPartition/WorldPartition.h
--- a/Engine/Source/Runtime/Engine/Public/WorldPartition/WorldPartition.h
+++ b/Engine/Source/Runtime/Engine/Public/WorldPartition/WorldPartition.h
@@ -340,6 +340,7 @@ public:
 	ENGINE_API FBox GetRuntimeWorldBounds() const;
 
 	ENGINE_API void ApplyRuntimeCellsTransformerStack(ULevel* InLevel);
+	ENGINE_API bool HasRuntimeCellTransformerOfType(const UClass* InCellTransformerClass);
 	
 	UHLODLayer* GetDefaultHLODLayer() const { return DefaultHLODLayer; }
 	void SetDefaultHLODLayer(UHLODLayer* InDefaultHLODLayer) { DefaultHLODLayer = InDefaultHLODLayer; }
