UE: Improve logging a little to show the package of an actor with a bad scale for collision

diff --git a/Engine/Source/Runtime/Engine/Private/Components/PrimitiveComponent.cpp b/Engine/Source/Runtime/Engine/Private/Components/PrimitiveComponent.cpp
--- a/Engine/Source/Runtime/Engine/Private/Components/PrimitiveComponent.cpp
+++ b/Engine/Source/Runtime/Engine/Private/Components/PrimitiveComponent.cpp
@@ -47,6 +47,7 @@
 #include "UObject/FortniteMainBranchObjectVersion.h"
 #include "UObject/UE5SpecialProjectStreamObjectVersion.h"
 #include "UObject/ObjectSaveContext.h"
+#include "UObject/Package.h"
 #include "Engine/DamageEvents.h"
 #include "MeshUVChannelInfo.h"
 #include "PrimitiveSceneDesc.h"
@@ -898,7 +899,7 @@ void UPrimitiveComponent::OnCreatePhysicsState()
 #if !(UE_BUILD_SHIPPING || UE_BUILD_TEST)
 			if ((BodyInstance.GetCollisionEnabled() != ECollisionEnabled::NoCollision) && (FMath::IsNearlyZero(BodyScale.X) || FMath::IsNearlyZero(BodyScale.Y) || FMath::IsNearlyZero(BodyScale.Z)))
 			{
-				UE_LOG(LogPhysics, Warning, TEXT("Scale for %s has a component set to zero, which will result in a bad body instance. Scale:%s"), *GetPathNameSafe(this), *BodyScale.ToString());
+				UE_LOG(LogPhysics, Warning, TEXT("Scale for %s has a component set to zero, which will result in a bad body instance. Scale:%s Package:%s"), *GetPathNameSafe(this), *BodyScale.ToString(), *GetPackage()->GetName());
 				
 				// User warning has been output - fix up the scale to be valid for physics
 				BodyTransform.SetScale3D(FVector(
