New material translator:
- Improved FValue prototype low-level construction and emission code in the emitter and added checks to make sure the hashing is always determinic, regardless of padding.
- Removed "NewCompositeValue()" special case which would allocate a composite value and later destroy it if already there. Added instead MakeCompositePrototype() to make a temporary composite prototype, that is emitted just like any value with EmitPrototype().

#rnx

diff --git a/Engine/Source/Runtime/Engine/Private/Materials/MaterialIREmitter.cpp b/Engine/Source/Runtime/Engine/Private/Materials/MaterialIREmitter.cpp
--- a/Engine/Source/Runtime/Engine/Private/Materials/MaterialIREmitter.cpp
+++ b/Engine/Source/Runtime/Engine/Private/Materials/MaterialIREmitter.cpp
@@ -89,76 +89,72 @@ struct FEmitter::FPrivate
 		return Bytes;
 	}
 
-	//
-	static void Deallocate(FEmitter& Emitter, void* Ptr)
-	{
-		// This is a NOOP, only there for code robustness.
-	}
-
 	// Pushes a new value to the list of values.
-	static void PushNewValue(FEmitter& Emitter, FValue* Value)
+	static void PushValueToModule(FEmitter& Emitter, FValue* Value)
 	{
 		Emitter.Module->Values.Add(Value);
 		Emitter.ValueSet.Add(Value);
 	}
 };
 
-// Creates a new composite value of specified type and returns it.
-static FComposite* NewCompositeValue(FEmitter& Emitter, FType Type, int NumComponents)
+// Makes an instance of specified FValue type T and initailizes its kind and type.
+// Note that this "prototype" value must be emitted using EmitPrototype(), which in turns
+// returns the real FValueRef.
+template <typename T>
+static T MakePrototype(FType InType)
 {
-	int32 SizeInBytes = sizeof(FComposite) + sizeof(FValue*) * NumComponents;
-
-	void* Bytes = FEmitter::FPrivate::Allocate(Emitter, SizeInBytes, alignof(FComposite));
+	 static_assert(std::is_trivially_destructible_v<T> && std::is_trivially_copy_constructible_v<T> && std::is_trivially_copy_assignable_v<T>,
+		"FValues are expected to be trivial types.");
 
-	FComposite* Value = new (Bytes) FComposite{};
-	Value->Kind = VK_Composite;
-	Value->Type = Type;
+	T Value;
+	FMemory::Memzero(Value);
+	Value.Kind = T::TypeKind;
+	Value.Type = InType;
 
 	return Value;
 }
 
-// It adds the given value to the module and returns it. If the exact value already existed,
-// it destroys the the given value and returns the existing one instead.
-static FValueRef EmitNew(FEmitter& Emitter, FValue* Value)
+// Creates a new composite value of specified type and returns it.
+static FComposite* MakeCompositePrototype(FEmitter& Emitter, FType Type, int NumComponents)
 {
-	if (FValue* Existing = FEmitter::FPrivate::FindValue(Emitter, Value))
-	{
-		FEmitter::FPrivate::Deallocate(Emitter, Value);
-		return Existing;
-	}
+	// Compute the total size of this composite value.
+	int32 SizeInBytes = sizeof(FComposite) + sizeof(FValue*) * NumComponents;
 
-	FEmitter::FPrivate::PushNewValue(Emitter, Value);
-	return Value;
-}
+	// Allocate a temporary buffer for it.
+	FComposite* Value = (FComposite*)FMemStack::Get().Alloc(SizeInBytes, alignof(FComposite));
 
-template <typename T>
-static T MakePrototype(FType InType)
-{
-	static_assert(std::is_trivially_destructible_v<T> && std::is_trivially_copy_constructible_v<T> && std::is_trivially_copy_assignable_v<T>,
-		"FValues are expected to be trivial types.");
+	// Zero its memory and set it up
+	FMemory::Memzero(Value, SizeInBytes);
+	Value->Kind = VK_Composite;
+	Value->Type = Type;
 
-	T Value;
-	FMemory::Memzero(Value);
-	Value.Kind = T::TypeKind;
-	Value.Type = InType;
 	return Value;
 }
 
 // Searches for an existing value in module that matches specified `Prototype`.
 // If none found, it creates a new value as a copy of the prototype, adds it to
 // the module then returns it.
-template <typename TValueType>
-static FValueRef EmitPrototype(FEmitter& Emitter, const TValueType& Prototype)
+static FValueRef EmitPrototype(FEmitter& Emitter, const FValue& Prototype)
 {
+	// Optimization: See if we emitted this value before, and if so, since MIR is SSA, with
+	// instructions having being the equivalent of "pure functions" with no side effects,
+	// simply return the existing value which holds the already computed result.
 	if (FValue* Existing = FEmitter::FPrivate::FindValue(Emitter, &Prototype))
 	{
 		return Existing;
 	}
 
-	void* Bytes = FEmitter::FPrivate::Allocate(Emitter, sizeof(TValueType), alignof(TValueType));
-	TValueType* Value = new (Bytes) TValueType{ Prototype };
+	// Otherwise, create a new value allocating the necessary memory in the module's arena.
+	FValue* Value = (FValue*)FEmitter::FPrivate::Allocate(Emitter, Prototype.GetSizeInBytes(), alignof(FValue));
 
-	FEmitter::FPrivate::PushNewValue(Emitter, Value);
+	// Bitcopy the prototype into the value memory.
+	FMemory::Memcpy(Value, &Prototype, Prototype.GetSizeInBytes());
+
+	// Push the value to the module.
+	FEmitter::FPrivate::PushValueToModule(Emitter, Value);
+
+	// Verify that value hashing is deterministic.
+	checkSlow(Value == FEmitter::FPrivate::FindValue(Emitter, &Prototype));
 
 	return Value;
 }
@@ -552,12 +548,12 @@ FValueRef FEmitter::ConstantDefault(FType Type)
 		}
 		else
 		{
-			FComposite* Composite = NewCompositeValue(*this, Type, PrimitiveType->NumComponents());
+			FComposite* Composite = MakeCompositePrototype(*this, Type, PrimitiveType->NumComponents());
 			for (FValue*& Component : Composite->GetMutableComponents())
 			{
 				Component = Zero;
 			}
-			return EmitNew(*this, Composite);
+			return EmitPrototype(*this, *Composite);
 		}
 	}
 	else if (const UMaterialAggregate* TypeAggregate = Type.AsAggregate())
@@ -879,7 +875,6 @@ FValueRef FEmitter::CustomPrimitiveData(uint32 PrimitiveDataIndex)
 	return {};
 }
 
-// Emit a scene texture value with the given Id
 FValueRef FEmitter::SceneTexture(ESceneTextureId SceneTextureId)
 {
 	FScreenTexture Proto = MakePrototype<FScreenTexture>(MIR::FType::MakeInt(1));
@@ -888,17 +883,15 @@ FValueRef FEmitter::SceneTexture(ESceneTextureId SceneTextureId)
 	return EmitPrototype(*this, Proto);
 }
 
-// Emit a user scene texture value with the given name
 FValueRef FEmitter::UserSceneTexture(FName UserSceneTexture)
 {
 	FScreenTexture Proto = MakePrototype<FScreenTexture>(MIR::FType::MakeInt(1));
 	Proto.TextureKind = EScreenTexture::UserSceneTexture;
-	Proto.Id = PPI_UserSceneTexture0;				// Actual ID determined later during HLSL generation
+	Proto.Id = PPI_UserSceneTexture0;				// Actual ID determined later during MIR lowering
 	Proto.UserSceneTexture = UserSceneTexture;
 	return EmitPrototype(*this, Proto);
 }
 
-// Emit a DBuffer texture value with the given Id
 FValueRef FEmitter::DBufferTexture(EDBufferTextureId DBufferId)
 {
 	FScreenTexture Proto = MakePrototype<FScreenTexture>(MIR::FType::MakeInt(1));
@@ -959,6 +952,7 @@ FValueRef FEmitter::Vector2(FValueRef InX, FValueRef InY)
 	check(InX->Type == InY->Type);
 
 	TComposite<2> Vector = MakePrototype<TComposite<2>>(FType::MakeVector(InX->Type.AsPrimitive()->ScalarKind, 2));
+
 	TArrayView<FValue*> Components = Vector.GetMutableComponents();
 	Components[0] = InX;
 	Components[1] = InY;
@@ -1049,7 +1043,7 @@ FValueRef FEmitter::Aggregate(const UMaterialAggregate* InAggregate, FValueRef I
 	}
 
 	// Create the new composite to store the aggregate attribute values.
-	FComposite* AggregateValue = NewCompositeValue(*this, FType::MakeAggregate(InAggregate), InAggregate->Attributes.Num());
+	FComposite* AggregateValue = MakeCompositePrototype(*this, FType::MakeAggregate(InAggregate), InAggregate->Attributes.Num());
 	
 	// Assign all components of the new composite value.
 	TArrayView<FValue*> Components = AggregateValue->GetMutableComponents();
@@ -1075,7 +1069,7 @@ FValueRef FEmitter::Aggregate(const UMaterialAggregate* InAggregate, FValueRef I
 		}
 	}
 
-	return EmitNew(*this, AggregateValue);
+	return EmitPrototype(*this, *AggregateValue);
 }
 
 FValueRef FEmitter::Aggregate(const UMaterialAggregate* InAggregate, FValueRef InPrototype, TConstArrayView<FAttributeAssignment> AttributeAssignments)
@@ -1090,7 +1084,7 @@ FValueRef FEmitter::Aggregate(const UMaterialAggregate* InAggregate, FValueRef I
 	}
 
 	// Create the new composite to store the aggregate attribute values.
-	FComposite* AggregateValue = NewCompositeValue(*this, FType::MakeAggregate(InAggregate), InAggregate->Attributes.Num());
+	FComposite* AggregateValue = MakeCompositePrototype(*this, FType::MakeAggregate(InAggregate), InAggregate->Attributes.Num());
 	
 	// Assign all components of the new composite value.
 	TArrayView<FValue*> Components = AggregateValue->GetMutableComponents();
@@ -1123,7 +1117,7 @@ FValueRef FEmitter::Aggregate(const UMaterialAggregate* InAggregate, FValueRef I
 		}
 	}
 
-	return EmitNew(*this, AggregateValue);
+	return EmitPrototype(*this, *AggregateValue);
 }
 
 /*--------------------------------- Operator emission ---------------------------------*/
@@ -2062,7 +2056,7 @@ static FValue* TryFoldComponentwiseOperator(FEmitter& Emitter, EOperator Op, FVa
 		}
 
 		// Make the new composite value
-		FComposite* Result = NewCompositeValue(Emitter, ResultType, TempResultComponents.Num());
+		FComposite* Result = MakeCompositePrototype(Emitter, ResultType, TempResultComponents.Num());
 
 		// Fetch the components array from the result composite
 		TArrayView<FValue*> ResultComponents = Result->GetMutableComponents();
@@ -2088,7 +2082,7 @@ static FValue* TryFoldComponentwiseOperator(FEmitter& Emitter, EOperator Op, FVa
 			}
 		}
 		
-		return EmitNew(Emitter, Result);
+		return EmitPrototype(Emitter, *Result);
 	}
 
 	return {};
@@ -2443,7 +2437,7 @@ FValueRef FEmitter::Swizzle(FValueRef Value, FSwizzleMask Mask)
 
 	// Make the result vector type.
 	FType ResultType = FType::MakeVector(PrimitiveType.ScalarKind, Mask.NumComponents);
-	FComposite* Result = NewCompositeValue(*this, ResultType, ResultType.GetPrimitive().NumComponents());
+	FComposite* Result = MakeCompositePrototype(*this, ResultType, ResultType.GetPrimitive().NumComponents());
 	TArrayView<FValue*> ResultComponents = Result->GetMutableComponents();
 
 	for (int32 i = 0; i < Mask.NumComponents; ++i)
@@ -2451,7 +2445,7 @@ FValueRef FEmitter::Swizzle(FValueRef Value, FSwizzleMask Mask)
 		ResultComponents[i] = Subscript(Value, (int32)Mask.Components[i]);
 	}
 
-	return Value.To(EmitNew(*this, Result));
+	return Value.To(EmitPrototype(*this, *Result));
 }
 
 static FValue* CastConstant(FEmitter& Emitter, FConstant* Constant, EScalarKind ConstantScalarKind, EScalarKind TargetKind)
@@ -2556,7 +2550,7 @@ static FValue* CastToPrimitive(FEmitter& Emitter, FValueRef Value, FType TargetT
 	if (ValuePrimitiveType.IsScalar())
 	{
 		// Create the result composite value.
-		FComposite* Result = NewCompositeValue(Emitter, TargetType, TargetPrimitiveType.NumComponents());
+		FComposite* Result = MakeCompositePrototype(Emitter, TargetType, TargetPrimitiveType.NumComponents());
 
 		// Create a composite and initialize each of its components to the conversion
 		// of initializer value to the single component type.
@@ -2571,7 +2565,7 @@ static FValue* CastToPrimitive(FEmitter& Emitter, FValueRef Value, FType TargetT
 			ResultComponents[i] = Component;
 		}
 		
-		return EmitNew(Emitter, Result);
+		return EmitPrototype(Emitter, *Result);
 	}
 
 	// Construct a vector from another vector. If constructed vector is larger, initialize
@@ -2585,7 +2579,7 @@ static FValue* CastToPrimitive(FEmitter& Emitter, FValueRef Value, FType TargetT
 		int32 InitializerNumComponents = ValuePrimitiveType.NumComponents();
 
 		// Create the result composite value.
-		FComposite* Result = NewCompositeValue(Emitter, TargetType, TargetPrimitiveType.NumComponents());
+		FComposite* Result = MakeCompositePrototype(Emitter, TargetType, TargetPrimitiveType.NumComponents());
 		TArrayView<FValue*> ResultComponents = Result->GetMutableComponents();
 
 		// Determine the result component type (scalar).
@@ -2607,7 +2601,7 @@ static FValue* CastToPrimitive(FEmitter& Emitter, FValueRef Value, FType TargetT
 			ResultComponents[Index] = Emitter.ConstantZero(ResultComponentType.GetPrimitive().ScalarKind);
 		}
 
-		return EmitNew(Emitter, Result);
+		return EmitPrototype(Emitter, *Result);
 	}
 	
 	// The two primitive types are identical matrices that differ only by their scalar type.
@@ -2617,7 +2611,7 @@ static FValue* CastToPrimitive(FEmitter& Emitter, FValueRef Value, FType TargetT
 		check(TargetPrimitiveType.IsMatrix());
 
 		// Create the result composite value.
-		FComposite* Result = NewCompositeValue(Emitter, TargetType, TargetPrimitiveType.NumComponents());
+		FComposite* Result = MakeCompositePrototype(Emitter, TargetType, TargetPrimitiveType.NumComponents());
 		TArrayView<FValue*> ResultComponents = Result->GetMutableComponents();
 			
 		// Determine the result component type (scalar).
@@ -2629,7 +2623,7 @@ static FValue* CastToPrimitive(FEmitter& Emitter, FValueRef Value, FType TargetT
 			ResultComponents[Index] = Emitter.Cast(Emitter.Subscript(Value, Index), ResultComponentType);
 		}
 
-		return EmitNew(Emitter, Result);
+		return EmitPrototype(Emitter, *Result);
 	}
 
 	// Initializer value cannot be used to construct this primitive type.
@@ -3232,14 +3226,18 @@ FValueRef FEmitter::AnalyticalPartialDerivative(FValueRef Value, EDerivativeAxis
 
 		case VK_Composite:
 		{
-			TConstArrayView<FValue*> ValueComponents = Value->As<FComposite>()->GetComponents();
-			FComposite* Derivative = NewCompositeValue(*this, Value->Type, Value->Type.GetPrimitive().NumComponents());
+			// Make a prototype composite to hold the derivatives of all its components
+			FComposite* Derivative = MakeCompositePrototype(*this, Value->Type, Value->Type.GetPrimitive().NumComponents());
 			TArrayView<FValue*> DerivativeComponents = Derivative->GetMutableComponents();
+
+			// Compute the derivative of each component
+			TConstArrayView<FValue*> ValueComponents = Value->As<FComposite>()->GetComponents();
 			for (int32 i = 0; i < ValueComponents.Num(); ++i)
 			{
 				DerivativeComponents[i] = AnalyticalPartialDerivative(ValueComponents[i], Axis);
 			}
-			return Derivative;
+
+			return EmitPrototype(*this, *Derivative);
 		}
 
 		case VK_Operator:
@@ -3384,19 +3382,6 @@ void FEmitter::Initialize()
 
 	Temp.Boolean = false;
 	FalseConstant = EmitPrototype(*this, Temp);
-
-	// This check ensures that FValue caching returns identical results for repeated values.
-	// If this assert fails, it's likely due to uninitialized padding bytes in FValue introduced by compiler-added padding.
-	// Review any recent changes to FValue to ensure no additional padding was introduced (i.e., the struct remains tightly packed).
-	// Temporary disable this check because it's causing a 100% assert on Mac
-	/*
-	if (FPlatformMisc::IsDebuggerPresent())
-	{
-		FValue* A = ConstantZero(MIR::EScalarKind::Float);
-		FValue* B = ConstantZero(MIR::EScalarKind::Float);
-		check(A == B);
-	}
-	*/
 }
 
 bool FEmitter::FValueKeyFuncs::Matches(KeyInitType A, KeyInitType B)
@@ -3406,7 +3391,7 @@ bool FEmitter::FValueKeyFuncs::Matches(KeyInitType A, KeyInitType B)
 
 uint32 FEmitter::FValueKeyFuncs::GetKeyHash(KeyInitType Key)
 {
-	return Internal::HashBytes((char*)Key, Key->GetSizeInBytes());
+	return Internal::HashBytes(Key, Key->GetSizeInBytes());
 }
 
 } // namespace MIR
diff --git a/Engine/Source/Runtime/Engine/Private/Materials/MaterialIRInternal.cpp b/Engine/Source/Runtime/Engine/Private/Materials/MaterialIRInternal.cpp
--- a/Engine/Source/Runtime/Engine/Private/Materials/MaterialIRInternal.cpp
+++ b/Engine/Source/Runtime/Engine/Private/Materials/MaterialIRInternal.cpp
@@ -98,29 +98,29 @@ EMaterialTextureParameterType TextureMaterialValueTypeToParameterType(EMaterialV
 	}
 }
 
-uint32 HashBytes(const char* Ptr, uint32 Size)
+uint32 HashBytes(const void* InPtr, uint32 Size)
 {
+	check((uintptr_t)InPtr % 4 == 0); // Require the pointer to be 4-byte aligned
+
 	uint32 Hash = 0;
 
-	const uint32* UPtr = (const uint32*)Ptr;
-	uint32 i = 0; 
-	for (uint32 SizeIn32Bits = Size / 4; i < SizeIn32Bits; ++i)
+	const uint32* WordPtr = (const uint32*)InPtr;
+	uint32 NumWords = Size / 4;
+	for (uint32 i = 0; i < NumWords; ++i)
 	{
-		Hash = HashCombineFast(Hash, UPtr[i]);
+		Hash = HashCombineFast(Hash, WordPtr[i]);
 	}
 
-	uint32 Modulo = Size % 4;
-	if (Modulo == 1)
-	{
-		Hash = HashCombineFast(Hash, UPtr[i] & 0xff000000);
-	}
-	else if (Modulo == 2)
-	{
-		Hash = HashCombineFast(Hash, UPtr[i] & 0xffff0000);
-	}
-	else if (Modulo == 3)
+	const uint8* EndPtr = (const uint8*)(WordPtr + NumWords);
+
+	uint32 End = 0;
+	switch (Size % 4)
 	{
-		Hash = HashCombineFast(Hash, UPtr[i] & 0xffffff00);
+		case 3: End |= uint32(EndPtr[2]) << 16; // fallthrough
+		case 2: End |= uint32(EndPtr[1]) << 8;  // fallthrough
+		case 1: End |= uint32(EndPtr[0]);       // fallthrough
+			Hash = HashCombineFast(Hash, End);
+		default: break;
 	}
 
 	return Hash;
diff --git a/Engine/Source/Runtime/Engine/Private/Materials/MaterialIRInternal.h b/Engine/Source/Runtime/Engine/Private/Materials/MaterialIRInternal.h
--- a/Engine/Source/Runtime/Engine/Private/Materials/MaterialIRInternal.h
+++ b/Engine/Source/Runtime/Engine/Private/Materials/MaterialIRInternal.h
@@ -29,8 +29,8 @@ void BindValueToExpressionInput(FMaterialIRModuleBuilderImpl* Builder, const FEx
 // Flows a value int given expression output.
 void BindValueToExpressionOutput(FMaterialIRModuleBuilderImpl* Builder, const FExpressionOutput* Output, FValue* Value);
 
-//
-uint32 HashBytes(const char* Ptr, uint32 Size);
+// Computes the hash of a blob of memory. Note: Ptr must be 4 bytes aligned.
+uint32 HashBytes(const void* Ptr, uint32 Size);
 
 /* Other helper functions */
 
diff --git a/Engine/Source/Runtime/Engine/Private/Materials/MaterialIRModuleBuilder.cpp b/Engine/Source/Runtime/Engine/Private/Materials/MaterialIRModuleBuilder.cpp
--- a/Engine/Source/Runtime/Engine/Private/Materials/MaterialIRModuleBuilder.cpp
+++ b/Engine/Source/Runtime/Engine/Private/Materials/MaterialIRModuleBuilder.cpp
@@ -141,6 +141,7 @@ struct FMaterialIRModuleBuilderImpl
 	{
 		while (true)
 		{
+			FMemMark MemMark { FMemStack::Get() };
 			FAnalysisContext& Context = AnalysisContextStack.Last();
 
 			if (!Context.ExpressionStack.IsEmpty())
@@ -702,6 +703,7 @@ struct FMaterialIRModuleBuilderImpl
 
 bool FMaterialIRModuleBuilder::Build(FMaterialIRModule* TargetModule)
 {
+	FMemMark MemMark { FMemStack::Get() };
 	FMaterialIRModuleBuilderImpl Impl{ this, TargetModule };
 
 	// Setup the emitter
diff --git a/Engine/Source/Runtime/Engine/Private/Materials/MaterialIRTypes.cpp b/Engine/Source/Runtime/Engine/Private/Materials/MaterialIRTypes.cpp
--- a/Engine/Source/Runtime/Engine/Private/Materials/MaterialIRTypes.cpp
+++ b/Engine/Source/Runtime/Engine/Private/Materials/MaterialIRTypes.cpp
@@ -160,7 +160,7 @@ FType FType::MakePoison()
 
 FType FType::MakeVoid()
 {
-	FType Type;
+	FType Type{};
 	Type.Kind = ETypeKind::Void;
 	return Type;
 }
@@ -203,7 +203,7 @@ FType FType::MakePrimitive(EScalarKind InScalarKind, int NumRows, int NumColumns
 	// LWC inverse matrix flag only supported for EScalarKind::LWC and 4x4 dimensions.
 	check(!bIsLWCInverseMatrix || (InScalarKind == EScalarKind::Double && NumRows == 4 && NumColumns == 4));
 	
-	FType Type;
+	FType Type{};
 	Type.Kind = ETypeKind::Primitive;
 	Type.Primitive.ScalarKind = InScalarKind;
 	Type.Primitive.NumRows = NumRows;
@@ -214,7 +214,7 @@ FType FType::MakePrimitive(EScalarKind InScalarKind, int NumRows, int NumColumns
 
 FType FType::MakeAggregate(const UMaterialAggregate* InAggregate)
 {
-	FType Type;
+	FType Type{};
 	Type.Kind = ETypeKind::Aggregate;
 	Type.Aggregate = InAggregate;
 	return Type;
@@ -222,21 +222,21 @@ FType FType::MakeAggregate(const UMaterialAggregate* InAggregate)
 
 FType FType::MakeTexture2D()
 {
-	FType Type;
+	FType Type{};
 	Type.Kind = ETypeKind::Texture2D;
 	return Type;
 }
 
 FType FType::MakeSubstrateData()
 {
-	FType Type;
+	FType Type{};
 	Type.Kind = ETypeKind::SubstrateData;
 	return Type;
 }
 
 FType FType::MakeVTPageTableResult()
 {
-	FType Type;
+	FType Type{};
 	Type.Kind = ETypeKind::VTPageTableResult;
 	return Type;}
 
diff --git a/Engine/Source/Runtime/Engine/Public/Materials/MaterialIRTypes.h b/Engine/Source/Runtime/Engine/Public/Materials/MaterialIRTypes.h
--- a/Engine/Source/Runtime/Engine/Public/Materials/MaterialIRTypes.h
+++ b/Engine/Source/Runtime/Engine/Public/Materials/MaterialIRTypes.h
@@ -259,9 +259,6 @@ private:
 
 	// Identifies the high-level kind of this type.
 	ETypeKind Kind : 8;
-
-	// Manual padding to fit FType size to a multiple of alignof(void*). This storage may be used in the future.
-	uint8 	  Unused[7];
 };
 
 // To avoid the compiler adding automatic padding, we must make sure that its size is a multiple of maximum alignment.
