Appendix to my previous new material translator change: improving documentation of static functions and removing redundant memzero.

#rnx

diff --git a/Engine/Source/Runtime/Engine/Private/Materials/MaterialIREmitter.cpp b/Engine/Source/Runtime/Engine/Private/Materials/MaterialIREmitter.cpp
--- a/Engine/Source/Runtime/Engine/Private/Materials/MaterialIREmitter.cpp
+++ b/Engine/Source/Runtime/Engine/Private/Materials/MaterialIREmitter.cpp
@@ -24,6 +24,7 @@ static TAutoConsoleVariable<bool> CVarMaterialIRDebugBreakOnPoison(
 namespace MIR
 {
 
+// Converts a vector component enum to its string representation ("x", "y", "z", or "w").
 const TCHAR* VectorComponentToString(EVectorComponent Component)
 {
 	static const TCHAR* Strings[] = { TEXT("x"), TEXT("y"), TEXT("z"), TEXT("w") };
@@ -74,22 +75,20 @@ void FSwizzleMask::Append(EVectorComponent Component)
 
 struct FEmitter::FPrivate
 {
-	// Looks for an existing value in the module that matches `Prototype` and returns it if found.
+	// Searches the emitter’s value set for an existing FValue matching Prototype, or returns nullptr.
 	static FValue* FindValue(FEmitter& Emitter, const FValue* Prototype)
 	{
 		FValue** Value = Emitter.ValueSet.Find(Prototype);
 		return Value ? *Value : nullptr;
 	}
 
-	// Allocates a block of memory.
+	// Allocates zero-initialized memory of given size and alignment in the module’s arena.
 	static void* Allocate(FEmitter& Emitter, int32 Size, int32 Alignment)
 	{
-		void* Bytes = (void*)Emitter.Module->Allocator.PushBytes(Size, Alignment);
-		FMemory::Memzero(Bytes, Size);
-		return Bytes;
+		return Emitter.Module->Allocator.PushBytes(Size, Alignment);
 	}
 
-	// Pushes a new value to the list of values.
+	// Registers a new value in the module.
 	static void PushValueToModule(FEmitter& Emitter, FValue* Value)
 	{
 		Emitter.Module->Values.Add(Value);
@@ -97,9 +96,9 @@ struct FEmitter::FPrivate
 	}
 };
 
-// Makes an instance of specified FValue type T and initailizes its kind and type.
-// Note that this "prototype" value must be emitted using EmitPrototype(), which in turns
-// returns the real FValueRef.
+// Creates a “prototype” value of type T with the specified MIR type.
+// Emit this value later with a matching EmitPrototype() call, which returns the actual value instance
+// after potential deduplication.
 template <typename T>
 static T MakePrototype(FType InType)
 {
@@ -114,7 +113,7 @@ static T MakePrototype(FType InType)
 	return Value;
 }
 
-// Creates a new composite value of specified type and returns it.
+// Allocates and initializes a temporary composite value with the given number of components.
 static FComposite* MakeCompositePrototype(FEmitter& Emitter, FType Type, int NumComponents)
 {
 	// Compute the total size of this composite value.
@@ -131,9 +130,7 @@ static FComposite* MakeCompositePrototype(FEmitter& Emitter, FType Type, int Num
 	return Value;
 }
 
-// Searches for an existing value in module that matches specified `Prototype`.
-// If none found, it creates a new value as a copy of the prototype, adds it to
-// the module then returns it.
+// Emits a prototype value into the module, deduplicating if an identical value was already emitted.
 static FValueRef EmitPrototype(FEmitter& Emitter, const FValue& Prototype)
 {
 	// Optimization: See if we emitted this value before, and if so, since MIR is SSA, with
