Fix for body setup race condition causing multiple threads to attempt to build collision data when the complex collision override was used. This happened as querying whether the cooked data was ready always reported that a complex mesh was ready, so the main thread would attempt to use the unbuilt phyiscs data and attemp to build it concurrently with the worker thread leading to the race.

diff --git a/Engine/Source/Runtime/Engine/Classes/Engine/StaticMesh.h b/Engine/Source/Runtime/Engine/Classes/Engine/StaticMesh.h
--- a/Engine/Source/Runtime/Engine/Classes/Engine/StaticMesh.h
+++ b/Engine/Source/Runtime/Engine/Classes/Engine/StaticMesh.h
@@ -2005,6 +2005,7 @@ public:
 
 private:
 	bool GetPhysicsTriMeshDataCheckComplex(struct FTriMeshCollisionData* CollisionData, bool bInUseAllTriData, bool bInCheckComplexCollisionMesh);
+	bool PollAsyncPhysicsTriMeshDataCheckComplex(bool InUseAllTriData, bool bInCheckComplexCollisionMesh) const;
 	bool ContainsPhysicsTriMeshDataCheckComplex(bool InUseAllTriData, bool bInCheckComplexCollisionMesh) const;
 	bool SectionHasCollisionEnabled(const FStaticMeshSection& Section, int32 LODIndex, int32 SectionIndex) const;
 
diff --git a/Engine/Source/Runtime/Engine/Private/StaticMesh.cpp b/Engine/Source/Runtime/Engine/Private/StaticMesh.cpp
--- a/Engine/Source/Runtime/Engine/Private/StaticMesh.cpp
+++ b/Engine/Source/Runtime/Engine/Private/StaticMesh.cpp
@@ -8805,6 +8805,18 @@ bool UStaticMesh::GetPhysicsTriMeshDataCheckComplex(struct FTriMeshCollisionData
 	return CollisionData->Vertices.Num() > 0 && CollisionData->Indices.Num() > 0;
 }
 
+bool UStaticMesh::PollAsyncPhysicsTriMeshDataCheckComplex(bool InUseAllTriData, bool bInCheckComplexCollisionMesh) const
+{
+#if WITH_EDITORONLY_DATA
+	if(ComplexCollisionMesh && ComplexCollisionMesh != this && bInCheckComplexCollisionMesh)
+	{
+		// Match GetPhysicsTriMeshDataCheck recursion limits on complex mesh to avoid loops
+		return ComplexCollisionMesh->PollAsyncPhysicsTriMeshDataCheckComplex(InUseAllTriData, false);
+	}
+#endif
+	return !IsCompiling();
+}
+
 bool UStaticMesh::ContainsPhysicsTriMeshData(bool bInUseAllTriData) const
 {
 	bool bInCheckComplexCollisionMesh = true;
@@ -8906,14 +8918,8 @@ bool UStaticMesh::SectionHasCollisionEnabled(const FStaticMeshSection& Section,
 
 bool UStaticMesh::PollAsyncPhysicsTriMeshData(bool InUseAllTriData) const
 {
-#if WITH_EDITORONLY_DATA
 	bool bInCheckComplexCollisionMesh = true;
-	if (ComplexCollisionMesh && ComplexCollisionMesh != this && bInCheckComplexCollisionMesh)
-	{
-		return true;
-	}
-#endif
-	return !IsCompiling();
+	return PollAsyncPhysicsTriMeshDataCheckComplex(InUseAllTriData, bInCheckComplexCollisionMesh);
 }
 
 void UStaticMesh::GetMeshId(FString& OutMeshId)
