Fix Texture.Usage (task) to always count as a DDC hit. This is purely tracking overhead, not ddc hit/miss.
#rnx

diff --git a/Engine/Source/Runtime/Engine/Private/TextureDerivedDataTask.cpp b/Engine/Source/Runtime/Engine/Private/TextureDerivedDataTask.cpp
--- a/Engine/Source/Runtime/Engine/Private/TextureDerivedDataTask.cpp
+++ b/Engine/Source/Runtime/Engine/Private/TextureDerivedDataTask.cpp
@@ -2032,6 +2032,7 @@ void FTextureCacheDerivedDataWorker::DoWork()
 
 	TRACE_CPUPROFILER_EVENT_SCOPE(FTextureCacheDerivedDataWorker::DoWork);
 	COOK_STAT(auto Timer = TextureCookStats::TaskUsageStats.TimeSyncWork());
+	COOK_STAT(Timer.AddHit(0));	//	Register as a DDC hit, but no bytes processed. Purely tracking processing overhead, not DDC transformation work
 
 	const bool bAllowAsyncBuild = EnumHasAnyFlags(CacheFlags, ETextureCacheFlags::AllowAsyncBuild);
 	const bool bAllowAsyncLoading = EnumHasAnyFlags(CacheFlags, ETextureCacheFlags::AllowAsyncLoading);
