Change SubstrateBsdfFeature to typed enum.

diff --git a/Engine/Source/Editor/MaterialEditor/Private/MaterialNodes/SGraphSubstrateMaterial.cpp b/Engine/Source/Editor/MaterialEditor/Private/MaterialNodes/SGraphSubstrateMaterial.cpp
--- a/Engine/Source/Editor/MaterialEditor/Private/MaterialNodes/SGraphSubstrateMaterial.cpp
+++ b/Engine/Source/Editor/MaterialEditor/Private/MaterialNodes/SGraphSubstrateMaterial.cpp
@@ -167,8 +167,8 @@ static const TSharedRef<SWidget> InternalProcessOperator(
 			auto GetEyeDesc = [&Op]()
 			{
 				FString EyeDesc = TEXT("Eye");
-				const bool bIrisN = Op.Has(ESubstrateBsdfFeature_EyeIrisNormalPluggedIn);
-				const bool bIrisT = Op.Has(ESubstrateBsdfFeature_EyeIrisTangentPluggedIn);
+				const bool bIrisN = Op.Has(ESubstrateBsdfFeature::EyeIrisNormalPluggedIn);
+				const bool bIrisT = Op.Has(ESubstrateBsdfFeature::EyeIrisTangentPluggedIn);
 				if (bIrisN || bIrisT)
 				{
 					EyeDesc += TEXT("Iris");
@@ -188,16 +188,16 @@ static const TSharedRef<SWidget> InternalProcessOperator(
 			FString BSDFDesc = OutputType == ESubstrateWidgetOutputType::Node ? 
 											 FString(TEXT("BSDF")) : 
 											 FString::Printf(TEXT("BSDF (%s%s%s%s%s%s%s%s%s%s)")
-											, Op.Has(ESubstrateBsdfFeature_EdgeColor) ? TEXT("F90 ") : TEXT("")
-											, Op.Has(ESubstrateBsdfFeature_SSS) ? TEXT("SSS ") : TEXT("")
-											, Op.Has(ESubstrateBsdfFeature_MFPPluggedIn) ? TEXT("MFP ") : TEXT("")
-											, Op.Has(ESubstrateBsdfFeature_Anisotropy) ? TEXT("Ani ") : TEXT("")
-											, Op.Has(ESubstrateBsdfFeature_SecondRoughnessOrSimpleClearCoat) ? TEXT("2Ro ") : TEXT("")
-											, Op.Has(ESubstrateBsdfFeature_Fuzz) ? TEXT("Fuz ") : TEXT("")
-											, Op.Has(ESubstrateBsdfFeature_Glint) ? TEXT("Gli ") : TEXT("")
-											, Op.Has(ESubstrateBsdfFeature_SpecularProfile) ? TEXT("Spc ") : TEXT("")
-											, Op.Has(ESubstrateBsdfFeature_Eye) ? *GetEyeDesc() : TEXT("")
-											, Op.Has(ESubstrateBsdfFeature_Hair) ? TEXT("Hai ") : TEXT("")
+											, Op.Has(ESubstrateBsdfFeature::EdgeColor) ? TEXT("F90 ") : TEXT("")
+											, Op.Has(ESubstrateBsdfFeature::SSS) ? TEXT("SSS ") : TEXT("")
+											, Op.Has(ESubstrateBsdfFeature::MFPPluggedIn) ? TEXT("MFP ") : TEXT("")
+											, Op.Has(ESubstrateBsdfFeature::Anisotropy) ? TEXT("Ani ") : TEXT("")
+											, Op.Has(ESubstrateBsdfFeature::SecondRoughnessOrSimpleClearCoat) ? TEXT("2Ro ") : TEXT("")
+											, Op.Has(ESubstrateBsdfFeature::Fuzz) ? TEXT("Fuz ") : TEXT("")
+											, Op.Has(ESubstrateBsdfFeature::Glint) ? TEXT("Gli ") : TEXT("")
+											, Op.Has(ESubstrateBsdfFeature::SpecularProfile) ? TEXT("Spc ") : TEXT("")
+											, Op.Has(ESubstrateBsdfFeature::Eye) ? *GetEyeDesc() : TEXT("")
+											, Op.Has(ESubstrateBsdfFeature::Hair) ? TEXT("Hai ") : TEXT("")
 											);
 
 			static FString ToolTip;
@@ -223,20 +223,20 @@ static const TSharedRef<SWidget> InternalProcessOperator(
 				{
 					SubSurfaceType = EMaterialSubSurfaceType::MSS_SimpleVolume;
 				}
-				else if (!Op.Has(ESubstrateBsdfFeature_SSS))
+				else if (!Op.Has(ESubstrateBsdfFeature::SSS))
 				{
 					SubSurfaceType = EMaterialSubSurfaceType::MSS_None;
 				}
 
-				const uint32 DisplayedFeatureMask = 
-					ESubstrateBsdfFeature::ESubstrateBsdfFeature_EdgeColor |
-					ESubstrateBsdfFeature::ESubstrateBsdfFeature_Fuzz |
-					ESubstrateBsdfFeature::ESubstrateBsdfFeature_SecondRoughnessOrSimpleClearCoat |
-					ESubstrateBsdfFeature::ESubstrateBsdfFeature_Anisotropy |
-					ESubstrateBsdfFeature::ESubstrateBsdfFeature_Glint |
-					ESubstrateBsdfFeature::ESubstrateBsdfFeature_SpecularProfile;
+				const ESubstrateBsdfFeature DisplayedFeatureMask = 
+					ESubstrateBsdfFeature::EdgeColor |
+					ESubstrateBsdfFeature::Fuzz |
+					ESubstrateBsdfFeature::SecondRoughnessOrSimpleClearCoat |
+					ESubstrateBsdfFeature::Anisotropy |
+					ESubstrateBsdfFeature::Glint |
+					ESubstrateBsdfFeature::SpecularProfile;
 
-				const bool bHasFeatures = (Op.BSDFFeatures & DisplayedFeatureMask) != 0;
+				const bool bHasFeatures = EnumHasAnyFlags(Op.BSDFFeatures, DisplayedFeatureMask);
 
 				if (bHasFeatures || SubSurfaceType != EMaterialSubSurfaceType::MSS_None)
 				{
@@ -277,27 +277,27 @@ static const TSharedRef<SWidget> InternalProcessOperator(
 							];
 						};
 
-						if (Op.Has(ESubstrateBsdfFeature::ESubstrateBsdfFeature_EdgeColor))
+						if (Op.Has(ESubstrateBsdfFeature::EdgeColor))
 						{
 							AddFeatureSlot(FText::FromString(TEXT("F90")));
 						}
-						if (Op.Has(ESubstrateBsdfFeature::ESubstrateBsdfFeature_Fuzz))
+						if (Op.Has(ESubstrateBsdfFeature::Fuzz))
 						{
 							AddFeatureSlot(FText::FromString(TEXT("Fuzz")));
 						}
-						if (Op.Has(ESubstrateBsdfFeature::ESubstrateBsdfFeature_SecondRoughnessOrSimpleClearCoat))
+						if (Op.Has(ESubstrateBsdfFeature::SecondRoughnessOrSimpleClearCoat))
 						{
 							AddFeatureSlot(FText::FromString(TEXT("Dual Spec.")));
 						}
-						if (Op.Has(ESubstrateBsdfFeature::ESubstrateBsdfFeature_Anisotropy))
+						if (Op.Has(ESubstrateBsdfFeature::Anisotropy))
 						{
 							AddFeatureSlot(FText::FromString(TEXT("Aniso")));
 						}
-						if (Op.Has(ESubstrateBsdfFeature::ESubstrateBsdfFeature_Glint))
+						if (Op.Has(ESubstrateBsdfFeature::Glint))
 						{
 							AddFeatureSlot(FText::FromString(TEXT("Glints")));
 						}
-						if (Op.Has(ESubstrateBsdfFeature::ESubstrateBsdfFeature_SpecularProfile))
+						if (Op.Has(ESubstrateBsdfFeature::SpecularProfile))
 						{
 							AddFeatureSlot(FText::FromString(TEXT("Spec. Profile")));
 						}
diff --git a/Engine/Source/Runtime/Core/Private/UObject/FortniteMainBranchObjectVersion.cpp b/Engine/Source/Runtime/Core/Private/UObject/FortniteMainBranchObjectVersion.cpp
--- a/Engine/Source/Runtime/Core/Private/UObject/FortniteMainBranchObjectVersion.cpp
+++ b/Engine/Source/Runtime/Core/Private/UObject/FortniteMainBranchObjectVersion.cpp
@@ -14,6 +14,6 @@ TMap<FGuid, FGuid> FFortniteMainBranchObjectVersion::GetSystemGuids()
 	SystemGuids.Add(DevGuids.Niagara_LatestScriptCompileVersion, FGuid("97F0FD7149D90D42B65D6557AF477D1A"));
 	SystemGuids.Add(DevGuids.SkeletalMeshDerivedDataVersion, FGuid("C998F2F3B2884A02B7E290B8C25835F0"));
 	SystemGuids.Add(DevGuids.STATICMESH_DERIVEDDATA_VER, FGuid("3ABF32149F9D4D838750368AB95573BE"));
-	SystemGuids.Add(DevGuids.MaterialTranslationDDCVersion, FGuid("73294AD79624440FAD3F53E9794C505A"));
+	SystemGuids.Add(DevGuids.MaterialTranslationDDCVersion, FGuid("439B1428AAED4BA489BF3D886E461F06"));
 	return SystemGuids;
 }
diff --git a/Engine/Source/Runtime/Engine/Private/Materials/HLSLMaterialTranslator.cpp b/Engine/Source/Runtime/Engine/Private/Materials/HLSLMaterialTranslator.cpp
--- a/Engine/Source/Runtime/Engine/Private/Materials/HLSLMaterialTranslator.cpp
+++ b/Engine/Source/Runtime/Engine/Private/Materials/HLSLMaterialTranslator.cpp
@@ -13280,7 +13280,7 @@ bool FHLSLMaterialTranslator::FSubstrateCompilationContext::SubstrateGenerateDer
 			{
 				// Disable all optional features for now to fit.
 				// SUBSTRATE_TODO we will need to refine that to account for platforms supporting SSS for instance.
-				It.BSDFFeatures = 0;
+				It.BSDFFeatures = ESubstrateBsdfFeature::None;
 			}
 			SubstrateSimplificationStatus.bSlabSimplificationStepHasBeenRun = true;
 		}
@@ -13650,6 +13650,7 @@ bool FHLSLMaterialTranslator::FSubstrateCompilationContext::SubstrateGenerateDer
 			SubstrateMaterialComplexity.bIsSimple = SubstrateMaterialEffectiveClosureCount == 1;
 			SubstrateMaterialComplexity.bIsSingle = SubstrateMaterialEffectiveClosureCount == 1;
 			SubstrateMaterialComplexity.bIsComplexSpecial = false;
+			ESubstrateBsdfFeature SubstrateMaterialBsdfFeatures = ESubstrateBsdfFeature::None;
 			bool bIsFastWaterPath = false;
 			bool bCustomEncoding = false;
 			bool bMayHaveCoverageLessThan1 = false;
@@ -13675,19 +13676,22 @@ bool FHLSLMaterialTranslator::FSubstrateCompilationContext::SubstrateGenerateDer
 					// From the compiler side, we can only assume the top layer has gray scale luminance weight.
 					const bool bMayHaveColoredWeight = !It.bIsTop;
 
+					// Aggregate all BSDFs features used by the material
+					SubstrateMaterialBsdfFeatures |= ESubstrateBsdfFeature(It.BSDFFeatures);
+
 					switch (It.BSDFType)
 					{
 					case SUBSTRATE_BSDF_TYPE_SLAB:
 					{
 						bool bIsComplexSpecial = false;
-						if (It.Has(ESubstrateBsdfFeature_ComplexSpecialMask))
+						if (It.Has(ESubstrateBsdfFeature::ComplexSpecialMask))
 						{
 							// We need to check each features if really enabled for a platofrms
-							bIsComplexSpecial = It.Has(ESubstrateBsdfFeature_Glint) && Substrate::IsGlintEnabled(ShaderPlatform);
+							bIsComplexSpecial = It.Has(ESubstrateBsdfFeature::Glint) && Substrate::IsGlintEnabled(ShaderPlatform);
 						}
 
-						SubstrateMaterialComplexity.bIsSimple = SubstrateMaterialComplexity.bIsSimple && !bMayHaveColoredWeight && !It.Has(ESubstrateBsdfFeature_ComplexMask) && !bIsComplexSpecial && !It.Has(ESubstrateBsdfFeature_SingleMask);
-						SubstrateMaterialComplexity.bIsSingle = SubstrateMaterialComplexity.bIsSingle && !bMayHaveColoredWeight && !It.Has(ESubstrateBsdfFeature_ComplexMask) && !bIsComplexSpecial;
+						SubstrateMaterialComplexity.bIsSimple = SubstrateMaterialComplexity.bIsSimple && !bMayHaveColoredWeight && !It.Has(ESubstrateBsdfFeature::ComplexMask) && !bIsComplexSpecial && !It.Has(ESubstrateBsdfFeature::SingleMask);
+						SubstrateMaterialComplexity.bIsSingle = SubstrateMaterialComplexity.bIsSingle && !bMayHaveColoredWeight && !It.Has(ESubstrateBsdfFeature::ComplexMask) && !bIsComplexSpecial;
 						SubstrateMaterialComplexity.bIsComplexSpecial |= bIsComplexSpecial;
 						break;
 					}
@@ -13821,13 +13825,13 @@ bool FHLSLMaterialTranslator::FSubstrateCompilationContext::SubstrateGenerateDer
 					case SUBSTRATE_BSDF_TYPE_SLAB:
 					{
 						// Compute values closer to the reality for HasSSS and IsSimpleVolume, now that we know that we know the topology of the material.
-						const bool bIsSimpleVolume = !It.bIsBottom && It.Has(ESubstrateBsdfFeature_MFPPluggedIn);
-						const bool bHasSSS = It.bIsBottom && It.Has(ESubstrateBsdfFeature_SSS);
+						const bool bIsSimpleVolume = !It.bIsBottom && It.Has(ESubstrateBsdfFeature::MFPPluggedIn);
+						const bool bHasSSS = It.bIsBottom && It.Has(ESubstrateBsdfFeature::SSS);
 
 						SubstrateMaterialRequestedSizeByte += UintByteSize;
 						SubstrateMaterialRequestedSizeByte += UintByteSize;
 
-						if (It.Has(ESubstrateBsdfFeature_EdgeColor | ESubstrateBsdfFeature_SecondRoughnessOrSimpleClearCoat))
+						if (It.Has(ESubstrateBsdfFeature::EdgeColor | ESubstrateBsdfFeature::SecondRoughnessOrSimpleClearCoat))
 						{
 							SubstrateMaterialRequestedSizeByte += UintByteSize;
 						}
@@ -13835,27 +13839,27 @@ bool FHLSLMaterialTranslator::FSubstrateCompilationContext::SubstrateGenerateDer
 						{
 							SubstrateMaterialRequestedSizeByte += UintByteSize;
 						}
-						if (It.Has(ESubstrateBsdfFeature_Fuzz))
+						if (It.Has(ESubstrateBsdfFeature::Fuzz))
 						{
 							SubstrateMaterialRequestedSizeByte += UintByteSize;
 						}
-						if (It.Has(ESubstrateBsdfFeature_Glint) && Substrate::IsGlintEnabled(ShaderPlatform))
+						if (It.Has(ESubstrateBsdfFeature::Glint) && Substrate::IsGlintEnabled(ShaderPlatform))
 						{
 							SubstrateMaterialRequestedSizeByte += UintByteSize;
 							SubstrateMaterialRequestedSizeByte += UintByteSize;
 							SubstrateMaterialRequestedSizeByte += UintByteSize;
 							SubstrateMaterialRequestedSizeByte += UintByteSize;
 						}
-						if (It.Has(ESubstrateBsdfFeature_SpecularProfile))
+						if (It.Has(ESubstrateBsdfFeature::SpecularProfile))
 						{
 							SubstrateMaterialRequestedSizeByte += UintByteSize;
 						}
-						if (It.Has(ESubstrateBsdfFeature_Eye))
+						if (It.Has(ESubstrateBsdfFeature::Eye))
 						{
 							SubstrateMaterialRequestedSizeByte += UintByteSize;
 							SubstrateMaterialRequestedSizeByte += UintByteSize;
 						}
-						if (It.Has(ESubstrateBsdfFeature_Hair))
+						if (It.Has(ESubstrateBsdfFeature::Hair))
 						{
 							SubstrateMaterialRequestedSizeByte += UintByteSize;
 							SubstrateMaterialRequestedSizeByte += UintByteSize;
@@ -16523,6 +16527,8 @@ void FHLSLMaterialTranslator::PrepareEnvironmentDefines()
 		FString SubstrateMaterialDescription;
 		for (uint32 SubstrateCompilationContextIndex = 0; SubstrateCompilationContextIndex < ESubstrateCompilationContext::SCC_MAX; ++SubstrateCompilationContextIndex)
 		{
+			ESubstrateBsdfFeature MaterialBsdfFeatures = ESubstrateBsdfFeature::None;
+
 			FSubstrateCompilationContext& SubstrateCtx = SubstrateCompilationContext[SubstrateCompilationContextIndex];
 			const FSubstrateSimplificationStatus& SubstrateSimplificationStatus = SubstrateCtx.SubstrateSimplificationStatus;
 
@@ -16581,8 +16587,8 @@ void FHLSLMaterialTranslator::PrepareEnvironmentDefines()
 				{
 					FSubstrateOperator& BSDFOperator = SubstrateCtx.SubstrateMaterialExpressionRegisteredOperators[OpIt];
 
-					EyeIrisNormalPluggedIn |= !!(BSDFOperator.BSDFFeatures & ESubstrateBsdfFeature_EyeIrisNormalPluggedIn);
-					EyeIrisTangentPluggedIn |= !!(BSDFOperator.BSDFFeatures & ESubstrateBsdfFeature_EyeIrisTangentPluggedIn);
+					EyeIrisNormalPluggedIn  = EyeIrisNormalPluggedIn  || BSDFOperator.Has(ESubstrateBsdfFeature::EyeIrisNormalPluggedIn);
+					EyeIrisTangentPluggedIn = EyeIrisTangentPluggedIn || BSDFOperator.Has(ESubstrateBsdfFeature::EyeIrisTangentPluggedIn);
 
 					if (BSDFOperator.BSDFIndex == INDEX_NONE || BSDFOperator.IsDiscarded())
 					{
@@ -16612,17 +16618,20 @@ void FHLSLMaterialTranslator::PrepareEnvironmentDefines()
 								SubstrateMaterialContextDescription += FString::Printf(TEXT("\tIdx=%d Op=%s ParentIdx=%d LeftIndex=%d RightIndex=%d BSDFIdx=%d LayerDepth=%d IsTop=%d IsBot=%d BSDFType=%s SSS=%d MFP=%d F90=%d Rough2=%d Fuzz=%d Aniso=%d Glint=%d SpecularProfile=%d Eye=%d Hair=%d\n GUID=%s\n"),
 									It.Index, GetSubstrateOperatorStr(It.OperatorType), It.ParentIndex, It.LeftIndex, It.RightIndex, It.BSDFIndex, It.LayerDepth, It.bIsTop, It.bIsBottom,
 									*GetSubstrateBSDFName(It.BSDFType), 
-									It.Has(ESubstrateBsdfFeature_SSS), 
-									It.Has(ESubstrateBsdfFeature_MFPPluggedIn), 
-									It.Has(ESubstrateBsdfFeature_EdgeColor), 
-									It.Has(ESubstrateBsdfFeature_SecondRoughnessOrSimpleClearCoat), 
-									It.Has(ESubstrateBsdfFeature_Fuzz), 
-									It.Has(ESubstrateBsdfFeature_Anisotropy), 
-									It.Has(ESubstrateBsdfFeature_Glint), 
-									It.Has(ESubstrateBsdfFeature_SpecularProfile), 
-									It.Has(ESubstrateBsdfFeature_Eye), 
-									It.Has(ESubstrateBsdfFeature_Hair),
+									It.Has(ESubstrateBsdfFeature::SSS), 
+									It.Has(ESubstrateBsdfFeature::MFPPluggedIn), 
+									It.Has(ESubstrateBsdfFeature::EdgeColor), 
+									It.Has(ESubstrateBsdfFeature::SecondRoughnessOrSimpleClearCoat), 
+									It.Has(ESubstrateBsdfFeature::Fuzz), 
+									It.Has(ESubstrateBsdfFeature::Anisotropy), 
+									It.Has(ESubstrateBsdfFeature::Glint), 
+									It.Has(ESubstrateBsdfFeature::SpecularProfile), 
+									It.Has(ESubstrateBsdfFeature::Eye), 
+									It.Has(ESubstrateBsdfFeature::Hair),
 									*It.MaterialExpressionGuid.ToString());
+
+								// Aggregate all the features used within the material
+								MaterialBsdfFeatures |= ESubstrateBsdfFeature(It.BSDFFeatures);
 							}
 						}
 					}
diff --git a/Engine/Source/Runtime/Engine/Private/Materials/MaterialExpressionSubstrate.cpp b/Engine/Source/Runtime/Engine/Private/Materials/MaterialExpressionSubstrate.cpp
--- a/Engine/Source/Runtime/Engine/Private/Materials/MaterialExpressionSubstrate.cpp
+++ b/Engine/Source/Runtime/Engine/Private/Materials/MaterialExpressionSubstrate.cpp
@@ -596,11 +596,11 @@ FSubstrateOperator* UMaterialExpressionSubstrateShadingModels::SubstrateGenerate
 	// Note Thickness has no meaning/usage in the context of SubstrateLegacyConversionNode
 	int32 ThicknessIndex = Compiler->SubstrateThicknessStackGetThicknessIndex();
 
-	auto AddDefaultWorstCase = [&](uint32 In)
+	auto AddDefaultWorstCase = [&](ESubstrateBsdfFeature In)
 		{
 			FSubstrateOperator& SlabOperator = Compiler->SubstrateCompilationRegisterOperator(SUBSTRATE_OPERATOR_BSDF_LEGACY, Compiler->SubstrateTreeStackGetPathUniqueId(), ThisExpressionGuid, Parent, Compiler->SubstrateTreeStackGetParentPathUniqueId());
 			SlabOperator.BSDFType = SUBSTRATE_BSDF_TYPE_SLAB;
-			SlabOperator.BSDFFeatures = In | (Anisotropy.IsConnected() ? ESubstrateBsdfFeature_Anisotropy : ESubstrateBsdfFeature_None);
+			SlabOperator.BSDFFeatures = In | (Anisotropy.IsConnected() ? ESubstrateBsdfFeature::Anisotropy : ESubstrateBsdfFeature::None);
 			SlabOperator.SubSurfaceType = uint8(EMaterialSubSurfaceType::MSS_None);
 			SlabOperator.bBSDFWritesEmissive = EmissiveColor.IsConnected();
 			SlabOperator.ThicknessIndex = ThicknessIndex;
@@ -615,8 +615,8 @@ FSubstrateOperator* UMaterialExpressionSubstrateShadingModels::SubstrateGenerate
 	bool bEyeIrisTangentPluggedIn = CustomTangent.IsConnected();
 	auto ApplyEyeIrisUsedFeatures = [&](FSubstrateOperator* Operator)
 	{
-		if (bEyeIrisNormalPluggedIn)	Operator->BSDFFeatures |= ESubstrateBsdfFeature_EyeIrisNormalPluggedIn;
-		if (bEyeIrisTangentPluggedIn)	Operator->BSDFFeatures |= ESubstrateBsdfFeature_EyeIrisTangentPluggedIn;
+		if (bEyeIrisNormalPluggedIn)	Operator->BSDFFeatures |= ESubstrateBsdfFeature::EyeIrisNormalPluggedIn;
+		if (bEyeIrisTangentPluggedIn)	Operator->BSDFFeatures |= ESubstrateBsdfFeature::EyeIrisTangentPluggedIn;
 	};
 
 	// Logic about shading models and complexity should match UMaterialExpressionSubstrateShadingModels::Compile.
@@ -634,9 +634,9 @@ FSubstrateOperator* UMaterialExpressionSubstrateShadingModels::SubstrateGenerate
 
 		// Be sure to track eye/hair feature, even though they are not part of Slab BSDF. This is important later for issuing 
 		// the correct material complexity, as hair/eye requires 'complex' complexity (not 'simple') for correct packing
-		uint32 BsdfFeatures = ESubstrateBsdfFeature_SSS | ESubstrateBsdfFeature_Fuzz;
-		if (ShadingModels.HasShadingModel(MSM_Eye))  { BsdfFeatures |= ESubstrateBsdfFeature_Eye; }
-		if (ShadingModels.HasShadingModel(MSM_Hair)) { BsdfFeatures |= ESubstrateBsdfFeature_Hair; }
+		ESubstrateBsdfFeature BsdfFeatures = ESubstrateBsdfFeature::SSS | ESubstrateBsdfFeature::Fuzz;
+		if (ShadingModels.HasShadingModel(MSM_Eye))  { BsdfFeatures |= ESubstrateBsdfFeature::Eye; }
+		if (ShadingModels.HasShadingModel(MSM_Hair)) { BsdfFeatures |= ESubstrateBsdfFeature::Hair; }
 		FSubstrateOperator* Operator = AddDefaultWorstCase(BsdfFeatures);
 		if (ShadingModels.HasShadingModel(MSM_Eye) || bHasShadingModelFromExpression)
 		{
@@ -658,38 +658,38 @@ FSubstrateOperator* UMaterialExpressionSubstrateShadingModels::SubstrateGenerate
 		}
 		else if (ShadingModels.HasShadingModel(MSM_DefaultLit))
 		{
-			return AddDefaultWorstCase(ESubstrateBsdfFeature_None);
+			return AddDefaultWorstCase(ESubstrateBsdfFeature::None);
 		}
 		else if (ShadingModels.HasShadingModel(MSM_ThinTranslucent))
 		{
-			return AddDefaultWorstCase(ESubstrateBsdfFeature_MFPPluggedIn);
+			return AddDefaultWorstCase(ESubstrateBsdfFeature::MFPPluggedIn);
 		}
 		else if (ShadingModels.HasShadingModel(MSM_SubsurfaceProfile))
 		{
-			return AddDefaultWorstCase(ESubstrateBsdfFeature_SSS);
+			return AddDefaultWorstCase(ESubstrateBsdfFeature::SSS);
 		}
 		else if (ShadingModels.HasShadingModel(MSM_Subsurface))
 		{
-			return AddDefaultWorstCase(ESubstrateBsdfFeature_SSS);
+			return AddDefaultWorstCase(ESubstrateBsdfFeature::SSS);
 		}
 		else if (ShadingModels.HasShadingModel(MSM_TwoSidedFoliage))
 		{
-			return AddDefaultWorstCase(ESubstrateBsdfFeature_SSS);
+			return AddDefaultWorstCase(ESubstrateBsdfFeature::SSS);
 		}
 		else if (ShadingModels.HasShadingModel(MSM_PreintegratedSkin))
 		{
-			return AddDefaultWorstCase(ESubstrateBsdfFeature_SSS);
+			return AddDefaultWorstCase(ESubstrateBsdfFeature::SSS);
 		}
 		else if (ShadingModels.HasShadingModel(MSM_Cloth))
 		{
-			return AddDefaultWorstCase(ESubstrateBsdfFeature_Fuzz);
+			return AddDefaultWorstCase(ESubstrateBsdfFeature::Fuzz);
 		}
 		else if (ShadingModels.HasShadingModel(MSM_ClearCoat))
 		{
 			FSubstrateOperator& Operator = Compiler->SubstrateCompilationRegisterOperator(SUBSTRATE_OPERATOR_BSDF_LEGACY, Compiler->SubstrateTreeStackGetPathUniqueId(), ThisExpressionGuid, Parent, Compiler->SubstrateTreeStackGetParentPathUniqueId());
 			Operator.BSDFType = SUBSTRATE_BSDF_TYPE_SLAB;
 			Operator.ThicknessIndex = ThicknessIndex;
-			Operator.BSDFFeatures = ESubstrateBsdfFeature_SecondRoughnessOrSimpleClearCoat | (Anisotropy.IsConnected() ? ESubstrateBsdfFeature_Anisotropy : ESubstrateBsdfFeature_None);
+			Operator.BSDFFeatures = ESubstrateBsdfFeature::SecondRoughnessOrSimpleClearCoat | (Anisotropy.IsConnected() ? ESubstrateBsdfFeature::Anisotropy : ESubstrateBsdfFeature::None);
 			Operator.bBSDFWritesEmissive = EmissiveColor.IsConnected();
 			return &Operator;
 		}
@@ -697,7 +697,7 @@ FSubstrateOperator* UMaterialExpressionSubstrateShadingModels::SubstrateGenerate
 		{
 			FSubstrateOperator& Operator = Compiler->SubstrateCompilationRegisterOperator(SUBSTRATE_OPERATOR_BSDF_LEGACY, Compiler->SubstrateTreeStackGetPathUniqueId(), ThisExpressionGuid, Parent, Compiler->SubstrateTreeStackGetParentPathUniqueId());
 			Operator.BSDFType = SUBSTRATE_BSDF_TYPE_HAIR;
-			Operator.BSDFFeatures = ESubstrateBsdfFeature_Hair;
+			Operator.BSDFFeatures = ESubstrateBsdfFeature::Hair;
 			Operator.ThicknessIndex = ThicknessIndex;
 			Operator.bBSDFWritesEmissive = EmissiveColor.IsConnected();
 			return &Operator;
@@ -706,7 +706,7 @@ FSubstrateOperator* UMaterialExpressionSubstrateShadingModels::SubstrateGenerate
 		{
 			FSubstrateOperator& Operator = Compiler->SubstrateCompilationRegisterOperator(SUBSTRATE_OPERATOR_BSDF_LEGACY, Compiler->SubstrateTreeStackGetPathUniqueId(), ThisExpressionGuid, Parent, Compiler->SubstrateTreeStackGetParentPathUniqueId());
 			Operator.BSDFType = SUBSTRATE_BSDF_TYPE_EYE;
-			Operator.BSDFFeatures = ESubstrateBsdfFeature_Eye;
+			Operator.BSDFFeatures = ESubstrateBsdfFeature::Eye;
 			ApplyEyeIrisUsedFeatures(&Operator);
 
 			Operator.ThicknessIndex = ThicknessIndex;
@@ -795,10 +795,10 @@ int32 UMaterialExpressionSubstrateSlabBSDF::Compile(class FMaterialCompiler* Com
 
 	// We also cannot ignore the tangent when using the default Tangent because GetTangentBasis
 	// used in SubstrateGetBSDFSharedBasis cannot be relied on for smooth tangent used for lighting on any mesh.
-	const bool bHasAnisotropy = SubstrateOperator.Has(ESubstrateBsdfFeature_Anisotropy);
+	const bool bHasAnisotropy = SubstrateOperator.Has(ESubstrateBsdfFeature::Anisotropy);
 
 	int32 SSSProfileCodeChunk = INDEX_NONE;
-	if (SubstrateOperator.Has(ESubstrateBsdfFeature_SSS) && SubsurfaceProfile)
+	if (SubstrateOperator.Has(ESubstrateBsdfFeature::SSS) && SubsurfaceProfile)
 	{
 		SSSProfileCodeChunk = CreateSubsurfaceProfileParameter(Compiler, SubsurfaceProfile);
 	}
@@ -808,7 +808,7 @@ int32 UMaterialExpressionSubstrateSlabBSDF::Compile(class FMaterialCompiler* Com
 	}
 
 	int32 SpecularProfileCodeChunk = INDEX_NONE;
-	if (SubstrateOperator.Has(ESubstrateBsdfFeature_SpecularProfile))
+	if (SubstrateOperator.Has(ESubstrateBsdfFeature::SpecularProfile))
 	{
 		const FName SpecularProfileParameterName = SpecularProfile::CreateSpecularProfileParameterName(SpecularProfile);
 		SpecularProfileCodeChunk = Compiler->ForceCast(Compiler->ScalarParameter(SpecularProfileParameterName, 1.0f), MCT_Float1);
@@ -852,31 +852,31 @@ int32 UMaterialExpressionSubstrateSlabBSDF::Compile(class FMaterialCompiler* Com
 	int32 GlintUVCodeChunk = CompileWithDefaultFloat2(Compiler, GlintUV, 0.0f, 0.0f);
 
 	// Disable some features if requested by the simplification process
-	if (!SubstrateOperator.Has(ESubstrateBsdfFeature_MFPPluggedIn))
+	if (!SubstrateOperator.Has(ESubstrateBsdfFeature::MFPPluggedIn))
 	{
 		SSSMFPCodeChunk = Compiler->Constant3(0.0f, 0.0f, 0.0f);
 	}
-	if (!SubstrateOperator.Has(ESubstrateBsdfFeature_EdgeColor))
+	if (!SubstrateOperator.Has(ESubstrateBsdfFeature::EdgeColor))
 	{
 		F90CodeChunk = Compiler->Constant3(1.0f, 1.0f, 1.0f);
 	}
-	if (!SubstrateOperator.Has(ESubstrateBsdfFeature_Fuzz))
+	if (!SubstrateOperator.Has(ESubstrateBsdfFeature::Fuzz))
 	{
 		FuzzAmountCodeChunk = Compiler->Constant(0.0f);
 	}
-	if (!SubstrateOperator.Has(ESubstrateBsdfFeature_SecondRoughnessOrSimpleClearCoat))
+	if (!SubstrateOperator.Has(ESubstrateBsdfFeature::SecondRoughnessOrSimpleClearCoat))
 	{
 		SecondRoughnessWeightCodeChunk = Compiler->Constant(0.0f);
 	}
-	if (!SubstrateOperator.Has(ESubstrateBsdfFeature_Anisotropy))
+	if (!SubstrateOperator.Has(ESubstrateBsdfFeature::Anisotropy))
 	{
 		AnisotropyCodeChunk = Compiler->Constant(0.0f);
 	}
-	if (!SubstrateOperator.Has(ESubstrateBsdfFeature_Glint))
+	if (!SubstrateOperator.Has(ESubstrateBsdfFeature::Glint))
 	{
 		GlintValueCodeChunk = Compiler->Constant(1.0f);
 	}
-	if (!SubstrateOperator.Has(ESubstrateBsdfFeature_SpecularProfile))
+	if (!SubstrateOperator.Has(ESubstrateBsdfFeature::SpecularProfile))
 	{
 		SpecularProfileCodeChunk = INDEX_NONE;
 	}
@@ -987,27 +987,27 @@ FSubstrateMaterialComplexity UMaterialExpressionSubstrateSlabBSDF::GetHighestCom
 	// This function returns the highest complexity of a materials.
 	// It will be lowered depending on features enabled per platform.
 
-	uint32 FeatureMask = ESubstrateBsdfFeature_None;
-	if (HasGlint()) 			{ FeatureMask |= ESubstrateBsdfFeature_Glint; } 
-	if (HasAnisotropy()) 		{ FeatureMask |= ESubstrateBsdfFeature_Anisotropy; } 
-	if (HasSpecularProfile()) 	{ FeatureMask |= ESubstrateBsdfFeature_SpecularProfile; } 
-	if (HasEdgeColor()) 		{ FeatureMask |= ESubstrateBsdfFeature_EdgeColor; } 
-	if (HasFuzz()) 				{ FeatureMask |= ESubstrateBsdfFeature_Fuzz; } 
-	if (HasSecondRoughness()) 	{ FeatureMask |= ESubstrateBsdfFeature_SecondRoughnessOrSimpleClearCoat; } 
-	if (HasMFPPluggedIn()) 		{ FeatureMask |= ESubstrateBsdfFeature_MFPPluggedIn; } 
-	if (HasSSS()) 				{ FeatureMask |= ESubstrateBsdfFeature_SSS; } 
+	ESubstrateBsdfFeature FeatureMask = ESubstrateBsdfFeature::None;
+	if (HasGlint()) 			{ FeatureMask |= ESubstrateBsdfFeature::Glint; } 
+	if (HasAnisotropy()) 		{ FeatureMask |= ESubstrateBsdfFeature::Anisotropy; } 
+	if (HasSpecularProfile()) 	{ FeatureMask |= ESubstrateBsdfFeature::SpecularProfile; } 
+	if (HasEdgeColor()) 		{ FeatureMask |= ESubstrateBsdfFeature::EdgeColor; } 
+	if (HasFuzz()) 				{ FeatureMask |= ESubstrateBsdfFeature::Fuzz; } 
+	if (HasSecondRoughness()) 	{ FeatureMask |= ESubstrateBsdfFeature::SecondRoughnessOrSimpleClearCoat; } 
+	if (HasMFPPluggedIn()) 		{ FeatureMask |= ESubstrateBsdfFeature::MFPPluggedIn; } 
+	if (HasSSS()) 				{ FeatureMask |= ESubstrateBsdfFeature::SSS; } 
 
 	FSubstrateMaterialComplexity Out;
 	Out.Reset();
-	if ((FeatureMask & ESubstrateBsdfFeature_ComplexSpecialMask) != 0)
+	if (EnumHasAnyFlags(FeatureMask,ESubstrateBsdfFeature::ComplexSpecialMask))
 	{
 		Out.bIsComplexSpecial = true;
 	}
-	else if ((FeatureMask & ESubstrateBsdfFeature_ComplexMask) != 0)
+	else if (EnumHasAnyFlags(FeatureMask,ESubstrateBsdfFeature::ComplexMask))
 	{
 		// Nothing
 	}
-	else if ((FeatureMask & ESubstrateBsdfFeature_SingleMask) != 0)
+	else if (EnumHasAnyFlags(FeatureMask,ESubstrateBsdfFeature::SingleMask))
 	{
 		Out.bIsSingle = true;
 	}
@@ -1257,15 +1257,15 @@ FSubstrateOperator* UMaterialExpressionSubstrateSlabBSDF::SubstrateGenerateMater
 	FSubstrateOperator& SubstrateOperator = Compiler->SubstrateCompilationRegisterOperator(SUBSTRATE_OPERATOR_BSDF, Compiler->SubstrateTreeStackGetPathUniqueId(), this->MaterialExpressionGuid, Parent, Compiler->SubstrateTreeStackGetParentPathUniqueId());
 	SubstrateOperator.BSDFType = SUBSTRATE_BSDF_TYPE_SLAB;
 
-	if (HasEdgeColor())			{ SubstrateOperator.BSDFFeatures |= ESubstrateBsdfFeature_EdgeColor; }
-	if (HasFuzz())				{ SubstrateOperator.BSDFFeatures |= ESubstrateBsdfFeature_Fuzz; }
-	if (HasSecondRoughness())	{ SubstrateOperator.BSDFFeatures |= ESubstrateBsdfFeature_SecondRoughnessOrSimpleClearCoat; }
-	if (HasSSS())				{ SubstrateOperator.BSDFFeatures |= ESubstrateBsdfFeature_SSS; }
-	if (HasMFPPluggedIn())		{ SubstrateOperator.BSDFFeatures |= ESubstrateBsdfFeature_MFPPluggedIn; }
-	if (HasAnisotropy())		{ SubstrateOperator.BSDFFeatures |= ESubstrateBsdfFeature_Anisotropy; }
-	if (HasGlint())				{ SubstrateOperator.BSDFFeatures |= ESubstrateBsdfFeature_Glint; }
-	if (HasSpecularProfile())	{ SubstrateOperator.BSDFFeatures |= ESubstrateBsdfFeature_SpecularProfile; }
-	if (HasSSSProfile())		{ SubstrateOperator.BSDFFeatures |= ESubstrateBsdfFeature_SecondRoughnessOrSimpleClearCoat; } // If a Slab has a subsurface profile, it will have haziness in order to to support the dual-specular lobe from the profile.
+	if (HasEdgeColor())			{ SubstrateOperator.BSDFFeatures |= ESubstrateBsdfFeature::EdgeColor; }
+	if (HasFuzz())				{ SubstrateOperator.BSDFFeatures |= ESubstrateBsdfFeature::Fuzz; }
+	if (HasSecondRoughness())	{ SubstrateOperator.BSDFFeatures |= ESubstrateBsdfFeature::SecondRoughnessOrSimpleClearCoat; }
+	if (HasSSS())				{ SubstrateOperator.BSDFFeatures |= ESubstrateBsdfFeature::SSS; }
+	if (HasMFPPluggedIn())		{ SubstrateOperator.BSDFFeatures |= ESubstrateBsdfFeature::MFPPluggedIn; }
+	if (HasAnisotropy())		{ SubstrateOperator.BSDFFeatures |= ESubstrateBsdfFeature::Anisotropy; }
+	if (HasGlint())				{ SubstrateOperator.BSDFFeatures |= ESubstrateBsdfFeature::Glint; }
+	if (HasSpecularProfile())	{ SubstrateOperator.BSDFFeatures |= ESubstrateBsdfFeature::SpecularProfile; }
+	if (HasSSSProfile())		{ SubstrateOperator.BSDFFeatures |= ESubstrateBsdfFeature::SecondRoughnessOrSimpleClearCoat; } // If a Slab has a subsurface profile, it will have haziness in order to to support the dual-specular lobe from the profile.
 
 	SubstrateOperator.SubSurfaceType = SubSurfaceType;
 	SubstrateOperator.ThicknessIndex = Compiler->SubstrateThicknessStackGetThicknessIndex();
@@ -1529,7 +1529,7 @@ FSubstrateOperator* UMaterialExpressionSubstrateSimpleClearCoatBSDF::SubstrateGe
 {
 	FSubstrateOperator& SubstrateOperator = Compiler->SubstrateCompilationRegisterOperator(SUBSTRATE_OPERATOR_BSDF, Compiler->SubstrateTreeStackGetPathUniqueId(), this->MaterialExpressionGuid, Parent, Compiler->SubstrateTreeStackGetParentPathUniqueId());
 	SubstrateOperator.BSDFType = SUBSTRATE_BSDF_TYPE_SLAB;
-	SubstrateOperator.BSDFFeatures = ESubstrateBsdfFeature_SecondRoughnessOrSimpleClearCoat; // This node explicitly requires simple clear coat
+	SubstrateOperator.BSDFFeatures = ESubstrateBsdfFeature::SecondRoughnessOrSimpleClearCoat; // This node explicitly requires simple clear coat
 	SubstrateOperator.ThicknessIndex = Compiler->SubstrateThicknessStackGetThicknessIndex();
 	SubstrateOperator.bBSDFWritesEmissive = EmissiveColor.IsConnected();
 	return &SubstrateOperator;
@@ -2207,7 +2207,7 @@ FSubstrateOperator* UMaterialExpressionSubstrateHairBSDF::SubstrateGenerateMater
 {
 	FSubstrateOperator& SubstrateOperator = Compiler->SubstrateCompilationRegisterOperator(SUBSTRATE_OPERATOR_BSDF, Compiler->SubstrateTreeStackGetPathUniqueId(), this->MaterialExpressionGuid, Parent, Compiler->SubstrateTreeStackGetParentPathUniqueId());
 	SubstrateOperator.BSDFType = SUBSTRATE_BSDF_TYPE_HAIR;
-	SubstrateOperator.BSDFFeatures = ESubstrateBsdfFeature_Hair;
+	SubstrateOperator.BSDFFeatures = ESubstrateBsdfFeature::Hair;
 	SubstrateOperator.ThicknessIndex = Compiler->SubstrateThicknessStackGetThicknessIndex();
 	SubstrateOperator.bBSDFWritesEmissive = EmissiveColor.IsConnected();
 	return &SubstrateOperator;
@@ -2328,14 +2328,14 @@ FSubstrateOperator* UMaterialExpressionSubstrateEyeBSDF::SubstrateGenerateMateri
 {
 	FSubstrateOperator& SubstrateOperator = Compiler->SubstrateCompilationRegisterOperator(SUBSTRATE_OPERATOR_BSDF, Compiler->SubstrateTreeStackGetPathUniqueId(), this->MaterialExpressionGuid, Parent, Compiler->SubstrateTreeStackGetParentPathUniqueId());
 	SubstrateOperator.BSDFType = SUBSTRATE_BSDF_TYPE_EYE;
-	SubstrateOperator.BSDFFeatures = ESubstrateBsdfFeature_Eye;
+	SubstrateOperator.BSDFFeatures = ESubstrateBsdfFeature::Eye;
 	if (IrisNormal.IsConnected())
 	{
-		SubstrateOperator.BSDFFeatures |= ESubstrateBsdfFeature_EyeIrisNormalPluggedIn;
+		SubstrateOperator.BSDFFeatures |= ESubstrateBsdfFeature::EyeIrisNormalPluggedIn;
 	}
 	if (IrisPlaneNormal.IsConnected())
 	{
-		SubstrateOperator.BSDFFeatures |= ESubstrateBsdfFeature_EyeIrisTangentPluggedIn;
+		SubstrateOperator.BSDFFeatures |= ESubstrateBsdfFeature::EyeIrisTangentPluggedIn;
 	}
 
 	SubstrateOperator.ThicknessIndex = Compiler->SubstrateThicknessStackGetThicknessIndex();
@@ -4111,11 +4111,11 @@ FSubstrateOperator* UMaterialExpressionSubstrateConvertMaterialAttributes::Subst
 
 	const bool bHasAnisotropy = FMaterialAttributeDefinitionMap::IsAttributeInBitmask(CachedConnectedMaterialAttributesInputs, MP_Anisotropy);
 
-	auto AddDefaultWorstCase = [&](uint32 In)
+	auto AddDefaultWorstCase = [&](ESubstrateBsdfFeature In)
 		{
 			FSubstrateOperator& SlabOperator = Compiler->SubstrateCompilationRegisterOperator(SUBSTRATE_OPERATOR_BSDF_LEGACY, Compiler->SubstrateTreeStackGetPathUniqueId(), ThisExpressionGuid, Parent, Compiler->SubstrateTreeStackGetParentPathUniqueId());
 			SlabOperator.BSDFType = SUBSTRATE_BSDF_TYPE_SLAB;
-			SlabOperator.BSDFFeatures = In | (bHasAnisotropy ? ESubstrateBsdfFeature_Anisotropy : ESubstrateBsdfFeature_None);
+			SlabOperator.BSDFFeatures = In | (bHasAnisotropy ? ESubstrateBsdfFeature::Anisotropy : ESubstrateBsdfFeature::None);
 			SlabOperator.ThicknessIndex = ThicknessIndex;
 
 			return &SlabOperator;
@@ -4128,9 +4128,9 @@ FSubstrateOperator* UMaterialExpressionSubstrateConvertMaterialAttributes::Subst
 	const bool bHasShadingModelFromExpression = FMaterialAttributeDefinitionMap::IsAttributeInBitmask(CachedConnectedMaterialAttributesInputs, MP_ShadingModel) || bShadingModelFromMaterialExpression; // We keep HasShadingModelFromExpression in case all shading models cannot be safely recovered from material functions.
 	if ((ShadingModels.CountShadingModels() > 1) || bHasShadingModelFromExpression)
 	{
-		uint32 Features = ESubstrateBsdfFeature_SSS | ESubstrateBsdfFeature_Fuzz;
-		if (ShadingModels.HasShadingModel(MSM_Eye)) { Features |= ESubstrateBsdfFeature_Eye; }
-		if (ShadingModels.HasShadingModel(MSM_Hair)) { Features |= ESubstrateBsdfFeature_Hair; }
+		ESubstrateBsdfFeature Features = ESubstrateBsdfFeature::SSS | ESubstrateBsdfFeature::Fuzz;
+		if (ShadingModels.HasShadingModel(MSM_Eye)) { Features |= ESubstrateBsdfFeature::Eye; }
+		if (ShadingModels.HasShadingModel(MSM_Hair)) { Features |= ESubstrateBsdfFeature::Hair; }
 		return AddDefaultWorstCase(Features);
 	}
 	else
@@ -4146,45 +4146,45 @@ FSubstrateOperator* UMaterialExpressionSubstrateConvertMaterialAttributes::Subst
 		}
 		else if (ShadingModels.HasShadingModel(MSM_DefaultLit))
 		{
-			return AddDefaultWorstCase(ESubstrateBsdfFeature_None);
+			return AddDefaultWorstCase(ESubstrateBsdfFeature::None);
 		}
 		else if (ShadingModels.HasShadingModel(MSM_ThinTranslucent))
 		{
-			return AddDefaultWorstCase(ESubstrateBsdfFeature_None);
+			return AddDefaultWorstCase(ESubstrateBsdfFeature::None);
 		}
 		else if (ShadingModels.HasShadingModel(MSM_SubsurfaceProfile))
 		{
-			return AddDefaultWorstCase(ESubstrateBsdfFeature_SSS);
+			return AddDefaultWorstCase(ESubstrateBsdfFeature::SSS);
 		}
 		else if (ShadingModels.HasShadingModel(MSM_Subsurface))
 		{
-			return AddDefaultWorstCase(ESubstrateBsdfFeature_SSS);
+			return AddDefaultWorstCase(ESubstrateBsdfFeature::SSS);
 		}
 		else if (ShadingModels.HasShadingModel(MSM_TwoSidedFoliage))
 		{
-			return AddDefaultWorstCase(ESubstrateBsdfFeature_SSS);
+			return AddDefaultWorstCase(ESubstrateBsdfFeature::SSS);
 		}
 		else if (ShadingModels.HasShadingModel(MSM_PreintegratedSkin))
 		{
-			return AddDefaultWorstCase(ESubstrateBsdfFeature_SSS);
+			return AddDefaultWorstCase(ESubstrateBsdfFeature::SSS);
 		}
 		else if (ShadingModels.HasShadingModel(MSM_Cloth))
 		{
-			return AddDefaultWorstCase(ESubstrateBsdfFeature_Fuzz);
+			return AddDefaultWorstCase(ESubstrateBsdfFeature::Fuzz);
 		}
 		else if (ShadingModels.HasShadingModel(MSM_ClearCoat))
 		{
 			FSubstrateOperator& Operator = Compiler->SubstrateCompilationRegisterOperator(SUBSTRATE_OPERATOR_BSDF_LEGACY, Compiler->SubstrateTreeStackGetPathUniqueId(), ThisExpressionGuid, Parent, Compiler->SubstrateTreeStackGetParentPathUniqueId());
 			Operator.BSDFType = SUBSTRATE_BSDF_TYPE_SLAB;
 			Operator.ThicknessIndex = ThicknessIndex;
-			Operator.BSDFFeatures = ESubstrateBsdfFeature_SecondRoughnessOrSimpleClearCoat | (bHasAnisotropy ? ESubstrateBsdfFeature_Anisotropy : ESubstrateBsdfFeature_None);
+			Operator.BSDFFeatures = ESubstrateBsdfFeature::SecondRoughnessOrSimpleClearCoat | (bHasAnisotropy ? ESubstrateBsdfFeature::Anisotropy : ESubstrateBsdfFeature::None);
 			return &Operator;
 		}
 		else if (ShadingModels.HasShadingModel(MSM_Hair))
 		{
 			FSubstrateOperator& Operator = Compiler->SubstrateCompilationRegisterOperator(SUBSTRATE_OPERATOR_BSDF_LEGACY, Compiler->SubstrateTreeStackGetPathUniqueId(), ThisExpressionGuid, Parent, Compiler->SubstrateTreeStackGetParentPathUniqueId());
 			Operator.BSDFType = SUBSTRATE_BSDF_TYPE_HAIR;
-			Operator.BSDFFeatures = ESubstrateBsdfFeature_Hair;
+			Operator.BSDFFeatures = ESubstrateBsdfFeature::Hair;
 			Operator.ThicknessIndex = ThicknessIndex;
 			return &Operator;
 		}
@@ -4192,7 +4192,7 @@ FSubstrateOperator* UMaterialExpressionSubstrateConvertMaterialAttributes::Subst
 		{
 			FSubstrateOperator& Operator = Compiler->SubstrateCompilationRegisterOperator(SUBSTRATE_OPERATOR_BSDF_LEGACY, Compiler->SubstrateTreeStackGetPathUniqueId(), ThisExpressionGuid, Parent, Compiler->SubstrateTreeStackGetParentPathUniqueId());
 			Operator.BSDFType = SUBSTRATE_BSDF_TYPE_EYE;
-			Operator.BSDFFeatures = ESubstrateBsdfFeature_Eye;
+			Operator.BSDFFeatures = ESubstrateBsdfFeature::Eye;
 			Operator.ThicknessIndex = ThicknessIndex;
 			return &Operator;
 		}
diff --git a/Engine/Source/Runtime/Engine/Public/Rendering/SubstrateMaterialShared.cpp b/Engine/Source/Runtime/Engine/Public/Rendering/SubstrateMaterialShared.cpp
--- a/Engine/Source/Runtime/Engine/Public/Rendering/SubstrateMaterialShared.cpp
+++ b/Engine/Source/Runtime/Engine/Public/Rendering/SubstrateMaterialShared.cpp
@@ -47,7 +47,7 @@ FSubstrateOperator::FSubstrateOperator()
 	BSDFIndex = INDEX_NONE;
 	BSDFType = 0;
 	BSDFRegisteredSharedLocalBasis = FSubstrateRegisteredSharedLocalBasis();
-	BSDFFeatures = 0;
+	BSDFFeatures = ESubstrateBsdfFeature::None;
 	SubUsage = SUBSTRATE_OPERATOR_SUBUSAGE_NONE;
 	SubSurfaceType = uint8(EMaterialSubSurfaceType::MSS_None);
 
diff --git a/Engine/Source/Runtime/Engine/Public/Rendering/SubstrateMaterialShared.h b/Engine/Source/Runtime/Engine/Public/Rendering/SubstrateMaterialShared.h
--- a/Engine/Source/Runtime/Engine/Public/Rendering/SubstrateMaterialShared.h
+++ b/Engine/Source/Runtime/Engine/Public/Rendering/SubstrateMaterialShared.h
@@ -6,6 +6,7 @@
 #include "SubstrateDefinitions.h"
 #include "Serialization/MemoryImage.h"
 
+#include "SubstrateMaterialShared.generated.h"
 
 // Structures in this files are only used a compilation result return by the compiler.
 // They are also used to present material information in the editor UI.
@@ -24,39 +25,42 @@ public:
 	LAYOUT_FIELD_EDITORONLY(uint8, GraphSharedLocalBasisIndex);
 };
 
-enum ESubstrateBsdfFeature
+UENUM()
+enum class ESubstrateBsdfFeature : uint16
 {
-	ESubstrateBsdfFeature_None                             = 0u,
-	ESubstrateBsdfFeature_SSS                              = 1u<<0u,
-	ESubstrateBsdfFeature_MFPPluggedIn                     = 1u<<1u,
-	ESubstrateBsdfFeature_EdgeColor                        = 1u<<2u,
-	ESubstrateBsdfFeature_Fuzz                             = 1u<<3u,
-	ESubstrateBsdfFeature_SecondRoughnessOrSimpleClearCoat = 1u<<4u,
-	ESubstrateBsdfFeature_Anisotropy                       = 1u<<5u,
-	ESubstrateBsdfFeature_Glint                            = 1u<<6u,
-	ESubstrateBsdfFeature_SpecularProfile                  = 1u<<7u,
-	ESubstrateBsdfFeature_Eye                              = 1u<<8u,
-	ESubstrateBsdfFeature_EyeIrisNormalPluggedIn           = 1u<<9u,
-	ESubstrateBsdfFeature_EyeIrisTangentPluggedIn          = 1u<<10u,
-	ESubstrateBsdfFeature_Hair                             = 1u<<11u,
+	None                             = 0u,
+	SSS                              = 1u<<0u,
+	MFPPluggedIn                     = 1u<<1u,
+	EdgeColor                        = 1u<<2u,
+	Fuzz                             = 1u<<3u,
+	SecondRoughnessOrSimpleClearCoat = 1u<<4u,
+	Anisotropy                       = 1u<<5u,
+	Glint                            = 1u<<6u,
+	SpecularProfile                  = 1u<<7u,
+	Eye                              = 1u<<8u,
+	EyeIrisNormalPluggedIn           = 1u<<9u,
+	EyeIrisTangentPluggedIn          = 1u<<10u,
+	Hair                             = 1u<<11u,
 
 	// Complexity masks
-	ESubstrateBsdfFeature_SingleMask =	
-		  ESubstrateBsdfFeature_EdgeColor
-		| ESubstrateBsdfFeature_Fuzz
-		| ESubstrateBsdfFeature_SSS
-		| ESubstrateBsdfFeature_SecondRoughnessOrSimpleClearCoat
-		| ESubstrateBsdfFeature_MFPPluggedIn,
-
-	ESubstrateBsdfFeature_ComplexMask = 
-		  ESubstrateBsdfFeature_Anisotropy
-		| ESubstrateBsdfFeature_SpecularProfile
-		| ESubstrateBsdfFeature_Eye
-		| ESubstrateBsdfFeature_Hair,
-
-	ESubstrateBsdfFeature_ComplexSpecialMask = 
-		  ESubstrateBsdfFeature_Glint
+	SingleMask =	
+		  EdgeColor
+		| Fuzz
+		| SSS
+		| SecondRoughnessOrSimpleClearCoat
+		| MFPPluggedIn,
+
+	ComplexMask = 
+		  Anisotropy
+		| SpecularProfile
+		| Eye
+		| Hair,
+
+	ComplexSpecialMask = 
+		  Glint
 };
+ENUM_CLASS_FLAGS(ESubstrateBsdfFeature)
+DECLARE_INTRINSIC_TYPE_LAYOUT(ESubstrateBsdfFeature);
 
 struct FSubstrateOperator
 {
@@ -82,7 +86,7 @@ public:
 	LAYOUT_FIELD_EDITORONLY(int32, BSDFIndex);		// Index in the array of BSDF if a BSDF operator
 	LAYOUT_FIELD_EDITORONLY(int32, BSDFType);
 	LAYOUT_FIELD_EDITORONLY(FSubstrateRegisteredSharedLocalBasis, BSDFRegisteredSharedLocalBasis);
-	LAYOUT_FIELD_EDITORONLY(uint32,BSDFFeatures);
+	LAYOUT_FIELD_EDITORONLY(ESubstrateBsdfFeature, BSDFFeatures);
 	LAYOUT_FIELD_EDITORONLY(uint8, SubUsage);		// Sometimes, Unlit or Weight operators are used to transport data for other meaning (e.g. Light Function or ConvertToDecal)
 	LAYOUT_FIELD_EDITORONLY(uint8, SubSurfaceType);
 
@@ -108,16 +112,7 @@ public:
 	bool Has(ESubstrateBsdfFeature In) const 
 	{ 
 	#if WITH_EDITOR
-		return (BSDFFeatures & In) != 0; 
-	#else
-		return false;
-	#endif
-	}
-
-	bool Has(uint32 In) const 
-	{ 
-	#if WITH_EDITOR
-		return (BSDFFeatures & In) != 0; 
+		return EnumHasAnyFlags(BSDFFeatures, In); 
 	#else
 		return false;
 	#endif
