Make sure all streams are added for default vertex element list of FGeometryCacheVertexVertexFactory

diff --git a/Engine/Source/Runtime/Engine/Private/GeometryCacheVertexFactory.cpp b/Engine/Source/Runtime/Engine/Private/GeometryCacheVertexFactory.cpp
--- a/Engine/Source/Runtime/Engine/Private/GeometryCacheVertexFactory.cpp
+++ b/Engine/Source/Runtime/Engine/Private/GeometryCacheVertexFactory.cpp
@@ -267,8 +267,22 @@ void FGeometryCacheVertexVertexFactory::GetPSOPrecacheVertexFetchElements(EVerte
 		switch (VertexInputStreamType)
 		{
 			case EVertexInputStreamType::Default:
-			{
-				OutElements.Add(FVertexElement(1, 0, VET_UInt, 13, sizeof(uint32), true));
+			{			
+				// Make sure all required elements are added because manual vertex fetch is not supported (see GetVertexElements above)
+				uint8 TangentBasisAttributes[2] = { 1, 2 };
+				for (int32 AxisIndex = 0; AxisIndex < 2; AxisIndex++)
+				{
+					OutElements.Add(FVertexElement(TangentBasisAttributes[AxisIndex], 0, VET_PackedNormal, TangentBasisAttributes[AxisIndex], sizeof(FPackedNormal), false));
+				}
+				OutElements.Add(FVertexElement(3, 0, VET_Color, 3, sizeof(FColor), false));
+				OutElements.Add(FVertexElement(4, 0, VET_Float3, 4, sizeof(FVector3f), false));
+				for (int32 CoordinateIndex = 0; CoordinateIndex < MAX_STATIC_TEXCOORDS / 2; CoordinateIndex++)
+				{
+					OutElements.Add(FVertexElement(5 + CoordinateIndex, 0, VET_Float2, 5 + CoordinateIndex, sizeof(FVector2f), false));
+				}
+
+				OutElements.Add(FVertexElement(9, 0, VET_UInt, 13, sizeof(uint32), true));
+
 				break;
 			}
 			case EVertexInputStreamType::PositionOnly:
