[Sound Attenuation] Fix == operator on FSoundAttenuationSettings.

diff --git a/Engine/Source/Runtime/Engine/Private/SoundAttenuation.cpp b/Engine/Source/Runtime/Engine/Private/SoundAttenuation.cpp
--- a/Engine/Source/Runtime/Engine/Private/SoundAttenuation.cpp
+++ b/Engine/Source/Runtime/Engine/Private/SoundAttenuation.cpp
@@ -189,7 +189,8 @@ bool FSoundAttenuationSettings::operator==(const FSoundAttenuationSettings& Othe
 			&& DistanceAlgorithm	    == Other.DistanceAlgorithm
 			&& AttenuationShape		    == Other.AttenuationShape
 			&& bAttenuateWithLPF		== Other.bAttenuateWithLPF
-			&& LPFRadiusMin				== Other.LPFRadiusMax
+			&& LPFRadiusMin				== Other.LPFRadiusMin
+			&& LPFRadiusMax				== Other.LPFRadiusMax
 			&& FalloffDistance		    == Other.FalloffDistance
 			&& AttenuationShapeExtents	== Other.AttenuationShapeExtents
 			&& SpatializationAlgorithm == Other.SpatializationAlgorithm
