Fix anisotropy rendering with material using Substrate's nodes when BlendableGBuffer is enabled.

This CL pipes BsdfFeature into compilation output to that it can be queried at runtime. This allows to know what feature a particular is using, and to enable it on the rendering side.

diff --git a/Engine/Shaders/Private/BasePassPixelShader.usf b/Engine/Shaders/Private/BasePassPixelShader.usf
--- a/Engine/Shaders/Private/BasePassPixelShader.usf
+++ b/Engine/Shaders/Private/BasePassPixelShader.usf
@@ -1810,6 +1810,15 @@ void FPixelShaderInOut_MainPS(
 		GBuffer.GBufferAO = MaterialAO;
 		GBuffer.PerObjectGBufferData = GetPrimitive_PerObjectGBufferData(MaterialParameters.PrimitiveId);
 		GBuffer.Depth = MaterialParameters.ScreenPosition.w;
+
+		#if USES_GBUFFER && !IS_NANITE_PASS
+		// This requires cleanup. Shader code that uses GBuffer.SelectiveOutputMask expects the outputmask to be in
+		// bits [4:7], but it gets packed as bits [0:3] in the flexible gbuffer since we might move it around.
+		if (Export.Anisotropy != 0)
+		{
+			GBuffer.SelectiveOutputMask = HAS_ANISOTROPY_MASK >> 4;
+		}
+		#endif
 		
 		// Patch material parameter to not have to change SetGBufferForShadingModel for now
 		MaterialParameters.WorldNormal = Export.WorldNormal;
diff --git a/Engine/Shaders/Public/ShaderVersion.ush b/Engine/Shaders/Public/ShaderVersion.ush
--- a/Engine/Shaders/Public/ShaderVersion.ush
+++ b/Engine/Shaders/Public/ShaderVersion.ush
@@ -3,4 +3,4 @@
 // Each time the console command is executed it generates a new GUID. As this file is included
 // in Platform.ush (which should be included in any shader) it allows to invalidate the shader DDC.
 // 
-#pragma message("UESHADERMETADATA_VERSION BBA73A42-4C16-45C4-A6D0-47E57DCFD9F4")
+#pragma message("UESHADERMETADATA_VERSION 794E2BBF-A6E7-43C2-BD06-A470603C59D5")
diff --git a/Engine/Source/Runtime/Core/Private/UObject/FortniteMainBranchObjectVersion.cpp b/Engine/Source/Runtime/Core/Private/UObject/FortniteMainBranchObjectVersion.cpp
--- a/Engine/Source/Runtime/Core/Private/UObject/FortniteMainBranchObjectVersion.cpp
+++ b/Engine/Source/Runtime/Core/Private/UObject/FortniteMainBranchObjectVersion.cpp
@@ -8,12 +8,12 @@ TMap<FGuid, FGuid> FFortniteMainBranchObjectVersion::GetSystemGuids()
 
 	SystemGuids.Add(DevGuids.GLOBALSHADERMAP_DERIVEDDATA_VER, FGuid("125CF21C881744418BD4A614B3B125FF"));
 	SystemGuids.Add(DevGuids.LANDSCAPE_MOBILE_COOK_VERSION, FGuid("32D02EF867C74B71A0D4E0FA41392732"));
-	SystemGuids.Add(DevGuids.MATERIALSHADERMAP_DERIVEDDATA_VER, FGuid("04B8C817071A453C9EBEFBFB51D5F785"));
+	SystemGuids.Add(DevGuids.MATERIALSHADERMAP_DERIVEDDATA_VER, FGuid("F051877974B44D409435F321BDB39DBB"));
 	SystemGuids.Add(DevGuids.NANITE_DERIVEDDATA_VER, FGuid("0092F35BA3894D8DBA29E99C9C8866A2"));
 	SystemGuids.Add(DevGuids.NIAGARASHADERMAP_DERIVEDDATA_VER, FGuid("6360A977062842A29ED30E3A7ACB0E64"));
 	SystemGuids.Add(DevGuids.Niagara_LatestScriptCompileVersion, FGuid("97F0FD7149D90D42B65D6557AF477D1A"));
 	SystemGuids.Add(DevGuids.SkeletalMeshDerivedDataVersion, FGuid("C998F2F3B2884A02B7E290B8C25835F0"));
 	SystemGuids.Add(DevGuids.STATICMESH_DERIVEDDATA_VER, FGuid("3ABF32149F9D4D838750368AB95573BE"));
-	SystemGuids.Add(DevGuids.MaterialTranslationDDCVersion, FGuid("439B1428AAED4BA489BF3D886E461F06"));
+	SystemGuids.Add(DevGuids.MaterialTranslationDDCVersion, FGuid("F051877974B44D409435F321BDB39DBB"));
 	return SystemGuids;
 }
diff --git a/Engine/Source/Runtime/Engine/Private/Materials/HLSLMaterialTranslator.cpp b/Engine/Source/Runtime/Engine/Private/Materials/HLSLMaterialTranslator.cpp
--- a/Engine/Source/Runtime/Engine/Private/Materials/HLSLMaterialTranslator.cpp
+++ b/Engine/Source/Runtime/Engine/Private/Materials/HLSLMaterialTranslator.cpp
@@ -13940,6 +13940,7 @@ bool FHLSLMaterialTranslator::FSubstrateCompilationContext::SubstrateGenerateDer
 				Compiler->MaterialCompilationOutput.SubstrateMaterialCompilationOutput.SubstrateMaterialType = SubstrateMaterialComplexity.SubstrateMaterialType();
 				Compiler->MaterialCompilationOutput.SubstrateMaterialCompilationOutput.SubstrateClosureCount = SubstrateMaterialEffectiveClosureCount;
 				Compiler->MaterialCompilationOutput.SubstrateMaterialCompilationOutput.SubstrateUintPerPixel = uint8(FMath::Clamp(RequestedSizeInUint, 0u, 0xFF));
+				Compiler->MaterialCompilationOutput.SubstrateMaterialCompilationOutput.SubstrateMaterialBsdfFeatures = SubstrateMaterialBsdfFeatures;
 
 #if WITH_EDITOR
 				Compiler->MaterialCompilationOutput.SubstrateMaterialCompilationOutput.SharedLocalBasesCount = 0; // FinalUsedSharedLocalBasesCount is not valid yet
@@ -16643,6 +16644,8 @@ void FHLSLMaterialTranslator::PrepareEnvironmentDefines()
 
 				SubstrateMaterialDescription += SubstrateMaterialContextDescription;
 
+				MaterialCompilationOutput.SubstrateMaterialCompilationOutput.SubstrateMaterialBsdfFeatures = MaterialBsdfFeatures;
+
 				if (SubstrateCompilationContextIndex == ESubstrateCompilationContext::SCC_Default)
 				{
 					MaterialCompilationOutput.SubstrateMaterialCompilationOutput.SharedLocalBasesCount = SubstrateCtx.FinalUsedSharedLocalBasesCount;
diff --git a/Engine/Source/Runtime/Engine/Private/Materials/MaterialShared.cpp b/Engine/Source/Runtime/Engine/Private/Materials/MaterialShared.cpp
--- a/Engine/Source/Runtime/Engine/Private/Materials/MaterialShared.cpp
+++ b/Engine/Source/Runtime/Engine/Private/Materials/MaterialShared.cpp
@@ -1142,13 +1142,13 @@ uint8 FMaterial::GetRuntimeVirtualTextureOutputAttibuteMask_RenderThread() const
 
 bool FMaterial::MaterialUsesAnisotropy_GameThread() const
 {
-	return GameThreadShaderMap ? GameThreadShaderMap->UsesAnisotropy() : false;
+	return GameThreadShaderMap ? (GameThreadShaderMap->UsesAnisotropy() || EnumHasAnyFlags(GameThreadShaderMap->GetSubstrateMaterialBsdfFeatures(), ESubstrateBsdfFeature::Anisotropy)) : false;
 }
 
 bool FMaterial::MaterialUsesAnisotropy_RenderThread() const
 {
 	check(IsInParallelRenderingThread());
-	return RenderingThreadShaderMap ? RenderingThreadShaderMap->UsesAnisotropy() : false;
+	return RenderingThreadShaderMap ? (RenderingThreadShaderMap->UsesAnisotropy() || EnumHasAnyFlags(RenderingThreadShaderMap->GetSubstrateMaterialBsdfFeatures(), ESubstrateBsdfFeature::Anisotropy)) : false;
 }
 
 bool FMaterial::MaterialIsLightFunctionAtlasCompatible_GameThread() const
@@ -1206,6 +1206,17 @@ bool FMaterial::MaterialGetSubstrateUsesComplexSpecialRenderPath_RenderThread()
 	return RenderingThreadShaderMap ? RenderingThreadShaderMap->GetSubstrateUsesComplexSpecialRenderPath() : false;
 }
 
+ESubstrateBsdfFeature FMaterial::MaterialGetSubstrateMaterialBsdfFeatures_GameThread() const
+{
+	return GameThreadShaderMap ? GameThreadShaderMap->GetSubstrateMaterialBsdfFeatures() : ESubstrateBsdfFeature::None;
+}
+
+ESubstrateBsdfFeature FMaterial::MaterialGetSubstrateMaterialBsdfFeatures_RenderThread() const
+{
+	check(IsInParallelRenderingThread());
+	return RenderingThreadShaderMap ? RenderingThreadShaderMap->GetSubstrateMaterialBsdfFeatures() : ESubstrateBsdfFeature::None;
+}
+
 void FMaterial::SetGameThreadShaderMap(FMaterialShaderMap* InMaterialShaderMap)
 {
 	checkSlow(IsInGameThread() || IsInAsyncLoadingThread());
@@ -5725,7 +5736,7 @@ FMaterialShaderParametersBase::FMaterialShaderParametersBase(const FMaterial* In
 		bHasMetallicConnected = InMaterial->HasMaterialPropertyConnected(MP_Metallic);
 		bHasEmissiveColorConnected = InMaterial->HasMaterialPropertyConnected(MP_EmissiveColor);
 		bHasAmbientOcclusionConnected = InMaterial->HasMaterialPropertyConnected(MP_AmbientOcclusion);
-		bHasAnisotropyConnected = InMaterial->HasMaterialPropertyConnected(MP_Anisotropy);
+		bHasAnisotropyConnected = InMaterial->HasMaterialPropertyConnected(MP_Anisotropy) || EnumHasAnyFlags(InMaterial->MaterialGetSubstrateMaterialBsdfFeatures_GameThread(), ESubstrateBsdfFeature::Anisotropy);
 	}
 	else
 	{
diff --git a/Engine/Source/Runtime/Engine/Public/MaterialShared.h b/Engine/Source/Runtime/Engine/Public/MaterialShared.h
--- a/Engine/Source/Runtime/Engine/Public/MaterialShared.h
+++ b/Engine/Source/Runtime/Engine/Public/MaterialShared.h
@@ -1673,7 +1673,8 @@ public:
 	uint8 GetSubstrateClosureCount() const { return GetSubstrateMaterialCompilationOutput().SubstrateClosureCount; }
 	uint8 GetSubstrateUintPerPixel() const { return GetSubstrateMaterialCompilationOutput().SubstrateUintPerPixel; }
 	bool GetSubstrateUsesComplexSpecialRenderPath() const { return GetSubstrateMaterialCompilationOutput().SubstrateMaterialType == SUBSTRATE_MATERIAL_TYPE_COMPLEX_SPECIAL; }
-	
+	ESubstrateBsdfFeature GetSubstrateMaterialBsdfFeatures() const { return ESubstrateBsdfFeature(GetSubstrateMaterialCompilationOutput().SubstrateMaterialBsdfFeatures); }
+
 #if WITH_EDITOR
 	uint32 GetNumUsedUVScalars() const { return GetContent()->MaterialCompilationOutput.NumUsedUVScalars; }
 	uint32 GetNumUsedCustomInterpolatorScalars() const { return GetContent()->MaterialCompilationOutput.NumUsedCustomInterpolatorScalars; }
@@ -2578,6 +2579,10 @@ public:
 	ENGINE_API bool MaterialGetSubstrateUsesComplexSpecialRenderPath_GameThread() const;
 	ENGINE_API bool MaterialGetSubstrateUsesComplexSpecialRenderPath_RenderThread() const;
 
+	/** Get Substrate material BSDFs features. */
+	ENGINE_API ESubstrateBsdfFeature MaterialGetSubstrateMaterialBsdfFeatures_GameThread() const;
+	ENGINE_API ESubstrateBsdfFeature MaterialGetSubstrateMaterialBsdfFeatures_RenderThread() const;
+
 	class FMaterialShaderMap* GetGameThreadShaderMap() const 
 	{ 
 		checkSlow(IsInGameThread() || IsInParallelGameThread() || IsInAsyncLoadingThread());
diff --git a/Engine/Source/Runtime/Engine/Public/Rendering/SubstrateMaterialShared.cpp b/Engine/Source/Runtime/Engine/Public/Rendering/SubstrateMaterialShared.cpp
--- a/Engine/Source/Runtime/Engine/Public/Rendering/SubstrateMaterialShared.cpp
+++ b/Engine/Source/Runtime/Engine/Public/Rendering/SubstrateMaterialShared.cpp
@@ -9,6 +9,7 @@ FSubstrateMaterialCompilationOutput::FSubstrateMaterialCompilationOutput()
 	: SubstrateMaterialType(0)
 	, SubstrateClosureCount(0)
 	, SubstrateUintPerPixel(0)
+	, SubstrateMaterialBsdfFeatures(ESubstrateBsdfFeature::None)
 #if WITH_EDITOR
 	, SubstrateMaterialDescription()
 	, SharedLocalBasesCount(0)
diff --git a/Engine/Source/Runtime/Engine/Public/Rendering/SubstrateMaterialShared.h b/Engine/Source/Runtime/Engine/Public/Rendering/SubstrateMaterialShared.h
--- a/Engine/Source/Runtime/Engine/Public/Rendering/SubstrateMaterialShared.h
+++ b/Engine/Source/Runtime/Engine/Public/Rendering/SubstrateMaterialShared.h
@@ -154,6 +154,9 @@ public:
 	/** Substrate uint per pixel, at compile time (0-255) */
 	LAYOUT_FIELD(uint8, SubstrateUintPerPixel);
 
+	/** Substrate feature used across the material (Aniso, SSS, Fuzz, ...)*/
+	LAYOUT_FIELD(ESubstrateBsdfFeature, SubstrateMaterialBsdfFeatures);
+
 	////
 	//// The following data is only needed when compiling with the editor.
 	////
