Ignore the invalid path warnings from AssetRegistry.ScanSynchronous when asking for levels AR properties.
In the case where the level mount point is not valid, just return false instead of logging a warning.

diff --git a/Engine/Source/Runtime/Engine/Private/Level.cpp b/Engine/Source/Runtime/Engine/Private/Level.cpp
--- a/Engine/Source/Runtime/Engine/Private/Level.cpp
+++ b/Engine/Source/Runtime/Engine/Private/Level.cpp
@@ -2798,8 +2798,9 @@ namespace LevelAssetRegistryHelper
 	{
 		IAssetRegistry& AssetRegistry = FModuleManager::LoadModuleChecked<FAssetRegistryModule>(TEXT("AssetRegistry")).Get();
 
-		// Always ask for scan in case no initial scan was done (Commandlets, IsRunningGame) or if AssetRegistry isn't done scanning
-		AssetRegistry.ScanSynchronous({ }, { LevelPackage.ToString()});
+		// Always ask for scan in case no initial scan was done (Commandlets, IsRunningGame) or if AssetRegistry isn't done scanning.
+		// Also ignore the invalid path warning as we just want to return false for unmounted level paths.
+		AssetRegistry.ScanSynchronous({ }, { LevelPackage.ToString()}, UE::AssetRegistry::EScanFlags::IgnoreInvalidPathWarning);
 		
 		TArray<FAssetData> LevelPackageAssets;
 		AssetRegistry.GetAssetsByPackageName(LevelPackage, LevelPackageAssets, true);
