NetDriver
* Added RawPing value to NetConnection and exported it to CSV stats
* This is the non-averaged ping value of the current frame and it stores the highest ping value if multiple packets are processed in a single update.

#rnx

diff --git a/Engine/Source/Runtime/Engine/Classes/Engine/NetConnection.h b/Engine/Source/Runtime/Engine/Classes/Engine/NetConnection.h
--- a/Engine/Source/Runtime/Engine/Classes/Engine/NetConnection.h
+++ b/Engine/Source/Runtime/Engine/Classes/Engine/NetConnection.h
@@ -541,6 +541,9 @@ public:
 	/** The average frame delta time over the last 1 second period.*/
 	double			AverageFrameTime;
 
+	/** The raw ping value of the last packet the connection processed */
+	double			RawPingInSeconds = 0.0;
+
 	/** Current jitter for this connection in milliseconds */
 	float			GetAverageJitterInMS() const { return AverageJitterInMS; }
 
@@ -626,6 +629,7 @@ public:
 	int32			InPacketId;						// Full incoming packet index.
 	int32			OutPacketId;					// Most recently sent packet.
 	int32 			OutAckPacketId;					// Most recently acked outgoing packet.
+	int32			OutFirstAckPacketId = -1;		// The id of the first ACK we processed each frame.
 
 	bool			bLastHasServerFrameTime;
 
diff --git a/Engine/Source/Runtime/Engine/Private/NetConnection.cpp b/Engine/Source/Runtime/Engine/Private/NetConnection.cpp
--- a/Engine/Source/Runtime/Engine/Private/NetConnection.cpp
+++ b/Engine/Source/Runtime/Engine/Private/NetConnection.cpp
@@ -2192,6 +2192,9 @@ void UNetConnection::PostTickDispatch()
 
 		FlushPacketOrderCache(/*bFlushWholeCache=*/true);
 		PacketAnalytics.Tick();
+
+		// Make sure to clear the AckPacketId for the next frame
+		OutFirstAckPacketId = -1;
 	}
 
 	if (RPCDoS.IsValid() && bIsServer && !IsReplay())
@@ -2637,12 +2640,18 @@ void UNetConnection::ReceivedAck(int32 AckPacketId, FChannelsToClose& OutChannel
 	// Advance OutAckPacketId
 	OutAckPacketId = AckPacketId;
 
+	// Store the first Ack we received each frame
+	if (OutFirstAckPacketId == -1)
+	{
+		OutFirstAckPacketId = AckPacketId;
+	}
+
 	// Process the bunch.
 	LastRecvAckTimestamp = Driver->GetElapsedTime();
 
 	PacketAnalytics.TrackAck(AckPacketId);
 
-	if (PackageMap != NULL)
+	if (PackageMap)
 	{
 		PackageMap->ReceivedAck( AckPacketId );
 	}
@@ -2903,41 +2912,59 @@ bool UNetConnection::ReadPacketInfo(FBitReader& Reader, bool bHasPacketInfoPaylo
 		return false;
 	}
 
-	// Update ping
-	// At this time we have updated OutAckPacketId to the latest received ack.
-	const int32 Index = OutAckPacketId & (UE_ARRAY_COUNT(OutLagPacketId)-1);
+	// Find the received time of the packet
+	double PacketReceiveTime = 0.0;
+	FTimespan& RecvTimespan = LastOSReceiveTime.Timestamp;
 
-	if ( OutLagPacketId[Index] == OutAckPacketId )
+	if (!RecvTimespan.IsZero() && Driver != nullptr && CVarPingUsePacketRecvTime.GetValueOnAnyThread())
 	{
-		OutLagPacketId[Index] = -1;		// Only use the ack once
-
-#if !UE_BUILD_SHIPPING
-		if ( CVarPingDisplayServerTime.GetValueOnAnyThread() > 0 )
+		if (bIsOSReceiveTimeLocal)
 		{
-			UE_LOG( LogNetTraffic, Warning, TEXT( "ServerFrameTime: %2.2f" ), ServerFrameTime * 1000.0f );
+			PacketReceiveTime = RecvTimespan.GetTotalSeconds();
 		}
-#endif
+		else if (ISocketSubsystem* SocketSubsystem = Driver->GetSocketSubsystem())
+		{
+			PacketReceiveTime = SocketSubsystem->TranslatePacketTimestamp(LastOSReceiveTime);
+		}
+	}
+
+	// Use FApp's time because it is set closer to the beginning of the frame - we don't care about the time so far of the current frame to process the packet
+	const double CurrentTime = (PacketReceiveTime != 0.0 ? PacketReceiveTime : FApp::GetCurrentTime());
+
+	// Update ping
 
-		double PacketReceiveTime = 0.0;
-		FTimespan& RecvTimespan = LastOSReceiveTime.Timestamp;
+	// Find the ping of the first packet we acked (highest ping)
+	if (OutFirstAckPacketId != -1)
+	{
+		const int32 FirstAckIndex = OutFirstAckPacketId & (UE_ARRAY_COUNT(OutLagPacketId) - 1);
 
-		if (!RecvTimespan.IsZero() && Driver != nullptr && CVarPingUsePacketRecvTime.GetValueOnAnyThread())
+		if (OutLagPacketId[FirstAckIndex] == OutFirstAckPacketId)
 		{
-			if (bIsOSReceiveTimeLocal)
-			{
-				PacketReceiveTime = RecvTimespan.GetTotalSeconds();
-			}
-			else if (ISocketSubsystem* SocketSubsystem = Driver->GetSocketSubsystem())
+			if (OutFirstAckPacketId != OutAckPacketId)
 			{
-				PacketReceiveTime = SocketSubsystem->TranslatePacketTimestamp(LastOSReceiveTime);
+				OutLagPacketId[FirstAckIndex] = -1;
 			}
+
+			const double PacketSendTime = OutLagTime[FirstAckIndex];
+			const double RTT = CurrentTime - PacketSendTime;
+			const double RTTExclFrame = FMath::Max(RTT - ServerFrameTime, 0.0);
+
+			RawPingInSeconds = RTTExclFrame;
 		}
 
+		OutFirstAckPacketId = -1;
+	}
+
+	// Find the ping of the last packet we acked (lowest ping)
+	const int32 LastAckIndex =  OutAckPacketId & (UE_ARRAY_COUNT(OutLagPacketId) - 1);
+
+	if (OutLagPacketId[LastAckIndex] == OutAckPacketId)
+	{
+		OutLagPacketId[LastAckIndex] = -1;		// Only use the ack once
 
-		// Use FApp's time because it is set closer to the beginning of the frame - we don't care about the time so far of the current frame to process the packet
 		const bool bExcludeFrameTime = !!CVarPingExcludeFrameTime.GetValueOnAnyThread();
-		const double CurrentTime	= (PacketReceiveTime != 0.0 ? PacketReceiveTime : FApp::GetCurrentTime());
-		const double RTT			= (CurrentTime - OutLagTime[Index]);
+		const double PacketSendTime = OutLagTime[LastAckIndex];
+		const double RTT			= (CurrentTime - PacketSendTime);
 		const double RTTExclFrame	= RTT - (bExcludeFrameTime ? ServerFrameTime : 0.0);
 		const double NewLag			= FMath::Max(RTTExclFrame, 0.0);
 
@@ -3263,11 +3290,6 @@ void UNetConnection::ReceivedPacket( FBitReader& Reader, bool bIsReinjectedPacke
 				return;
 			}
 
-			if (MissingPacketCount > 10)
-			{
-				UE_LOG(LogNetTraffic, Verbose, TEXT("High single frame packet loss. PacketsLost: %i %s" ), MissingPacketCount, *Describe());
-			}
-
 			InPacketsLost += MissingPacketCount;
 			InTotalPacketsLost += MissingPacketCount;
 			Driver->InPacketsLost += MissingPacketCount;
diff --git a/Engine/Source/Runtime/Engine/Private/NetDriver.cpp b/Engine/Source/Runtime/Engine/Private/NetDriver.cpp
--- a/Engine/Source/Runtime/Engine/Private/NetDriver.cpp
+++ b/Engine/Source/Runtime/Engine/Private/NetDriver.cpp
@@ -528,8 +528,12 @@ namespace UE::Net
 			CleanUpRenamedDynamicActors,
 			TEXT("When enabled, dynamic actors that change outers via Rename() on the server will close their channel ")
 			TEXT("or send a destruction info to clients without the actor's new level loaded."));
+
+		static FName StatNameRawPing(TEXT("RawPing"));
 	}
-}
+
+} // end namespace UE::Net
+
 
 
 /*-----------------------------------------------------------------------------
@@ -8464,11 +8468,17 @@ void UNetDriver::UpdateNetworkStats()
 				ClientsOutPacketsThisFrameAvg /= NumClients;
 			}
 
+			if (ServerConnection)
+			{
+				const int32 RawPingValueInMS = (int32)(ServerConnection->RawPingInSeconds * 1000.f);
+				FCsvProfiler::RecordCustomStat(StatNameRawPing, CSV_CATEGORY_INDEX(Networking), RawPingValueInMS, ECsvCustomStatOp::Set);
+			}
+
 			GetMetrics()->SetInt(UE::Net::Metric::InPacketsClientAvg, ClientsInPacketsThisFrameAvg);
 			GetMetrics()->SetInt(UE::Net::Metric::InPacketsClientMax, ClientsInPacketsThisFrameMax);
 			GetMetrics()->SetInt(UE::Net::Metric::OutPacketsClientAvg, ClientsOutPacketsThisFrameAvg);
 			GetMetrics()->SetInt(UE::Net::Metric::OutPacketsClientMax, ClientsOutPacketsThisFrameMax);
-#else
+#else // #if CSV_PROFILER_STATS
 			// Reset the per-frame stats here
 			for (UNetConnection* Client : ClientConnections)
 			{
@@ -8478,7 +8488,7 @@ void UNetDriver::UpdateNetworkStats()
 					Client->OutPacketsThisFrame = 0;
 				}
 			}
-#endif
+#endif // #else CSV_PROFILER_STATS
 		}
 	} // bCollectNetStats ||(USE_SERVER_PERF_COUNTERS) || STATS
 	else
