Splines - Prevent querying the reparam table for an empty spline.

diff --git a/Engine/Source/Runtime/Engine/Private/Curves/Spline.cpp b/Engine/Source/Runtime/Engine/Private/Curves/Spline.cpp
--- a/Engine/Source/Runtime/Engine/Private/Curves/Spline.cpp
+++ b/Engine/Source/Runtime/Engine/Private/Curves/Spline.cpp
@@ -156,6 +156,11 @@ public:
         return Ar;
     }
 
+	/**
+	 * Returns the parameter space of FNewSpline. All public functions taking a parameter are only defined when the provided value is in this interval.
+	 * @return [0, NumSegments] when NumSegment != 0, otherwise undefined (empty interval).
+	 */
+	FInterval1f GetSegmentSpace() const;
 	float GetSegmentLength(const int32 Index, const float Param, const FVector& Scale3D = FVector(1.0f)) const;
 	float GetSplineLength();
 	
@@ -1749,6 +1754,11 @@ bool FNewSpline::Serialize(FArchive& Ar)
     return true;
 }
 
+FInterval1f FNewSpline::GetSegmentSpace() const
+{
+	return GetNumberOfSegments() == 0 ? FInterval1f::Empty() : FInterval1f(0.f, static_cast<float>(GetNumberOfSegments()));
+}
+
 float FNewSpline::GetSegmentLength(const int32 Index, const float Param, const FVector& Scale3D) const
 {
 	const int32 NumPoints = GetSpline().GetNumPoints();
@@ -1812,7 +1822,7 @@ float FNewSpline::GetSegmentLength(const int32 Index, const float Param, const F
 float FNewSpline::GetSplineLength()
 {	
 	// Evaluate the ParamToLength channel at the very end of the spline
-	const float MaxParameter = static_cast<float>(GetSpline().GetNumberOfSegments());
+	const float MaxParameter = static_cast<float>(GetNumberOfSegments());
     return GetDistanceAtParameter(MaxParameter);
 }
 
@@ -2227,6 +2237,11 @@ float FNewSpline::GetParameterAtDistance(float Distance)
 
 float FNewSpline::GetDistanceAtParameter(float Parameter)
 {
+	if (!GetSegmentSpace().Contains(Parameter))
+	{
+		return 0.f;
+	}
+	
 	UpdateReparamTable();
 
 	const float ParameterMax = static_cast<float>(GetNumberOfSegments());
@@ -2311,12 +2326,12 @@ void FNewSpline::UpdateSpline(const FSpline::FUpdateSplineParams& Params)
 
 float FNewSpline::FromNormalizedSpace(float Parameter) const
 {
-	return GetSpline().GetNumberOfSegments() * Parameter;
+	return GetNumberOfSegments() * Parameter;
 }
 
 float FNewSpline::ToNormalizedSpace(float Parameter) const
 {
-	return Parameter / GetSpline().GetNumberOfSegments();
+	return Parameter / GetNumberOfSegments();
 }
 
 float FNewSpline::FromInternalSplineSpace(float Parameter) const
@@ -2327,11 +2342,11 @@ float FNewSpline::FromInternalSplineSpace(float Parameter) const
 }
 float FNewSpline::ToInternalSplineSpace(float Parameter) const
 {
-	const int32 NumSegments = GetSpline().GetNumberOfSegments();
+	const int32 NumSegments = GetNumberOfSegments();
 	float ClampedParam = FMath::Clamp(Parameter, 0.0f, static_cast<float>(NumSegments));
 	const int32 SegmentIndex = FMath::Min(FMath::FloorToInt(ClampedParam), NumSegments - 1);
 	const float LocalT = ClampedParam - static_cast<float>(SegmentIndex); // Parameter [0,1] within segment
-	return GetSpline().GetSegmentParameterRange(SegmentIndex).Interpolate(LocalT);
+	return GetSegmentParameterRange(SegmentIndex).Interpolate(LocalT);
 }
 
 void FNewSpline::UpdateReparamTable()
@@ -2341,7 +2356,7 @@ void FNewSpline::UpdateReparamTable()
 		return;
 	}
 
-	const int32 NumSegments = GetSpline().GetNumberOfSegments();
+	const int32 NumSegments = GetNumberOfSegments();
 
 	// Start by clearing it
 	ReparamTable.Points.Reset(NumSegments * ReparamStepsPerSegment + 1);
