Fix other Substrate feature rendering (cloth, SSS, ...) with BlendableGbuffer when using Substrate's node.

diff --git a/Engine/Shaders/Public/ShaderVersion.ush b/Engine/Shaders/Public/ShaderVersion.ush
--- a/Engine/Shaders/Public/ShaderVersion.ush
+++ b/Engine/Shaders/Public/ShaderVersion.ush
@@ -3,4 +3,4 @@
 // Each time the console command is executed it generates a new GUID. As this file is included
 // in Platform.ush (which should be included in any shader) it allows to invalidate the shader DDC.
 // 
-#pragma message("UESHADERMETADATA_VERSION 794E2BBF-A6E7-43C2-BD06-A470603C59D5")
+#pragma message("UESHADERMETADATA_VERSION ED27A877-556A-41DB-A67D-D0880AFFE659")
diff --git a/Engine/Source/Runtime/Core/Private/UObject/FortniteMainBranchObjectVersion.cpp b/Engine/Source/Runtime/Core/Private/UObject/FortniteMainBranchObjectVersion.cpp
--- a/Engine/Source/Runtime/Core/Private/UObject/FortniteMainBranchObjectVersion.cpp
+++ b/Engine/Source/Runtime/Core/Private/UObject/FortniteMainBranchObjectVersion.cpp
@@ -14,6 +14,6 @@ TMap<FGuid, FGuid> FFortniteMainBranchObjectVersion::GetSystemGuids()
 	SystemGuids.Add(DevGuids.Niagara_LatestScriptCompileVersion, FGuid("97F0FD7149D90D42B65D6557AF477D1A"));
 	SystemGuids.Add(DevGuids.SkeletalMeshDerivedDataVersion, FGuid("C998F2F3B2884A02B7E290B8C25835F0"));
 	SystemGuids.Add(DevGuids.STATICMESH_DERIVEDDATA_VER, FGuid("3ABF32149F9D4D838750368AB95573BE"));
-	SystemGuids.Add(DevGuids.MaterialTranslationDDCVersion, FGuid("F051877974B44D409435F321BDB39DBB"));
+	SystemGuids.Add(DevGuids.MaterialTranslationDDCVersion, FGuid("05D49B4F487C48F29A0B6AB785F506AF"));
 	return SystemGuids;
 }
diff --git a/Engine/Source/Runtime/Engine/Private/Materials/HLSLMaterialTranslator.cpp b/Engine/Source/Runtime/Engine/Private/Materials/HLSLMaterialTranslator.cpp
--- a/Engine/Source/Runtime/Engine/Private/Materials/HLSLMaterialTranslator.cpp
+++ b/Engine/Source/Runtime/Engine/Private/Materials/HLSLMaterialTranslator.cpp
@@ -311,6 +311,7 @@ struct FHLSLMaterialTranslator::FEnvironmentDefines
 	uint8 NumCustomizedUVs;
 	bool bTextureSampleDebug;
 	bool bSubstrateRoughnessTracking;
+	ESubstrateBsdfFeature SubstrateMaterialBsdfFeatures;
 	bool bMaterialEnableTranslucentLocalLightShadow;
 	bool bMaterialEnableTranslucentHighQualityLocalLightShadow;
 	bool bMaterialEnableTranslucentHighQualityDirectionalLightShadow;
@@ -416,6 +417,7 @@ struct FHLSLMaterialTranslator::FEnvironmentDefines
 		Ar << bTextureSampleDebug;
 		Ar << SubstrateDefines;
 		Ar << bSubstrateRoughnessTracking;
+		Ar << SubstrateMaterialBsdfFeatures;
 		Ar << bMaterialEnableTranslucentLocalLightShadow;
 		Ar << bMaterialEnableTranslucentHighQualityLocalLightShadow;
 		Ar << bMaterialEnableTranslucentHighQualityDirectionalLightShadow;
@@ -3019,17 +3021,17 @@ void FHLSLMaterialTranslator::GetMaterialEnvironment(EShaderPlatform InPlatform,
 			OutEnvironment.SetDefine(TEXT("MATERIAL_SHADINGMODEL_DEFAULT_LIT"), TEXT("1"));
 		}
 
-		if (EnvironmentDefines->HasShadingModel(MSM_Subsurface))
+		if (EnvironmentDefines->HasShadingModel(MSM_Subsurface) || EnumHasAnyFlags(EnvironmentDefines->SubstrateMaterialBsdfFeatures, ESubstrateBsdfFeature::SSS))
 		{
 			OutEnvironment.SetDefine(TEXT("MATERIAL_SHADINGMODEL_SUBSURFACE"), TEXT("1"));
 		}
 
-		if (EnvironmentDefines->HasShadingModel(MSM_PreintegratedSkin))
+		if (EnvironmentDefines->HasShadingModel(MSM_PreintegratedSkin) || EnumHasAnyFlags(EnvironmentDefines->SubstrateMaterialBsdfFeatures, ESubstrateBsdfFeature::SSS))
 		{
 			OutEnvironment.SetDefine(TEXT("MATERIAL_SHADINGMODEL_PREINTEGRATED_SKIN"), TEXT("1"));
 		}
 
-		if (EnvironmentDefines->HasShadingModel(MSM_SubsurfaceProfile))
+		if (EnvironmentDefines->HasShadingModel(MSM_SubsurfaceProfile) || EnumHasAnyFlags(EnvironmentDefines->SubstrateMaterialBsdfFeatures, ESubstrateBsdfFeature::SSS))
 		{
 			OutEnvironment.SetDefine(TEXT("MATERIAL_SHADINGMODEL_SUBSURFACE_PROFILE"), TEXT("1"));
 		}
@@ -3039,27 +3041,27 @@ void FHLSLMaterialTranslator::GetMaterialEnvironment(EShaderPlatform InPlatform,
 			OutEnvironment.SetDefine(TEXT("MATERIAL_SUBSURFACE_PROFILE_USE_CURVATURE"), TEXT("1"));
 		}
 
-		if (EnvironmentDefines->HasShadingModel(MSM_ClearCoat))
+		if (EnvironmentDefines->HasShadingModel(MSM_ClearCoat) || EnumHasAnyFlags(EnvironmentDefines->SubstrateMaterialBsdfFeatures, ESubstrateBsdfFeature::SecondRoughnessOrSimpleClearCoat))
 		{
 			OutEnvironment.SetDefine(TEXT("MATERIAL_SHADINGMODEL_CLEAR_COAT"), TEXT("1"));
 		}
 
-		if (EnvironmentDefines->HasShadingModel(MSM_TwoSidedFoliage))
+		if (EnvironmentDefines->HasShadingModel(MSM_TwoSidedFoliage) || EnumHasAnyFlags(EnvironmentDefines->SubstrateMaterialBsdfFeatures, ESubstrateBsdfFeature::SSS))
 		{
 			OutEnvironment.SetDefine(TEXT("MATERIAL_SHADINGMODEL_TWOSIDED_FOLIAGE"), TEXT("1"));
 		}
 
-		if (EnvironmentDefines->HasShadingModel(MSM_Hair))
+		if (EnvironmentDefines->HasShadingModel(MSM_Hair) || EnumHasAnyFlags(EnvironmentDefines->SubstrateMaterialBsdfFeatures, ESubstrateBsdfFeature::Hair))
 		{
 			OutEnvironment.SetDefine(TEXT("MATERIAL_SHADINGMODEL_HAIR"), TEXT("1"));
 		}
 
-		if (EnvironmentDefines->HasShadingModel(MSM_Cloth))
+		if (EnvironmentDefines->HasShadingModel(MSM_Cloth) || EnumHasAnyFlags(EnvironmentDefines->SubstrateMaterialBsdfFeatures, ESubstrateBsdfFeature::Fuzz))
 		{
 			OutEnvironment.SetDefine(TEXT("MATERIAL_SHADINGMODEL_CLOTH"), TEXT("1"));
 		}
 
-		if (EnvironmentDefines->HasShadingModel(MSM_Eye))
+		if (EnvironmentDefines->HasShadingModel(MSM_Eye) || EnumHasAnyFlags(EnvironmentDefines->SubstrateMaterialBsdfFeatures, ESubstrateBsdfFeature::Eye))
 		{
 			OutEnvironment.SetDefine(TEXT("MATERIAL_SHADINGMODEL_EYE"), TEXT("1"));
 		}
@@ -16657,6 +16659,7 @@ void FHLSLMaterialTranslator::PrepareEnvironmentDefines()
 
 		EnvironmentDefines->bSubstrateLegacyIrisNormal = EyeIrisNormalPluggedIn;
 		EnvironmentDefines->bSubstrateLegacyIrisTangent = EyeIrisTangentPluggedIn;
+		EnvironmentDefines->SubstrateMaterialBsdfFeatures = MaterialCompilationOutput.SubstrateMaterialCompilationOutput.SubstrateMaterialBsdfFeatures;
 	}
 
 	EnvironmentDefines->bTextureSampleDebug = IsDebugTextureSampleEnabled();
