Fix linker error in KismetMathLibrary when compiling debug configuration.

#rnx

diff --git a/Engine/Source/Runtime/Engine/Classes/Kismet/KismetMathLibrary.inl b/Engine/Source/Runtime/Engine/Classes/Kismet/KismetMathLibrary.inl
--- a/Engine/Source/Runtime/Engine/Classes/Kismet/KismetMathLibrary.inl
+++ b/Engine/Source/Runtime/Engine/Classes/Kismet/KismetMathLibrary.inl
@@ -2,79 +2,87 @@
 
 #pragma once
 
-inline
+// When KISMET_MATH_INLINE_ENABLED is true this file is included in the header
+// so we need to mark these functions as inline.
+#if KISMET_MATH_INLINE_ENABLED
+	#define KISMET_MATH_INLINE			inline
+#else
+	#define KISMET_MATH_INLINE			// nothing
+#endif
+
+KISMET_MATH_INLINE
 bool UKismetMathLibrary::RandomBool()
 {
 	return FMath::RandBool();
 }
 
-inline
+KISMET_MATH_INLINE
 bool UKismetMathLibrary::Not_PreBool(bool A)
 {
 	return !A;
 }
 
-inline
+KISMET_MATH_INLINE
 bool UKismetMathLibrary::EqualEqual_BoolBool(bool A, bool B)
 {
 	return ((!A) == (!B));
 }
 
-inline
+KISMET_MATH_INLINE
 bool UKismetMathLibrary::NotEqual_BoolBool(bool A, bool B)
 {
 	return ((!A) != (!B));
 }
 
-inline
+KISMET_MATH_INLINE
 bool UKismetMathLibrary::BooleanAND(bool A, bool B)
 {
 	return A && B;
 }
 
-inline
+KISMET_MATH_INLINE
 bool UKismetMathLibrary::BooleanNAND(bool A, bool B)
 {
 	return !(A && B);
 }
 
-inline
+KISMET_MATH_INLINE
 bool UKismetMathLibrary::BooleanOR(bool A, bool B)
 {
 	return A || B;
 }
 
-inline
+KISMET_MATH_INLINE
 bool UKismetMathLibrary::BooleanXOR(bool A, bool B)
 {
 	return A ^ B;
 }
 
-inline
+KISMET_MATH_INLINE
 bool UKismetMathLibrary::BooleanNOR(bool A, bool B)
 {
 	return !(A || B);
 }
 
-inline
+KISMET_MATH_INLINE
 uint8 UKismetMathLibrary::Multiply_ByteByte(uint8 A, uint8 B)
 {
 	return A * B;
 }
 
-inline
+KISMET_MATH_INLINE
 uint8 UKismetMathLibrary::Add_ByteByte(uint8 A, uint8 B)
 {
 	return A + B;
 }
 
-inline
+KISMET_MATH_INLINE
 uint8 UKismetMathLibrary::Subtract_ByteByte(uint8 A, uint8 B)
 {
 	return A - B;
 }
 
-inline
+KISMET_MATH_INLINE
 uint8 UKismetMathLibrary::Divide_ByteByte(uint8 A, uint8 B)
 {
 	if (B == 0)
@@ -86,7 +94,7 @@ uint8 UKismetMathLibrary::Divide_ByteByte(uint8 A, uint8 B)
 	return (A / B);
 }
 
-inline
+KISMET_MATH_INLINE
 uint8 UKismetMathLibrary::Percent_ByteByte(uint8 A, uint8 B)
 {
 	if (B == 0)
@@ -98,169 +106,169 @@ uint8 UKismetMathLibrary::Percent_ByteByte(uint8 A, uint8 B)
 	return (A % B);
 }
 
-inline
+KISMET_MATH_INLINE
 uint8 UKismetMathLibrary::BMin(uint8 A, uint8 B)
 {
 	return FMath::Min<uint8>(A, B);
 }
 
-inline
+KISMET_MATH_INLINE
 uint8 UKismetMathLibrary::BMax(uint8 A, uint8 B)
 {
 	return FMath::Max<uint8>(A, B);
 }
 
-inline
+KISMET_MATH_INLINE
 bool UKismetMathLibrary::Less_ByteByte(uint8 A, uint8 B)
 {
 	return A < B;
 }
 
-inline
+KISMET_MATH_INLINE
 bool UKismetMathLibrary::Greater_ByteByte(uint8 A, uint8 B)
 {
 	return A > B;
 }
 
-inline
+KISMET_MATH_INLINE
 bool UKismetMathLibrary::LessEqual_ByteByte(uint8 A, uint8 B)
 {
 	return A <= B;
 }
 
-inline
+KISMET_MATH_INLINE
 bool UKismetMathLibrary::GreaterEqual_ByteByte(uint8 A, uint8 B)
 {
 	return A >= B;
 }
 
-inline
+KISMET_MATH_INLINE
 bool UKismetMathLibrary::EqualEqual_ByteByte(uint8 A, uint8 B)
 {
 	return A == B;
 }
 
-inline
+KISMET_MATH_INLINE
 bool UKismetMathLibrary::NotEqual_ByteByte(uint8 A, uint8 B)
 {
 	return A != B;
 }
 
-inline
+KISMET_MATH_INLINE
 int32 UKismetMathLibrary::Multiply_IntInt(int32 A, int32 B)
 {
 	return A * B;
 }
 
-inline
+KISMET_MATH_INLINE
 int32 UKismetMathLibrary::Add_IntInt(int32 A, int32 B)
 {
 	return A + B;
 }
 
-inline
+KISMET_MATH_INLINE
 int32 UKismetMathLibrary::Subtract_IntInt(int32 A, int32 B)
 {
 	return A - B;
 }
 
-inline
+KISMET_MATH_INLINE
 bool UKismetMathLibrary::Less_IntInt(int32 A, int32 B)
 {
 	return A < B;
 }
 
-inline
+KISMET_MATH_INLINE
 bool UKismetMathLibrary::Greater_IntInt(int32 A, int32 B)
 {
 	return A > B;
 }
 
-inline
+KISMET_MATH_INLINE
 bool UKismetMathLibrary::LessEqual_IntInt(int32 A, int32 B)
 {
 	return A <= B;
 }
 
-inline
+KISMET_MATH_INLINE
 bool UKismetMathLibrary::GreaterEqual_IntInt(int32 A, int32 B)
 {
 	return A >= B;
 }
 
-inline
+KISMET_MATH_INLINE
 bool UKismetMathLibrary::EqualEqual_IntInt(int32 A, int32 B)
 {
 	return A == B;
 }
 
-inline
+KISMET_MATH_INLINE
 bool UKismetMathLibrary::NotEqual_IntInt(int32 A, int32 B)
 {
 	return A != B;
 }
 
-inline
+KISMET_MATH_INLINE
 int32 UKismetMathLibrary::And_IntInt(int32 A, int32 B)
 {
 	return A & B;
 }
 
-inline
+KISMET_MATH_INLINE
 int32 UKismetMathLibrary::Xor_IntInt(int32 A, int32 B)
 {
 	return A ^ B;
 }
 
-inline
+KISMET_MATH_INLINE
 int32 UKismetMathLibrary::Or_IntInt(int32 A, int32 B)
 {
 	return A | B;
 }
 
-inline
+KISMET_MATH_INLINE
 int32 UKismetMathLibrary::Not_Int(int32 A)
 {
 	return ~A;
 }
 
-inline
+KISMET_MATH_INLINE
 int32 UKismetMathLibrary::SignOfInteger(int32 A)
 {
 	return FMath::Sign<int32>(A);
 }
 
-inline
+KISMET_MATH_INLINE
 int32 UKismetMathLibrary::RandomInteger(int32 A)
 {
 	return FMath::RandHelper(A);
 }
 
-inline
+KISMET_MATH_INLINE
 int32 UKismetMathLibrary::RandomIntegerInRange(int32 Min, int32 Max)
 {
 	return FMath::RandRange(Min, Max);
 }
 
-inline
+KISMET_MATH_INLINE
 int32 UKismetMathLibrary::Min(int32 A, int32 B)
 {
 	return FMath::Min(A, B);
 }
 
-inline
+KISMET_MATH_INLINE
 int32 UKismetMathLibrary::Max(int32 A, int32 B)
 {
 	return FMath::Max(A, B);
 }
 
-inline
+KISMET_MATH_INLINE
 int32 UKismetMathLibrary::Clamp(int32 V, int32 A, int32 B)
 {
 	return FMath::Clamp(V, A, B);
 }
 
-inline
+KISMET_MATH_INLINE
 int32 UKismetMathLibrary::Wrap(int32 Value, int32 Min, int32 Max)
 {
 PRAGMA_DISABLE_DEPRECATION_WARNINGS
@@ -268,145 +276,145 @@ PRAGMA_DISABLE_DEPRECATION_WARNINGS
 PRAGMA_ENABLE_DEPRECATION_WARNINGS
 }
 
-inline
+KISMET_MATH_INLINE
 int32 UKismetMathLibrary::Abs_Int(int32 A)
 {
 	return FMath::Abs(A);
 }
 
-inline
+KISMET_MATH_INLINE
 int64 UKismetMathLibrary::Multiply_Int64Int64(int64 A, int64 B)
 {
 	return A * B;
 }
 
-inline
+KISMET_MATH_INLINE
 int64 UKismetMathLibrary::Add_Int64Int64(int64 A, int64 B)
 {
 	return A + B;
 }
 
-inline
+KISMET_MATH_INLINE
 int64 UKismetMathLibrary::Subtract_Int64Int64(int64 A, int64 B)
 {
 	return A - B;
 }
 
-inline
+KISMET_MATH_INLINE
 bool UKismetMathLibrary::Less_Int64Int64(int64 A, int64 B)
 {
 	return A < B;
 }
 
-inline
+KISMET_MATH_INLINE
 bool UKismetMathLibrary::Greater_Int64Int64(int64 A, int64 B)
 {
 	return A > B;
 }
 
-inline
+KISMET_MATH_INLINE
 bool UKismetMathLibrary::LessEqual_Int64Int64(int64 A, int64 B)
 {
 	return A <= B;
 }
 
-inline
+KISMET_MATH_INLINE
 bool UKismetMathLibrary::GreaterEqual_Int64Int64(int64 A, int64 B)
 {
 	return A >= B;
 }
 
-inline
+KISMET_MATH_INLINE
 bool UKismetMathLibrary::EqualEqual_Int64Int64(int64 A, int64 B)
 {
 	return A == B;
 }
 
-inline
+KISMET_MATH_INLINE
 bool UKismetMathLibrary::NotEqual_Int64Int64(int64 A, int64 B)
 {
 	return A != B;
 }
 
-inline
+KISMET_MATH_INLINE
 int64 UKismetMathLibrary::And_Int64Int64(int64 A, int64 B)
 {
 	return A & B;
 }
 
-inline
+KISMET_MATH_INLINE
 int64 UKismetMathLibrary::Xor_Int64Int64(int64 A, int64 B)
 {
 	return A ^ B;
 }
 
-inline
+KISMET_MATH_INLINE
 int64 UKismetMathLibrary::Or_Int64Int64(int64 A, int64 B)
 {
 	return A | B;
 }
 
-inline
+KISMET_MATH_INLINE
 int64 UKismetMathLibrary::Not_Int64(int64 A)
 {
 	return ~A;
 }
 
-inline
+KISMET_MATH_INLINE
 int64 UKismetMathLibrary::SignOfInteger64(int64 A)
 {
 	return FMath::Sign<int64>(A);
 }
 
-inline
+KISMET_MATH_INLINE
 int64 UKismetMathLibrary::RandomInteger64(int64 A)
 {
 	return FMath::RandHelper64(A);
 }
 
-inline
+KISMET_MATH_INLINE
 int64 UKismetMathLibrary::RandomInteger64InRange(int64 Min, int64 Max)
 {
 	return FMath::RandRange(Min, Max);
 }
 
-inline
+KISMET_MATH_INLINE
 int64 UKismetMathLibrary::MinInt64(int64 A, int64 B)
 {
 	return FMath::Min(A, B);
 }
 
-inline
+KISMET_MATH_INLINE
 int64 UKismetMathLibrary::MaxInt64(int64 A, int64 B)
 {
 	return FMath::Max(A, B);
 }
 
-inline
+KISMET_MATH_INLINE
 int64 UKismetMathLibrary::ClampInt64(int64 V, int64 A, int64 B)
 {
 	return FMath::Clamp(V, A, B);
 }
 
-inline
+KISMET_MATH_INLINE
 int64 UKismetMathLibrary::Abs_Int64(int64 A)
 {
 	return FMath::Abs(A);
 }
 
-inline
+KISMET_MATH_INLINE
 double UKismetMathLibrary::MultiplyMultiply_FloatFloat(double Base, double Exp)
 {
 	return FMath::Pow(Base, Exp);
 }
 
-inline
+KISMET_MATH_INLINE
 double UKismetMathLibrary::Multiply_IntFloat(int32 A, double B)
 {
 	return A * B;
 }	
 
-inline
+KISMET_MATH_INLINE
 int32 UKismetMathLibrary::Divide_IntInt(int32 A, int32 B)
 {
 	if (B == 0)
@@ -418,7 +426,7 @@ int32 UKismetMathLibrary::Divide_IntInt(int32 A, int32 B)
 	return (A / B);
 }
 
-inline
+KISMET_MATH_INLINE
 int64 UKismetMathLibrary::Divide_Int64Int64(int64 A, int64 B)
 {
 	if (B == 0)
@@ -431,7 +439,7 @@ int64 UKismetMathLibrary::Divide_Int64Int64(int64 A, int64 B)
 }
 
 
-inline
+KISMET_MATH_INLINE
 int32 UKismetMathLibrary::Percent_IntInt(int32 A, int32 B)
 {
 	if (B == 0)
@@ -443,7 +451,7 @@ int32 UKismetMathLibrary::Percent_IntInt(int32 A, int32 B)
 	return (A % B);
 }
 
-inline
+KISMET_MATH_INLINE
 int64 UKismetMathLibrary::Percent_Int64Int64(int64 A, int64 B)
 {
 	if (B == 0)
@@ -455,31 +463,31 @@ int64 UKismetMathLibrary::Percent_Int64Int64(int64 A, int64 B)
 	return (A % B);
 }
 
-inline
+KISMET_MATH_INLINE
 double UKismetMathLibrary::Fraction(double A)
 {
 	return FMath::Fractional(A);
 }
 
-inline
+KISMET_MATH_INLINE
 double UKismetMathLibrary::Add_DoubleDouble(double A, double B)
 {
 	return A + B;
 }
 
-inline
+KISMET_MATH_INLINE
 double UKismetMathLibrary::Subtract_DoubleDouble(double A, double B)
 {
 	return A - B;
 }
 
-inline
+KISMET_MATH_INLINE
 double UKismetMathLibrary::Multiply_DoubleDouble(double A, double B)
 {
 	return A * B;
 }
 
-inline
+KISMET_MATH_INLINE
 double UKismetMathLibrary::Divide_DoubleDouble(double A, double B)
 {
 	if (B == 0.0)
@@ -490,187 +498,187 @@ double UKismetMathLibrary::Divide_DoubleDouble(double A, double B)
 	return A / B;
 }
 
-inline
+KISMET_MATH_INLINE
 bool UKismetMathLibrary::Less_DoubleDouble(double A, double B)
 {
 	return A < B;
 }
 
-inline
+KISMET_MATH_INLINE
 bool UKismetMathLibrary::Greater_DoubleDouble(double A, double B)
 {
 	return A > B;
 }
 
-inline
+KISMET_MATH_INLINE
 bool UKismetMathLibrary::LessEqual_DoubleDouble(double A, double B)
 {
 	return A <= B;
 }
 
-inline
+KISMET_MATH_INLINE
 bool UKismetMathLibrary::GreaterEqual_DoubleDouble(double A, double B)
 {
 	return A >= B;
 }
 
-inline
+KISMET_MATH_INLINE
 bool UKismetMathLibrary::EqualEqual_DoubleDouble(double A, double B)
 {
 	return A == B;
 }
 
-inline
+KISMET_MATH_INLINE
 bool UKismetMathLibrary::NearlyEqual_FloatFloat(double A, double B, double ErrorTolerance)
 {
 	return FMath::IsNearlyEqual(A, B, ErrorTolerance);
 }
 
-inline
+KISMET_MATH_INLINE
 bool UKismetMathLibrary::NotEqual_DoubleDouble(double A, double B)
 {
 	return A != B;
 }
 
-inline
+KISMET_MATH_INLINE
 double UKismetMathLibrary::GridSnap_Float(double Location, double GridSize)
 {
 	return FMath::GridSnap(Location, GridSize);
 }
 
-inline
+KISMET_MATH_INLINE
 double UKismetMathLibrary::GetPI()
 {
 	return UE_DOUBLE_PI;
 }
 
-inline
+KISMET_MATH_INLINE
 double UKismetMathLibrary::GetTAU()
 {
 	return 2.0 * UE_DOUBLE_PI;
 }
 
-inline
+KISMET_MATH_INLINE
 double UKismetMathLibrary::DegreesToRadians(double A)
 {
 	return FMath::DegreesToRadians(A);
 }
 
-inline
+KISMET_MATH_INLINE
 double UKismetMathLibrary::RadiansToDegrees(double A)
 {
 	return FMath::RadiansToDegrees(A);
 }
 
-inline
+KISMET_MATH_INLINE
 double UKismetMathLibrary::Abs(double A)
 {
 	return FMath::Abs(A);
 }	
 
-inline
+KISMET_MATH_INLINE
 double UKismetMathLibrary::Sin(double A)
 {
 	return FMath::Sin(A);
 }	
 
-inline
+KISMET_MATH_INLINE
 double UKismetMathLibrary::Asin(double A)
 {
 	return FMath::Asin(A);
 }	
 
-inline
+KISMET_MATH_INLINE
 double UKismetMathLibrary::Cos(double A)
 {
 	return FMath::Cos(A);
 }	
 
-inline
+KISMET_MATH_INLINE
 double UKismetMathLibrary::Acos(double A)
 {
 	return FMath::Acos(A);
 }	
 
-inline
+KISMET_MATH_INLINE
 double UKismetMathLibrary::Tan(double A)
 {
 	return FMath::Tan(A);
 }	
 
-inline
+KISMET_MATH_INLINE
 double UKismetMathLibrary::Atan(double A)
 {
 	return FMath::Atan(A);
 }
 
-inline
+KISMET_MATH_INLINE
 double UKismetMathLibrary::Atan2(double Y, double X)
 {
 	return FMath::Atan2(Y, X);
 }	
 
-inline
+KISMET_MATH_INLINE
 double UKismetMathLibrary::DegSin(double A)
 {
 	return FMath::Sin(UE_DOUBLE_PI/(180.0) * A);
 }
 
-inline
+KISMET_MATH_INLINE
 double UKismetMathLibrary::DegAsin(double A)
 {
 	return (180.0)/UE_DOUBLE_PI * FMath::Asin(A);
 }
 
-inline
+KISMET_MATH_INLINE
 double UKismetMathLibrary::DegCos(double A)
 {
 	return FMath::Cos(UE_DOUBLE_PI/(180.0) * A);
 }
 
-inline
+KISMET_MATH_INLINE
 double UKismetMathLibrary::DegAcos(double A)
 {
 	return (180.0)/UE_DOUBLE_PI * FMath::Acos(A);
 }
 
-inline
+KISMET_MATH_INLINE
 double UKismetMathLibrary::DegTan(double A)
 {
 	return FMath::Tan(UE_DOUBLE_PI/(180.0) * A);
 }
 
-inline
+KISMET_MATH_INLINE
 double UKismetMathLibrary::DegAtan(double A)
 {
 	return (180.0)/UE_DOUBLE_PI * FMath::Atan(A);
 }
 
-inline
+KISMET_MATH_INLINE
 double UKismetMathLibrary::DegAtan2(double Y, double X)
 {
 	return (180.0)/UE_DOUBLE_PI * FMath::Atan2(Y, X);
 }
 
-inline
+KISMET_MATH_INLINE
 double UKismetMathLibrary::ClampAngle(double AngleDegrees, double MinAngleDegrees, double MaxAngleDegrees)
 {
 	return FMath::ClampAngle(AngleDegrees, MinAngleDegrees, MaxAngleDegrees);
 }
 
-inline
+KISMET_MATH_INLINE
 double UKismetMathLibrary::Exp(double A)
 {
 	return FMath::Exp(A);
 }	
 
-inline
+KISMET_MATH_INLINE
 double UKismetMathLibrary::Loge(double A)
 {
 	return FMath::Loge(A);
 }	
 
-inline
+KISMET_MATH_INLINE
 double UKismetMathLibrary::Sqrt(double A)
 {
 	if (A >= 0.0)
@@ -684,127 +692,127 @@ double UKismetMathLibrary::Sqrt(double A)
 	}
 }
 
-inline
+KISMET_MATH_INLINE
 double UKismetMathLibrary::Square(double A)
 {
 	return FMath::Square(A);
 }	
 
-inline
+KISMET_MATH_INLINE
 int32 UKismetMathLibrary::Round(double A)
 {
 	return static_cast<int32>(FMath::RoundToInt(A));
 }	
 
-inline
+KISMET_MATH_INLINE
 int32 UKismetMathLibrary::FFloor(double A)
 {
 	return static_cast<int32>(FMath::FloorToInt(A));
 }	
 
-inline
+KISMET_MATH_INLINE
 int32 UKismetMathLibrary::FTrunc(double A)
 {
 	return static_cast<int32>(FMath::TruncToInt(A));
 }	
 
-inline
+KISMET_MATH_INLINE
 int64 UKismetMathLibrary::Round64(double A)
 {
 	return (int64)FMath::RoundToDouble(A);
 }
 
-inline
+KISMET_MATH_INLINE
 int64 UKismetMathLibrary::FFloor64(double A)
 {
 	return (int64)FMath::FloorToDouble(A);
 }
 
-inline
+KISMET_MATH_INLINE
 int64 UKismetMathLibrary::FTrunc64(double A)
 {
 	return (int64)A;
 }
 
-inline
+KISMET_MATH_INLINE
 int64 UKismetMathLibrary::FCeil64(double A)
 {
 	return (int64)FMath::CeilToDouble(A);
 }
 
-inline
+KISMET_MATH_INLINE
 int32 UKismetMathLibrary::FCeil(double A)
 {
 	return static_cast<int32>(FMath::CeilToInt(A));
 }	
 
-inline
+KISMET_MATH_INLINE
 double UKismetMathLibrary::SignOfFloat(double A)
 {
 	return FMath::Sign<double>(A);
 }
 
-inline
+KISMET_MATH_INLINE
 double UKismetMathLibrary::MultiplyByPi(double Value)
 {
 	return Value * UE_DOUBLE_PI;
 }
 
-inline
+KISMET_MATH_INLINE
 float UKismetMathLibrary::FixedTurn(float InCurrent, float InDesired, float InDeltaRate)
 {
 	return FMath::FixedTurn(InCurrent, InDesired, InDeltaRate);
 }
 
-inline
+KISMET_MATH_INLINE
 double UKismetMathLibrary::RandomFloat()
 {
 	return FMath::FRand(); // LWC_todo: add and switch to double version.
 }	
 
-inline
+KISMET_MATH_INLINE
 double UKismetMathLibrary::RandomFloatInRange(double Min, double Max)
 {
 	return FMath::FRandRange(Min, Max);
 }	
 
-inline
+KISMET_MATH_INLINE
 double UKismetMathLibrary::FMin(double A, double B)
 {
 	return FMath::Min(A, B);
 }	
 
-inline
+KISMET_MATH_INLINE
 double UKismetMathLibrary::FMax(double A, double B)
 {
 	return FMath::Max(A, B);
 }	
 
-inline
+KISMET_MATH_INLINE
 double UKismetMathLibrary::FClamp(double V, double A, double B)
 {
 	return FMath::Clamp(V, A, B);
 }	
 
-inline
+KISMET_MATH_INLINE
 double UKismetMathLibrary::FWrap(double Value, double Min, double Max)
 {
 	return FMath::Wrap(Value, Min, Max);
 }
 
-inline
+KISMET_MATH_INLINE
 double UKismetMathLibrary::Lerp(double A, double B, double V)
 {
 	return A + V*(B-A);
 }	
 
-inline
+KISMET_MATH_INLINE
 double UKismetMathLibrary::FInterpTo(double Current, double Target, double DeltaTime, double InterpSpeed)
 {
 	return FMath::FInterpTo(Current, Target, DeltaTime, InterpSpeed);
 }
 
-inline
+KISMET_MATH_INLINE
 double UKismetMathLibrary::FInterpTo_Constant(double Current, double Target, double DeltaTime, double InterpSpeed)
 {
 	return FMath::FInterpConstantTo(Current, Target, DeltaTime, InterpSpeed);
@@ -814,44 +822,44 @@ double UKismetMathLibrary::FInterpTo_Constant(double Current, double Target, dou
 /* FRotator
 *****************************************************************************/
 
-inline
+KISMET_MATH_INLINE
 FRotator UKismetMathLibrary::RInterpTo_Constant(FRotator Current, FRotator Target, float DeltaTime, float InterpSpeed)
 {
 	return FMath::RInterpConstantTo(Current, Target, DeltaTime, InterpSpeed);
 }
 
 
-inline
+KISMET_MATH_INLINE
 bool UKismetMathLibrary::EqualEqual_RotatorRotator(FRotator A, FRotator B, float ErrorTolerance)
 {
 	return A.Equals(B, ErrorTolerance);
 }	
 
-inline
+KISMET_MATH_INLINE
 bool UKismetMathLibrary::NotEqual_RotatorRotator(FRotator A, FRotator B, float ErrorTolerance)
 {
 	return !A.Equals(B, ErrorTolerance);
 }
 
-inline
+KISMET_MATH_INLINE
 FRotator UKismetMathLibrary::Multiply_RotatorFloat(FRotator A, float B)
 {
 	return A * B;
 }
 
-inline
+KISMET_MATH_INLINE
 FRotator UKismetMathLibrary::Multiply_RotatorInt(FRotator A, int32 B)
 {
 	return A * (float)B;
 }	
 
-inline
+KISMET_MATH_INLINE
 FRotator UKismetMathLibrary::NegateRotator(FRotator A)
 {
 	return A.GetInverse();
 }
 
-inline
+KISMET_MATH_INLINE
 FRotator UKismetMathLibrary::RInterpTo(FRotator Current, FRotator Target, float DeltaTime, float InterpSpeed)
 {
 	return FMath::RInterpTo(Current, Target, DeltaTime, InterpSpeed);
@@ -861,13 +869,13 @@ FRotator UKismetMathLibrary::RInterpTo(FRotator Current, FRotator Target, float
 /* FTransform
 *****************************************************************************/
 
-inline
+KISMET_MATH_INLINE
 FTransform UKismetMathLibrary::MakeTransform(FVector Translation, FRotator Rotation, FVector Scale)
 {
 	return FTransform(Rotation, Translation, Scale);
 }
 
-inline
+KISMET_MATH_INLINE
 void UKismetMathLibrary::BreakTransform(const FTransform& InTransform, FVector& Translation, FRotator& Rotation, FVector& Scale)
 {
 	Translation = InTransform.GetLocation();
@@ -875,49 +883,49 @@ void UKismetMathLibrary::BreakTransform(const FTransform& InTransform, FVector&
 	Scale = InTransform.GetScale3D();
 }
 
-inline
+KISMET_MATH_INLINE
 FVector UKismetMathLibrary::TransformLocation(const FTransform& T, FVector Location)
 {
 	return T.TransformPosition(Location);
 }
 
-inline
+KISMET_MATH_INLINE
 FVector UKismetMathLibrary::TransformDirection(const FTransform& T, FVector Direction)
 {
 	return T.TransformVectorNoScale(Direction);
 }
 
-inline
+KISMET_MATH_INLINE
 FVector UKismetMathLibrary::InverseTransformLocation(const FTransform& T, FVector Location)
 {
 	return T.InverseTransformPosition(Location);
 }
 
-inline
+KISMET_MATH_INLINE
 FVector UKismetMathLibrary::InverseTransformDirection(const FTransform& T, FVector Direction)
 {
 	return T.InverseTransformVectorNoScale(Direction);
 }
 
-inline
+KISMET_MATH_INLINE
 FTransform UKismetMathLibrary::ComposeTransforms(const FTransform& A, const FTransform& B)
 {
 	return A * B;
 }
 
-inline
+KISMET_MATH_INLINE
 FTransform UKismetMathLibrary::MakeRelativeTransform(const FTransform& A, const FTransform& RelativeTo)
 {
 	return A.GetRelativeTransform(RelativeTo);
 }
 
-inline
+KISMET_MATH_INLINE
 FTransform UKismetMathLibrary::InvertTransform(const FTransform& T)
 {
 	return T.Inverse();
 }
 
-inline
+KISMET_MATH_INLINE
 bool UKismetMathLibrary::EqualEqual_TransformTransform(const FTransform& A, const FTransform& B)
 {
 	return NearlyEqual_TransformTransform(A, B);
@@ -927,86 +935,86 @@ bool UKismetMathLibrary::EqualEqual_TransformTransform(const FTransform& A, cons
 /* FIntPoint
 *****************************************************************************/
 
-inline
+KISMET_MATH_INLINE
 FIntPoint UKismetMathLibrary::IntPoint_Zero()
 {
 	return FIntPoint::ZeroValue;
 }
 
-inline
+KISMET_MATH_INLINE
 FIntPoint UKismetMathLibrary::IntPoint_One()
 {
 	return FIntPoint(1);
 }
 
-inline
+KISMET_MATH_INLINE
 FIntPoint UKismetMathLibrary::IntPoint_Up()
 {
 	return FIntPoint(0, -1);
 }
 
-inline
+KISMET_MATH_INLINE
 FIntPoint UKismetMathLibrary::IntPoint_Left()
 {
 	return FIntPoint(-1, 0);
 }
 
-inline
+KISMET_MATH_INLINE
 FIntPoint UKismetMathLibrary::IntPoint_Right()
 {
 	return FIntPoint(1, 0);
 }
 
-inline
+KISMET_MATH_INLINE
 FIntPoint UKismetMathLibrary::IntPoint_Down()
 {
 	return FIntPoint(0, 1);
 }
 
-inline
+KISMET_MATH_INLINE
 FVector2D UKismetMathLibrary::Conv_IntPointToVector2D(FIntPoint InIntPoint)
 {
 	return FVector2D(InIntPoint.X, InIntPoint.Y);
 }
 
 
-inline
+KISMET_MATH_INLINE
 FIntPoint UKismetMathLibrary::Add_IntPointIntPoint(FIntPoint A, FIntPoint B)
 {
 	return A + B;
 }
 
-inline
+KISMET_MATH_INLINE
 FIntPoint UKismetMathLibrary::Add_IntPointInt(FIntPoint A, int32 B)
 {
 	return A + FIntPoint(B);
 }
 
-inline
+KISMET_MATH_INLINE
 FIntPoint UKismetMathLibrary::Subtract_IntPointIntPoint(FIntPoint A, FIntPoint B)
 {
 	return A - B;
 }
 
-inline
+KISMET_MATH_INLINE
 FIntPoint UKismetMathLibrary::Subtract_IntPointInt(FIntPoint A, int32 B)
 {
 	return A - FIntPoint(B);
 }
 
-inline
+KISMET_MATH_INLINE
 FIntPoint UKismetMathLibrary::Multiply_IntPointIntPoint(FIntPoint A, FIntPoint B)
 {
 	return A * B;
 }
 
-inline
+KISMET_MATH_INLINE
 FIntPoint UKismetMathLibrary::Multiply_IntPointInt(FIntPoint A, int32 B)
 {
 	return A * FIntPoint(B);
 }
 
-inline
+KISMET_MATH_INLINE
 FIntPoint UKismetMathLibrary::Divide_IntPointIntPoint(FIntPoint A, FIntPoint B)
 {
 	if (B.X == 0 || B.Y == 0)
@@ -1017,7 +1025,7 @@ FIntPoint UKismetMathLibrary::Divide_IntPointIntPoint(FIntPoint A, FIntPoint B)
 	return A / B;
 }
 
-inline
+KISMET_MATH_INLINE
 FIntPoint UKismetMathLibrary::Divide_IntPointInt(FIntPoint A, int32 B)
 {
 	if (B == 0)
@@ -1029,13 +1037,13 @@ FIntPoint UKismetMathLibrary::Divide_IntPointInt(FIntPoint A, int32 B)
 	return A / FIntPoint(B);
 }
 
-inline
+KISMET_MATH_INLINE
 bool UKismetMathLibrary::Equal_IntPointIntPoint(FIntPoint A, FIntPoint B)
 {
 	return A == B;
 }
 
-inline
+KISMET_MATH_INLINE
 bool UKismetMathLibrary::NotEqual_IntPointIntPoint(FIntPoint A, FIntPoint B)
 {
 	return A != B;
@@ -1045,74 +1053,74 @@ bool UKismetMathLibrary::NotEqual_IntPointIntPoint(FIntPoint A, FIntPoint B)
 /* FVector2D
 *****************************************************************************/
 
-inline
+KISMET_MATH_INLINE
 FVector2D UKismetMathLibrary::Vector2D_One()
 {
 	return FVector2D::UnitVector;
 }
 
-inline
+KISMET_MATH_INLINE
 FVector2D UKismetMathLibrary::Vector2D_Unit45Deg()
 {
 	return FVector2D::Unit45Deg;
 }
 
-inline
+KISMET_MATH_INLINE
 FVector2D UKismetMathLibrary::Vector2D_Zero()
 {
 	return FVector2D::ZeroVector;
 }
 
-inline
+KISMET_MATH_INLINE
 FVector2D UKismetMathLibrary::MakeVector2D(double X, double Y)
 {
 	return FVector2D(X, Y);
 }
 
-inline
+KISMET_MATH_INLINE
 void UKismetMathLibrary::BreakVector2D(FVector2D InVec, double& X, double& Y)
 {
 	X = InVec.X;
 	Y = InVec.Y;
 }
 
-inline
+KISMET_MATH_INLINE
 FVector UKismetMathLibrary::Conv_Vector2DToVector(FVector2D InVec2D, float Z)
 {
 	return FVector(InVec2D, Z);
 }
 
-inline
+KISMET_MATH_INLINE
 FIntPoint UKismetMathLibrary::Conv_Vector2DToIntPoint(FVector2D InVec2D)
 {
 	return InVec2D.IntPoint();
 }
 
-inline
+KISMET_MATH_INLINE
 FVector2D UKismetMathLibrary::Add_Vector2DVector2D(FVector2D A, FVector2D B)
 {
 	return A + B;
 }
 
-inline
+KISMET_MATH_INLINE
 FVector2D UKismetMathLibrary::Subtract_Vector2DVector2D(FVector2D A, FVector2D B)
 {
 	return A - B;
 }
 
-inline
+KISMET_MATH_INLINE
 FVector2D UKismetMathLibrary::Multiply_Vector2DFloat(FVector2D A, double B)
 {
 	return A * B;
 }
 
-inline
+KISMET_MATH_INLINE
 FVector2D UKismetMathLibrary::Multiply_Vector2DVector2D(FVector2D A, FVector2D B)
 {
 	return A * B;
 }
 
-inline
+KISMET_MATH_INLINE
 FVector2D UKismetMathLibrary::Divide_Vector2DFloat(FVector2D A, double B)
 {
 	if (B == 0.f)
@@ -1124,7 +1132,7 @@ FVector2D UKismetMathLibrary::Divide_Vector2DFloat(FVector2D A, double B)
 	return A / B;
 }
 
-inline
+KISMET_MATH_INLINE
 FVector2D UKismetMathLibrary::Divide_Vector2DVector2D(FVector2D A, FVector2D B)
 {
 	if (B.X == 0.f || B.Y == 0.f)
@@ -1136,187 +1144,187 @@ FVector2D UKismetMathLibrary::Divide_Vector2DVector2D(FVector2D A, FVector2D B)
 	return A / B;
 }
 
-inline
+KISMET_MATH_INLINE
 FVector2D UKismetMathLibrary::Add_Vector2DFloat(FVector2D A, double B)
 {
 	return A+B;
 }
 
-inline
+KISMET_MATH_INLINE
 FVector2D UKismetMathLibrary::Subtract_Vector2DFloat(FVector2D A, double B)
 {
 	return A-B;
 }
 
-inline
+KISMET_MATH_INLINE
 bool UKismetMathLibrary::EqualExactly_Vector2DVector2D(FVector2D A, FVector2D B)
 { 
 	return A == B;
 }
 
-inline
+KISMET_MATH_INLINE
 bool UKismetMathLibrary::EqualEqual_Vector2DVector2D(FVector2D A, FVector2D B, float ErrorTolerance)
 { 
 	return A.Equals(B,ErrorTolerance);
 }
 
-inline
+KISMET_MATH_INLINE
 bool UKismetMathLibrary::NotEqualExactly_Vector2DVector2D(FVector2D A, FVector2D B)
 {
 	return A != B;
 }
 
-inline
+KISMET_MATH_INLINE
 bool UKismetMathLibrary::NotEqual_Vector2DVector2D(FVector2D A, FVector2D B, float ErrorTolerance)
 {
 	return !A.Equals(B,ErrorTolerance);
 }
 
-inline
+KISMET_MATH_INLINE
 FVector2D UKismetMathLibrary::Negated2D(const FVector2D& A)
 {
 	return FVector2D(-A.X, -A.Y);
 }
 
-inline
+KISMET_MATH_INLINE
 void UKismetMathLibrary::Set2D(FVector2D& A, double X, double Y)
 {
 	A.Set(X, Y);
 }
 
-inline
+KISMET_MATH_INLINE
 FVector2D UKismetMathLibrary::ClampAxes2D(FVector2D A, double MinAxisVal, double MaxAxisVal)
 {
 	return A.ClampAxes(MinAxisVal, MaxAxisVal);
 }
 
-inline
+KISMET_MATH_INLINE
 double UKismetMathLibrary::CrossProduct2D(FVector2D A, FVector2D B)
 {
 	return FVector2D::CrossProduct(A, B);
 }
 
-inline
+KISMET_MATH_INLINE
 double UKismetMathLibrary::Distance2D(FVector2D V1, FVector2D V2)
 {
 	return FVector2D::Distance(V1, V2);
 }
 
-inline
+KISMET_MATH_INLINE
 double UKismetMathLibrary::DistanceSquared2D(FVector2D V1, FVector2D V2)
 {
 	return FVector2D::DistSquared(V1, V2);
 }
 
-inline
+KISMET_MATH_INLINE
 double UKismetMathLibrary::DotProduct2D(FVector2D A, FVector2D B)
 {
 	return FVector2D::DotProduct(A, B);
 }
 
-inline
+KISMET_MATH_INLINE
 FVector2D UKismetMathLibrary::GetAbs2D(FVector2D A)
 {
 	return A.GetAbs();
 }
 
-inline
+KISMET_MATH_INLINE
 double UKismetMathLibrary::GetAbsMax2D(FVector2D A)
 {
 	return A.GetAbsMax();
 }
 
-inline
+KISMET_MATH_INLINE
 double UKismetMathLibrary::GetMax2D(FVector2D A)
 {
 	return A.GetMax();
 }
 
-inline
+KISMET_MATH_INLINE
 double UKismetMathLibrary::GetMin2D(FVector2D A)
 {
 	return A.GetMin();
 }
 
-inline
+KISMET_MATH_INLINE
 FVector2D UKismetMathLibrary::GetRotated2D(FVector2D A, float AngleDeg)
 {
 	return A.GetRotated(AngleDeg);
 }
 
-inline
+KISMET_MATH_INLINE
 bool UKismetMathLibrary::IsNearlyZero2D(const FVector2D& A, float Tolerance)
 {
 	return A.IsNearlyZero(Tolerance);
 }
 
-inline
+KISMET_MATH_INLINE
 bool UKismetMathLibrary::IsZero2D(const FVector2D& A)
 {
 	return A.IsZero();
 }
 
-inline
+KISMET_MATH_INLINE
 FVector2D UKismetMathLibrary::Vector2DInterpTo(FVector2D Current, FVector2D Target, float DeltaTime, float InterpSpeed)
 {
 	return FMath::Vector2DInterpTo(Current, Target, DeltaTime, InterpSpeed);
 }
 
-inline
+KISMET_MATH_INLINE
 FVector2D UKismetMathLibrary::Vector2DInterpTo_Constant(FVector2D Current, FVector2D Target, float DeltaTime, float InterpSpeed)
 {
 	return FMath::Vector2DInterpConstantTo(Current, Target, DeltaTime, InterpSpeed);
 }
 
-inline
+KISMET_MATH_INLINE
 FVector2D UKismetMathLibrary::Normal2D(FVector2D A)
 {
 	return A.GetSafeNormal();
 }
 
-inline
+KISMET_MATH_INLINE
 FVector2D UKismetMathLibrary::NormalSafe2D(FVector2D A, float Tolerance)
 {
 	return A.GetSafeNormal(Tolerance);
 }
 
-inline
+KISMET_MATH_INLINE
 void UKismetMathLibrary::Normalize2D(FVector2D& A, float Tolerance)
 {
 	A.Normalize(Tolerance);
 }
 
-inline
+KISMET_MATH_INLINE
 FVector UKismetMathLibrary::Spherical2DToUnitCartesian(FVector2D A)
 {
 	return A.SphericalToUnitCartesian();
 }
 
-inline
+KISMET_MATH_INLINE
 void UKismetMathLibrary::ToDirectionAndLength2D(FVector2D A, FVector2D& OutDir, double& OutLength)
 {
 	A.ToDirectionAndLength(OutDir, OutLength);
 }
 
-inline
+KISMET_MATH_INLINE
 FVector2D UKismetMathLibrary::ToRounded2D(FVector2D A)
 {
 	return A.RoundToVector();
 }
 
-inline
+KISMET_MATH_INLINE
 FVector2D UKismetMathLibrary::ToSign2D(FVector2D A)
 {
 	return A.GetSignVector();
 }
 
-inline
+KISMET_MATH_INLINE
 double UKismetMathLibrary::VSize2D(FVector2D A)
 {
 	return A.Size();
 }
 
-inline
+KISMET_MATH_INLINE
 double UKismetMathLibrary::VSize2DSquared(FVector2D A)
 {
 	return A.SizeSquared();
@@ -1326,73 +1334,73 @@ double UKismetMathLibrary::VSize2DSquared(FVector2D A)
 /* FVector
 *****************************************************************************/
 
-inline
+KISMET_MATH_INLINE
 FVector UKismetMathLibrary::Vector_Zero()
 {
 	return FVector::ZeroVector;
 }
 
-inline
+KISMET_MATH_INLINE
 FVector UKismetMathLibrary::Vector_One()
 {
 	return FVector::OneVector;
 }
 
-inline
+KISMET_MATH_INLINE
 FVector UKismetMathLibrary::Vector_Forward()
 {
 	return FVector::ForwardVector;
 }
 
-inline
+KISMET_MATH_INLINE
 FVector UKismetMathLibrary::Vector_Backward()
 {
 	return FVector::BackwardVector;
 }
 
-inline
+KISMET_MATH_INLINE
 FVector UKismetMathLibrary::Vector_Up()
 {
 	return FVector::UpVector;
 }
 
-inline
+KISMET_MATH_INLINE
 FVector UKismetMathLibrary::Vector_Down()
 {
 	return FVector::DownVector;
 }
 
-inline
+KISMET_MATH_INLINE
 FVector UKismetMathLibrary::Vector_Right()
 {
 	return FVector::RightVector;
 }
 
-inline
+KISMET_MATH_INLINE
 FVector UKismetMathLibrary::Vector_Left()
 {
 	return FVector::LeftVector;
 }
 
-inline
+KISMET_MATH_INLINE
 FVector UKismetMathLibrary::MakeVector(double X, double Y, double Z)
 {
 	return FVector(X, Y, Z);
 }
 
-inline
+KISMET_MATH_INLINE
 void UKismetMathLibrary::Vector_Assign(FVector& A, const FVector& InVector)
 {
 	A = InVector;
 }
 
-inline
+KISMET_MATH_INLINE
 void UKismetMathLibrary::Vector_Set(FVector& A, double X, double Y, double Z)
 {
 	A.Set(X, Y, Z);
 }
 
-inline
+KISMET_MATH_INLINE
 void UKismetMathLibrary::BreakVector(FVector InVec, double& X, double& Y, double& Z)
 {
 	X = InVec.X;
@@ -1400,7 +1408,7 @@ void UKismetMathLibrary::BreakVector(FVector InVec, double& X, double& Y, double
 	Z = InVec.Z;
 }
 
-inline
+KISMET_MATH_INLINE
 void UKismetMathLibrary::BreakVector3f(FVector3f InVec, float& X, float& Y, float& Z)
 {
 	X = InVec.X;
@@ -1408,67 +1416,67 @@ void UKismetMathLibrary::BreakVector3f(FVector3f InVec, float& X, float& Y, floa
 	Z = InVec.Z;
 }
 
-inline
+KISMET_MATH_INLINE
 FRotator UKismetMathLibrary::Conv_VectorToRotator(FVector InVec)
 {
 	return InVec.ToOrientationRotator();
 }
 
-inline
+KISMET_MATH_INLINE
 FQuat UKismetMathLibrary::Conv_VectorToQuaternion(FVector InVec)
 {
 	return InVec.ToOrientationQuat();
 }
 
-inline
+KISMET_MATH_INLINE
 FVector UKismetMathLibrary::Vector_SlerpVectorToDirection(FVector Vector, FVector Direction, double Alpha)
 {
 	return FVector::SlerpVectorToDirection(Vector, Direction, Alpha);
 }
 
-inline
+KISMET_MATH_INLINE
 FVector UKismetMathLibrary::Vector_SlerpNormals(FVector NormalA, FVector NormalB, double Alpha)
 {
 	return FVector::SlerpNormals(NormalA, NormalB, Alpha);
 }
 
-inline
+KISMET_MATH_INLINE
 FLinearColor UKismetMathLibrary::Conv_VectorToLinearColor(FVector InVec)
 {
 	return FLinearColor(InVec);	
 }
 
-inline
+KISMET_MATH_INLINE
 FVector2D UKismetMathLibrary::Conv_VectorToVector2D(FVector InVec)
 {
 	return FVector2D(InVec);
 }
 
-inline
+KISMET_MATH_INLINE
 FTransform UKismetMathLibrary::Conv_VectorToTransform(FVector InTranslation)
 {
 	return FTransform(InTranslation);
 }
 
-inline
+KISMET_MATH_INLINE
 FVector UKismetMathLibrary::Multiply_VectorFloat(FVector A, double B)
 {
 	return A * B;
 }
 
-inline
+KISMET_MATH_INLINE
 FVector UKismetMathLibrary::Multiply_VectorInt(FVector A, int32 B)
 {
 	return A * (float)B;
 }
 
-inline
+KISMET_MATH_INLINE
 FVector UKismetMathLibrary::Multiply_VectorVector(FVector A, FVector B)
 {
 	return A * B;
 }	
 
-inline
+KISMET_MATH_INLINE
 FVector UKismetMathLibrary::Divide_VectorFloat(FVector A, double B)
 {
 	if (B == 0.f)
@@ -1480,7 +1488,7 @@ FVector UKismetMathLibrary::Divide_VectorFloat(FVector A, double B)
 	return A / B;
 }
 
-inline
+KISMET_MATH_INLINE
 FVector UKismetMathLibrary::Divide_VectorInt(FVector A, int32 B)
 {
 	if (B == 0)
@@ -1492,7 +1500,7 @@ FVector UKismetMathLibrary::Divide_VectorInt(FVector A, int32 B)
 	return A / (float)B;
 }
 
-inline
+KISMET_MATH_INLINE
 FVector UKismetMathLibrary::Divide_VectorVector(FVector A, FVector B)
 {
 	if (B.X == 0.f || B.Y == 0.f || B.Z == 0.f)
@@ -1504,223 +1512,223 @@ FVector UKismetMathLibrary::Divide_VectorVector(FVector A, FVector B)
 	return A / B;
 }
 
-inline
+KISMET_MATH_INLINE
 FVector UKismetMathLibrary::Add_VectorVector(FVector A, FVector B)
 {
 	return A + B;
 }	
 
-inline
+KISMET_MATH_INLINE
 FVector UKismetMathLibrary::Add_VectorFloat(FVector A, double B)
 {
 	return A + B;
 }	
 
-inline
+KISMET_MATH_INLINE
 FVector UKismetMathLibrary::Add_VectorInt(FVector A, int32 B)
 {
 	return A + (float)B;
 }
 
-inline
+KISMET_MATH_INLINE
 FVector UKismetMathLibrary::Subtract_VectorVector(FVector A, FVector B)
 {
 	return A - B;
 }	
 
-inline
+KISMET_MATH_INLINE
 FVector UKismetMathLibrary::Subtract_VectorFloat(FVector A, double B)
 {
 	return A - B;
 }
 
-inline
+KISMET_MATH_INLINE
 FVector UKismetMathLibrary::Subtract_VectorInt(FVector A, int32 B)
 {
 	return A - (float)B;
 }
 
-inline
+KISMET_MATH_INLINE
 FVector UKismetMathLibrary::LessLess_VectorRotator(FVector A, FRotator B)
 {
 	return B.UnrotateVector(A);
 }	
 
-inline
+KISMET_MATH_INLINE
 FVector UKismetMathLibrary::GreaterGreater_VectorRotator(FVector A, FRotator B)
 {
 	return B.RotateVector(A);
 }	
 
-inline
+KISMET_MATH_INLINE
 bool UKismetMathLibrary::EqualExactly_VectorVector(FVector A, FVector B)
 {
 	return A == B;
 }	
 
-inline
+KISMET_MATH_INLINE
 bool UKismetMathLibrary::EqualEqual_VectorVector(FVector A, FVector B, float ErrorTolerance)
 {
 	return A.Equals(B, ErrorTolerance);
 }	
 
-inline
+KISMET_MATH_INLINE
 bool UKismetMathLibrary::NotEqualExactly_VectorVector(FVector A, FVector B)
 {
 	return A != B;
 }	
 
-inline
+KISMET_MATH_INLINE
 bool UKismetMathLibrary::NotEqual_VectorVector(FVector A, FVector B, float ErrorTolerance)
 {
 	return !A.Equals(B, ErrorTolerance);
 }	
 
-inline
+KISMET_MATH_INLINE
 FVector UKismetMathLibrary::Cross_VectorVector(FVector A, FVector B)
 {
 	return FVector::CrossProduct(A, B);
 }
 
-inline
+KISMET_MATH_INLINE
 double UKismetMathLibrary::Dot_VectorVector(FVector A, FVector B)
 {
 	return FVector::DotProduct(A, B);
 }
 
-inline
+KISMET_MATH_INLINE
 double UKismetMathLibrary::Vector_Distance(FVector V1, FVector V2)
 {
 	return FVector::Distance(V1, V2);
 }
 
-inline
+KISMET_MATH_INLINE
 double UKismetMathLibrary::Vector_DistanceSquared(FVector V1, FVector V2)
 {
 	return FVector::DistSquared(V1, V2);
 }
 
-inline
+KISMET_MATH_INLINE
 double UKismetMathLibrary::Vector_Distance2D(FVector V1, FVector V2)
 {
 	return FVector::DistXY(V1, V2);
 }
 
-inline
+KISMET_MATH_INLINE
 double UKismetMathLibrary::Vector_Distance2DSquared(FVector V1, FVector V2)
 {
 	return FVector::DistSquaredXY(V1, V2);
 }
 
-inline
+KISMET_MATH_INLINE
 double UKismetMathLibrary::VSize(FVector A)
 {
 	return A.Size();
 }
 
-inline
+KISMET_MATH_INLINE
 double UKismetMathLibrary::VSizeSquared(FVector A)
 {
 	return A.SizeSquared();
 }
 
-inline
+KISMET_MATH_INLINE
 double UKismetMathLibrary::VSizeXY(FVector A)
 {
 	return A.Size2D();
 }	
 
-inline
+KISMET_MATH_INLINE
 double UKismetMathLibrary::VSizeXYSquared(FVector A)
 {
 	return A.SizeSquared2D();
 }	
 
-inline
+KISMET_MATH_INLINE
 bool UKismetMathLibrary::Vector_IsNearlyZero(const FVector& A, float Tolerance)
 {
 	return A.IsNearlyZero(Tolerance);
 }
 
-inline
+KISMET_MATH_INLINE
 bool UKismetMathLibrary::Vector_IsZero(const FVector& A)
 {
 	return (A.X == 0.0f) && (A.Y == 0.0f) && (A.Z == 0.0f);
 }
 
-inline
+KISMET_MATH_INLINE
 bool UKismetMathLibrary::Vector_IsNAN(const FVector& A)
 {
 	return A.ContainsNaN();
 }
 
-inline
+KISMET_MATH_INLINE
 bool UKismetMathLibrary::Vector_IsUniform(const FVector& A, float Tolerance)
 {
 	return A.AllComponentsEqual(Tolerance);
 }
 
-inline
+KISMET_MATH_INLINE
 bool UKismetMathLibrary::Vector_IsUnit(const FVector& A, float SquaredLenthTolerance)
 {
 	return A.IsUnit(SquaredLenthTolerance);
 }
 
-inline
+KISMET_MATH_INLINE
 bool UKismetMathLibrary::Vector_IsNormal(const FVector& A)
 {
 	return A.IsNormalized();
 }
 
-inline
+KISMET_MATH_INLINE
 FVector UKismetMathLibrary::Normal(FVector A, float Tolerance)
 {
 	return A.GetSafeNormal(Tolerance);
 }
 
-inline
+KISMET_MATH_INLINE
 FVector UKismetMathLibrary::Vector_Normal2D(FVector A, float Tolerance)
 {
 	return A.GetSafeNormal2D(Tolerance);
 }
 
-inline
+KISMET_MATH_INLINE
 FVector UKismetMathLibrary::Vector_NormalUnsafe(const FVector& A)
 {
 	return A.GetUnsafeNormal();
 }
 
-inline
+KISMET_MATH_INLINE
 void UKismetMathLibrary::Vector_Normalize(FVector& A, float Tolerance)
 {
 	A.Normalize(Tolerance);
 }
 
-inline
+KISMET_MATH_INLINE
 FVector UKismetMathLibrary::VLerp(FVector A, FVector B, float V)
 {
 	return A + V * (B - A);
 }	
 
-inline
+KISMET_MATH_INLINE
 FVector UKismetMathLibrary::VInterpTo(FVector Current, FVector Target, float DeltaTime, float InterpSpeed)
 {
 	return FMath::VInterpTo(Current, Target, DeltaTime, InterpSpeed);
 }
 
-inline
+KISMET_MATH_INLINE
 FVector UKismetMathLibrary::VInterpTo_Constant(FVector Current, FVector Target, float DeltaTime, float InterpSpeed)
 {
 	return FMath::VInterpConstantTo(Current, Target, DeltaTime, InterpSpeed);
 }
 
-inline
+KISMET_MATH_INLINE
 FVector UKismetMathLibrary::Vector_Reciprocal(const FVector& A)
 {
 	return A.Reciprocal();
 }	
 
-inline
+KISMET_MATH_INLINE
 FVector UKismetMathLibrary::RandomPointInBoundingBox(const FVector Center, const FVector HalfSize)
 {
 	const FVector BoxMin = Center - HalfSize;
@@ -1728,69 +1736,69 @@ FVector UKismetMathLibrary::RandomPointInBoundingBox(const FVector Center, const
 	return FMath::RandPointInBox(FBox(BoxMin, BoxMax));
 }
 
-inline
+KISMET_MATH_INLINE
 FVector UKismetMathLibrary::RandomPointInBoundingBox_Box(const FBox Box)
 {
 	return FMath::RandPointInBox(Box);
 }
 
-inline
+KISMET_MATH_INLINE
 FVector UKismetMathLibrary::RandomUnitVectorInConeInRadians(FVector ConeDir, float ConeHalfAngleInRadians)
 {
 	return FMath::VRandCone(ConeDir, ConeHalfAngleInRadians);
 }
 
-inline
+KISMET_MATH_INLINE
 FVector UKismetMathLibrary::GetReflectionVector(FVector Direction, FVector SurfaceNormal)
 {
 	// Same as MirrorVectorByNormal() below - though they are bound to different names so keeping both for now
 	return FMath::GetReflectionVector(Direction, SurfaceNormal);
 }
 
-inline
+KISMET_MATH_INLINE
 FVector UKismetMathLibrary::MirrorVectorByNormal(FVector InVect, FVector InNormal)
 {
 	// Same as GetReflectionVector() above - though they are bound to cdifferent names so keeping both for now
 	return FMath::GetReflectionVector(InVect, InNormal);
 }
 
-inline
+KISMET_MATH_INLINE
 FVector UKismetMathLibrary::Vector_MirrorByPlane(FVector A, const FPlane& InPlane)
 {
 	return A.MirrorByPlane(InPlane);
 }
 
-inline
+KISMET_MATH_INLINE
 FVector UKismetMathLibrary::Vector_SnappedToGrid(FVector InVect, float InGridSize)
 {
 	return InVect.GridSnap(InGridSize);
 }
 
-inline
+KISMET_MATH_INLINE
 FVector UKismetMathLibrary::Vector_BoundedToCube(FVector InVect, float InRadius)
 {
 	return InVect.BoundToCube(InRadius);
 }
 
-inline
+KISMET_MATH_INLINE
 void UKismetMathLibrary::Vector_AddBounded(FVector& A, FVector InAddVect, float InRadius)
 {
 	return A.AddBounded(InAddVect, InRadius);
 }
 
-inline
+KISMET_MATH_INLINE
 FVector UKismetMathLibrary::Vector_BoundedToBox(FVector InVect, FVector InBoxMin, FVector InBoxMax)
 {
 	return InVect.BoundToBox(InBoxMin, InBoxMax);
 }
 
-inline
+KISMET_MATH_INLINE
 FVector UKismetMathLibrary::Vector_ProjectOnToNormal(FVector V, FVector InNormal)
 {
 	return V.ProjectOnToNormal(InNormal);
 }
 
-inline
+KISMET_MATH_INLINE
 FVector UKismetMathLibrary::ProjectVectorOnToVector(FVector V, FVector Target)
 {
 	if (Target.SizeSquared() > UE_SMALL_NUMBER)
@@ -1804,169 +1812,169 @@ FVector UKismetMathLibrary::ProjectVectorOnToVector(FVector V, FVector Target)
 	}
 }
 
-inline
+KISMET_MATH_INLINE
 void UKismetMathLibrary::FindNearestPointsOnLineSegments(FVector Segment1Start, FVector Segment1End, FVector Segment2Start, FVector Segment2End, FVector& Segment1Point, FVector& Segment2Point)
 {
 	FMath::SegmentDistToSegmentSafe(Segment1Start, Segment1End, Segment2Start, Segment2End, Segment1Point, Segment2Point);
 }
 
-inline
+KISMET_MATH_INLINE
 FVector UKismetMathLibrary::FindClosestPointOnSegment(FVector Point, FVector SegmentStart, FVector SegmentEnd)
 {
 	return FMath::ClosestPointOnSegment(Point, SegmentStart, SegmentEnd);
 }
 
-inline
+KISMET_MATH_INLINE
 float UKismetMathLibrary::GetPointDistanceToSegment(FVector Point, FVector SegmentStart, FVector SegmentEnd)
 {
 	return FMath::PointDistToSegment(Point, SegmentStart, SegmentEnd);
 }
 
-inline
+KISMET_MATH_INLINE
 float UKismetMathLibrary::GetPointDistanceToLine(FVector Point, FVector LineOrigin, FVector LineDirection)
 {
 	return FMath::PointDistToLine(Point, LineDirection, LineOrigin);
 }
 
-inline
+KISMET_MATH_INLINE
 FVector UKismetMathLibrary::ProjectPointOnToPlane(FVector Point, FVector PlaneBase, FVector PlaneNormal)
 {
 	return FVector::PointPlaneProject(Point, PlaneBase, PlaneNormal);
 }
 
-inline
+KISMET_MATH_INLINE
 FVector UKismetMathLibrary::ProjectVectorOnToPlane(FVector V, FVector PlaneNormal)
 {
 	return FVector::VectorPlaneProject(V, PlaneNormal);
 }
 
-inline
+KISMET_MATH_INLINE
 FVector UKismetMathLibrary::NegateVector(FVector A)
 {
 	return -A;
 }
 
-inline
+KISMET_MATH_INLINE
 void UKismetMathLibrary::Vector_UnwindEuler(FVector& A)
 {
 	A.UnwindEuler();
 }
 
-inline
+KISMET_MATH_INLINE
 FVector UKismetMathLibrary::ClampVectorSize(FVector A, double Min, double Max)
 {
 	return A.GetClampedToSize(Min, Max);
 }
 
-inline
+KISMET_MATH_INLINE
 FVector UKismetMathLibrary::Vector_ClampSize2D(FVector A, double Min, double Max)
 {
 	return A.GetClampedToSize2D(Min, Max);
 }
 
-inline
+KISMET_MATH_INLINE
 FVector UKismetMathLibrary::Vector_ClampSizeMax(FVector A, double Max)
 {
 	return A.GetClampedToMaxSize(Max);
 }
 
-inline
+KISMET_MATH_INLINE
 FVector UKismetMathLibrary::Vector_ClampSizeMax2D(FVector A, double Max)
 {
 	return A.GetClampedToMaxSize2D(Max);
 }
 
-inline
+KISMET_MATH_INLINE
 double UKismetMathLibrary::GetMinElement(FVector A)
 {
 	return A.GetMin();
 }
 
-inline
+KISMET_MATH_INLINE
 double UKismetMathLibrary::GetMaxElement(FVector A)
 {
 	return A.GetMax();
 }
 
-inline
+KISMET_MATH_INLINE
 double UKismetMathLibrary::Vector_GetAbsMax(FVector A)
 {
 	return A.GetAbsMax();
 }
 
-inline
+KISMET_MATH_INLINE
 double UKismetMathLibrary::Vector_GetAbsMin(FVector A)
 {
 	return A.GetAbsMin();
 }
 
-inline
+KISMET_MATH_INLINE
 FVector UKismetMathLibrary::Vector_GetAbs(FVector A)
 {
 	return A.GetAbs();
 }
 
-inline
+KISMET_MATH_INLINE
 FVector UKismetMathLibrary::Vector_ComponentMin(FVector A, FVector B)
 {
 	return A.ComponentMin(B);
 }
 
-inline
+KISMET_MATH_INLINE
 FVector UKismetMathLibrary::Vector_ComponentMax(FVector A, FVector B)
 {
 	return A.ComponentMax(B);
 }
 
-inline
+KISMET_MATH_INLINE
 FVector UKismetMathLibrary::Vector_GetSignVector(FVector A)
 {
 	return A.GetSignVector();
 }
 
-inline
+KISMET_MATH_INLINE
 FVector UKismetMathLibrary::Vector_GetProjection(FVector A)
 {
 	return A.Projection();
 }
 
-inline
+KISMET_MATH_INLINE
 double UKismetMathLibrary::Vector_HeadingAngle(FVector A)
 {
 	return A.HeadingAngle();
 }
 
-inline
+KISMET_MATH_INLINE
 double UKismetMathLibrary::Vector_CosineAngle2D(FVector A, FVector B)
 {
 	return A.CosineAngle2D(B);
 }
 
-inline
+KISMET_MATH_INLINE
 FVector UKismetMathLibrary::Vector_ToRadians(FVector A)
 {
 	return FVector::DegreesToRadians(A);
 }
 
-inline
+KISMET_MATH_INLINE
 FVector UKismetMathLibrary::Vector_ToDegrees(FVector A)
 {
 	return FVector::RadiansToDegrees(A);
 }
 
-inline
+KISMET_MATH_INLINE
 FVector2D UKismetMathLibrary::Vector_UnitCartesianToSpherical(FVector A)
 {
 	return A.UnitCartesianToSpherical();
 }
 
-inline
+KISMET_MATH_INLINE
 FVector UKismetMathLibrary::GetDirectionUnitVector(FVector From, FVector To)
 {
 	return (To - From).GetSafeNormal();
 }
 
-inline
+KISMET_MATH_INLINE
 FIntVector UKismetMathLibrary::FTruncVector(const FVector& InVector)
 {
 	return FIntVector(static_cast<int32>(FMath::TruncToInt(InVector.X)),
@@ -1974,7 +1982,7 @@ FIntVector UKismetMathLibrary::FTruncVector(const FVector& InVector)
 		static_cast<int32>(FMath::TruncToInt(InVector.Z)));
 }
 
-inline
+KISMET_MATH_INLINE
 FIntVector2 UKismetMathLibrary::FTruncVector2D(const FVector2D& InVector2D)
 {
 	return FIntVector2(static_cast<int32>(FMath::TruncToInt(InVector2D.X)),
@@ -1984,19 +1992,19 @@ FIntVector2 UKismetMathLibrary::FTruncVector2D(const FVector2D& InVector2D)
 /* FVector4
 *****************************************************************************/
 
-inline
+KISMET_MATH_INLINE
 FVector4 UKismetMathLibrary::Vector4_Zero()
 {
 	return FVector4(ForceInitToZero);
 }
 
-inline
+KISMET_MATH_INLINE
 FVector4 UKismetMathLibrary::MakeVector4(double X, double Y, double Z, double W)
 {
 	return FVector4(X, Y, Z, W);
 }
 
-inline
+KISMET_MATH_INLINE
 void UKismetMathLibrary::BreakVector4(const FVector4& InVec, double& X, double& Y, double& Z, double& W)
 {
 	X = InVec.X;
@@ -2005,187 +2013,187 @@ void UKismetMathLibrary::BreakVector4(const FVector4& InVec, double& X, double&
 	W = InVec.W;
 }
 
-inline
+KISMET_MATH_INLINE
 FVector UKismetMathLibrary::Conv_Vector4ToVector(const FVector4& InVector4)
 {
 	return FVector(InVector4.X, InVector4.Y, InVector4.Z);
 }
 
-inline
+KISMET_MATH_INLINE
 FRotator UKismetMathLibrary::Conv_Vector4ToRotator(const FVector4& InVec)
 {
 	return InVec.ToOrientationRotator();
 }
 
-inline
+KISMET_MATH_INLINE
 FQuat UKismetMathLibrary::Conv_Vector4ToQuaternion(const FVector4& InVec)
 {
 	return InVec.ToOrientationQuat();
 }
 
-inline
+KISMET_MATH_INLINE
 FVector4 UKismetMathLibrary::Add_Vector4Vector4(const FVector4& A, const FVector4& B)
 {
 	return A + B;
 }	
 
-inline
+KISMET_MATH_INLINE
 FVector4 UKismetMathLibrary::Subtract_Vector4Vector4(const FVector4& A, const FVector4& B)
 {
 	return A - B;
 }	
 
-inline
+KISMET_MATH_INLINE
 FVector4 UKismetMathLibrary::Multiply_Vector4Vector4(const FVector4& A, const FVector4& B)
 {
 	return A * B;
 }	
 
-inline
+KISMET_MATH_INLINE
 FVector4 UKismetMathLibrary::Divide_Vector4Vector4(const FVector4& A, const FVector4& B)
 {
 	return A / B;
 }	
 
-inline
+KISMET_MATH_INLINE
 bool UKismetMathLibrary::EqualExactly_Vector4Vector4(const FVector4& A, const FVector4& B)
 {
 	return A == B;
 }	
 
-inline
+KISMET_MATH_INLINE
 bool UKismetMathLibrary::EqualEqual_Vector4Vector4(const FVector4& A, const FVector4& B, float ErrorTolerance)
 {
 	return A.Equals(B, ErrorTolerance);
 }	
 
-inline
+KISMET_MATH_INLINE
 bool UKismetMathLibrary::NotEqualExactly_Vector4Vector4(const FVector4& A, const FVector4& B)
 {
 	return A != B;
 }	
 
-inline
+KISMET_MATH_INLINE
 bool UKismetMathLibrary::NotEqual_Vector4Vector4(const FVector4& A, const FVector4& B, float ErrorTolerance)
 {
 	return !A.Equals(B, ErrorTolerance);
 }	
 
-inline
+KISMET_MATH_INLINE
 FVector4 UKismetMathLibrary::Vector4_Negated(const FVector4& A)
 {
 	return -A;
 }
 
-inline
+KISMET_MATH_INLINE
 void UKismetMathLibrary::Vector4_Assign(FVector4& A, const FVector4& InVector)
 {
 	A = InVector;
 }
 
-inline
+KISMET_MATH_INLINE
 void UKismetMathLibrary::Vector4_Set(FVector4& A, double X, double Y, double Z, double W)
 {
 	A.Set(X, Y, Z, W);
 }
 
-inline
+KISMET_MATH_INLINE
 FVector4 UKismetMathLibrary::Vector4_CrossProduct3(const FVector4& A, const FVector4& B)
 {
 	return A ^ B;
 }
 
-inline
+KISMET_MATH_INLINE
 double UKismetMathLibrary::Vector4_DotProduct(const FVector4& A, const FVector4& B)
 {
 	return Dot4(A, B);
 }
 
-inline
+KISMET_MATH_INLINE
 double UKismetMathLibrary::Vector4_DotProduct3(const FVector4& A, const FVector4& B)
 {
 	return Dot3(A, B);
 }
 
-inline
+KISMET_MATH_INLINE
 bool UKismetMathLibrary::Vector4_IsNAN(const FVector4& A)
 {
 	return A.ContainsNaN();
 }
 
-inline
+KISMET_MATH_INLINE
 bool UKismetMathLibrary::Vector4_IsNearlyZero3(const FVector4& A, float Tolerance)
 {
 	return A.IsNearlyZero3(Tolerance);
 }
 
-inline
+KISMET_MATH_INLINE
 bool UKismetMathLibrary::Vector4_IsZero(const FVector4& A)
 {
 	return (A.X == 0.0f) && (A.Y == 0.0f) && (A.Z == 0.0f) && (A.W == 0.0f);
 }
 
-inline
+KISMET_MATH_INLINE
 double UKismetMathLibrary::Vector4_Size(const FVector4& A)
 {
 	return A.Size();
 }
 
-inline
+KISMET_MATH_INLINE
 double UKismetMathLibrary::Vector4_SizeSquared(const FVector4& A)
 {
 	return A.SizeSquared();
 }
 
-inline
+KISMET_MATH_INLINE
 double UKismetMathLibrary::Vector4_Size3(const FVector4& A)
 {
 	return A.Size3();
 }
 
-inline
+KISMET_MATH_INLINE
 double UKismetMathLibrary::Vector4_SizeSquared3(const FVector4& A)
 {
 	return A.SizeSquared3();
 }
 
-inline
+KISMET_MATH_INLINE
 bool UKismetMathLibrary::Vector4_IsUnit3(const FVector4& A, float SquaredLenthTolerance)
 {
 	return A.IsUnit3(SquaredLenthTolerance);
 }
 
-inline
+KISMET_MATH_INLINE
 bool UKismetMathLibrary::Vector4_IsNormal3(const FVector4& A)
 {
 	return A.IsUnit3(UE_THRESH_VECTOR_NORMALIZED);
 }
 
-inline
+KISMET_MATH_INLINE
 FVector4 UKismetMathLibrary::Vector4_Normal3(const FVector4& A, float Tolerance)
 {
 	return A.GetSafeNormal(Tolerance);
 }
 
-inline
+KISMET_MATH_INLINE
 FVector4 UKismetMathLibrary::Vector4_NormalUnsafe3(const FVector4& A)
 {
 	return A.GetUnsafeNormal3();
 }
 
-inline
+KISMET_MATH_INLINE
 void UKismetMathLibrary::Vector4_Normalize3(FVector4& A, float Tolerance)
 {
 	A = A.GetSafeNormal(Tolerance);
 }
 
-inline
+KISMET_MATH_INLINE
 FVector4 UKismetMathLibrary::Vector4_MirrorByVector3(const FVector4& Direction, const FVector4& SurfaceNormal)
 {
 	return Direction.Reflect3(SurfaceNormal);
 }
 
-inline
+KISMET_MATH_INLINE
 FVector4 UKismetMathLibrary::TransformVector4(const FMatrix& Matrix, const FVector4& Vec4)
 {
 	return Matrix.TransformFVector4(Vec4);
@@ -2195,37 +2203,37 @@ FVector4 UKismetMathLibrary::TransformVector4(const FMatrix& Matrix, const FVect
 /* FName
 *****************************************************************************/
 
-inline
+KISMET_MATH_INLINE
 bool UKismetMathLibrary::EqualEqual_NameName(FName A, FName B)
 {
 	return A == B;
 }
 
-inline
+KISMET_MATH_INLINE
 bool UKismetMathLibrary::NotEqual_NameName(FName A, FName B)
 {
 	return A != B;
 }
 
-inline
+KISMET_MATH_INLINE
 bool UKismetMathLibrary::EqualEqual_ObjectObject(class UObject* A, class UObject* B)
 {
 	return A == B;
 }
 
-inline
+KISMET_MATH_INLINE
 bool UKismetMathLibrary::NotEqual_ObjectObject(class UObject* A, class UObject* B)
 {
 	return A != B;
 }
 
-inline
+KISMET_MATH_INLINE
 bool UKismetMathLibrary::EqualEqual_ClassClass(class UClass* A, class UClass* B)
 {
 	return A == B;
 }
 
-inline
+KISMET_MATH_INLINE
 bool UKismetMathLibrary::NotEqual_ClassClass(class UClass* A, class UClass* B)
 {
 	return A != B;
@@ -2234,145 +2242,145 @@ bool UKismetMathLibrary::NotEqual_ClassClass(class UClass* A, class UClass* B)
 /* DateTime functions
 *****************************************************************************/
 
-inline
+KISMET_MATH_INLINE
 FDateTime UKismetMathLibrary::Add_DateTimeTimespan(FDateTime A, FTimespan B)
 {
 	return A + B;
 }
 
-inline
+KISMET_MATH_INLINE
 FDateTime UKismetMathLibrary::Subtract_DateTimeTimespan(FDateTime A, FTimespan B)
 {
 	return A - B;
 }
 
-inline
+KISMET_MATH_INLINE
 FTimespan UKismetMathLibrary::Subtract_DateTimeDateTime(FDateTime A, FDateTime B)
 {
 	return A - B;
 }
 
-inline
+KISMET_MATH_INLINE
 FDateTime UKismetMathLibrary::Add_DateTimeDateTime(FDateTime A, FDateTime B)
 {
 	return A + FTimespan(B.GetTicks());
 }
 
-inline
+KISMET_MATH_INLINE
 bool UKismetMathLibrary::EqualEqual_DateTimeDateTime(FDateTime A, FDateTime B)
 {
 	return A == B;
 }
 
-inline
+KISMET_MATH_INLINE
 bool UKismetMathLibrary::NotEqual_DateTimeDateTime(FDateTime A, FDateTime B)
 {
 	return A != B;
 }
 
-inline
+KISMET_MATH_INLINE
 bool UKismetMathLibrary::Greater_DateTimeDateTime(FDateTime A, FDateTime B)
 {
 	return A > B;
 }
 
-inline
+KISMET_MATH_INLINE
 bool UKismetMathLibrary::GreaterEqual_DateTimeDateTime(FDateTime A, FDateTime B)
 {
 	return A >= B;
 }
 
-inline
+KISMET_MATH_INLINE
 bool UKismetMathLibrary::Less_DateTimeDateTime(FDateTime A, FDateTime B)
 {
 	return A < B;
 }
 
-inline
+KISMET_MATH_INLINE
 bool UKismetMathLibrary::LessEqual_DateTimeDateTime(FDateTime A, FDateTime B)
 {
 	return A <= B;
 }
 
-inline
+KISMET_MATH_INLINE
 FDateTime UKismetMathLibrary::GetDate(FDateTime A)
 {
 	return A.GetDate();
 }
 
-inline
+KISMET_MATH_INLINE
 int32 UKismetMathLibrary::GetDay(FDateTime A)
 {
 	return A.GetDay();
 }
 
-inline
+KISMET_MATH_INLINE
 int32 UKismetMathLibrary::GetDayOfYear(FDateTime A)
 {
 	return A.GetDayOfYear();
 }
 
-inline
+KISMET_MATH_INLINE
 int32 UKismetMathLibrary::GetHour(FDateTime A)
 {
 	return A.GetHour();
 }
 
-inline
+KISMET_MATH_INLINE
 int32 UKismetMathLibrary::GetHour12(FDateTime A)
 {
 	return A.GetHour12();
 }
 
-inline
+KISMET_MATH_INLINE
 int32 UKismetMathLibrary::GetMillisecond(FDateTime A)
 {
 	return A.GetMillisecond();
 }
 
-inline
+KISMET_MATH_INLINE
 int32 UKismetMathLibrary::GetMinute(FDateTime A)
 {
 	return A.GetMinute();
 }
 
-inline
+KISMET_MATH_INLINE
 int32 UKismetMathLibrary::GetMonth(FDateTime A)
 {
 	return A.GetMonth();
 }
 
-inline
+KISMET_MATH_INLINE
 int32 UKismetMathLibrary::GetSecond(FDateTime A)
 {
 	return A.GetSecond();
 }
 
-inline
+KISMET_MATH_INLINE
 FTimespan UKismetMathLibrary::GetTimeOfDay(FDateTime A)
 {
 	return A.GetTimeOfDay();
 }
 
-inline
+KISMET_MATH_INLINE
 int32 UKismetMathLibrary::GetYear(FDateTime A)
 {
 	return A.GetYear();
 }
 
-inline
+KISMET_MATH_INLINE
 bool UKismetMathLibrary::IsAfternoon(FDateTime A)
 {
 	return A.IsAfternoon();
 }
 
-inline
+KISMET_MATH_INLINE
 bool UKismetMathLibrary::IsMorning(FDateTime A)
 {
 	return A.IsMorning();
 }
 
-inline
+KISMET_MATH_INLINE
 int32 UKismetMathLibrary::DaysInMonth(int32 Year, int32 Month)
 {
 	if ((Month < 1) || (Month > 12))
@@ -2384,73 +2392,73 @@ int32 UKismetMathLibrary::DaysInMonth(int32 Year, int32 Month)
 	return FDateTime::DaysInMonth(Year, Month);
 }
 
-inline
+KISMET_MATH_INLINE
 int32 UKismetMathLibrary::DaysInYear(int32 Year)
 {
 	return FDateTime::DaysInYear(Year);
 }
 
-inline
+KISMET_MATH_INLINE
 bool UKismetMathLibrary::IsLeapYear(int32 Year)
 {
 	return FDateTime::IsLeapYear(Year);
 }
 
-inline
+KISMET_MATH_INLINE
 FDateTime UKismetMathLibrary::DateTimeMaxValue()
 {
 	return FDateTime::MaxValue();
 }
 
-inline
+KISMET_MATH_INLINE
 FDateTime UKismetMathLibrary::DateTimeMinValue()
 {
 	return FDateTime::MinValue();
 }
 
-inline
+KISMET_MATH_INLINE
 FDateTime UKismetMathLibrary::Now()
 {
 	return FDateTime::Now();
 }
 
-inline
+KISMET_MATH_INLINE
 FDateTime UKismetMathLibrary::Today()
 {
 	return FDateTime::Today();
 }
 
-inline
+KISMET_MATH_INLINE
 FDateTime UKismetMathLibrary::UtcNow()
 {
 	return FDateTime::UtcNow();
 }
 
-inline
+KISMET_MATH_INLINE
 bool UKismetMathLibrary::DateTimeFromIsoString(FString IsoString, FDateTime& Result)
 {
 	return FDateTime::ParseIso8601(*IsoString, Result);
 }
 
-inline
+KISMET_MATH_INLINE
 bool UKismetMathLibrary::DateTimeFromString(FString DateTimeString, FDateTime& Result)
 {
 	return FDateTime::Parse(DateTimeString, Result);
 }
 
-inline
+KISMET_MATH_INLINE
 int64 UKismetMathLibrary::ToUnixTimestamp(const FDateTime& Time)
 {
 	return Time.ToUnixTimestamp();
 }
 
-inline
+KISMET_MATH_INLINE
 double UKismetMathLibrary::ToUnixTimestampDouble(const FDateTime& Time)
 {
 	return Time.ToUnixTimestampDecimal();
 }
 
-inline
+KISMET_MATH_INLINE
 FDateTime UKismetMathLibrary::FromUnixTimestamp(const int64 UnixTime)
 {
 	return FDateTime::FromUnixTimestamp(UnixTime);
@@ -2459,157 +2467,157 @@ FDateTime UKismetMathLibrary::FromUnixTimestamp(const int64 UnixTime)
 /* Timespan functions
 *****************************************************************************/
 
-inline
+KISMET_MATH_INLINE
 FTimespan UKismetMathLibrary::Add_TimespanTimespan(FTimespan A, FTimespan B)
 {
 	return A + B;
 }
 
-inline
+KISMET_MATH_INLINE
 FTimespan UKismetMathLibrary::Subtract_TimespanTimespan(FTimespan A, FTimespan B)
 {
 	return A - B;
 }
 
-inline
+KISMET_MATH_INLINE
 FTimespan UKismetMathLibrary::Multiply_TimespanFloat(FTimespan A, float Scalar)
 {
 	return A * Scalar;
 }
 
-inline
+KISMET_MATH_INLINE
 FTimespan UKismetMathLibrary::Divide_TimespanFloat(FTimespan A, float Scalar)
 {
 	return A / Scalar;
 }
 
-inline
+KISMET_MATH_INLINE
 bool UKismetMathLibrary::EqualEqual_TimespanTimespan(FTimespan A, FTimespan B)
 {
 	return A == B;
 }
 
-inline
+KISMET_MATH_INLINE
 bool UKismetMathLibrary::NotEqual_TimespanTimespan(FTimespan A, FTimespan B)
 {
 	return A != B;
 }
 
-inline
+KISMET_MATH_INLINE
 bool UKismetMathLibrary::Greater_TimespanTimespan(FTimespan A, FTimespan B)
 {
 	return A > B;
 }
 
-inline
+KISMET_MATH_INLINE
 bool UKismetMathLibrary::GreaterEqual_TimespanTimespan(FTimespan A, FTimespan B)
 {
 	return A >= B;
 }
 
-inline
+KISMET_MATH_INLINE
 bool UKismetMathLibrary::Less_TimespanTimespan(FTimespan A, FTimespan B)
 {
 	return A < B;
 }
 
-inline
+KISMET_MATH_INLINE
 bool UKismetMathLibrary::LessEqual_TimespanTimespan(FTimespan A, FTimespan B)
 {
 	return A <= B;
 }
 
-inline
+KISMET_MATH_INLINE
 int32 UKismetMathLibrary::GetDays(FTimespan A)
 {
 	return A.GetDays();
 }
 
-inline
+KISMET_MATH_INLINE
 FTimespan UKismetMathLibrary::GetDuration(FTimespan A)
 {
 	return A.GetDuration();
 }
 
-inline
+KISMET_MATH_INLINE
 int32 UKismetMathLibrary::GetHours(FTimespan A)
 {
 	return A.GetHours();
 }
 
-inline
+KISMET_MATH_INLINE
 int32 UKismetMathLibrary::GetMilliseconds(FTimespan A)
 {
 	return A.GetFractionMilli();
 }
 
-inline
+KISMET_MATH_INLINE
 int32 UKismetMathLibrary::GetMinutes(FTimespan A)
 {
 	return A.GetMinutes();
 }
 
-inline
+KISMET_MATH_INLINE
 int32 UKismetMathLibrary::GetSeconds(FTimespan A)
 {
 	return A.GetSeconds();
 }
 
-inline
+KISMET_MATH_INLINE
 double UKismetMathLibrary::GetTotalDays(FTimespan A)
 {
 	return A.GetTotalDays();
 }
 
-inline
+KISMET_MATH_INLINE
 double UKismetMathLibrary::GetTotalHours(FTimespan A)
 {
 	return A.GetTotalHours();
 }
 
-inline
+KISMET_MATH_INLINE
 double UKismetMathLibrary::GetTotalMilliseconds(FTimespan A)
 {
 	return A.GetTotalMilliseconds();
 }
 
-inline
+KISMET_MATH_INLINE
 double UKismetMathLibrary::GetTotalMinutes(FTimespan A)
 {
 	return A.GetTotalMinutes();
 }
 
-inline
+KISMET_MATH_INLINE
 double UKismetMathLibrary::GetTotalSeconds(FTimespan A)
 {
 	return A.GetTotalSeconds();
 }
 
-inline
+KISMET_MATH_INLINE
 FTimespan UKismetMathLibrary::TimespanMaxValue()
 {
 	return FTimespan::MaxValue();
 }
 
-inline
+KISMET_MATH_INLINE
 FTimespan UKismetMathLibrary::TimespanMinValue()
 {
 	return FTimespan::MinValue();
 }
 
-inline
+KISMET_MATH_INLINE
 float UKismetMathLibrary::TimespanRatio(FTimespan A, FTimespan B)
 {
 	return static_cast<float>(FTimespan::Ratio(A, B));
 }
 
-inline
+KISMET_MATH_INLINE
 FTimespan UKismetMathLibrary::TimespanZeroValue()
 {
 	return FTimespan::Zero();
 }
 
-inline
+KISMET_MATH_INLINE
 bool UKismetMathLibrary::TimespanFromString(FString TimespanString, FTimespan& Result)
 {
 	return FTimespan::Parse(TimespanString, Result);
@@ -2619,193 +2627,193 @@ bool UKismetMathLibrary::TimespanFromString(FString TimespanString, FTimespan& R
 /* K2 Utilities
 *****************************************************************************/
 
-inline
+KISMET_MATH_INLINE
 float UKismetMathLibrary::Conv_ByteToFloat(uint8 InByte)
 {
 	return (float)InByte;
 }
 
-inline
+KISMET_MATH_INLINE
 double UKismetMathLibrary::Conv_ByteToDouble(uint8 InByte)
 {
 	return (double)InByte;
 }
 
-inline
+KISMET_MATH_INLINE
 float UKismetMathLibrary::Conv_IntToFloat(int32 InInt)
 {
 	return (float)InInt;
 }
 
-inline
+KISMET_MATH_INLINE
 double UKismetMathLibrary::Conv_IntToDouble(int32 InInt)
 {
 	return (double)InInt;
 }
 
-inline
+KISMET_MATH_INLINE
 int64 UKismetMathLibrary::Conv_IntToInt64(int32 InInt)
 {
 	return (int64)InInt;
 }
 
-inline
+KISMET_MATH_INLINE
 uint8 UKismetMathLibrary::Conv_IntToByte(int32 InInt)
 {
 	return (uint8)InInt;
 }
 
-inline
+KISMET_MATH_INLINE
 int32 UKismetMathLibrary::Conv_Int64ToInt(int64 InInt)
 {
 	return (int32)InInt;
 }
 
-inline
+KISMET_MATH_INLINE
 float UKismetMathLibrary::Conv_DoubleToFloat(double InDouble)
 {
 	return (float)InDouble;
 }
 
-inline
+KISMET_MATH_INLINE
 double UKismetMathLibrary::Conv_FloatToDouble(float InFloat)
 {
 	return (double)InFloat;
 }
 
-inline
+KISMET_MATH_INLINE
 uint8 UKismetMathLibrary::Conv_Int64ToByte(int64 InInt)
 {
 	return (uint8)InInt;
 }
 
-inline
+KISMET_MATH_INLINE
 int64 UKismetMathLibrary::Conv_DoubleToInt64(double InDouble)
 {
 	return (int64)InDouble;
 }
 
-inline
+KISMET_MATH_INLINE
 double UKismetMathLibrary::Conv_Int64ToDouble(int64 InInt)
 {
 	return (double)InInt;
 }
 
-inline
+KISMET_MATH_INLINE
 FIntVector UKismetMathLibrary::Conv_IntToIntVector(int32 InInt)
 {
 	return FIntVector(InInt, InInt, InInt);
 }
 
-inline
+KISMET_MATH_INLINE
 FIntVector2 UKismetMathLibrary::Conv_IntToIntVector2(int32 InInt)
 {
 	return FIntVector2(InInt, InInt);
 }
 
-inline
+KISMET_MATH_INLINE
 FVector UKismetMathLibrary::Conv_IntToVector(int32 InInt)
 {
 	return FVector((float)InInt);
 }
 
-inline
+KISMET_MATH_INLINE
 bool UKismetMathLibrary::Conv_IntToBool(int32 InInt)
 {
 	return InInt == 0 ? false : true;
 }
 
-inline
+KISMET_MATH_INLINE
 int32 UKismetMathLibrary::Conv_BoolToInt(bool InBool)
 {
 	return InBool ? 1 : 0;
 }
 
-inline
+KISMET_MATH_INLINE
 uint8 UKismetMathLibrary::Conv_BoolToByte(bool InBool)
 {
 	return InBool ? 1 : 0;
 }
 
-inline
+KISMET_MATH_INLINE
 float UKismetMathLibrary::Conv_BoolToFloat(bool InBool)
 {
 	return InBool ? 1.0f : 0.0f;
 }
 
-inline
+KISMET_MATH_INLINE
 double UKismetMathLibrary::Conv_BoolToDouble(bool InBool)
 {
 	return InBool ? 1.0 : 0.0;
 }
 
-inline
+KISMET_MATH_INLINE
 int32 UKismetMathLibrary::Conv_ByteToInt(uint8 InByte)
 {
 	return (int32)InByte;
 }
 
-inline
+KISMET_MATH_INLINE
 int64 UKismetMathLibrary::Conv_ByteToInt64(uint8 InByte)
 {
 	return (int64)InByte;
 }
 
-inline
+KISMET_MATH_INLINE
 FTransform UKismetMathLibrary::Conv_RotatorToTransform(const FRotator& InRotator)
 {
 	return FTransform(InRotator);
 }
 
-inline
+KISMET_MATH_INLINE
 FVector UKismetMathLibrary::Conv_RotatorToVector(FRotator InRot)
 {
 	return InRot.Vector();
 }
 
-inline
+KISMET_MATH_INLINE
 FVector UKismetMathLibrary::Conv_IntVectorToVector(const FIntVector& InIntVector)
 {
 	return FVector(InIntVector);
 }
 
-inline
+KISMET_MATH_INLINE
 FVector2D UKismetMathLibrary::Conv_IntVector2ToVector2D(const FIntVector2& InIntVector2)
 {
 	return FVector2D(InIntVector2.X, InIntVector2.Y);
 }
 
-inline
+KISMET_MATH_INLINE
 FLinearColor UKismetMathLibrary::Conv_ColorToLinearColor(FColor InColor)
 {
 	return FLinearColor(InColor);
 }
 
-inline
+KISMET_MATH_INLINE
 FVector UKismetMathLibrary::Conv_FloatToVector(float InFloat)
 {
 	return FVector(InFloat);
 }
 
-inline
+KISMET_MATH_INLINE
 FVector UKismetMathLibrary::Conv_DoubleToVector(double InDouble)
 {
 	return FVector(InDouble);
 }
 
-inline
+KISMET_MATH_INLINE
 FVector2D UKismetMathLibrary::Conv_DoubleToVector2D(double InDouble)
 {
 	return FVector2D(InDouble);
 }
 
-inline
+KISMET_MATH_INLINE
 FLinearColor UKismetMathLibrary::Conv_FloatToLinearColor(float InFloat)
 {
 	return FLinearColor(InFloat, InFloat, InFloat);
 }
 
-inline
+KISMET_MATH_INLINE
 FLinearColor UKismetMathLibrary::Conv_DoubleToLinearColor(double InDouble)
 {
 	// potential precision loss by float conversion
@@ -2813,25 +2821,25 @@ FLinearColor UKismetMathLibrary::Conv_DoubleToLinearColor(double InDouble)
 	return FLinearColor(Value, Value, Value);
 }
 
-inline
+KISMET_MATH_INLINE
 FBox UKismetMathLibrary::MakeBox(FVector Min, FVector Max)
 {
 	return FBox(Min, Max);
 }
 
-inline
+KISMET_MATH_INLINE
 FBox2D UKismetMathLibrary::MakeBox2D(FVector2D Min, FVector2D Max)
 {
 	return FBox2D(Min, Max);
 }
 
-inline
+KISMET_MATH_INLINE
 FBoxSphereBounds UKismetMathLibrary::MakeBoxSphereBounds(FVector Origin, FVector BoxExtent, float SphereRadius)
 {
 	return FBoxSphereBounds(Origin, BoxExtent, SphereRadius);
 }
 
-inline
+KISMET_MATH_INLINE
 void UKismetMathLibrary::BreakBoxSphereBounds(const FBoxSphereBounds& InBoxSphereBounds, FVector& Origin, FVector& BoxExtent, float& SphereRadius)
 {
 	Origin = InBoxSphereBounds.Origin;
@@ -2839,73 +2847,73 @@ void UKismetMathLibrary::BreakBoxSphereBounds(const FBoxSphereBounds& InBoxSpher
 	SphereRadius = (float)InBoxSphereBounds.SphereRadius;
 }
 
-inline
+KISMET_MATH_INLINE
 FRotator UKismetMathLibrary::MakeRotator(float Roll, float Pitch, float Yaw)
 {
 	return FRotator(Pitch,Yaw,Roll);
 }
 
-inline
+KISMET_MATH_INLINE
 FRotator UKismetMathLibrary::FindLookAtRotation(const FVector& Start, const FVector& Target)
 {
 	return MakeRotFromX(Target - Start);
 }
 
-inline
+KISMET_MATH_INLINE
 FRotator UKismetMathLibrary::MakeRotFromX(const FVector& X)
 {
 	return FRotationMatrix::MakeFromX(X).Rotator();
 }
 
-inline
+KISMET_MATH_INLINE
 FRotator UKismetMathLibrary::MakeRotFromY(const FVector& Y)
 {
 	return FRotationMatrix::MakeFromY(Y).Rotator();
 }
 
-inline
+KISMET_MATH_INLINE
 FRotator UKismetMathLibrary::MakeRotFromZ(const FVector& Z)
 {
 	return FRotationMatrix::MakeFromZ(Z).Rotator();
 }
 
-inline
+KISMET_MATH_INLINE
 FRotator UKismetMathLibrary::MakeRotFromXY(const FVector& X, const FVector& Y)
 {
 	return FRotationMatrix::MakeFromXY(X, Y).Rotator();
 }
 
-inline
+KISMET_MATH_INLINE
 FRotator UKismetMathLibrary::MakeRotFromXZ(const FVector& X, const FVector& Z)
 {
 	return FRotationMatrix::MakeFromXZ(X, Z).Rotator();
 }
 
-inline
+KISMET_MATH_INLINE
 FRotator UKismetMathLibrary::MakeRotFromYX(const FVector& Y, const FVector& X)
 {
 	return FRotationMatrix::MakeFromYX(Y, X).Rotator();
 }
 
-inline
+KISMET_MATH_INLINE
 FRotator UKismetMathLibrary::MakeRotFromYZ(const FVector& Y, const FVector& Z)
 {
 	return FRotationMatrix::MakeFromYZ(Y, Z).Rotator();
 }
 
-inline
+KISMET_MATH_INLINE
 FRotator UKismetMathLibrary::MakeRotFromZX(const FVector& Z, const FVector& X)
 {
 	return FRotationMatrix::MakeFromZX(Z, X).Rotator();
 }
 
-inline
+KISMET_MATH_INLINE
 FRotator UKismetMathLibrary::MakeRotFromZY(const FVector& Z, const FVector& Y)
 {
 	return FRotationMatrix::MakeFromZY(Z, Y).Rotator();
 }
 
-inline
+KISMET_MATH_INLINE
 void UKismetMathLibrary::BreakRotator(FRotator InRot, float& Roll, float& Pitch, float& Yaw)
 {
 	Pitch = static_cast<float>(InRot.Pitch);
@@ -2917,157 +2925,157 @@ void UKismetMathLibrary::BreakRotator(FRotator InRot, float& Roll, float& Pitch,
 /* FMatrix
 *****************************************************************************/
 
-inline
+KISMET_MATH_INLINE
 FTransform UKismetMathLibrary::Conv_MatrixToTransform(const FMatrix& InMatrix)
 {
 	return FTransform(InMatrix);
 }
 
-inline
+KISMET_MATH_INLINE
 FRotator UKismetMathLibrary::Conv_MatrixToRotator(const FMatrix& InMatrix)
 {
 	return InMatrix.Rotator();
 }
 
-inline
+KISMET_MATH_INLINE
 FVector UKismetMathLibrary::Matrix_GetOrigin(const FMatrix& InMatrix)
 {
 	return InMatrix.GetOrigin();
 }
 
-inline
+KISMET_MATH_INLINE
 FMatrix UKismetMathLibrary::Matrix_Identity()
 {
 	return FMatrix::Identity;
 }
 
-inline
+KISMET_MATH_INLINE
 FMatrix UKismetMathLibrary::Multiply_MatrixMatrix(const FMatrix& A, const FMatrix& B)
 {
 	return A * B;
 }
 
-inline
+KISMET_MATH_INLINE
 FMatrix UKismetMathLibrary::Add_MatrixMatrix(const FMatrix& A, const FMatrix& B)
 {
 	return A + B;
 }
 
-inline
+KISMET_MATH_INLINE
 FMatrix UKismetMathLibrary::Multiply_MatrixFloat(const FMatrix& A, double B)
 {
 	return A * B;
 }
 
-inline
+KISMET_MATH_INLINE
 bool UKismetMathLibrary::EqualEqual_MatrixMatrix(const FMatrix& A, const FMatrix& B, float Tolerance /*= KINDA_SMALL_NUMBER*/)
 {
 	return A.Equals(B, Tolerance);
 }
 
-inline
+KISMET_MATH_INLINE
 bool UKismetMathLibrary::NotEqual_MatrixMatrix(const FMatrix& A, const FMatrix& B, float Tolerance /*= KINDA_SMALL_NUMBER*/)
 {
 	return !(A.Equals(B, Tolerance));
 }
 
-inline
+KISMET_MATH_INLINE
 FVector4 UKismetMathLibrary::Matrix_TransformVector4(const FMatrix& M, FVector4 V)
 {
 	return M.TransformFVector4(V);
 }
 
-inline
+KISMET_MATH_INLINE
 FVector4 UKismetMathLibrary::Matrix_TransformPosition(const FMatrix& M, FVector V)
 {
 	return M.TransformPosition(V);
 }
 
-inline
+KISMET_MATH_INLINE
 FVector UKismetMathLibrary::Matrix_InverseTransformPosition(const FMatrix& M, FVector V)
 {
 	return M.InverseTransformPosition(V);
 }
 
-inline
+KISMET_MATH_INLINE
 FVector4 UKismetMathLibrary::Matrix_TransformVector(const FMatrix& M, FVector V)
 {
 	return M.TransformVector(V);
 }
 
-inline
+KISMET_MATH_INLINE
 FVector UKismetMathLibrary::Matrix_InverseTransformVector(const FMatrix& M, FVector V)
 {
 	return M.InverseTransformVector(V);
 }
 
-inline
+KISMET_MATH_INLINE
 FMatrix UKismetMathLibrary::Matrix_GetTransposed(const FMatrix& M)
 {
 	return M.GetTransposed();
 }
 
-inline
+KISMET_MATH_INLINE
 float UKismetMathLibrary::Matrix_GetDeterminant(const FMatrix& M)
 {
 	return static_cast<float>(M.Determinant());
 }
 
-inline
+KISMET_MATH_INLINE
 float UKismetMathLibrary::Matrix_GetRotDeterminant(const FMatrix& M)
 {
 	return static_cast<float>(M.RotDeterminant());
 }
 
-inline
+KISMET_MATH_INLINE
 FMatrix UKismetMathLibrary::Matrix_GetInverse(const FMatrix& M)
 {
 	return M.Inverse();
 }
 
-inline
+KISMET_MATH_INLINE
 FMatrix UKismetMathLibrary::Matrix_GetTransposeAdjoint(const FMatrix& M)
 {
 	return M.TransposeAdjoint();
 }
 
-inline
+KISMET_MATH_INLINE
 void UKismetMathLibrary::Matrix_RemoveScaling(FMatrix& M, float Tolerance /*= SMALL_NUMBER*/)
 {
 	M.RemoveScaling(Tolerance);
 }
 
-inline
+KISMET_MATH_INLINE
 FMatrix UKismetMathLibrary::Matrix_GetMatrixWithoutScale(const FMatrix& M, float Tolerance /*= SMALL_NUMBER*/)
 {
 	return M.GetMatrixWithoutScale(Tolerance);
 }
 
-inline
+KISMET_MATH_INLINE
 FVector UKismetMathLibrary::Matrix_GetScaleVector(const FMatrix& M, float Tolerance /*= SMALL_NUMBER*/)
 {
 	return M.GetScaleVector(Tolerance);
 }
 
-inline
+KISMET_MATH_INLINE
 FMatrix UKismetMathLibrary::Matrix_RemoveTranslation(const FMatrix& M)
 {
 	return M.RemoveTranslation();
 }
 
-inline
+KISMET_MATH_INLINE
 FMatrix UKismetMathLibrary::Matrix_ConcatenateTranslation(const FMatrix& M, FVector Translation)
 {
 	return M.ConcatTranslation(Translation);
 }
 
-inline
+KISMET_MATH_INLINE
 bool UKismetMathLibrary::Matrix_ContainsNaN(const FMatrix& M)
 {
 	return M.ContainsNaN();
 }
 
-inline
+KISMET_MATH_INLINE
 FMatrix UKismetMathLibrary::Matrix_ScaleTranslation(const FMatrix& M, FVector Scale3D)
 {
 	FMatrix ScaledMatrix = M;
@@ -3075,115 +3083,115 @@ FMatrix UKismetMathLibrary::Matrix_ScaleTranslation(const FMatrix& M, FVector Sc
 	return ScaledMatrix;
 }
 
-inline
+KISMET_MATH_INLINE
 float UKismetMathLibrary::Matrix_GetMaximumAxisScale(const FMatrix& M)
 {
 	return static_cast<float>(M.GetMaximumAxisScale());
 }
 
-inline
+KISMET_MATH_INLINE
 FMatrix UKismetMathLibrary::Matrix_ApplyScale(const FMatrix& M, float Scale)
 {
 	return M.ApplyScale(Scale);
 }
 
-inline
+KISMET_MATH_INLINE
 FVector UKismetMathLibrary::Matrix_GetScaledAxis(const FMatrix& M, TEnumAsByte<EAxis::Type> Axis)
 {
 	return M.GetScaledAxis(Axis);
 }
 
-inline
+KISMET_MATH_INLINE
 void UKismetMathLibrary::Matrix_GetScaledAxes(const FMatrix& M, FVector &X, FVector &Y, FVector &Z)
 {
 	M.GetScaledAxes(X, Y, Z);
 }
 
-inline
+KISMET_MATH_INLINE
 FVector UKismetMathLibrary::Matrix_GetUnitAxis(const FMatrix& M, TEnumAsByte<EAxis::Type> Axis)
 {
 	return M.GetUnitAxis(Axis);
 }
 
-inline
+KISMET_MATH_INLINE
 void UKismetMathLibrary::Matrix_GetUnitAxes(const FMatrix& M, FVector &X, FVector &Y, FVector &Z)
 {
 	M.GetUnitAxes(X, Y, Z);
 }
 
-inline
+KISMET_MATH_INLINE
 void UKismetMathLibrary::Matrix_SetAxis(FMatrix& M, TEnumAsByte<EAxis::Type> Axis, FVector AxisVector)
 {
 	M.SetAxis(StaticCast<int32>(Axis), AxisVector);
 }
 
-inline
+KISMET_MATH_INLINE
 void UKismetMathLibrary::Matrix_SetOrigin(FMatrix& M, FVector NewOrigin)
 {
 	M.SetOrigin(NewOrigin);
 }
 
-inline
+KISMET_MATH_INLINE
 FVector UKismetMathLibrary::Matrix_GetColumn(const FMatrix& M, TEnumAsByte<EMatrixColumns::Type> Column)
 {
 	return M.GetColumn(StaticCast<int32>(Column));
 }
 
-inline
+KISMET_MATH_INLINE
 void UKismetMathLibrary::Matrix_SetColumn(FMatrix& M, TEnumAsByte<EMatrixColumns::Type> Column, FVector Value)
 {
 	M.SetColumn(StaticCast<int32>(Column), Value);
 }
 
-inline
+KISMET_MATH_INLINE
 FRotator UKismetMathLibrary::Matrix_GetRotator(const FMatrix& M)
 {
 	return M.Rotator();
 }
 
-inline
+KISMET_MATH_INLINE
 FQuat UKismetMathLibrary::Matrix_ToQuat(const FMatrix& M)
 {
 	return M.ToQuat();
 }
 
-inline
+KISMET_MATH_INLINE
 bool UKismetMathLibrary::Matrix_GetFrustumNearPlane(const FMatrix& M, FPlane& OutPlane)
 {
 	return M.GetFrustumNearPlane(OutPlane);
 }
 
-inline
+KISMET_MATH_INLINE
 bool UKismetMathLibrary::Matrix_GetFrustumFarPlane(const FMatrix& M, FPlane& OutPlane)
 {
 	return M.GetFrustumFarPlane(OutPlane);
 }
 
-inline
+KISMET_MATH_INLINE
 bool UKismetMathLibrary::Matrix_GetFrustumLeftPlane(const FMatrix& M, FPlane& OutPlane)
 {
 	return M.GetFrustumLeftPlane(OutPlane);
 }
 
-inline
+KISMET_MATH_INLINE
 bool UKismetMathLibrary::Matrix_GetFrustumRightPlane(const FMatrix& M, FPlane& OutPlane)
 {
 	return M.GetFrustumRightPlane(OutPlane);
 }
 
-inline
+KISMET_MATH_INLINE
 bool UKismetMathLibrary::Matrix_GetFrustumTopPlane(const FMatrix& M, FPlane& OutPlane)
 {
 	return M.GetFrustumTopPlane(OutPlane);
 }
 
-inline
+KISMET_MATH_INLINE
 bool UKismetMathLibrary::Matrix_GetFrustumBottomPlane(const FMatrix& M, FPlane& OutPlane)
 {
 	return M.GetFrustumBottomPlane(OutPlane);
 }
 
-inline
+KISMET_MATH_INLINE
 FMatrix UKismetMathLibrary::Matrix_Mirror(const FMatrix& M, TEnumAsByte<EAxis::Type> MirrorAxis, TEnumAsByte<EAxis::Type> FlipAxis)
 {
 	FMatrix MirrorMatrix = M;
@@ -3194,49 +3202,49 @@ FMatrix UKismetMathLibrary::Matrix_Mirror(const FMatrix& M, TEnumAsByte<EAxis::T
 /* FQuat
 *****************************************************************************/
 
-inline
+KISMET_MATH_INLINE
 FQuat UKismetMathLibrary::Quat_Identity()
 {
 	return FQuat::Identity;
 }
 
-inline
+KISMET_MATH_INLINE
 bool UKismetMathLibrary::EqualEqual_QuatQuat(const FQuat& A, const FQuat& B, float Tolerance)
 {
 	return A.Equals(B, Tolerance);
 }
 
-inline
+KISMET_MATH_INLINE
 bool UKismetMathLibrary::NotEqual_QuatQuat(const FQuat& A, const FQuat& B, float Tolerance)
 {
 	return !A.Equals(B, Tolerance);
 }
 
-inline
+KISMET_MATH_INLINE
 FQuat UKismetMathLibrary::Add_QuatQuat(const FQuat& A, const FQuat& B)
 {
 	return A + B;
 }
 
-inline
+KISMET_MATH_INLINE
 FQuat UKismetMathLibrary::Subtract_QuatQuat(const FQuat& A, const FQuat& B)
 {
 	return A - B;
 }
 
-inline
+KISMET_MATH_INLINE
 FQuat UKismetMathLibrary::Multiply_QuatQuat(const FQuat& A, const FQuat& B)
 {
 	return A * B;
 }
 
-inline
+KISMET_MATH_INLINE
 FQuat UKismetMathLibrary::MakeQuat(float X, float Y, float Z, float W)
 {
 	return FQuat(X, Y, Z, W);
 }
 
-inline
+KISMET_MATH_INLINE
 void UKismetMathLibrary::BreakQuat(const FQuat& InQuat, float& X, float& Y, float& Z, float& W)
 {
 	X = static_cast<float>(InQuat.X);
@@ -3245,127 +3253,127 @@ void UKismetMathLibrary::BreakQuat(const FQuat& InQuat, float& X, float& Y, floa
 	W = static_cast<float>(InQuat.W);
 }
 
-inline
+KISMET_MATH_INLINE
 bool UKismetMathLibrary::Quat_IsIdentity(const FQuat& Q, float Tolerance)
 {
 	return Q.IsIdentity(Tolerance);
 }
 
-inline
+KISMET_MATH_INLINE
 bool UKismetMathLibrary::Quat_IsNormalized(const FQuat& Q)
 {
 	return Q.IsNormalized();
 }
 
-inline
+KISMET_MATH_INLINE
 bool UKismetMathLibrary::Quat_IsFinite(const FQuat& Q)
 {
 	return !Q.ContainsNaN();
 }
 
-inline
+KISMET_MATH_INLINE
 bool UKismetMathLibrary::Quat_IsNonFinite(const FQuat& Q)
 {
 	return Q.ContainsNaN();
 }
 
-inline
+KISMET_MATH_INLINE
 float UKismetMathLibrary::Quat_AngularDistance(const FQuat& A, const FQuat& B)
 {
 	return static_cast<float>(A.AngularDistance(B));
 }
 
-inline
+KISMET_MATH_INLINE
 void UKismetMathLibrary::Quat_EnforceShortestArcWith(FQuat& A, const FQuat& B)
 {
 	A.EnforceShortestArcWith(B);
 }
 
-inline
+KISMET_MATH_INLINE
 FVector UKismetMathLibrary::Quat_Euler(const FQuat& Q)
 {
 	return Q.Euler();
 }
 
-inline
+KISMET_MATH_INLINE
 FQuat UKismetMathLibrary::Quat_Exp(const FQuat& Q)
 {
 	return Q.Exp();
 }
 
-inline
+KISMET_MATH_INLINE
 float UKismetMathLibrary::Quat_GetAngle(const FQuat& Q)
 {
 	return static_cast<float>(Q.GetAngle());
 }
 
-inline
+KISMET_MATH_INLINE
 FVector UKismetMathLibrary::Quat_GetAxisX(const FQuat& Q)
 {
 	return Q.GetAxisX();
 }
 
-inline
+KISMET_MATH_INLINE
 FVector UKismetMathLibrary::Quat_GetAxisY(const FQuat& Q)
 {
 	return Q.GetAxisY();
 }
 
-inline
+KISMET_MATH_INLINE
 FVector UKismetMathLibrary::Quat_GetAxisZ(const FQuat& Q)
 {
 	return Q.GetAxisZ();
 }
 
-inline
+KISMET_MATH_INLINE
 FVector UKismetMathLibrary::Quat_VectorForward(const FQuat& Q)
 {
 	return Q.GetForwardVector();
 }
 
-inline
+KISMET_MATH_INLINE
 FVector UKismetMathLibrary::Quat_VectorRight(const FQuat& Q)
 {
 	return Q.GetRightVector();
 }
 
-inline
+KISMET_MATH_INLINE
 FVector UKismetMathLibrary::Quat_VectorUp(const FQuat& Q)
 {
 	return Q.GetUpVector();
 }
 
-inline
+KISMET_MATH_INLINE
 void UKismetMathLibrary::Quat_Normalize(FQuat& Q, float Tolerance)
 {
 	Q.Normalize(Tolerance);
 }
 
-inline
+KISMET_MATH_INLINE
 FQuat UKismetMathLibrary::Quat_Normalized(const FQuat& Q, float Tolerance)
 {
 	return Q.GetNormalized(Tolerance);
 }
 
-inline
+KISMET_MATH_INLINE
 FVector UKismetMathLibrary::Quat_GetRotationAxis(const FQuat& Q)
 {
 	return Q.GetRotationAxis();
 }
 
-inline
+KISMET_MATH_INLINE
 FQuat UKismetMathLibrary::Quat_Inversed(const FQuat& Q)
 {
 	return Q.Inverse();
 }
 
-inline
+KISMET_MATH_INLINE
 FQuat UKismetMathLibrary::Quat_Log(const FQuat& Q)
 {
 	return Q.Log();
 }
 
-inline
+KISMET_MATH_INLINE
 void UKismetMathLibrary::Quat_SetComponents(FQuat& Q, float X, float Y, float Z, float W)
 {
 	Q.X = X;
@@ -3374,67 +3382,67 @@ void UKismetMathLibrary::Quat_SetComponents(FQuat& Q, float X, float Y, float Z,
 	Q.W = W;
 }
 
-inline
+KISMET_MATH_INLINE
 void UKismetMathLibrary::Quat_SetFromEuler(FQuat& Q, const FVector& Euler)
 {
 	Q = FQuat::MakeFromEuler(Euler);
 }
 
-inline
+KISMET_MATH_INLINE
 FQuat UKismetMathLibrary::Quat_MakeFromEuler(const FVector& Euler)
 {
 	return FQuat::MakeFromEuler(Euler);
 }
 
-inline
+KISMET_MATH_INLINE
 FRotator UKismetMathLibrary::Quat_Rotator(const FQuat& Q)
 {
 	return Q.Rotator();
 }
 
-inline
+KISMET_MATH_INLINE
 FQuat UKismetMathLibrary::Conv_RotatorToQuaternion(FRotator InRot)
 {
 	return InRot.Quaternion();
 }
 
-inline
+KISMET_MATH_INLINE
 float UKismetMathLibrary::Quat_Size(const FQuat& Q)
 {
 	return static_cast<float>(Q.Size());
 }
 
-inline
+KISMET_MATH_INLINE
 float UKismetMathLibrary::Quat_SizeSquared(const FQuat& Q)
 {
 	return static_cast<float>(Q.SizeSquared());
 }
 
-inline
+KISMET_MATH_INLINE
 FVector UKismetMathLibrary::Quat_RotateVector(const FQuat& Q, const FVector& V)
 {
 	return Q.RotateVector(V);
 }
 
-inline
+KISMET_MATH_INLINE
 FVector UKismetMathLibrary::Quat_UnrotateVector(const FQuat& Q, const FVector& V)
 {
 	return Q.UnrotateVector(V);
 }
 
-inline
+KISMET_MATH_INLINE
 FQuat UKismetMathLibrary::Quat_Slerp(const FQuat& A, const FQuat& B, double Alpha)
 {
 	return FQuat::Slerp(A, B, Alpha);
 }
 
-inline
+KISMET_MATH_INLINE
 FQuat UKismetMathLibrary::Quat_FindBetweenVectors(FVector Start, FVector End)
 {
 	return FQuat::FindBetweenVectors(Start, End);
 }
 
-inline
+KISMET_MATH_INLINE
 FQuat UKismetMathLibrary::Quat_FindBetweenNormals(FVector StartNormal, FVector EndNormal)
 {
 	return FQuat::FindBetweenNormals(StartNormal, EndNormal);
@@ -3446,49 +3454,49 @@ FQuat UKismetMathLibrary::Quat_FindBetweenNormals(FVector StartNormal, FVector E
 
 // Constants
 
-inline
+KISMET_MATH_INLINE
 FLinearColor UKismetMathLibrary::LinearColor_White()
 {
 	return FLinearColor::White;
 }
 
-inline
+KISMET_MATH_INLINE
 FLinearColor UKismetMathLibrary::LinearColor_Gray()
 {
 	return FLinearColor::Gray;
 }
 
-inline
+KISMET_MATH_INLINE
 FLinearColor UKismetMathLibrary::LinearColor_Black()
 {
 	return FLinearColor::Black;
 }
 
-inline
+KISMET_MATH_INLINE
 FLinearColor UKismetMathLibrary::LinearColor_Red()
 {
 	return FLinearColor::Red;
 }
 
-inline
+KISMET_MATH_INLINE
 FLinearColor UKismetMathLibrary::LinearColor_Green()
 {
 	return FLinearColor::Green;
 }
 
-inline
+KISMET_MATH_INLINE
 FLinearColor UKismetMathLibrary::LinearColor_Blue()
 {
 	return FLinearColor::Blue;
 }
 
-inline
+KISMET_MATH_INLINE
 FLinearColor UKismetMathLibrary::LinearColor_Yellow()
 {
 	return FLinearColor::Yellow;
 }
 
-inline
+KISMET_MATH_INLINE
 FLinearColor UKismetMathLibrary::LinearColor_Transparent()
 {
 	return FLinearColor::Transparent;
@@ -3496,13 +3504,13 @@ FLinearColor UKismetMathLibrary::LinearColor_Transparent()
 
 // Functions
 
-inline
+KISMET_MATH_INLINE
 FLinearColor UKismetMathLibrary::MakeColor(float R, float G, float B, float A)
 {
 	return FLinearColor(R,G,B,A);
 }
 
-inline
+KISMET_MATH_INLINE
 void UKismetMathLibrary::BreakColor(FLinearColor InColor, float& R, float& G, float& B, float& A)
 {
 	R = InColor.R;
@@ -3511,7 +3519,7 @@ void UKismetMathLibrary::BreakColor(FLinearColor InColor, float& R, float& G, fl
 	A = InColor.A;
 }
 
-inline
+KISMET_MATH_INLINE
 void UKismetMathLibrary::LinearColor_Set(FLinearColor& InOutColor, FLinearColor InColor)
 {
 	InOutColor.R = InColor.R;
@@ -3520,7 +3528,7 @@ void UKismetMathLibrary::LinearColor_Set(FLinearColor& InOutColor, FLinearColor
 	InOutColor.A = InColor.A;
 }
 
-inline
+KISMET_MATH_INLINE
 void UKismetMathLibrary::LinearColor_SetRGBA(FLinearColor& InOutColor, float R, float G, float B, float A)
 {
 	InOutColor.R = R;
@@ -3529,63 +3537,63 @@ void UKismetMathLibrary::LinearColor_SetRGBA(FLinearColor& InOutColor, float R,
 	InOutColor.A = A;
 }
 
-inline
+KISMET_MATH_INLINE
 void UKismetMathLibrary::LinearColor_SetFromHSV(FLinearColor& InOutColor, float H, float S, float V, float A)
 {
 	const FLinearColor HSV(H, S, V, A);
 	InOutColor = HSV.HSVToLinearRGB();
 }
 
-inline
+KISMET_MATH_INLINE
 void UKismetMathLibrary::LinearColor_SetFromSRGB(FLinearColor& InOutColor, const FColor& InSRGB)
 {
 	InOutColor = FLinearColor::FromSRGBColor(InSRGB);
 }
 
-inline
+KISMET_MATH_INLINE
 void UKismetMathLibrary::LinearColor_SetFromPow22(FLinearColor& InOutColor, const FColor& InColor)
 {
 	InOutColor = FLinearColor::FromPow22Color(InColor);
 }
 
-inline
+KISMET_MATH_INLINE
 void UKismetMathLibrary::LinearColor_SetTemperature(FLinearColor& InOutColor, float InTemperature)
 {
 	InOutColor = FLinearColor::MakeFromColorTemperature(InTemperature);
 }
 
-inline
+KISMET_MATH_INLINE
 void UKismetMathLibrary::LinearColor_SetRandomHue(FLinearColor& InOutColor)
 {
 	InOutColor = FLinearColor::MakeRandomColor();
 }
 
-inline
+KISMET_MATH_INLINE
 FLinearColor UKismetMathLibrary::HSVToRGB(float H, float S, float V, float A)
 {
 	const FLinearColor HSV(H, S, V, A);
 	return HSV.HSVToLinearRGB();
 }
 
-inline
+KISMET_MATH_INLINE
 FColor UKismetMathLibrary::LinearColor_ToRGBE(FLinearColor InLinearColor)
 {
 	return InLinearColor.ToRGBE();
 }
 
-inline
+KISMET_MATH_INLINE
 FVector UKismetMathLibrary::Conv_LinearColorToVector(FLinearColor InLinearColor)
 {
 	return FVector(InLinearColor);
 }
 
-inline
+KISMET_MATH_INLINE
 FColor UKismetMathLibrary::Conv_LinearColorToColor(FLinearColor InLinearColor, bool InUseSRGB)
 {
 	return InLinearColor.ToFColor(InUseSRGB);
 }
 
-inline
+KISMET_MATH_INLINE
 void UKismetMathLibrary::RGBToHSV(FLinearColor InColor, float& H, float& S, float& V, float& A)
 {
 	const FLinearColor HSV(InColor.LinearRGBToHSV());
@@ -3595,31 +3603,31 @@ void UKismetMathLibrary::RGBToHSV(FLinearColor InColor, float& H, float& S, floa
 	A = HSV.A;
 }
 
-inline
+KISMET_MATH_INLINE
 void UKismetMathLibrary::HSVToRGB_Vector(FLinearColor HSV, FLinearColor& RGB)
 {
 	RGB = HSV.HSVToLinearRGB();
 }
 
-inline
+KISMET_MATH_INLINE
 FLinearColor UKismetMathLibrary::HSVToRGBLinear(FLinearColor HSV)
 {
 	return HSV.HSVToLinearRGB();
 }
 
-inline
+KISMET_MATH_INLINE
 void UKismetMathLibrary::RGBToHSV_Vector(FLinearColor RGB, FLinearColor& HSV)
 {
 	HSV = RGB.LinearRGBToHSV();
 }
 
-inline
+KISMET_MATH_INLINE
 FLinearColor UKismetMathLibrary::RGBLinearToHSV(FLinearColor RGB)
 {
 	return RGB.LinearRGBToHSV();
 }
 
-inline
+KISMET_MATH_INLINE
 FColor UKismetMathLibrary::LinearColor_Quantize(FLinearColor InColor)
 {
 	// @todo remove LinearColor_Quantize from Kismet?
@@ -3628,187 +3636,187 @@ FColor UKismetMathLibrary::LinearColor_Quantize(FLinearColor InColor)
 	return InColor.QuantizeFloor();
 }
 
-inline
+KISMET_MATH_INLINE
 FColor UKismetMathLibrary::LinearColor_QuantizeRound(FLinearColor InColor)
 {
 	return InColor.QuantizeRound();
 }
 
-inline
+KISMET_MATH_INLINE
 FLinearColor UKismetMathLibrary::LinearColor_Desaturated(FLinearColor InColor, float InDesaturation)
 {
 	return InColor.Desaturate(InDesaturation);
 }
 
-inline
+KISMET_MATH_INLINE
 float UKismetMathLibrary::LinearColor_Distance(FLinearColor C1, FLinearColor C2)
 {
 	return FLinearColor::Dist(C1, C2);
 }
 
-inline
+KISMET_MATH_INLINE
 FLinearColor UKismetMathLibrary::LinearColor_ToNewOpacity(FLinearColor InColor, float InOpacity)
 {
 	return InColor.CopyWithNewOpacity(InOpacity);
 }
 
-inline
+KISMET_MATH_INLINE
 float UKismetMathLibrary::LinearColor_GetLuminance(FLinearColor InColor)
 {
 	return InColor.GetLuminance();
 }
 
-inline
+KISMET_MATH_INLINE
 float UKismetMathLibrary::LinearColor_GetMax(FLinearColor InColor)
 {
 	return InColor.GetMax();
 }
 
-inline
+KISMET_MATH_INLINE
 float UKismetMathLibrary::LinearColor_GetMin(FLinearColor InColor)
 {
 	return InColor.GetMin();
 }
 
-inline
+KISMET_MATH_INLINE
 FLinearColor UKismetMathLibrary::CInterpTo(FLinearColor Current, FLinearColor Target, float DeltaTime, float InterpSpeed)
 {
 	return FMath::CInterpTo(Current, Target, DeltaTime, InterpSpeed);
 }
 
-inline
+KISMET_MATH_INLINE
 FLinearColor UKismetMathLibrary::LinearColorLerp(FLinearColor A, FLinearColor B, float Alpha)
 {
 	return A + Alpha * (B - A);
 }
 
-inline
+KISMET_MATH_INLINE
 FLinearColor UKismetMathLibrary::LinearColorLerpUsingHSV(FLinearColor A, FLinearColor B, float Alpha)
 {
 	return FLinearColor::LerpUsingHSV(A, B, Alpha);
 }
 
-inline
+KISMET_MATH_INLINE
 bool UKismetMathLibrary::LinearColor_IsNearEqual(FLinearColor A, FLinearColor B, float Tolerance)
 {
 	return A.Equals(B, Tolerance);
 }
 
-inline
+KISMET_MATH_INLINE
 bool UKismetMathLibrary::EqualEqual_LinearColorLinearColor(FLinearColor A, FLinearColor B)
 {
 	return A == B;
 }
 
-inline
+KISMET_MATH_INLINE
 bool UKismetMathLibrary::NotEqual_LinearColorLinearColor(FLinearColor A, FLinearColor B)
 {
 	return A != B;
 }
 
-inline
+KISMET_MATH_INLINE
 FLinearColor UKismetMathLibrary::Add_LinearColorLinearColor(FLinearColor A, FLinearColor B)
 {
 	return A + B;
 }
 
-inline
+KISMET_MATH_INLINE
 FLinearColor UKismetMathLibrary::Subtract_LinearColorLinearColor(FLinearColor A, FLinearColor B)
 {
 	return A - B;
 }
 
-inline
+KISMET_MATH_INLINE
 FLinearColor UKismetMathLibrary::Multiply_LinearColorLinearColor(FLinearColor A, FLinearColor B)
 {
 	return A * B;
 }
 
-inline
+KISMET_MATH_INLINE
 FLinearColor UKismetMathLibrary::Multiply_LinearColorFloat(FLinearColor A, float B)
 {
 	return A * B;
 }
 
-inline
+KISMET_MATH_INLINE
 FLinearColor UKismetMathLibrary::Divide_LinearColorLinearColor(FLinearColor A, FLinearColor B)
 {
 	return A / B;
 }
 
-inline
+KISMET_MATH_INLINE
 FString UKismetMathLibrary::ToHex_LinearColor(FLinearColor InColor)
 {
 	return InColor.ToFColor(true).ToHex();
 }
 
-inline
+KISMET_MATH_INLINE
 FString UKismetMathLibrary::SelectString(const FString& A, const FString& B, bool bSelectA)
 {
 	return bSelectA ? A : B;
 }
 
-inline
+KISMET_MATH_INLINE
 FText UKismetMathLibrary::SelectText(const FText A, const FText B, bool bSelectA)
 {
 	return bSelectA ? A : B;
 }
 
-inline
+KISMET_MATH_INLINE
 FName UKismetMathLibrary::SelectName(const FName A, const FName B, bool bSelectA)
 {
 	return bSelectA ? A : B;
 }
 
-inline
+KISMET_MATH_INLINE
 int32 UKismetMathLibrary::SelectInt(int32 A, int32 B, bool bSelectA)
 {
 	return bSelectA ? A : B;
 }
 
-inline
+KISMET_MATH_INLINE
 double UKismetMathLibrary::SelectFloat(double A, double B, bool bSelectA)
 {
 	return bSelectA ? A : B;
 }
 
-inline
+KISMET_MATH_INLINE
 FVector UKismetMathLibrary::SelectVector(FVector A, FVector B, bool bSelectA)
 {
 	return bSelectA ? A : B;
 }
 
-inline
+KISMET_MATH_INLINE
 FRotator UKismetMathLibrary::SelectRotator(FRotator A, FRotator B, bool bSelectA)
 {
 	return bSelectA ? A : B;
 }
 
-inline
+KISMET_MATH_INLINE
 FLinearColor UKismetMathLibrary::SelectColor(FLinearColor A, FLinearColor B, bool bSelectA)
 {
 	return bSelectA ? A : B;
 }
 
-inline
+KISMET_MATH_INLINE
 FTransform UKismetMathLibrary::SelectTransform(const FTransform& A, const FTransform& B, bool bSelectA)
 {
 	return bSelectA ? A : B;
 }
 
-inline
+KISMET_MATH_INLINE
 UObject* UKismetMathLibrary::SelectObject(UObject* A, UObject* B, bool bSelectA)
 {
 	return bSelectA ? A : B;
 }
 
-inline
+KISMET_MATH_INLINE
 UClass* UKismetMathLibrary::SelectClass(UClass* A, UClass* B, bool bSelectA)
 {
 	return bSelectA ? A : B;
 }
 
-inline
+KISMET_MATH_INLINE
 bool UKismetMathLibrary::PointsAreCoplanar(const TArray<FVector>& Points, float Tolerance)
 {
 	return FMath::PointsAreCoplanar(Points, Tolerance);
