Control Streamable debug name stripping through CVars. Add a CVar to strip debug names for non-shipping. Remove hardcoded stripping of names for mobile.

diff --git a/Engine/Config/Android/AndroidEngine.ini b/Engine/Config/Android/AndroidEngine.ini
--- a/Engine/Config/Android/AndroidEngine.ini
+++ b/Engine/Config/Android/AndroidEngine.ini
@@ -101,6 +101,10 @@ r.Vulkan.DiagnosticBuffer=1
 ; Transient memory allocator
 r.Vulkan.EnableTransientResourceAllocator=0
 
+; enable stripping debug strings from streamable manager in shipping to save memory
+s.StreamableStripDebugName=true
+s.StreamableStripDebugNameInShipping=true
+
 [/Script/Engine.RendererSettings]
 r.DistanceFields=1
 r.TemporalAA.Mobile.UseCompute=0
diff --git a/Engine/Config/IOS/BaseIOSEngine.ini b/Engine/Config/IOS/BaseIOSEngine.ini
--- a/Engine/Config/IOS/BaseIOSEngine.ini
+++ b/Engine/Config/IOS/BaseIOSEngine.ini
@@ -72,6 +72,10 @@ r.HairStrands.Binding=0
 r.Forward.LightGridSizeZ=8
 r.Forward.MaxCulledLightsPerCell=8
 
+; enable stripping debug strings from streamable manager in shipping to save memory
+s.StreamableStripDebugName=true
+s.StreamableStripDebugNameInShipping=true
+
 [DeviceProfileManager]
 DeviceProfileSelectionModule="IOSDeviceProfileSelector"
 
diff --git a/Engine/Source/Runtime/Engine/Private/StreamableManager.cpp b/Engine/Source/Runtime/Engine/Private/StreamableManager.cpp
--- a/Engine/Source/Runtime/Engine/Private/StreamableManager.cpp
+++ b/Engine/Source/Runtime/Engine/Private/StreamableManager.cpp
@@ -42,11 +42,19 @@ static FAutoConsoleVariableRef CVarStreamableFlushAllAsyncLoadRequestsOnWait(
 	ECVF_Default
 );
 
+static bool GStreamableStripDebugName = false;
+static FAutoConsoleVariableRef CVarStreamableStripDebugName(
+	TEXT("s.StreamableStripDebugName"),
+	GStreamableStripDebugName,
+	TEXT("Whether to strip the bundle names DebugName when changing bundle. Saves memory when disabled, provides more debug information when enabled."),
+	ECVF_Default
+);
+
 static bool GStreamableStripDebugNameInShipping = false;
 static FAutoConsoleVariableRef CVarStreamableStripDebugNameInShipping(
 	TEXT("s.StreamableStripDebugNameInShipping"),
 	GStreamableStripDebugNameInShipping,
-	TEXT("Whether to strip the bundle names DebugName when changing bundle. Saves memory when disabled, provides more debug information when enabled."),
+	TEXT("Whether to strip the bundle names DebugName when changing bundle. Saves memory when disabled, provides more debug information when enabled. Only applied to shipping builds."),
 	ECVF_Default
 );
 
@@ -1748,11 +1756,9 @@ FStreamable* FStreamableManager::StreamInternal(const FSoftObjectPath& InTargetN
 
 bool FStreamableManager::ShouldStripDebugName()
 {
-	bool bShouldStripDebugName = false;
-#if (PLATFORM_IOS || PLATFORM_ANDROID)
-	bShouldStripDebugName = true;
-#elif UE_BUILD_SHIPPING
-	bShouldStripDebugName = GStreamableStripDebugNameInShipping;
+	bool bShouldStripDebugName = GStreamableStripDebugName;
+#if UE_BUILD_SHIPPING
+	bShouldStripDebugName = bShouldStripDebugName || GStreamableStripDebugNameInShipping;
 #endif
 	return bShouldStripDebugName;
 }
