Adding null checks to bypass PVS warnings: https://epicgames.slack.com/archives/C0905SFAAQK/p1749246807305409

diff --git a/Engine/Source/Runtime/Engine/Private/RepLayout.cpp b/Engine/Source/Runtime/Engine/Private/RepLayout.cpp
--- a/Engine/Source/Runtime/Engine/Private/RepLayout.cpp
+++ b/Engine/Source/Runtime/Engine/Private/RepLayout.cpp
@@ -2854,7 +2854,7 @@ void FRepLayout::SendProperties_r(
 		}
 
 		// Use shared serialization if was found
-		if (SharedPropInfo)
+		if (SharedPropInfo && SharedInfo)
 		{
 			check(SharedInfo->SerializedProperties.IsValid());
 
@@ -4071,6 +4071,8 @@ bool FRepLayout::ReceiveProperties_BackwardsCompatible_r(
 			FRepShadowDataBuffer LocalShadowData = ShadowData;
 
 			bool bShadowDataCopied = false;
+			TArray<FProperty*>* RepNotifies = (ShadowData && RepState) ? &RepState->RepNotifies : nullptr;
+
 
 			FGuidReferencesMap* NewGuidReferencesArray = PrepReceivedArray(
 				ArrayNum,
@@ -4083,7 +4085,7 @@ bool FRepLayout::ReceiveProperties_BackwardsCompatible_r(
 				CmdIndex,
 				&LocalShadowData,
 				&LocalData,
-				ShadowData ? &RepState->RepNotifies : nullptr,
+				RepNotifies,
 				bShadowDataCopied,
 				TempReader.PackageMap);
 
@@ -4433,7 +4435,7 @@ void FRepLayout::UpdateUnmappedObjects_r(
 			Cmd.Property->NetSerializeItem(Reader, Connection->PackageMap, Data + AbsOffset);
 
 			// Check to see if this property changed
-			if (bUpdateShadowState)
+			if (bUpdateShadowState && RepState)
 			{
 				// I have a sneaking suspicion that this is broken.
 				// AbsOffset could be Cmd.Offset, but we also may be recursing into an Array, and that
