New Material Translator:  Material Parameter Collection support, including Collection Parameter and Collection Transform nodes.  Bug fix in "LowerBranch" HLSL generation -- need underscore prefix on temporary variable names.

#rnx

diff --git a/Engine/Source/Runtime/Engine/Private/Materials/MaterialExpressionsIR.cpp b/Engine/Source/Runtime/Engine/Private/Materials/MaterialExpressionsIR.cpp
--- a/Engine/Source/Runtime/Engine/Private/Materials/MaterialExpressionsIR.cpp
+++ b/Engine/Source/Runtime/Engine/Private/Materials/MaterialExpressionsIR.cpp
@@ -40,6 +40,7 @@
 #include "Materials/MaterialExpressionClearCoatNormalCustomOutput.h"
 #include "Materials/MaterialExpressionCloudLayer.h"
 #include "Materials/MaterialExpressionCollectionParameter.h"
+#include "Materials/MaterialExpressionCollectionTransform.h"
 #include "Materials/MaterialExpressionColorRamp.h"
 #include "Materials/MaterialExpressionComponentMask.h"
 #include "Materials/MaterialExpressionConstant.h"
@@ -2093,4 +2094,178 @@ void UMaterialExpressionBlend::Build(MIR::FEmitter& Em)
 	Em.Output(0, Result);
 }
 
+static FValueRef EmitParameterCollectionVectorInlineHLSL(MIR::FEmitter& Em, MIR::FValueRef CollectionValue, int32 ParameterIndex)
+{
+	return Em.InlineHLSL(MIR::FType::MakeFloat(4), TEXTVIEW("MaterialCollection$0.Vectors[$1]"), { CollectionValue, Em.ConstantInt(ParameterIndex) });
+}
+
+static bool GetExpressionCollectionParameter(MIR::FEmitter& Em, UMaterialParameterCollection* Collection, FName ParameterName, const FGuid& ParameterId, int32& OutParamIndex, int32& OutComponentIndex)
+{
+	if (!Collection)
+	{
+		Em.Errorf(TEXT("CollectionParameter has invalid Collection!"));
+		return false;
+	}
+
+	Collection->GetParameterIndex(ParameterId, OutParamIndex, OutComponentIndex);
+	if (OutParamIndex == INDEX_NONE)
+	{
+		Em.Errorf(TEXT("CollectionParameter has invalid parameter %s"), *ParameterName.ToString());
+		return false;
+	}
+
+	return true;
+}
+
+void UMaterialExpressionCollectionParameter::Build(MIR::FEmitter& Em)
+{
+	int32 ParameterIndex = INDEX_NONE;
+	int32 ComponentIndex = INDEX_NONE;
+	if (!GetExpressionCollectionParameter(Em, Collection, ParameterName, ParameterId, ParameterIndex, ComponentIndex))
+	{
+		return;
+	}
+
+	FValueRef Result = EmitParameterCollectionVectorInlineHLSL(Em, Em.MaterialParameterCollection(Collection), ParameterIndex);
+	if (ComponentIndex != INDEX_NONE)
+	{
+		Result = Em.Subscript(Result, ComponentIndex % 4);
+	}
+	Em.Output(0, Result);
+}
+
+void UMaterialExpressionCollectionTransform::Build(MIR::FEmitter& Em)
+{
+	int32 ParameterIndex = INDEX_NONE;
+	int32 ComponentIndex = INDEX_NONE;
+	if (!GetExpressionCollectionParameter(Em, Collection, ParameterName, ParameterId, ParameterIndex, ComponentIndex))
+	{
+		return;
+	}
+
+	if (ComponentIndex != INDEX_NONE)
+	{
+		Em.Errorf(TEXT("CollectionTransform parameter %s is scalar, vectors are required"), *ParameterName.ToString());
+		return;
+	}
+
+	FValueRef Value = Em.TryInput(&Input);
+	if (!Value)
+	{
+		Em.Error(TEXT("CollectionTransform missing input"));
+		return;
+	}
+
+	if (!Value->Type.IsAnyFloat() || Value->Type.GetPrimitive().NumRows < 3 || Value->Type.GetPrimitive().NumColumns != 1)
+	{
+		Em.Error(TEXT("CollectionTransform requires float vector input"));
+		return;
+	}
+
+	int32 NumVectors = 0;
+	if (TransformType == EParameterCollectionTransformType::Position || TransformType == EParameterCollectionTransformType::Projection)
+	{
+		if (ParameterIndex + 4 > Collection->GetTotalVectorStorage())
+		{
+			Em.Errorf(TEXT("CollectionTransform parameter %s requires 4 vectors for Position or Projection matrix"), *ParameterName.ToString());
+			return;
+		}
+		NumVectors = 4;
+	}
+	else if (TransformType == EParameterCollectionTransformType::Vector)
+	{
+		if (ParameterIndex + 3 > Collection->GetTotalVectorStorage())
+		{
+			Em.Errorf(TEXT("CollectionTransform parameter %s requires 3 vectors for Vector matrix"), *ParameterName.ToString());
+			return;
+		}
+		NumVectors = 3;
+	}
+	else
+	{
+		check(TransformType == EParameterCollectionTransformType::LocalToWorld || TransformType == EParameterCollectionTransformType::WorldToLocal);
+
+		if (ParameterIndex + 5 > Collection->GetTotalVectorStorage())
+		{
+			Em.Errorf(TEXT("CollectionTransform parameter %s requires 5 vectors for LWC Matrix"), *ParameterName.ToString());
+			return;
+		}
+		NumVectors = 5;
+	}
+
+	FValueRef CollectionValue = Em.MaterialParameterCollection(Collection);
+	TArray<FValueRef, TFixedAllocator<5>> CollectionParameters;
+
+	for (int32 i = 0; i < NumVectors; i++)
+	{
+		CollectionParameters.Add(EmitParameterCollectionVectorInlineHLSL(Em, CollectionValue, ParameterIndex + i));
+	}
+
+	MIR::FValueRef Result;
+
+	// Matrix transforms cobbled together from primitive ops (rather than using mul or LWCMultiply), so analytic derivatives are supported for free
+	if (TransformType == EParameterCollectionTransformType::Vector)
+	{
+		// Treat input as a direction vector (w = 0)
+		Value = Em.Cast(Value, MIR::FType::MakeFloat(3));
+
+		Result =        Em.Multiply(Em.Subscript(Value, 0), Em.Swizzle(CollectionParameters[0], MIR::FSwizzleMask::XYZ()));
+		Result = Em.Add(Em.Multiply(Em.Subscript(Value, 1), Em.Swizzle(CollectionParameters[1], MIR::FSwizzleMask::XYZ())), Result);
+		Result = Em.Add(Em.Multiply(Em.Subscript(Value, 2), Em.Swizzle(CollectionParameters[2], MIR::FSwizzleMask::XYZ())), Result);
+	}
+	else if (TransformType == EParameterCollectionTransformType::Projection)
+	{
+		// Optimized to save many ALU for a standard perspective or orthographic projection matrix, where most of the elements of the matrix are zero.
+		Result = Em.Vector4(
+			Em.Multiply(Em.Subscript(Value, 0), Em.Subscript(CollectionParameters[0], 0)),														// Value.x * Matrix._00
+			Em.Multiply(Em.Subscript(Value, 1), Em.Subscript(CollectionParameters[1], 1)),														// Value.y * Matrix._11
+			Em.Add(Em.Multiply(Em.Subscript(Value, 2), Em.Subscript(CollectionParameters[2], 2)), Em.Subscript(CollectionParameters[3], 2)),	// Value.z * Matrix._22 + Matrix._32
+			Em.Add(Em.Multiply(Em.Subscript(Value, 2), Em.Subscript(CollectionParameters[2], 3)), Em.Subscript(CollectionParameters[3], 3)));	// Value.z * Matrix._23 + Matrix._33
+	}
+	else
+	{
+		// Position, LocalToWorld, WorldToLocal
+		if (TransformType == EParameterCollectionTransformType::WorldToLocal)
+		{
+			// Pre subtract tile value, to convert this to float (LWC inverse matrices have their tile negated, so adding means we are subtracting the tile value).
+			// The tile value only applies to XYZ -- if the input Value has a fourth component, the Add operation will pad the tile argument with zero if needed.
+			Value = Em.Add(Em.CastToScalarKind(Value, MIR::EScalarKind::Double), Em.LWCTile(Em.Swizzle(CollectionParameters[4], MIR::FSwizzleMask::XYZ())));
+		}
+
+		Value = Em.CastToFloatKind(Value);
+
+		// If a 3-element vector is provided as input, we want to generate a 3-element vector as output.  Swizzle the collection parameters to achieve this.
+		if (Value->Type.GetPrimitive().NumRows == 3)
+		{
+			for (int32 i = 0; i < 4; i++)
+			{
+				CollectionParameters[i] = Em.Swizzle(CollectionParameters[i], MIR::FSwizzleMask::XYZ());
+			}
+		}
+
+		Result =        Em.Multiply(Em.Subscript(Value, 0), CollectionParameters[0]);
+		Result = Em.Add(Em.Multiply(Em.Subscript(Value, 1), CollectionParameters[1]), Result);
+		Result = Em.Add(Em.Multiply(Em.Subscript(Value, 2), CollectionParameters[2]), Result);
+
+		if (Value->Type.GetPrimitive().NumRows == 3)
+		{
+			// Treat input as a translation vector (w = 1)
+			Result = Em.Add(CollectionParameters[3], Result);
+		}
+		else
+		{
+			// Treat input as a homogenous vector (w = user specified)
+			Result = Em.Add(Em.Multiply(Em.Subscript(Value, 3), CollectionParameters[3]), Result);
+		}
+
+		if (TransformType == EParameterCollectionTransformType::LocalToWorld)
+		{
+			// Post add tile value, to convert this to LWC
+			Result = Em.Add(Em.CastToScalarKind(Result, MIR::EScalarKind::Double), Em.LWCTile(Em.Swizzle(CollectionParameters[4], MIR::FSwizzleMask::XYZ())));
+		}
+	}
+
+	Em.Output(0, Result);
+}
+
 #endif // WITH_EDITOR
diff --git a/Engine/Source/Runtime/Engine/Private/Materials/MaterialIR.cpp b/Engine/Source/Runtime/Engine/Private/Materials/MaterialIR.cpp
--- a/Engine/Source/Runtime/Engine/Private/Materials/MaterialIR.cpp
+++ b/Engine/Source/Runtime/Engine/Private/Materials/MaterialIR.cpp
@@ -42,6 +42,7 @@ const TCHAR* LexToString(EValueKind Kind)
 		case VK_Poison: return TEXT("Poison");
 		case VK_Constant: return TEXT("Constant");
 		case VK_ExternalInput: return TEXT("ExternalInput");
+		case VK_MaterialParameterCollection: return TEXT("MaterialParameterCollection");
 		case VK_ScreenTexture: return TEXT("ScreenTexture");
 		case VK_TextureObject: return TEXT("TextureObject");
 		case VK_UniformParameter: return TEXT("UniformParameter");
@@ -102,6 +103,7 @@ uint32 FValue::GetSizeInBytes() const
 		case VK_Poison: return sizeof(FPoison);
 		case VK_Constant: return sizeof(FConstant);
 		case VK_ExternalInput: return sizeof(FExternalInput);
+		case VK_MaterialParameterCollection: return sizeof(FMaterialParameterCollection);
 		case VK_ScreenTexture: return sizeof(FScreenTexture);
 		case VK_TextureObject: return sizeof(FTextureObject);
 		case VK_UniformParameter: return sizeof(FUniformParameter);
@@ -701,6 +703,7 @@ const TCHAR* LexToString(EOperator Op)
 		case UO_Logarithm: return TEXT("Logarithm");
 		case UO_Logarithm10: return TEXT("Logarithm10");
 		case UO_Logarithm2: return TEXT("Logarithm2");
+		case UO_LWCTile: return TEXT("LWCTile");
 		case UO_Negate: return TEXT("Negate");
 		case UO_Not: return TEXT("Not");
 		case UO_Reciprocal: return TEXT("Reciprocal");
diff --git a/Engine/Source/Runtime/Engine/Private/Materials/MaterialIREmitter.cpp b/Engine/Source/Runtime/Engine/Private/Materials/MaterialIREmitter.cpp
--- a/Engine/Source/Runtime/Engine/Private/Materials/MaterialIREmitter.cpp
+++ b/Engine/Source/Runtime/Engine/Private/Materials/MaterialIREmitter.cpp
@@ -726,6 +726,13 @@ FValueRef FEmitter::ExternalInput(EExternalInput Id)
 	return EmitPrototype(*this, Prototype);
 }
 
+FValueRef FEmitter::MaterialParameterCollection(UMaterialParameterCollection* Collection)
+{
+	FMaterialParameterCollection Prototype = MakePrototype<FMaterialParameterCollection>(FType::MakeParameterCollection());
+	Prototype.Collection = Collection;
+	return EmitPrototype(*this, Prototype);
+}
+
 FValueRef FEmitter::TextureObject(UTexture* Texture, EMaterialSamplerType SamplerType)
 {
 	switch (Texture->GetMaterialType())
@@ -1661,6 +1668,7 @@ static const FOperatorSignature* GetOperatorSignature(EOperator Op)
 		S[UO_Logarithm] 				= { { OPTF_CheckIsArithmetic | OPTF_CheckIsNonZeroFloatConst | OPTF_CheckIsNonNegativeFloatConst | OPTF_CastToAnyFloat }, ORT_FirstArgumentType };
 		S[UO_Logarithm10]				= { { OPTF_CheckIsArithmetic | OPTF_CheckIsNonZeroFloatConst | OPTF_CheckIsNonNegativeFloatConst | OPTF_CastToAnyFloat }, ORT_FirstArgumentType };
 		S[UO_Logarithm2] 				= { { OPTF_CheckIsArithmetic | OPTF_CheckIsNonZeroFloatConst | OPTF_CheckIsNonNegativeFloatConst | OPTF_CastToAnyFloat }, ORT_FirstArgumentType };
+		S[UO_LWCTile]					= { };		// UNUSED
 		S[UO_Reciprocal]				= { { OPTF_CheckIsArithmetic | OPTF_CheckIsNonZeroFloatConst | OPTF_CastToAnyFloat }, ORT_FirstArgumentType };
 		S[UO_Round] 					= UnaryFloatLWC;
 		S[UO_Rsqrt]						= { { OPTF_CheckIsArithmetic | OPTF_CheckIsNonZeroFloatConst | OPTF_CastToAnyFloat }, ORT_FirstArgumentType };
@@ -1737,6 +1745,17 @@ static FType ValidateOperatorAndGetResultType(FEmitter& Emitter, EOperator Op, F
 		return FType::MakePoison();
 	}
 
+	// Given input float3, generates LWC3 type with the given tile value and zero offset.  Caller must pass in float3.
+	if (Op == UO_LWCTile)
+	{
+		if (!A->Type.GetPrimitive().IsFloat() || !A->Type.GetPrimitive().IsVector() || A->Type.GetPrimitive().NumComponents() != 3)
+		{
+			Emitter.Errorf(A, TEXT("Argument of LWCTile operator expected to be a 3D float vector."));
+			return FType::MakePoison();
+		}
+		return FType::MakeDouble(3);
+	}
+
 	// For Clamp, we do a special case and demote the first LWC argument if the second and third arguments (min / max) are non-LWC.
 	// We want to do this before fetching FirstArgumentPrimitiveType.
 	if (Op == TO_Clamp && A->Type.GetPrimitive().IsLWC() && !B->Type.IsLWC() && !C->Type.IsLWC())
@@ -3187,6 +3206,7 @@ static FValue* DifferentiateOperator(FEmitter& E, FOperator* Op, EDerivativeAxis
 		case UO_IsFinite:
 		case UO_IsInf:
 		case UO_IsNan:
+		case UO_LWCTile:
 		case UO_Sign:
 		case BO_Modulo:
 		case BO_BitwiseAnd:
diff --git a/Engine/Source/Runtime/Engine/Private/Materials/MaterialIRModule.cpp b/Engine/Source/Runtime/Engine/Private/Materials/MaterialIRModule.cpp
--- a/Engine/Source/Runtime/Engine/Private/Materials/MaterialIRModule.cpp
+++ b/Engine/Source/Runtime/Engine/Private/Materials/MaterialIRModule.cpp
@@ -3,6 +3,7 @@
 #include "Materials/MaterialIRModule.h"
 #include "Materials/MaterialIR.h"
 #include "Materials/MaterialIRTypes.h"
+#include "ParameterCollection.h"
 
 #if WITH_EDITOR
 
@@ -56,4 +57,24 @@ void FMaterialIRModule::AddError(UMaterialExpression* Expression, FString Messag
 	Errors.Push({ Expression, MoveTemp(Message) });
 }
 
+int32 FMaterialIRModule::FindOrAddParameterCollection(UMaterialParameterCollection* ParameterCollection)
+{
+	int32 CollectionIndex = ParameterCollections.Find(ParameterCollection);
+
+	if (CollectionIndex == INDEX_NONE)
+	{
+		if (ParameterCollections.Num() >= MaxNumParameterCollectionsPerMaterial)
+		{
+			return INDEX_NONE;
+		}
+		else
+		{
+			ParameterCollections.Add(ParameterCollection);
+			CollectionIndex = ParameterCollections.Num() - 1;
+		}
+	}
+
+	return CollectionIndex;
+}
+
 #endif // #if WITH_EDITOR
diff --git a/Engine/Source/Runtime/Engine/Private/Materials/MaterialIRModuleBuilder.cpp b/Engine/Source/Runtime/Engine/Private/Materials/MaterialIRModuleBuilder.cpp
--- a/Engine/Source/Runtime/Engine/Private/Materials/MaterialIRModuleBuilder.cpp
+++ b/Engine/Source/Runtime/Engine/Private/Materials/MaterialIRModuleBuilder.cpp
@@ -543,7 +543,7 @@ struct FMaterialIRModuleBuilderImpl
 		// Derive additional defines from the final state of the IR analysis
 		DeriveEnvironmentDefinesFromAnalysisState();
 
-		// Move final environemnt defines from analyzer into output module
+		// Move final environment defines from analyzer into output module
 		Module->EnvironmentDefines = MoveTemp(ValueAnalyzer.EnvironmentDefines);
 	}
 
@@ -654,6 +654,8 @@ struct FMaterialIRModuleBuilderImpl
 		FMaterialCompilationOutput& CompilationOutput = Module->CompilationOutput;
 		CompilationOutput.NumUsedUVScalars = Module->Statistics.NumPixelTexCoords * 2;
 
+		CompilationOutput.UniformExpressionSet.SetParameterCollections(Module->ParameterCollections);
+
 		int32 NumPostProcessInputs = CompilationOutput.GetNumPostProcessInputsUsed();
 		if (NumPostProcessInputs > kPostProcessMaterialInputCountMax)
 		{
diff --git a/Engine/Source/Runtime/Engine/Private/Materials/MaterialIRToHLSLTranslator.cpp b/Engine/Source/Runtime/Engine/Private/Materials/MaterialIRToHLSLTranslator.cpp
--- a/Engine/Source/Runtime/Engine/Private/Materials/MaterialIRToHLSLTranslator.cpp
+++ b/Engine/Source/Runtime/Engine/Private/Materials/MaterialIRToHLSLTranslator.cpp
@@ -12,9 +12,11 @@
 #include "Materials/MaterialIR.h"
 #include "Materials/MaterialIRModule.h"
 #include "Materials/MaterialIRTypes.h"
+#include "Materials/MaterialParameterCollection.h"
 #include "MaterialShared.h"
 #include "MaterialSharedPrivate.h"
 #include "Misc/LargeWorldRenderPosition.h"
+#include "ParameterCollection.h"
 #include "RenderUtils.h"
 #include "ShaderCore.h"
 
@@ -425,6 +427,7 @@ struct FPrivate : FMaterialIRToHLSLTranslation
 		{
 			case MIR::VK_Constant: LowerConstant(static_cast<const MIR::FConstant*>(InValue)); break;
 			case MIR::VK_ExternalInput: LowerExternalInput(static_cast<const MIR::FExternalInput*>(InValue)); break;
+			case MIR::VK_MaterialParameterCollection: LowerMaterialParameterCollection(static_cast<const MIR::FMaterialParameterCollection*>(InValue)); break;
 			case MIR::VK_ScreenTexture: LowerScreenTexture(static_cast<const MIR::FScreenTexture*>(InValue)); break;
 			case MIR::VK_TextureObject: LowerTextureObject(static_cast<const MIR::FTextureObject*>(InValue)); break;
 			case MIR::VK_UniformParameter: LowerUniformParameter(static_cast<const MIR::FUniformParameter*>(InValue)); break;
@@ -570,6 +573,11 @@ struct FPrivate : FMaterialIRToHLSLTranslation
 		}
 	}
 
+	void LowerMaterialParameterCollection(const MIR::FMaterialParameterCollection* MaterialParameterCollection)
+	{
+		Printer << MaterialParameterCollection->Analysis_CollectionIndex;
+	}
+
 	void LowerInlineHLSLWithArgumentsInternal(FStringView Code, const TConstArrayView<MIR::FValue*>& InArguments)
 	{
 		// Substitute argument tokens with instruction arguments
@@ -1235,6 +1243,14 @@ struct FPrivate : FMaterialIRToHLSLTranslation
 
 	void LowerOperator(const MIR::FOperator* Operator)
 	{
+		// LWCTile operator is special in that it has an extra zero parameter, so it can't go through the normal operator code path
+		if (Operator->Op == MIR::UO_LWCTile)
+		{
+			// Given input float3, generate LWC3 type with the given tile value and zero offset.
+			Printer << TEXTVIEW("MakeLWCVector3(") << LowerValue(Operator->AArg) << TEXTVIEW(", 0)");
+			return;
+		}
+
 		bool bIsLWC = Operator->AArg->Type.IsLWC() || (Operator->Op == MIR::TO_Select && Operator->BArg->Type.IsLWC());
 
 		if (IsOperatorInfix(Operator->Op) && !bIsLWC)
@@ -1417,11 +1433,11 @@ struct FPrivate : FMaterialIRToHLSLTranslation
 			Printer << EndOfStatement;
 			Printer << TEXTVIEW("if (") << LowerValue(Branch->ConditionArg) << TEXTVIEW(")") << NewLine << OpenBrace;
 			Printer << LowerBlock(Branch->TrueBlock[CurrentStage]);
-			Printer << InstrToLocalIndex[Branch] << TEXTVIEW(" = ") << LowerValue(Branch->TrueArg) << EndOfStatement;
+			Printer << TEXTVIEW("_") << InstrToLocalIndex[Branch] << TEXTVIEW(" = ") << LowerValue(Branch->TrueArg) << EndOfStatement;
 			Printer << CloseBrace << NewLine;
 			Printer << TEXTVIEW("else") << NewLine << OpenBrace;
 			Printer << LowerBlock(Branch->FalseBlock[CurrentStage]);
-			Printer << InstrToLocalIndex[Branch] << TEXTVIEW(" = ") << LowerValue(Branch->FalseArg) << EndOfStatement;
+			Printer << TEXTVIEW("_") << InstrToLocalIndex[Branch] << TEXTVIEW(" = ") << LowerValue(Branch->FalseArg) << EndOfStatement;
 			Printer << CloseBrace;
 		}
 	}
@@ -2075,6 +2091,41 @@ struct FPrivate : FMaterialIRToHLSLTranslation
 			FString PageTableUniformValue = FString::Printf(TEXT("Material.VTPackedPageTableUniform[%d*2], Material.VTPackedPageTableUniform[%d*2+1]"), VTStackIndex, VTStackIndex);
 			OutEnvironment.SetDefine(*FString::Printf(TEXT("VIRTUALTEXTURE_PAGETABLE_UNIFORM_%d"), VTStackIndex), *PageTableUniformValue);
 		}
+
+		TConstArrayView<UMaterialParameterCollection*> ParameterCollections = Module->GetParameterCollections();
+		for (int32 CollectionIndex = 0; CollectionIndex < ParameterCollections.Num(); CollectionIndex++)
+		{
+			// Add uniform buffer declarations for any parameter collections referenced
+			static_assert(MaxNumParameterCollectionsPerMaterial == 2);
+			static const TCHAR* CollectionNames[MaxNumParameterCollectionsPerMaterial] =
+			{
+				TEXT("MaterialCollection0"),
+				TEXT("MaterialCollection1"),
+			};
+
+			// Check that the parameter collection loaded successfully.
+			UMaterialParameterCollection* ParameterCollection = ParameterCollections[CollectionIndex];
+			if (!ParameterCollection)
+			{
+				UE_LOG(LogMaterial, Warning, TEXT("Null parameter collection found in environment defines while translating material."));
+				continue;
+			}
+
+			// Ensure PostLoad is called so the uniform buffers are created in case the parameter collection was loaded async 
+			ParameterCollection->ConditionalPostLoad();
+
+			// Check that the parameter collection uniform buffer structure is valid
+			if (!ParameterCollection->HasValidUniformBufferStruct())
+			{
+				UE_LOG(LogMaterial, Warning, TEXT("Invalid parameter collection uniform buffer struct found in environment defines while translating material."));
+				continue;
+			}
+
+			// This can potentially become an issue for MaterialCollection Uniform Buffers if they ever get non-numeric resources (eg Textures), as
+			// OutEnvironment.ResourceTableMap has a map by name, and the N ParameterCollection Uniform Buffers ALL are names "MaterialCollection"
+			// (and the hlsl cbuffers are named MaterialCollection0, etc, so the names don't match the layout)
+			FShaderUniformBufferParameter::ModifyCompilationEnvironment(CollectionNames[CollectionIndex], ParameterCollection->GetUniformBufferStruct(), Module->GetShaderPlatform(), OutEnvironment);
+		}
 	}
 };
 
diff --git a/Engine/Source/Runtime/Engine/Private/Materials/MaterialIRTypes.cpp b/Engine/Source/Runtime/Engine/Private/Materials/MaterialIRTypes.cpp
--- a/Engine/Source/Runtime/Engine/Private/Materials/MaterialIRTypes.cpp
+++ b/Engine/Source/Runtime/Engine/Private/Materials/MaterialIRTypes.cpp
@@ -16,6 +16,7 @@ const TCHAR* TypeKindToString(ETypeKind Kind)
 		case ETypeKind::Void: return TEXT("Void");
 		case ETypeKind::Primitive: return TEXT("Primitive");
 		case ETypeKind::Aggregate: return TEXT("Aggregate");
+		case ETypeKind::ParameterCollection: return TEXT("ParameterCollection");
 		case ETypeKind::Texture2D: return TEXT("Texture2D");
 		case ETypeKind::SubstrateData: return TEXT("SubstrateData");
 		case ETypeKind::VTPageTableResult: return TEXT("VTPageTableResult");
@@ -220,6 +221,13 @@ FType FType::MakeAggregate(const UMaterialAggregate* InAggregate)
 	return Type;
 }
 
+FType FType::MakeParameterCollection()
+{
+	FType Type{};
+	Type.Kind = ETypeKind::ParameterCollection;
+	return Type;
+}
+
 FType FType::MakeTexture2D()
 {
 	FType Type{};
diff --git a/Engine/Source/Runtime/Engine/Private/Materials/MaterialIRValueAnalyzer.cpp b/Engine/Source/Runtime/Engine/Private/Materials/MaterialIRValueAnalyzer.cpp
--- a/Engine/Source/Runtime/Engine/Private/Materials/MaterialIRValueAnalyzer.cpp
+++ b/Engine/Source/Runtime/Engine/Private/Materials/MaterialIRValueAnalyzer.cpp
@@ -18,6 +18,7 @@
 #include "MaterialCachedData.h"
 #include "MaterialDomain.h"
 #include "MaterialShared.h"
+#include "ParameterCollection.h"
 #include "Engine/Texture.h"
 #include "UObject/Package.h"
 #include "PostProcess/PostProcessMaterialInputs.h"
@@ -601,6 +602,18 @@ static void AnalyzeSetMaterialOutput(FMaterialIRValueAnalyzer& Analyzer, MIR::FS
 	}
 }
 
+static void AnalyzeMaterialParameterCollection(FMaterialIRValueAnalyzer& Analyzer, MIR::FMaterialParameterCollection* MaterialParameterCollection)
+{
+	int32 CollectionIndex = Analyzer.Module->FindOrAddParameterCollection(MaterialParameterCollection->Collection);
+		
+	if (CollectionIndex == INDEX_NONE)
+	{
+		Analyzer.Module->AddError(nullptr, TEXT("Material references too many MaterialParameterCollections!  A material may only reference 2 different collections."));
+	}
+
+	MaterialParameterCollection->Analysis_CollectionIndex = CollectionIndex;
+}
+
 //////////////////////////////////////////////////////////////////////////////////////////
 // FMaterialIRValueAnalyzer
 //////////////////////////////////////////////////////////////////////////////////////////
@@ -640,6 +653,7 @@ void FMaterialIRValueAnalyzer::Analyze(MIR::FValue* Value)
 	switch (Value->Kind)
 	{
 		EXPAND_CASE(ExternalInput);
+		EXPAND_CASE(MaterialParameterCollection);
 		EXPAND_CASE(TextureObject);
 		EXPAND_CASE(UniformParameter);
 		EXPAND_CASE(SetMaterialOutput);
diff --git a/Engine/Source/Runtime/Engine/Public/Materials/MaterialExpressionCollectionParameter.h b/Engine/Source/Runtime/Engine/Public/Materials/MaterialExpressionCollectionParameter.h
--- a/Engine/Source/Runtime/Engine/Public/Materials/MaterialExpressionCollectionParameter.h
+++ b/Engine/Source/Runtime/Engine/Public/Materials/MaterialExpressionCollectionParameter.h
@@ -52,6 +52,7 @@ class UMaterialExpressionCollectionParameter : public UMaterialExpression
 
 	//~ Begin UMaterialExpression Interface
 #if WITH_EDITOR
+	virtual void Build(MIR::FEmitter& Emitter) override;
 	virtual int32 Compile(class FMaterialCompiler* Compiler, int32 OutputIndex) override;
 	virtual void GetCaption(TArray<FString>& OutCaptions) const override;
 
diff --git a/Engine/Source/Runtime/Engine/Public/Materials/MaterialExpressionCollectionTransform.h b/Engine/Source/Runtime/Engine/Public/Materials/MaterialExpressionCollectionTransform.h
--- a/Engine/Source/Runtime/Engine/Public/Materials/MaterialExpressionCollectionTransform.h
+++ b/Engine/Source/Runtime/Engine/Public/Materials/MaterialExpressionCollectionTransform.h
@@ -57,6 +57,7 @@ class UMaterialExpressionCollectionTransform : public UMaterialExpression
 
 	//~ Begin UMaterialExpression Interface
 #if WITH_EDITOR
+	virtual void Build(MIR::FEmitter& Emitter) override;
 	virtual int32 Compile(class FMaterialCompiler* Compiler, int32 OutputIndex) override;
 	virtual void GetCaption(TArray<FString>& OutCaptions) const override;
 
diff --git a/Engine/Source/Runtime/Engine/Public/Materials/MaterialExpressionOperator.h b/Engine/Source/Runtime/Engine/Public/Materials/MaterialExpressionOperator.h
--- a/Engine/Source/Runtime/Engine/Public/Materials/MaterialExpressionOperator.h
+++ b/Engine/Source/Runtime/Engine/Public/Materials/MaterialExpressionOperator.h
@@ -40,6 +40,7 @@ enum class EMaterialExpressionOperatorKind
 	Logarithm,
 	Logarithm10,
 	Logarithm2,
+	LWCTile,
 	Reciprocal,
 	Round,
 	Rsqrt,
diff --git a/Engine/Source/Runtime/Engine/Public/Materials/MaterialIR.h b/Engine/Source/Runtime/Engine/Public/Materials/MaterialIR.h
--- a/Engine/Source/Runtime/Engine/Public/Materials/MaterialIR.h
+++ b/Engine/Source/Runtime/Engine/Public/Materials/MaterialIR.h
@@ -103,6 +103,7 @@ enum EValueKind
 	VK_Poison,
 	VK_Constant,
 	VK_ExternalInput,
+	VK_MaterialParameterCollection,
 	VK_ScreenTexture,
 	VK_TextureObject,
 	VK_UniformParameter,
@@ -451,6 +452,16 @@ struct FExternalInput : TValue<VK_ExternalInput>
 	EExternalInput Id;
 };
 
+// It represents a Material Parameter Collection object value.
+struct FMaterialParameterCollection : TValue<VK_MaterialParameterCollection>
+{
+	UMaterialParameterCollection* Collection;
+
+	// Index of this collection in the list of collections referenced by the material.
+	// Note: This field is automatically set by the builder.
+	uint32 Analysis_CollectionIndex;
+};
+
 // It represents a texture object value.
 struct FTextureObject : TValue<VK_TextureObject>
 {
@@ -678,6 +689,7 @@ enum EOperator
 	UO_Logarithm,
 	UO_Logarithm10,
 	UO_Logarithm2,
+	UO_LWCTile,
 	UO_Reciprocal,
 	UO_Round,
 	UO_Rsqrt,
diff --git a/Engine/Source/Runtime/Engine/Public/Materials/MaterialIREmitter.h b/Engine/Source/Runtime/Engine/Public/Materials/MaterialIREmitter.h
--- a/Engine/Source/Runtime/Engine/Public/Materials/MaterialIREmitter.h
+++ b/Engine/Source/Runtime/Engine/Public/Materials/MaterialIREmitter.h
@@ -331,6 +331,9 @@ public:
 	// Emits an external input value of given input Id and returns it.
 	FValueRef ExternalInput(EExternalInput Id);
 
+	// Emits a material parameter collection.
+	FValueRef MaterialParameterCollection(UMaterialParameterCollection* Collection);
+
 	// Emits a texture object value of fiven sampler type and returns it.
 	FValueRef TextureObject(UTexture* Texture, EMaterialSamplerType SamplerType);
 
@@ -546,6 +549,7 @@ public:
 	FValueRef Logarithm(FValueRef A) { return Operator(UO_Logarithm, A); }
 	FValueRef Logarithm2(FValueRef A) { return Operator(UO_Logarithm2, A); }
 	FValueRef Logarithm10(FValueRef A) { return Operator(UO_Logarithm10, A); }
+	FValueRef LWCTile(FValueRef A) { return Operator(UO_LWCTile, A); }
 	FValueRef Reciprocal(FValueRef A) { return Operator(UO_Reciprocal, A); }
 	FValueRef Round(FValueRef A) { return Operator(UO_Round, A); }
 	FValueRef Rsqrt(FValueRef A) { return Operator(UO_Rsqrt, A); }
diff --git a/Engine/Source/Runtime/Engine/Public/Materials/MaterialIRModule.h b/Engine/Source/Runtime/Engine/Public/Materials/MaterialIRModule.h
--- a/Engine/Source/Runtime/Engine/Public/Materials/MaterialIRModule.h
+++ b/Engine/Source/Runtime/Engine/Public/Materials/MaterialIRModule.h
@@ -80,6 +80,12 @@ public:
 	// Retrieves parameter metadata for a given parameter ID.
 	const FMaterialParameterMetadata& GetParameterMetadata(uint32 ParameterId) const { return ParameterIdToData[ParameterId].Value; }
 
+	// Find existing or add a material parameter collection -- may return INDEX_NONE if too many have already been added.
+	int32 FindOrAddParameterCollection(UMaterialParameterCollection* ParameterCollection);
+
+	// Returns the list of material parameter collections used by this material.
+	TConstArrayView<UMaterialParameterCollection*> GetParameterCollections() const { return ParameterCollections; }
+
 	// Stores a user-defined string and returns a view of it.
 	FStringView InternString(FStringView InString);
 
@@ -129,6 +135,9 @@ private:
 	// Environment define names for shader compilation.
 	TSet<FName> EnvironmentDefines;
 
+	// Parameter collections used by this material.
+	TArray<UMaterialParameterCollection*> ParameterCollections;
+
 	// List of compilation errors.
 	TArray<FError> Errors;
 
diff --git a/Engine/Source/Runtime/Engine/Public/Materials/MaterialIRTypes.h b/Engine/Source/Runtime/Engine/Public/Materials/MaterialIRTypes.h
--- a/Engine/Source/Runtime/Engine/Public/Materials/MaterialIRTypes.h
+++ b/Engine/Source/Runtime/Engine/Public/Materials/MaterialIRTypes.h
@@ -21,9 +21,10 @@ enum class ETypeKind : uint8
 
 	// Opaque Objects
 	
-	Texture2D,			// A 2D texture object.
-	SubstrateData,		// Substrate-specific data structure.
-	VTPageTableResult,	// A virtual texture page table lookup result.
+	ParameterCollection,	// Material parameter collection object.
+	Texture2D,				// A 2D texture object.
+	SubstrateData,			// Substrate-specific data structure.
+	VTPageTableResult,		// A virtual texture page table lookup result.
 };
 
 // Returns the string representation of a ETypeKind.
@@ -154,6 +155,9 @@ struct FType
 	// Returns the aggregate type associated with the given material aggregate definition.
 	static FType MakeAggregate(const UMaterialAggregate* Aggregate);
 
+	// Returns the Material Parameter Collection object type.
+	static FType MakeParameterCollection();
+
 	// Returns the Texture2D object type.
 	static FType MakeTexture2D();
 	
@@ -205,6 +209,9 @@ struct FType
 	// Returns true if the given type is a vector primitive.
 	bool IsVector() const { return Is(ETypeKind::Primitive) && Primitive.IsVector(); }
 
+	// Returns whether this type is a Material Parameter Collection object.
+	bool IsParameterCollection() const { return Kind == ETypeKind::ParameterCollection; }
+
 	// Returns whether this type is a Texture2D object.
 	bool IsTexture2D() const { return Kind == ETypeKind::Texture2D; }
 
