Add DrawDebugCircleArc version accepting a FQuat for orienting the arc with an arbitrary roll

diff --git a/Engine/Source/Runtime/Engine/Private/DrawDebugHelpers.cpp b/Engine/Source/Runtime/Engine/Private/DrawDebugHelpers.cpp
--- a/Engine/Source/Runtime/Engine/Private/DrawDebugHelpers.cpp
+++ b/Engine/Source/Runtime/Engine/Private/DrawDebugHelpers.cpp
@@ -534,6 +534,46 @@ void DrawDebugCircleArc(const UWorld* InWorld, const FVector& Center, float Radi
 
 }
 
+void DrawDebugCircleArc(const UWorld* InWorld, const FVector& Center, float Radius, const FQuat& Rotation, float AngleWidth, int32 Segments, const FColor& Color, bool PersistentLines, float LifeTime, uint8 DepthPriority, float Thickness)
+{
+	RETURN_IF_DEBUG_DRAW_DISABLED;
+	
+	if (GEngine->GetNetMode(InWorld) != NM_DedicatedServer)
+	{
+		if (ULineBatchComponent* const LineBatcher = GetDebugLineBatcher(InWorld, PersistentLines, LifeTime, (DepthPriority == SDPG_Foreground)))
+		{
+			const float LineLifeTime = GetDebugLineLifeTime(LineBatcher, LifeTime, PersistentLines);
+
+			// Need at least 4 segments
+			Segments = FMath::Max(Segments, 4);
+			const float AngleStep = AngleWidth / static_cast<float>(Segments) * 2.f;
+
+			const FVector AxisX = Rotation.Vector();
+			const FVector AxisY = Rotation.GetAxisY();
+
+			TArray<FBatchedLine> Lines;
+			Lines.Reserve(Segments);
+			float Angle = -AngleWidth;
+			
+			FVector PrevVertex = Center + Radius * (AxisX * FMath::Cos(Angle) + AxisY * FMath::Sin(Angle));
+			
+			while (Segments--)
+			{
+				Angle += AngleStep;
+				const FVector NextVertex = Center + Radius * (AxisX * FMath::Cos(Angle) + AxisY * FMath::Sin(Angle));
+				Lines.Emplace(FBatchedLine(PrevVertex, NextVertex, Color, LineLifeTime, Thickness, DepthPriority));
+				PrevVertex = NextVertex;
+			}
+
+			LineBatcher->DrawLines(Lines);
+		}
+	}
+	else
+	{
+		UE_DRAW_SERVER_DEBUG_ON_EACH_CLIENT(DrawDebugCircleArc, Center, Radius, Rotation, AngleWidth, Segments, AdjustColorForServer(Color), PersistentLines, LifeTime, DepthPriority, Thickness);
+	}
+}
+
 void DrawDebug2DDonut(const UWorld* InWorld, const FMatrix& TransformMatrix, float InnerRadius, float OuterRadius, int32 Segments, const FColor& Color, bool bPersistentLines, float LifeTime, uint8 DepthPriority, float Thickness)
 {
 	RETURN_IF_DEBUG_DRAW_DISABLED;
diff --git a/Engine/Source/Runtime/Engine/Public/DrawDebugHelpers.h b/Engine/Source/Runtime/Engine/Public/DrawDebugHelpers.h
--- a/Engine/Source/Runtime/Engine/Public/DrawDebugHelpers.h
+++ b/Engine/Source/Runtime/Engine/Public/DrawDebugHelpers.h
@@ -38,6 +38,7 @@ ENGINE_API void DrawDebugCircle(const UWorld* InWorld, const FMatrix& TransformM
 ENGINE_API void DrawDebugCircle(const UWorld* InWorld, FVector Center, double Radius, int32 Segments, const FColor& Color, bool bPersistentLines = false, float LifeTime=-1.f, uint8 DepthPriority = 0, float Thickness = 0.f, FVector YAxis=FVector(0.f,1.f,0.f), FVector ZAxis=FVector(0.f,0.f,1.f), bool bDrawAxis=true);
 /** Draw Debug Circle Arc */
 ENGINE_API void DrawDebugCircleArc(const UWorld* InWorld, const FVector& Center, float Radius, const FVector& Direction, float AngleWidth, int32 Segments, const FColor& Color, bool PersistentLines = false, float LifeTime = -1.f, uint8 DepthPriority = 0, float Thickness = 0.f);
+ENGINE_API void DrawDebugCircleArc(const UWorld* InWorld, const FVector& Center, float Radius, const FQuat& Rotation, float AngleWidth, int32 Segments, const FColor& Color, bool PersistentLines = false, float LifeTime = -1.f, uint8 DepthPriority = 0, float Thickness = 0.f);
 /** Draw Debug 2D donut */
 ENGINE_API void DrawDebug2DDonut(const UWorld* InWorld, const FMatrix& TransformMatrix, float InnerRadius, float OuterRadius, int32 Segments, const FColor& Color, bool bPersistentLines = false, float LifeTime=-1.f, uint8 DepthPriority = 0, float Thickness = 0.f);
 /** Draw a debug sphere */
