Reverting changes from CL 36193418 after discussion with the author of that CL.

My specific interest here was in GetNumMaterials() sometimes returning 0 just because the mesh component was compiling even though that value was inaccurate, which caused a bug where we weren't clearing materials properly.

GetNumMaterials should NOT return 0 in such cases, as it's unintuitive and inconsistent with the behavior of other components.  Besides, it was possible to get around it by simply calling GetSkinnedAsset and then getting the Materials and checking Num() directly, which isn't a good thing to encourage people to do.

Here I'm reverting ALL of the changes at Alexis' suggestion, rather than just changing GetNumMaterials itself.  Note that the CL that made this change said: "Interchange: fix deadlock when a selected skeletal mesh in the world is re-import."  So engine may need to look out for such issues again with this change and investigate a different fix if the issue recurs.

Posted for review, will make any needed changes once review is done.

#RNX #NoReleaseNotes

diff --git a/Engine/Source/Runtime/Engine/Private/Components/SkinnedMeshComponent.cpp b/Engine/Source/Runtime/Engine/Private/Components/SkinnedMeshComponent.cpp
--- a/Engine/Source/Runtime/Engine/Private/Components/SkinnedMeshComponent.cpp
+++ b/Engine/Source/Runtime/Engine/Private/Components/SkinnedMeshComponent.cpp
@@ -1712,19 +1712,16 @@ void USkinnedMeshComponent::UpdateFollowerComponent()
 // this has to be skeletalmesh material. You can't have more than what SkeletalMesh materials have
 int32 USkinnedMeshComponent::GetNumMaterials() const
 {
-	if (GetSkinnedAsset() && !GetSkinnedAsset()->IsCompiling())
+	if (GetSkinnedAsset())
 	{
 		return GetSkinnedAsset()->GetMaterials().Num();
 	}
+
 	return 0;
 }
 
 void USkinnedMeshComponent::GetPrimitiveStats(FPrimitiveStats& PrimitiveStats) const
 {
-	if (!IsRegistered() || IsCompiling())
-	{
-		return;
-	}
 	const TIndirectArray<FSkeletalMeshLODRenderData>& LODRenderData = GetSkeletalMeshRenderData()->LODRenderData;
 	const TArray<FSkeletalMaterial>& Materials = GetSkinnedAsset()->GetMaterials();
 	PrimitiveStats.LODStats.Reserve(LODRenderData.Num());
@@ -1757,7 +1754,7 @@ UMaterialInterface* USkinnedMeshComponent::GetMaterial(int32 MaterialIndex) cons
 
 int32 USkinnedMeshComponent::GetMaterialIndex(FName MaterialSlotName) const
 {
-	if (GetSkinnedAsset() && !GetSkinnedAsset()->IsCompiling())
+	if (GetSkinnedAsset())
 	{
 		const TArray<FSkeletalMaterial>& SkeletalMeshMaterials = GetSkinnedAsset()->GetMaterials();
 		for (int32 MaterialIndex = 0; MaterialIndex < SkeletalMeshMaterials.Num(); ++MaterialIndex)
@@ -1775,7 +1772,7 @@ int32 USkinnedMeshComponent::GetMaterialIndex(FName MaterialSlotName) const
 TArray<FName> USkinnedMeshComponent::GetMaterialSlotNames() const
 {
 	TArray<FName> MaterialNames;
-	if (GetSkinnedAsset() && !GetSkinnedAsset()->IsCompiling())
+	if (GetSkinnedAsset())
 	{
 		const TArray<FSkeletalMaterial>& SkeletalMeshMaterials = GetSkinnedAsset()->GetMaterials();
 		for (int32 MaterialIndex = 0; MaterialIndex < SkeletalMeshMaterials.Num(); ++MaterialIndex)
