Remove r.VRS.DisableForCustomPrimitiveData because the bug it was designed to work around has been fixed.
It was already disabled by default and tooltipped as likely unnecessary.

diff --git a/Engine/Source/Runtime/Engine/Private/Materials/MaterialShared.cpp b/Engine/Source/Runtime/Engine/Private/Materials/MaterialShared.cpp
--- a/Engine/Source/Runtime/Engine/Private/Materials/MaterialShared.cpp
+++ b/Engine/Source/Runtime/Engine/Private/Materials/MaterialShared.cpp
@@ -105,13 +105,6 @@ static TAutoConsoleVariable<bool> CVarUsingUseNewMaterialTranslatorPrototype(
 	TEXT("Controls whether to enable the new material translator prototype (WIP) ."),
 	ECVF_RenderThreadSafe);
 
-static TAutoConsoleVariable<bool> CVarDisableVRSForCustomPrimitiveData(
-	TEXT("r.VRS.DisableForCustomPrimitiveData"),
-	false,
-	TEXT("Controls whether or not to disable VRS for materials using Custom Primitive Data.\n")
-	TEXT("This was a workaround for Nanite not using centroid sampling for VRS clusters, and should no longer be necessary in 5.6."),
-	ECVF_RenderThreadSafe | ECVF_ReadOnly);
-
 bool IsUsingNewMaterialTranslatorPrototype()
 {
 	return CVarUsingUseNewMaterialTranslatorPrototype.GetValueOnAnyThread();
@@ -1647,12 +1640,9 @@ bool FMaterialResource::HasVertexPositionOffsetConnected() const { return GetCac
 bool FMaterialResource::HasPixelDepthOffsetConnected() const { return Material->HasPixelDepthOffsetConnected(); }
 EMaterialShadingRate FMaterialResource::GetShadingRate() const { return Material->ShadingRate; }
 bool FMaterialResource::IsVariableRateShadingAllowed() const {
-	const bool bDisableForCustomPrimitiveData = Material->HasCustomPrimitiveData() && CVarDisableVRSForCustomPrimitiveData.GetValueOnAnyThread();
-
 	return Material->bAllowVariableRateShading
 		&& !IsMasked()								// When using pixel discard, coarse shading causes the whole block to get discarded resulting in noticeable artifacts
-		&& !bDisableForCustomPrimitiveData			// When using custom primitive data, Nanite can't determine which primitive ID to use for shading clusters, leading to edge artifacts
-		&& !HasPixelDepthOffsetConnected();			// When using pixel depth offset, coarse shading causes the whole block to output the same depth value resulting in noticeable artifacts
+		&& !HasPixelDepthOffsetConnected();			// When writing depth from the pixel shader, coarse shading causes the whole block to output the same depth value resulting in noticeable artifacts
 }
 FString FMaterialResource::GetBaseMaterialPathName() const { return Material->GetPathName(); }
 FString FMaterialResource::GetDebugName() const
