Fix dllexport api issues in several headers

#rnx

diff --git a/Engine/Plugins/AudioInsights/Source/AudioInsights/Public/Views/SoundPlotsDashboardViewFactory.h b/Engine/Plugins/AudioInsights/Source/AudioInsights/Public/Views/SoundPlotsDashboardViewFactory.h
--- a/Engine/Plugins/AudioInsights/Source/AudioInsights/Public/Views/SoundPlotsDashboardViewFactory.h
+++ b/Engine/Plugins/AudioInsights/Source/AudioInsights/Public/Views/SoundPlotsDashboardViewFactory.h
@@ -15,7 +15,7 @@ namespace UE::Audio::Insights
 	class FSoundPlotsDashboardViewFactory : public IDashboardViewFactory, public TSharedFromThis<FSoundPlotsDashboardViewFactory>
 	{
 	public:
-		UE_API virtual ~FSoundPlotsDashboardViewFactory() = default;
+		virtual ~FSoundPlotsDashboardViewFactory() = default;
 
 		UE_API virtual FName GetName() const override;
 		UE_API virtual FText GetDisplayName() const override;
diff --git a/Engine/Plugins/Experimental/AnimNext/Source/AnimNext/Public/Variables/AnimNextSoftVariableReference.h b/Engine/Plugins/Experimental/AnimNext/Source/AnimNext/Public/Variables/AnimNextSoftVariableReference.h
--- a/Engine/Plugins/Experimental/AnimNext/Source/AnimNext/Public/Variables/AnimNextSoftVariableReference.h
+++ b/Engine/Plugins/Experimental/AnimNext/Source/AnimNext/Public/Variables/AnimNextSoftVariableReference.h
@@ -1,4 +1,4 @@
-﻿// Copyright Epic Games, Inc. All Rights Reserved.
+// Copyright Epic Games, Inc. All Rights Reserved.
 
 #pragma once
 
@@ -16,7 +16,7 @@ struct FAnimNextSoftVariableReference
 	FAnimNextSoftVariableReference() = default;
 
 	// Construct a soft reference to an asset/struct variable
-	ANIMNEXT_API explicit FAnimNextSoftVariableReference(FName InName, const FSoftObjectPath& InSoftObjectPath)
+	explicit FAnimNextSoftVariableReference(FName InName, const FSoftObjectPath& InSoftObjectPath)
 		: Name(InName)
 		, SoftObjectPath(InSoftObjectPath)
 	{}
diff --git a/Engine/Source/Runtime/Engine/Classes/Kismet/KismetMathLibrary.h b/Engine/Source/Runtime/Engine/Classes/Kismet/KismetMathLibrary.h
--- a/Engine/Source/Runtime/Engine/Classes/Kismet/KismetMathLibrary.h
+++ b/Engine/Source/Runtime/Engine/Classes/Kismet/KismetMathLibrary.h
@@ -213,7 +213,7 @@ class UKismetMathLibrary : public UBlueprintFunctionLibrary
 	static ENGINE_API bool RandomBoolWithWeightFromStream(const FRandomStream& RandomStream, float Weight);
 
 	UE_DEPRECATED(5.2, "Use RandomBoolWithWeightFromStream taking the FRandomStream as the first argument.")
-	static UE_INL_API bool RandomBoolWithWeightFromStream(float Weight, const FRandomStream& RandomStream)
+	static bool RandomBoolWithWeightFromStream(float Weight, const FRandomStream& RandomStream)
 	{
 		return RandomBoolWithWeightFromStream(RandomStream, Weight);
 	}
@@ -383,7 +383,7 @@ class UKismetMathLibrary : public UBlueprintFunctionLibrary
 
 	/** Returns a uniformly distributed random number between 0 and Max - 1 */
 	UFUNCTION(BlueprintPure, Category="Math|Random", meta=(NotBlueprintThreadSafe))
-	static ENGINE_API int32 RandomInteger(UPARAM(DisplayName = "Exclusive Max") int32 Max);
+	static UE_INL_API int32 RandomInteger(UPARAM(DisplayName = "Exclusive Max") int32 Max);
 
 	/** Return a random integer between Min and Max (>= Min and <= Max) */
 	UFUNCTION(BlueprintPure, Category="Math|Random", meta = (NotBlueprintThreadSafe))
@@ -399,7 +399,7 @@ class UKismetMathLibrary : public UBlueprintFunctionLibrary
 
 	/** Returns Value clamped to be between A and B (inclusive) */
 	UFUNCTION(BlueprintPure, meta=(DisplayName = "Clamp (Integer)"), Category="Math|Integer")
-	static ENGINE_API int32 Clamp(int32 Value, int32 Min, int32 Max);
+	static UE_INL_API int32 Clamp(int32 Value, int32 Min, int32 Max);
 
 	/** Returns Value between A and B (inclusive) that wraps around */
 	UFUNCTION(BlueprintPure, meta = (DisplayName = "Wrap (Integer)", Min = "0", Max = "100"), Category = "Math|Integer")
@@ -486,7 +486,7 @@ class UKismetMathLibrary : public UBlueprintFunctionLibrary
 
 	/** Returns a uniformly distributed random number between 0 and Max - 1 */
 	UFUNCTION(BlueprintPure, Category="Math|Random", meta=(NotBlueprintThreadSafe))
-	static ENGINE_API int64 RandomInteger64(UPARAM(DisplayName = "Exclusive Max") int64 Max);
+	static UE_INL_API int64 RandomInteger64(UPARAM(DisplayName = "Exclusive Max") int64 Max);
 
 	/** Return a random integer64 between Min and Max (>= Min and <= Max) */
 	UFUNCTION(BlueprintPure, Category="Math|Random", meta = (NotBlueprintThreadSafe))
@@ -502,7 +502,7 @@ class UKismetMathLibrary : public UBlueprintFunctionLibrary
 
 	/** Returns Value clamped to be between A and B (inclusive) */
 	UFUNCTION(BlueprintPure, meta=(DisplayName = "Clamp (Integer64)"), Category="Math|Integer64")
-	static ENGINE_API int64 ClampInt64(int64 Value, int64 Min, int64 Max);
+	static UE_INL_API int64 ClampInt64(int64 Value, int64 Min, int64 Max);
 
 	/** Returns the absolute (positive) value of A */
 	UFUNCTION(BlueprintPure, meta=(DisplayName = "Absolute (Integer64)", CompactNodeTitle = "ABS"), Category="Math|Integer64")
@@ -735,7 +735,7 @@ class UKismetMathLibrary : public UBlueprintFunctionLibrary
 
 	/** Returns Value clamped between A and B (inclusive) */
 	UFUNCTION(BlueprintPure, meta=(DisplayName = "Clamp (Float)", Min="0.0", Max="1.0"), Category="Math|Float")
-	static ENGINE_API double FClamp(double Value, double Min, double Max);
+	static UE_INL_API double FClamp(double Value, double Min, double Max);
 
 	/** Returns Value wrapped from A and B (inclusive) */
 	UFUNCTION(BlueprintPure, meta = (DisplayName = "Wrap (Float)", Min = "0.0", Max = "1.0"), Category = "Math|Float")
@@ -779,9 +779,9 @@ class UKismetMathLibrary : public UBlueprintFunctionLibrary
 
 	/** Linearly interpolates between A and B based on Alpha (100% of A when Alpha=0 and 100% of B when Alpha=1) */
 	UFUNCTION(BlueprintPure, Category="Math|Float")
-	static ENGINE_API double Lerp(double A, double B, double Alpha);
+	static UE_INL_API double Lerp(double A, double B, double Alpha);
 	
-	/** Easeing  between A and B using a specified easing function */
+	/** Easing between A and B using a specified easing function */
 	UFUNCTION(BlueprintPure, meta = (DisplayName = "Ease", BlueprintInternalUseOnly = "true"), Category = "Math|Interpolation")
 	static ENGINE_API double Ease(double A, double B, double Alpha, TEnumAsByte<EEasingFunc::Type> EasingFunc, double BlendExp = 2, int32 Steps = 2);
 
@@ -995,11 +995,11 @@ class UKismetMathLibrary : public UBlueprintFunctionLibrary
 
 	/** Converts a Vector2D to a Vector */
 	UFUNCTION(BlueprintPure, meta = (DisplayName = "To Vector (Vector2D)", CompactNodeTitle = "->", ScriptMethod = "Vector", Keywords = "cast convert", BlueprintAutocast), Category = "Math|Conversions")
-	static ENGINE_API FVector Conv_Vector2DToVector(FVector2D InVector2D, float Z = 0);
+	static UE_INL_API FVector Conv_Vector2DToVector(FVector2D InVector2D, float Z = 0);
 
 	/** Converts a Vector2D to an IntPoint */
 	UFUNCTION(BlueprintPure, meta = (DisplayName = "To IntPoint (Vector2D)", CompactNodeTitle = "->", ScriptMethod = "IntPoint", Keywords = "cast convert", BlueprintAutocast), Category = "Math|Conversions")
-	static ENGINE_API FIntPoint Conv_Vector2DToIntPoint(FVector2D InVector2D);
+	static UE_INL_API FIntPoint Conv_Vector2DToIntPoint(FVector2D InVector2D);
 
 	/** Returns addition of Vector A and Vector B (A + B) */
 	UFUNCTION(BlueprintPure, meta = (DisplayName = "vector2d + vector2d", CompactNodeTitle = "+", ScriptMethod = "Add", ScriptOperator = "+;+=", Keywords = "+ add plus", CommutativeAssociativeBinaryOperator = "true"), Category = "Math|Vector2D")
@@ -1060,7 +1060,7 @@ class UKismetMathLibrary : public UBlueprintFunctionLibrary
 	 * @param InY New Y coordinate.
 	 */
 	UFUNCTION(BlueprintCallable, meta = (ScriptMethod = "Set"), Category = "Math|Vector2D")
-	static ENGINE_API void Set2D(UPARAM(ref) FVector2D& A, double X, double Y);
+	static UE_INL_API void Set2D(UPARAM(ref) FVector2D& A, double X, double Y);
 
 	/**
 	 * Creates a copy of this vector with both axes clamped to the given range.
@@ -1200,7 +1200,7 @@ class UKismetMathLibrary : public UBlueprintFunctionLibrary
 	 * @see NormalSafe2D()
 	 */
 	UFUNCTION(BlueprintCallable, meta = (DisplayName = "Normalize In Place (Vector2D)", Keywords = "Unit Vector", ScriptMethod = "Normalize"), Category = "Math|Vector2D")
-	static ENGINE_API void Normalize2D(UPARAM(ref) FVector2D& A, float Tolerance = 1.e-8f);
+	static UE_INL_API void Normalize2D(UPARAM(ref) FVector2D& A, float Tolerance = 1.e-8f);
 
 	/** Converts spherical coordinates on the unit sphere into a Cartesian unit length vector. */
 	UFUNCTION(BlueprintPure, meta = (DisplayName = "Spherical2D To Unit Cartesian", Keywords = "Unit Vector", ScriptMethod = "SphericalToUnitCartesian"), Category = "Math|Vector2D")
@@ -1295,7 +1295,7 @@ class UKismetMathLibrary : public UBlueprintFunctionLibrary
 	 * @param InVector Vector to copy values from.
 	 */
 	UFUNCTION(BlueprintCallable, meta = (ScriptMethod = "Assign"), Category = "Math|Vector")
-	static ENGINE_API void Vector_Assign(UPARAM(ref) FVector& A, const FVector& InVector);
+	static UE_INL_API void Vector_Assign(UPARAM(ref) FVector& A, const FVector& InVector);
 
 	/**
 	 * Set the values of the vector directly.
@@ -1305,7 +1305,7 @@ class UKismetMathLibrary : public UBlueprintFunctionLibrary
 	 * @param InZ New Z coordinate.
 	 */
 	UFUNCTION(BlueprintCallable, meta = (ScriptMethod = "Set"), Category = "Math|Vector")
-	static ENGINE_API void Vector_Set(UPARAM(ref) FVector& A, double X, double Y, double Z);
+	static UE_INL_API void Vector_Set(UPARAM(ref) FVector& A, double X, double Y, double Z);
 
 	/** Breaks a vector apart into X, Y, Z */
 	UFUNCTION(BlueprintPure, Category="Math|Vector", meta=(NativeBreakFunc))
@@ -1411,7 +1411,7 @@ class UKismetMathLibrary : public UBlueprintFunctionLibrary
 
 	/** Element-wise Vector division (Result = {A.x/B.x, A.y/B.y, A.z/B.z}) */
 	UFUNCTION(BlueprintPure, meta=(DisplayName = "vector / vector", CompactNodeTitle = "/", ScriptMethod = "Divide", ScriptOperator = "/;/=", Keywords = "/ divide division"), Category="Math|Vector")
-	static ENGINE_API FVector Divide_VectorVector(FVector A, FVector B = FVector(1.f,1.f,1.f));
+	static UE_INL_API FVector Divide_VectorVector(FVector A, FVector B = FVector(1.f,1.f,1.f));
 
 	/** Vector divide by a float */
 	UFUNCTION(BlueprintPure, meta=(DisplayName = "vector / float", CompactNodeTitle = "/", ScriptMethod = "DivideFloat", ScriptOperator = "/;/=", Keywords = "/ divide division"), Category="Math|Vector")
@@ -1463,7 +1463,7 @@ class UKismetMathLibrary : public UBlueprintFunctionLibrary
 
 	/** When this vector contains Euler angles (degrees), ensure that angles are between +/-180 */
 	UFUNCTION(BlueprintCallable, meta = (ScriptMethod = "UnwindEuler"), Category = "Math|Vector")
-	static ENGINE_API void Vector_UnwindEuler(UPARAM(ref) FVector& A);
+	static UE_INL_API void Vector_UnwindEuler(UPARAM(ref) FVector& A);
 
 	/** Create a copy of this vector, with its magnitude/size/length clamped between Min and Max. */
 	UFUNCTION(BlueprintPure, meta=(ScriptMethod = "ClampedSize"), Category="Math|Vector")
@@ -1589,11 +1589,11 @@ class UKismetMathLibrary : public UBlueprintFunctionLibrary
 
 	/** Rounds A to an integer with truncation towards zero for each element in a vector.  (e.g. -1.7 truncated to -1, 2.8 truncated to 2) */
 	UFUNCTION(BlueprintPure, meta = (DisplayName = "Truncate (Vector)", ScriptMethod = "Truncated", BlueprintAutocast), Category = "Math|Float")
-	static FIntVector FTruncVector(const FVector& InVector);
+	static UE_INL_API FIntVector FTruncVector(const FVector& InVector);
 
 	/** Rounds A to an integer with truncation towards zero for each element in a vector2D.  (e.g. -1.7 truncated to -1, 2.8 truncated to 2) */
 	UFUNCTION(BlueprintPure, meta = (DisplayName = "Truncate (Vector2D)", ScriptMethod = "Truncated", BlueprintAutocast), Category = "Math|Float")
-	static FIntVector2 FTruncVector2D(const FVector2D& InVector2D);
+	static UE_INL_API FIntVector2 FTruncVector2D(const FVector2D& InVector2D);
 
 	/**
 	 * Distance between two points.
@@ -1735,11 +1735,11 @@ class UKismetMathLibrary : public UBlueprintFunctionLibrary
 	 * @param Tolerance Minimum squared length of vector for normalization.
 	 */
 	UFUNCTION(BlueprintCallable, meta = (DisplayName = "Normalize In Place (Vector)", ScriptMethod = "Normalize", Keywords = "Unit Vector"), Category = "Math|Vector")
-	static ENGINE_API void Vector_Normalize(UPARAM(ref) FVector& A, float Tolerance = 1.e-8f);
+	static UE_INL_API void Vector_Normalize(UPARAM(ref) FVector& A, float Tolerance = 1.e-8f);
 
 	/** Linearly interpolates between A and B based on Alpha (100% of A when Alpha=0 and 100% of B when Alpha=1) */
 	UFUNCTION(BlueprintPure, meta=(DisplayName = "Lerp (Vector)", ScriptMethod = "LerpTo"), Category="Math|Vector")
-	static ENGINE_API FVector VLerp(FVector A, FVector B, float Alpha);
+	static UE_INL_API FVector VLerp(FVector A, FVector B, float Alpha);
 
 	/** Easing between A and B using a specified easing function */
 	UFUNCTION(BlueprintPure, meta = (DisplayName = "Ease (Vector)", BlueprintInternalUseOnly = "true"), Category = "Math|Interpolation")
@@ -2008,7 +2008,7 @@ class UKismetMathLibrary : public UBlueprintFunctionLibrary
 	 * @param ConeHalfAngleInDegrees	The half-angle of the cone (from ConeDir to edge), in degrees.
 	 */
 	UFUNCTION(BlueprintPure, Category = "Math|Random", meta=(NotBlueprintThreadSafe))
-	static inline FVector RandomUnitVectorInConeInDegrees(FVector ConeDir, float ConeHalfAngleInDegrees)
+	static FVector RandomUnitVectorInConeInDegrees(FVector ConeDir, float ConeHalfAngleInDegrees)
 	{
 		return RandomUnitVectorInConeInRadians(ConeDir, FMath::DegreesToRadians(ConeHalfAngleInDegrees));
 	}
@@ -2031,7 +2031,7 @@ class UKismetMathLibrary : public UBlueprintFunctionLibrary
 	* @param MaxPitchInDegrees	The pitch angle of the cone (from ConeDir to vertical edge), in degrees.	
 	*/
 	UFUNCTION(BlueprintPure, Category = "Math|Random", meta = (Keywords = "RandomVector Pitch Yaw", NotBlueprintThreadSafe))
-	static inline FVector RandomUnitVectorInEllipticalConeInDegrees(FVector ConeDir, float MaxYawInDegrees, float MaxPitchInDegrees)
+	static FVector RandomUnitVectorInEllipticalConeInDegrees(FVector ConeDir, float MaxYawInDegrees, float MaxPitchInDegrees)
 	{
 		return RandomUnitVectorInEllipticalConeInRadians(ConeDir, FMath::DegreesToRadians(MaxYawInDegrees), FMath::DegreesToRadians(MaxPitchInDegrees));
 	}
@@ -2124,7 +2124,7 @@ class UKismetMathLibrary : public UBlueprintFunctionLibrary
 	 * @param InVector Vector to copy values from.
 	 */
 	UFUNCTION(BlueprintCallable, meta = (ScriptMethod = "Assign"), Category = "Math|Vector4")
-	static ENGINE_API void Vector4_Assign(UPARAM(ref) FVector4& A, const FVector4& InVector);
+	static UE_INL_API void Vector4_Assign(UPARAM(ref) FVector4& A, const FVector4& InVector);
 
 	/**
 	 * Set the values of the vector directly.
@@ -2135,7 +2135,7 @@ class UKismetMathLibrary : public UBlueprintFunctionLibrary
 	 * @param InW New W coordinate.
 	 */
 	UFUNCTION(BlueprintCallable, meta = (ScriptMethod = "Set"), Category = "Math|Vector4")
-	static ENGINE_API void Vector4_Set(UPARAM(ref) FVector4& A, double X, double Y, double Z, double W);
+	static UE_INL_API void Vector4_Set(UPARAM(ref) FVector4& A, double X, double Y, double Z, double W);
 
 	/** Returns the cross product of two vectors - see  http://mathworld.wolfram.com/CrossProduct.html */
 	UFUNCTION(BlueprintPure, meta = (DisplayName = "Cross Product XYZ (Vector4)", CompactNodeTitle = "cross3", ScriptMethod = "Cross3"), Category = "Math|Vector4")
@@ -2230,7 +2230,7 @@ class UKismetMathLibrary : public UBlueprintFunctionLibrary
 	 * @param Tolerance Minimum squared length of vector for normalization.
 	 */
 	UFUNCTION(BlueprintCallable, meta = (DisplayName = "Normalize In Place XYZ (Vector4)", ScriptMethod = "Normalize3", Keywords = "Unit Vector"), Category = "Math|Vector4")
-	static ENGINE_API void Vector4_Normalize3(UPARAM(ref) FVector4& A, float Tolerance = 1.e-8f);
+	static UE_INL_API void Vector4_Normalize3(UPARAM(ref) FVector4& A, float Tolerance = 1.e-8f);
 
 	/** 
 	 * Given a direction vector and a surface normal, returns the vector reflected across the surface normal.
@@ -2259,7 +2259,7 @@ class UKismetMathLibrary : public UBlueprintFunctionLibrary
 
 	/** Makes a rotator {Roll, Pitch, Yaw} from rotation values supplied in degrees */
 	UFUNCTION(BlueprintPure, Category="Math|Rotator", meta=(Keywords="construct build rotation rotate rotator makerotator", NativeMakeFunc))
-	static ENGINE_API FRotator MakeRotator(
+	static UE_INL_API FRotator MakeRotator(
 		UPARAM(DisplayName="X (Roll)") float Roll,	
 		UPARAM(DisplayName="Y (Pitch)") float Pitch,
 		UPARAM(DisplayName="Z (Yaw)") float Yaw);
@@ -2317,7 +2317,7 @@ class UKismetMathLibrary : public UBlueprintFunctionLibrary
 
 	/** Breaks apart a rotator into {Roll, Pitch, Yaw} angles in degrees */
 	UFUNCTION(BlueprintPure, Category = "Math|Rotator", meta = (Keywords = "rotation rotate rotator breakrotator", NativeBreakFunc))
-	static ENGINE_API void BreakRotator(
+	static UE_INL_API void BreakRotator(
 		UPARAM(DisplayName="Rotation") FRotator InRot,
 		UPARAM(DisplayName="X (Roll)") float& Roll,
 		UPARAM(DisplayName="Y (Pitch)") float& Pitch,
@@ -2437,7 +2437,7 @@ class UKismetMathLibrary : public UBlueprintFunctionLibrary
 
 	// Identity matrix
 	UFUNCTION(BlueprintPure, meta = (DisplayName = "Identity (Matrix)", ScriptConstant = "Identity", ScriptConstantHost = "/Script/CoreUObject.Matrix"), Category = "Math|Matrix")
-	static FMatrix Matrix_Identity();
+	static UE_INL_API FMatrix Matrix_Identity();
 
 	/**
 	 * Gets the result of multiplying a Matrix to this.
@@ -2446,7 +2446,7 @@ class UKismetMathLibrary : public UBlueprintFunctionLibrary
 	 * @return The result of multiplication.
 	 */
 	UFUNCTION(BlueprintPure, meta = (DisplayName = "Matrix * Matrix", CompactNodeTitle = "*", ScriptMethod = "Multiply", ScriptOperator = "*;*=", Keywords = "* multiply"), Category = "Math|Matrix")
-	static FMatrix Multiply_MatrixMatrix (const FMatrix& A, const FMatrix& B);
+	static UE_INL_API FMatrix Multiply_MatrixMatrix (const FMatrix& A, const FMatrix& B);
 
 	/**
 	 * Gets the result of adding a matrix to this.
@@ -2455,14 +2455,14 @@ class UKismetMathLibrary : public UBlueprintFunctionLibrary
 	 * @return The result of addition.
 	 */
 	UFUNCTION(BlueprintPure, meta = (DisplayName = "Matrix + Matrix", CompactNodeTitle = "+", ScriptMethod = "Add", ScriptOperator = "+;+=", Keywords = "+ add plus", CommutativeAssociativeBinaryOperator = "true"), Category = "Math|Matrix")
-	static FMatrix Add_MatrixMatrix (const FMatrix& A, const FMatrix& B);
+	static UE_INL_API FMatrix Add_MatrixMatrix (const FMatrix& A, const FMatrix& B);
 
 	/**
 	  * Multiplies all values of the matrix by a float.
 	  * If your Matrix represents a Transform that you wish to scale you should use Apply Scale instead
 	  */
 	UFUNCTION(BlueprintPure, meta = (DisplayName = "Matrix * Float", CompactNodeTitle = "*", ScriptMethod = "MultiplyFloat", ScriptOperator = "*;*=", Keywords = "* multiply"), Category = "Math|Matrix")
-	static FMatrix Multiply_MatrixFloat (const FMatrix& A, double B);
+	static UE_INL_API FMatrix Multiply_MatrixFloat (const FMatrix& A, double B);
 
 	/**
 	 * Checks whether another Matrix is equal to this, within specified tolerance.
@@ -2520,7 +2520,7 @@ class UKismetMathLibrary : public UBlueprintFunctionLibrary
 
 	// Transpose.
 	UFUNCTION(BlueprintPure, meta = (DisplayName = "Get Transposed (Matrix)", ScriptMethod = "GetTransposed"), Category = "Math|Matrix")
-	static FMatrix Matrix_GetTransposed(const FMatrix& M);
+	static UE_INL_API FMatrix Matrix_GetTransposed(const FMatrix& M);
 
 	// @return determinant of this matrix.
 	UFUNCTION(BlueprintPure, meta = (DisplayName = "Get Determinant (Matrix)", ScriptMethod = "GetDeterminant"), Category = "Math|Matrix")
@@ -2534,23 +2534,23 @@ class UKismetMathLibrary : public UBlueprintFunctionLibrary
 
 	/** Get the inverse of the Matrix. Handles nil matrices. */
 	UFUNCTION(BlueprintPure, meta = (DisplayName = "GetInverse (Matrix)", ScriptMethod = "GetInverse"), Category = "Math|Matrix")
-	static FMatrix Matrix_GetInverse(const FMatrix& M);
+	static UE_INL_API FMatrix Matrix_GetInverse(const FMatrix& M);
 
 	/** Get the Transose Adjoint of the Matrix. */
 	UFUNCTION(BlueprintPure, meta = (DisplayName = "Get Transpose Adjoint (Matrix)", ScriptMethod = "GetTransposeAdjoint"), Category = "Math|Matrix")
-	static FMatrix Matrix_GetTransposeAdjoint(const FMatrix& M);
+	static UE_INL_API FMatrix Matrix_GetTransposeAdjoint(const FMatrix& M);
 
 	/** Remove any scaling from this matrix (ie magnitude of each row is 1) with error Tolerance
 	* (Assumes Matrix represents a transform) 
 	*/
 	UFUNCTION(BlueprintCallable, meta = (DisplayName = "Remove Scaling (Matrix)", ScriptMethod = "RemoveScaling"), Category = "Math|Matrix")
-	static ENGINE_API void Matrix_RemoveScaling(UPARAM(Ref) FMatrix& M, float Tolerance = 1.e-8f);
+	static UE_INL_API void Matrix_RemoveScaling(UPARAM(Ref) FMatrix& M, float Tolerance = 1.e-8f);
 
 	/** Returns matrix after RemoveScaling with error Tolerance
 	* (Assumes Matrix represents a transform)
 	*/
 	UFUNCTION(BlueprintPure, meta = (DisplayName = "Get Matrix Without Scale (Matrix)", ScriptMethod = "GetMatrixWithoutScale"), Category = "Math|Matrix")
-	static FMatrix Matrix_GetMatrixWithoutScale(const FMatrix& M, float Tolerance = 1.e-8f);
+	static UE_INL_API FMatrix Matrix_GetMatrixWithoutScale(const FMatrix& M, float Tolerance = 1.e-8f);
 
 	/** return a 3D scale vector calculated from this matrix (where each component is the magnitude of a row vector) with error Tolerance.
 	* (Assumes Matrix represents a transform)
@@ -2562,13 +2562,13 @@ class UKismetMathLibrary : public UBlueprintFunctionLibrary
 	* (Assumes Matrix represents a transform)
 	*/
 	UFUNCTION(BlueprintPure, meta = (DisplayName = "Remove Translation (Matrix)", ScriptMethod = "RemoveTranslation"), Category = "Math|Matrix")
-	static FMatrix Matrix_RemoveTranslation(const FMatrix& M);
+	static UE_INL_API FMatrix Matrix_RemoveTranslation(const FMatrix& M);
 
 	/** Returns a matrix with an additional translation concatenated.
 	* (Assumes Matrix represents a transform)
 	*/
 	UFUNCTION(BlueprintPure, meta = (DisplayName = "Concatenate Translation (Matrix)", ScriptMethod = "ConcatenateTranslation"), Category = "Math|Matrix")
-	static FMatrix Matrix_ConcatenateTranslation(const FMatrix& M, FVector Translation);
+	static UE_INL_API FMatrix Matrix_ConcatenateTranslation(const FMatrix& M, FVector Translation);
 
 	/** Returns true if any element of this matrix is NaN */
 	UFUNCTION(BlueprintPure, meta = (DisplayName = "Contains NaN (Matrix)", ScriptMethod = "ContainsNaN"), Category = "Math|Matrix")
@@ -2578,7 +2578,7 @@ class UKismetMathLibrary : public UBlueprintFunctionLibrary
 	* (Assumes Matrix represents a transform)
 	*/
 	UFUNCTION(BlueprintPure, meta = (DisplayName = "Scale Translation (Matrix)", ScriptMethod = "ScaleTranslation"), Category = "Math|Matrix")
-	static FMatrix Matrix_ScaleTranslation(const FMatrix& M, FVector Scale3D);
+	static UE_INL_API FMatrix Matrix_ScaleTranslation(const FMatrix& M, FVector Scale3D);
 
 	/** @return the maximum magnitude of any row of the matrix.
 	* (Assumes Matrix represents a transform)
@@ -2590,7 +2590,7 @@ class UKismetMathLibrary : public UBlueprintFunctionLibrary
 	* (Assumes Matrix represents a transform)
 	*/
 	UFUNCTION(BlueprintPure, meta = (DisplayName = "Apply Scale (Matrix)", ScriptMethod = "ApplyScale"), Category = "Math|Matrix")
-	static FMatrix Matrix_ApplyScale(const FMatrix& M, float Scale);
+	static UE_INL_API FMatrix Matrix_ApplyScale(const FMatrix& M, float Scale);
 
 	/**
 	 * get axis of this matrix scaled by the scale of the matrix
@@ -2642,13 +2642,13 @@ class UKismetMathLibrary : public UBlueprintFunctionLibrary
 	 * @param Axis vector of the axis
 	 */
 	UFUNCTION(BlueprintCallable, meta = (DisplayName = "Set Axis (Matrix)", ScriptMethod = "SetAxis"), Category = "Math|Matrix")
-	static ENGINE_API void Matrix_SetAxis(UPARAM(Ref) FMatrix& M, TEnumAsByte<EAxis::Type> Axis, FVector AxisVector);
+	static UE_INL_API void Matrix_SetAxis(UPARAM(Ref) FMatrix& M, TEnumAsByte<EAxis::Type> Axis, FVector AxisVector);
 
 	/** Set the origin of the coordinate system to the given vector
 	* (Assumes Matrix represents a transform)
 	*/
 	UFUNCTION(BlueprintCallable, meta = (DisplayName = "Set Origin (Matrix)", ScriptMethod = "SetOrigin"), Category = "Math|Matrix")
-	static ENGINE_API void Matrix_SetOrigin(UPARAM(Ref) FMatrix& M, FVector NewOrigin);
+	static UE_INL_API void Matrix_SetOrigin(UPARAM(Ref) FMatrix& M, FVector NewOrigin);
 
 	/**
 	 * get a column of this matrix
@@ -2660,7 +2660,7 @@ class UKismetMathLibrary : public UBlueprintFunctionLibrary
 	static UE_INL_API FVector Matrix_GetColumn(const FMatrix& M, TEnumAsByte<EMatrixColumns::Type> Column);
 
 	UFUNCTION(BlueprintCallable, meta = (DisplayName = "Set Column (Matrix)", ScriptMethod = "SetColumn"), Category = "Math|Matrix")
-	static ENGINE_API void Matrix_SetColumn(UPARAM(Ref) FMatrix& M, TEnumAsByte<EMatrixColumns::Type> Column, FVector Value);
+	static UE_INL_API void Matrix_SetColumn(UPARAM(Ref) FMatrix& M, TEnumAsByte<EMatrixColumns::Type> Column, FVector Value);
 
 	/** 
 	* Get the rotator representation of this matrix
@@ -2728,7 +2728,7 @@ class UKismetMathLibrary : public UBlueprintFunctionLibrary
 	 * (Assumes Matrix represents a transform)
 	 */
 	UFUNCTION(BlueprintPure, meta = (DisplayName = "Mirror (Matrix)", ScriptMethod = "Mirror"), Category = "Math|Matrix")
-	static FMatrix Matrix_Mirror(const FMatrix& M, TEnumAsByte<EAxis::Type> MirrorAxis, TEnumAsByte<EAxis::Type> FlipAxis);
+	static UE_INL_API FMatrix Matrix_Mirror(const FMatrix& M, TEnumAsByte<EAxis::Type> MirrorAxis, TEnumAsByte<EAxis::Type> FlipAxis);
 
 	//
 	// Quat constants - exposed for scripting
@@ -2748,7 +2748,7 @@ class UKismetMathLibrary : public UBlueprintFunctionLibrary
 
 	/** Returns true if Quat A is not equal to Quat B (A != B) within a specified error tolerance */
 	UFUNCTION(BlueprintPure, meta = (DisplayName = "Not Equal (Quat)", CompactNodeTitle = "!=", ScriptMethod = "NotEqual", ScriptOperator = "!=", Keywords = "!= not equal"), Category = "Math|Quat")
-	static ENGINE_API bool NotEqual_QuatQuat(const FQuat& A, const FQuat& B, float ErrorTolerance = 1.e-4f);
+	static UE_INL_API bool NotEqual_QuatQuat(const FQuat& A, const FQuat& B, float ErrorTolerance = 1.e-4f);
 
 	/** Returns addition of Vector A and Vector B (A + B) */
 	UFUNCTION(BlueprintPure, meta = (DisplayName = "Quat + Quat", CompactNodeTitle = "+", ScriptMethod = "Add", ScriptOperator = "+;+=", Keywords = "+ add plus", CommutativeAssociativeBinaryOperator = "true"), Category = "Math|Quat")
@@ -2811,7 +2811,7 @@ class UKismetMathLibrary : public UBlueprintFunctionLibrary
 
 	/** Modify the quaternion to ensure that the delta between it and B represents the shortest possible rotation angle. */
 	UFUNCTION(BlueprintCallable, meta = (DisplayName = "Ensure Shortest Arc To (Quat)", ScriptMethod = "EnsureShortestArcTo"), Category = "Math|Quat")
-	static ENGINE_API void Quat_EnforceShortestArcWith(UPARAM(ref) FQuat& A, const FQuat& B);
+	static UE_INL_API void Quat_EnforceShortestArcWith(UPARAM(ref) FQuat& A, const FQuat& B);
 
 	/**	Convert a Quaternion into floating-point Euler angles (in degrees). */
 	UFUNCTION(BlueprintPure, meta = (DisplayName = "Euler (Quat)", ScriptMethod = "Euler"), Category = "Math|Quat")
@@ -2860,7 +2860,7 @@ class UKismetMathLibrary : public UBlueprintFunctionLibrary
 	 * @param Tolerance Minimum squared length of quaternion for normalization.
 	 */
 	UFUNCTION(BlueprintCallable, meta = (DisplayName = "Normalize (Quat)", ScriptMethod = "Normalize"), Category = "Math|Quat")
-	static ENGINE_API void Quat_Normalize(UPARAM(ref) FQuat& Q, float Tolerance = 1.e-4f);
+	static UE_INL_API void Quat_Normalize(UPARAM(ref) FQuat& Q, float Tolerance = 1.e-4f);
 
 	/**
 	 * Get a normalized copy of this quaternion.
@@ -2889,7 +2889,7 @@ class UKismetMathLibrary : public UBlueprintFunctionLibrary
 
 	/** Set X, Y, Z, W components of Quaternion. */
 	UFUNCTION(BlueprintCallable, meta = (DisplayName = "Set Components (Quat)", ScriptMethod = "SetComponents"), Category = "Math|Quat")
-	static ENGINE_API void Quat_SetComponents(UPARAM(ref) FQuat& Q, float X, float Y, float Z, float W);
+	static UE_INL_API void Quat_SetComponents(UPARAM(ref) FQuat& Q, float X, float Y, float Z, float W);
 
 	/**
 	 * Convert a vector of floating-point Euler angles (in degrees) into a Quaternion.
@@ -2898,7 +2898,7 @@ class UKismetMathLibrary : public UBlueprintFunctionLibrary
 	 * @param Euler the Euler angles
 	 */
 	UFUNCTION(BlueprintCallable, meta = (DisplayName = "Set from Euler (Quat)", ScriptMethod = "SetFromEuler"), Category = "Math|Quat")
-	static ENGINE_API void Quat_SetFromEuler(UPARAM(ref) FQuat& Q, const FVector& Euler);
+	static UE_INL_API void Quat_SetFromEuler(UPARAM(ref) FQuat& Q, const FVector& Euler);
 
 	/**
 	 * Convert a vector of floating-point Euler angles (in degrees) into a Quaternion.
@@ -3033,37 +3033,37 @@ class UKismetMathLibrary : public UBlueprintFunctionLibrary
 
 	/** Assign contents of InColor */
 	UFUNCTION(BlueprintCallable, meta = (ScriptMethod = "Set"), Category = "Math|Color")
-	static ENGINE_API void LinearColor_Set(UPARAM(ref) FLinearColor& InOutColor, FLinearColor InColor);
+	static UE_INL_API void LinearColor_Set(UPARAM(ref) FLinearColor& InOutColor, FLinearColor InColor);
 
 	/** Assign individual linear RGBA components. */
 	UFUNCTION(BlueprintCallable, meta = (ScriptMethod = "SetRGBA"), Category = "Math|Color")
-	static ENGINE_API void LinearColor_SetRGBA(UPARAM(ref) FLinearColor& InOutColor, float R, float G, float B, float A = 1.0f);
+	static UE_INL_API void LinearColor_SetRGBA(UPARAM(ref) FLinearColor& InOutColor, float R, float G, float B, float A = 1.0f);
 
 	/** Assigns an HSV color to a linear space RGB color */
 	UFUNCTION(BlueprintCallable, meta = (ScriptMethod = "SetFromHSV"), Category = "Math|Color")
-	static ENGINE_API void LinearColor_SetFromHSV(UPARAM(ref) FLinearColor& InOutColor, float H, float S, float V, float A = 1.0f);
+	static UE_INL_API void LinearColor_SetFromHSV(UPARAM(ref) FLinearColor& InOutColor, float H, float S, float V, float A = 1.0f);
 
 	/**
 	 * Assigns an FColor coming from an observed sRGB output, into a linear color.
 	 * @param InSRGB The sRGB color that needs to be converted into linear space.
 	 */
 	UFUNCTION(BlueprintCallable, meta = (ScriptMethod = "SetFromSRGB"), Category = "Math|Color")
-	static ENGINE_API void LinearColor_SetFromSRGB(UPARAM(ref) FLinearColor& InOutColor, const FColor& InSRGB);
+	static UE_INL_API void LinearColor_SetFromSRGB(UPARAM(ref) FLinearColor& InOutColor, const FColor& InSRGB);
 
 	/**
 	 * Assigns an FColor coming from an observed Pow(1/2.2) output, into a linear color.
 	 * @param InColor The Pow(1/2.2) color that needs to be converted into linear space.
 	 */
 	UFUNCTION(BlueprintCallable, meta = (ScriptMethod = "SetFromPow22"), Category = "Math|Color")
-	static ENGINE_API void LinearColor_SetFromPow22(UPARAM(ref) FLinearColor& InOutColor, const FColor& InColor);
+	static UE_INL_API void LinearColor_SetFromPow22(UPARAM(ref) FLinearColor& InOutColor, const FColor& InColor);
 
 	/** Converts temperature in Kelvins of a black body radiator to RGB chromaticity. */
 	UFUNCTION(BlueprintCallable, meta = (ScriptMethod = "SetTemperature"), Category = "Math|Color")
-	static ENGINE_API void LinearColor_SetTemperature(UPARAM(ref) FLinearColor& InOutColor, float InTemperature);
+	static UE_INL_API void LinearColor_SetTemperature(UPARAM(ref) FLinearColor& InOutColor, float InTemperature);
 
 	/** Sets to a random color. Choses a quite nice color based on a random hue. */
 	UFUNCTION(BlueprintCallable, meta = (ScriptMethod = "SetRandomHue"), Category = "Math|Color")
-	static ENGINE_API void LinearColor_SetRandomHue(UPARAM(ref) FLinearColor& InOutColor);
+	static UE_INL_API void LinearColor_SetRandomHue(UPARAM(ref) FLinearColor& InOutColor);
 
 	/** Converts a float into a LinearColor, where each element is a float */
 	UE_DEPRECATED(5.2, "This method has been deprecated and will be removed. Use the double version instead.")
@@ -3103,19 +3103,19 @@ class UKismetMathLibrary : public UBlueprintFunctionLibrary
 
 	/** Converts from linear to 8-bit RGBE as outlined in Gregory Ward's Real Pixels article, Graphics Gems II, page 80. */
 	UFUNCTION(BlueprintPure, meta = (DisplayName = "To RGBE (LinearColor)", ScriptMethod = "ToRGBE", Keywords = "cast convert", BlueprintAutocast), Category = "Math|Color")
-	static FColor LinearColor_ToRGBE(FLinearColor InLinearColor);
+	static UE_INL_API FColor LinearColor_ToRGBE(FLinearColor InLinearColor);
 
 	/** Quantizes the linear color and returns the result as a FColor with optional sRGB conversion and quality as goal. */
 	UFUNCTION(BlueprintPure, meta = (DisplayName = "To Color (LinearColor)", ScriptMethod = "ToColor", CompactNodeTitle = "->", Keywords = "cast convert", BlueprintAutocast), Category = "Math|Conversions")
-	static FColor Conv_LinearColorToColor(FLinearColor InLinearColor, bool InUseSRGB = true);
+	static UE_INL_API FColor Conv_LinearColorToColor(FLinearColor InLinearColor, bool InUseSRGB = true);
 
 	/** Quantizes the linear color and returns the result as an 8-bit color.  This bypasses the SRGB conversion. */
 	UFUNCTION(BlueprintPure, meta = (DeprecatedFunction, DeprecationMessage = "Use LinearColor_QuantizeRound instead for correct color conversion.", DisplayName = "Quantize to 8-bit (LinearColor)", ScriptMethod = "Quantize", Keywords = "cast convert"), Category = "Math|Color")
-	static FColor LinearColor_Quantize(FLinearColor InColor);
+	static UE_INL_API FColor LinearColor_Quantize(FLinearColor InColor);
 
 	/** Quantizes the linear color with rounding and returns the result as an 8-bit color.  This bypasses the SRGB conversion. */
 	UFUNCTION(BlueprintPure, meta = (DisplayName = "Quantize With Rounding To 8-bit (LinearColor)", ScriptMethod = "QuantizeRound", Keywords = "cast convert"), Category = "Math|Color")
-	static FColor LinearColor_QuantizeRound(FLinearColor InColor);
+	static UE_INL_API FColor LinearColor_QuantizeRound(FLinearColor InColor);
 
 	/**
 	 * Returns a desaturated color, with 0 meaning no desaturation and 1 == full desaturation
@@ -3216,7 +3216,7 @@ class UKismetMathLibrary : public UBlueprintFunctionLibrary
 	
 	/** Converts this color value to a hexadecimal string. The format of the string is RRGGBBAA. */
 	UFUNCTION(BlueprintPure, meta = (DisplayName = "ToHex"), Category = "Math|Color")
-	static FString ToHex_LinearColor(FLinearColor InColor);
+	static UE_INL_API FString ToHex_LinearColor(FLinearColor InColor);
 
 	//
 	// Plane functions.
@@ -3259,7 +3259,7 @@ class UKismetMathLibrary : public UBlueprintFunctionLibrary
 
 	/** Subtraction (A - B) */
 	UFUNCTION(BlueprintPure, meta = (IgnoreTypePromotion, DisplayName = "DateTime - DateTime", CompactNodeTitle = "-", Keywords = "- subtract minus"), Category = "Math|DateTime")
-	static FTimespan Subtract_DateTimeDateTime(FDateTime A, FDateTime B);
+	static UE_INL_API FTimespan Subtract_DateTimeDateTime(FDateTime A, FDateTime B);
 
 	/** Returns true if the values are equal (A == B) */
 	UFUNCTION(BlueprintPure, meta=(IgnoreTypePromotion, DisplayName="Equal (DateTime)", CompactNodeTitle="==", Keywords="== equal"), Category="Math|DateTime")
@@ -3323,7 +3323,7 @@ class UKismetMathLibrary : public UBlueprintFunctionLibrary
 
 	/** Returns the time elapsed since midnight of A */
 	UFUNCTION(BlueprintPure, meta=(DisplayName="Get Time Of Day"), Category="Math|DateTime")
-	static FTimespan GetTimeOfDay( FDateTime A );
+	static UE_INL_API FTimespan GetTimeOfDay( FDateTime A );
 
 	/** Returns the year component of A */
 	UFUNCTION(BlueprintPure, meta=(DisplayName="Get Year"), Category="Math|DateTime")
@@ -3411,15 +3411,15 @@ class UKismetMathLibrary : public UBlueprintFunctionLibrary
 
 	/** Returns the maximum time span value */
 	UFUNCTION(BlueprintPure, meta=(DisplayName="Max Value (Timespan)", ScriptConstant = "MaxValue", ScriptConstantHost = "/Script/CoreUObject.Timespan"), Category="Math|Timespan")
-	static FTimespan TimespanMaxValue( );
+	static UE_INL_API FTimespan TimespanMaxValue( );
 
 	/** Returns the minimum time span value */
 	UFUNCTION(BlueprintPure, meta=(DisplayName="Min Value (Timespan)", ScriptConstant = "MinValue", ScriptConstantHost = "/Script/CoreUObject.Timespan"), Category="Math|Timespan")
-	static FTimespan TimespanMinValue( );
+	static UE_INL_API FTimespan TimespanMinValue( );
 
 	/** Returns a zero time span value */
 	UFUNCTION(BlueprintPure, meta=(DisplayName="Zero Value (Timespan)", ScriptConstant = "Zero", ScriptConstantHost = "/Script/CoreUObject.Timespan"), Category="Math|Timespan")
-	static FTimespan TimespanZeroValue( );
+	static UE_INL_API FTimespan TimespanZeroValue( );
 
 	//
 	// Timespan functions.
@@ -3443,19 +3443,19 @@ class UKismetMathLibrary : public UBlueprintFunctionLibrary
 
 	/** Addition (A + B) */
 	UFUNCTION(BlueprintPure, meta=(DisplayName="Timespan + Timespan", CompactNodeTitle="+", Keywords="+ add plus"), Category="Math|Timespan")
-	static FTimespan Add_TimespanTimespan( FTimespan A, FTimespan B );
+	static UE_INL_API FTimespan Add_TimespanTimespan( FTimespan A, FTimespan B );
 
 	/** Subtraction (A - B) */
 	UFUNCTION(BlueprintPure, meta=(DisplayName="Timespan - Timespan", CompactNodeTitle="-", Keywords="- subtract minus"), Category="Math|Timespan")
-	static FTimespan Subtract_TimespanTimespan( FTimespan A, FTimespan B );
+	static UE_INL_API FTimespan Subtract_TimespanTimespan( FTimespan A, FTimespan B );
 
 	/** Scalar multiplication (A * s) */
 	UFUNCTION(BlueprintPure, meta=(DisplayName="Timespan * float", CompactNodeTitle="*", Keywords="* multiply"), Category="Math|Timespan")
-	static FTimespan Multiply_TimespanFloat( FTimespan A, float Scalar );
+	static UE_INL_API FTimespan Multiply_TimespanFloat( FTimespan A, float Scalar );
 
 	/** Scalar division (A / s) */
 	UFUNCTION(BlueprintPure, meta=(DisplayName="Timespan / float", CompactNodeTitle="/", Keywords="/ divide"), Category="Math|Timespan")
-	static FTimespan Divide_TimespanFloat( FTimespan A, float Scalar );
+	static UE_INL_API FTimespan Divide_TimespanFloat( FTimespan A, float Scalar );
 
 	/** Returns true if the values are equal (A == B) */
 	UFUNCTION(BlueprintPure, meta=(DisplayName="Equal (Timespan)", CompactNodeTitle="==", Keywords="== equal"), Category="Math|Timespan")
@@ -3487,7 +3487,7 @@ class UKismetMathLibrary : public UBlueprintFunctionLibrary
 
 	/** Returns the absolute value of A */
 	UFUNCTION(BlueprintPure, meta=(DisplayName="Get Duration"), Category="Math|Timespan")
-	static FTimespan GetDuration( FTimespan A );
+	static UE_INL_API FTimespan GetDuration( FTimespan A );
 
 	/** Returns the hours component of A */
 	UFUNCTION(BlueprintPure, meta=(DisplayName="Get Hours"), Category="Math|Timespan")
@@ -3625,11 +3625,11 @@ class UKismetMathLibrary : public UBlueprintFunctionLibrary
 
 	/** Converts an integer to an IntVector*/
 	UFUNCTION(BlueprintPure, meta = (DisplayName = "To IntVector (Integer)", CompactNodeTitle = "->", Keywords = "cast convert", BlueprintAutocast), Category = "Math|Conversions")
-	static FIntVector Conv_IntToIntVector(int32 InInt);
+	static UE_INL_API FIntVector Conv_IntToIntVector(int32 InInt);
 
 	/** Converts an integer to an IntVector2*/
 	UFUNCTION(BlueprintPure, meta = (DisplayName = "To IntVector2 (Integer)", CompactNodeTitle = "->", Keywords = "cast convert", BlueprintAutocast), Category = "Math|Conversions")
-	static FIntVector2 Conv_IntToIntVector2(int32 InInt);
+	static UE_INL_API FIntVector2 Conv_IntToIntVector2(int32 InInt);
 
 	/** Converts an integer to a FVector */
 	UFUNCTION(BlueprintPure, meta = (DisplayName = "To Vector (Integer)", CompactNodeTitle = "->", Keywords = "cast convert", BlueprintAutocast), Category = "Math|Conversions")
@@ -3693,7 +3693,7 @@ class UKismetMathLibrary : public UBlueprintFunctionLibrary
 
 	/** Makes an FBox from Min and Max and sets IsValid to true */
 	UFUNCTION(BlueprintPure, Category="Math|Box", meta=(Keywords="construct build", NativeMakeFunc))
-	static FBox MakeBox(FVector Min, FVector Max);
+	static UE_INL_API FBox MakeBox(FVector Min, FVector Max);
 
 	/** 
 	 * Utility function to build an box from an Origin and Extent 
@@ -3791,7 +3791,7 @@ class UKismetMathLibrary : public UBlueprintFunctionLibrary
 
 	/** Makes an FBox2D from Min and Max and sets IsValid to true */
 	UFUNCTION(BlueprintPure, Category = "Math|Box2D", meta = (Keywords = "construct build", NativeMakeFunc))
-	static FBox2D MakeBox2D(FVector2D Min, FVector2D Max);
+	static UE_INL_API FBox2D MakeBox2D(FVector2D Min, FVector2D Max);
 
 	
 	//
@@ -3800,7 +3800,7 @@ class UKismetMathLibrary : public UBlueprintFunctionLibrary
 	
 	/** Makes an FBoxSphereBounds given an origin, extent, and radius */
 	UFUNCTION(BlueprintPure, meta = (Keywords = "construct build", NativeMakeFunc), Category = "Math|BoxSphereBounds")
-	static FBoxSphereBounds MakeBoxSphereBounds(FVector Origin, FVector BoxExtent, float SphereRadius);
+	static UE_INL_API FBoxSphereBounds MakeBoxSphereBounds(FVector Origin, FVector BoxExtent, float SphereRadius);
 
 	/** Breaks an FBoxSphereBounds into origin, extent, and radius */
 	UFUNCTION(BlueprintPure, meta = (NativeBreakFunc), Category = "Math|BoxSphereBounds")
@@ -3831,47 +3831,47 @@ class UKismetMathLibrary : public UBlueprintFunctionLibrary
 
 	/** If bPickA is true, A is returned, otherwise B is */
 	UFUNCTION(BlueprintPure, Category="Utilities|String")
-	static ENGINE_API FString SelectString(const FString& A, const FString& B, bool bPickA);
+	static UE_INL_API FString SelectString(const FString& A, const FString& B, bool bPickA);
 
 	/** If bPickA is true, A is returned, otherwise B is */
 	UFUNCTION(BlueprintPure, Category="Utilities|String")
-	static ENGINE_API FText SelectText(const FText A, const FText B, bool bPickA);
+	static UE_INL_API FText SelectText(const FText A, const FText B, bool bPickA);
 
 	/** If bPickA is true, A is returned, otherwise B is */
 	UFUNCTION(BlueprintPure, Category="Utilities|String")
-	static ENGINE_API FName SelectName(const FName A, const FName B, bool bPickA);
+	static UE_INL_API FName SelectName(const FName A, const FName B, bool bPickA);
 
 	/** If bPickA is true, A is returned, otherwise B is */
 	UFUNCTION(BlueprintPure, Category="Math|Integer")
-	static ENGINE_API int32 SelectInt(int32 A, int32 B, bool bPickA);
+	static UE_INL_API int32 SelectInt(int32 A, int32 B, bool bPickA);
 
 	/** If bPickA is true, A is returned, otherwise B is */
 	UFUNCTION(BlueprintPure, Category="Math|Float")
-	static ENGINE_API double SelectFloat(double A, double B, bool bPickA);
+	static UE_INL_API double SelectFloat(double A, double B, bool bPickA);
 
 	/** If bPickA is true, A is returned, otherwise B is */
 	UFUNCTION(BlueprintPure, Category="Math|Vector")
-	static ENGINE_API FVector SelectVector(FVector A, FVector B, bool bPickA);
+	static UE_INL_API FVector SelectVector(FVector A, FVector B, bool bPickA);
 
 	/** If bPickA is true, A is returned, otherwise B is */
 	UFUNCTION(BlueprintPure, Category="Math|Rotator", meta=(Keywords="rotation rotate"))
-	static ENGINE_API FRotator SelectRotator(FRotator A, FRotator B, bool bPickA);
+	static UE_INL_API FRotator SelectRotator(FRotator A, FRotator B, bool bPickA);
 
 	/** If bPickA is true, A is returned, otherwise B is */
 	UFUNCTION(BlueprintPure, Category="Math|Color")
-	static ENGINE_API FLinearColor SelectColor(FLinearColor A, FLinearColor B, bool bPickA);
+	static UE_INL_API FLinearColor SelectColor(FLinearColor A, FLinearColor B, bool bPickA);
 
 	/** If bPickA is true, A is returned, otherwise B is */
 	UFUNCTION(BlueprintPure, Category="Math|Transform")
-	static ENGINE_API FTransform SelectTransform(const FTransform& A, const FTransform& B, bool bPickA);
+	static UE_INL_API FTransform SelectTransform(const FTransform& A, const FTransform& B, bool bPickA);
 
 	/** If bPickA is true, A is returned, otherwise B is */
 	UFUNCTION(BlueprintPure, Category="Utilities")
-	static UObject* SelectObject(UObject* A, UObject* B, bool bSelectA);
+	static UE_INL_API UObject* SelectObject(UObject* A, UObject* B, bool bSelectA);
 
 	/** If bPickA is true, A is returned, otherwise B is */
 	UFUNCTION(BlueprintPure, Category = "Utilities")
-	static UClass* SelectClass(UClass* A, UClass* B, bool bSelectA);
+	static UE_INL_API UClass* SelectClass(UClass* A, UClass* B, bool bSelectA);
 
 
 	//
@@ -3927,11 +3927,11 @@ class UKismetMathLibrary : public UBlueprintFunctionLibrary
 	
 	/** Make a transform from location, rotation and scale */
 	UFUNCTION(BlueprintPure, meta = (Keywords = "construct build", NativeMakeFunc), Category = "Math|Transform")
-	static ENGINE_API FTransform MakeTransform(FVector Location, FRotator Rotation, FVector Scale = FVector(1,1,1));
+	static UE_INL_API FTransform MakeTransform(FVector Location, FRotator Rotation, FVector Scale = FVector(1,1,1));
 
 	/** Breaks apart a transform into location, rotation and scale */
 	UFUNCTION(BlueprintPure, Category = "Math|Transform", meta = (NativeBreakFunc))
-	static ENGINE_API void BreakTransform(const FTransform& InTransform, FVector& Location, FRotator& Rotation, FVector& Scale);
+	static UE_INL_API void BreakTransform(const FTransform& InTransform, FVector& Location, FRotator& Rotation, FVector& Scale);
 
 	/** Returns true if transform A is equal to transform B */
 	UFUNCTION(BlueprintPure, meta = (DisplayName = "Equal (Transform)", CompactNodeTitle = "==", ScriptMethod = "Equals", ScriptOperator = "==", Keywords = "== equal"), Category="Math|Transform")
@@ -4163,7 +4163,7 @@ class UKismetMathLibrary : public UBlueprintFunctionLibrary
 	static ENGINE_API int32 RandomIntegerFromStream(const FRandomStream& Stream, int32 Max);
 
 	UE_DEPRECATED(5.2, "Use RandomIntegerFromStream taking the FRandomStream as the first argument.")
-	static UE_INL_API int32 RandomIntegerFromStream(int32 Max, const FRandomStream& Stream)
+	static int32 RandomIntegerFromStream(int32 Max, const FRandomStream& Stream)
 	{
 		return RandomIntegerFromStream(Stream, Max);
 	}
@@ -4173,7 +4173,7 @@ class UKismetMathLibrary : public UBlueprintFunctionLibrary
 	static ENGINE_API int32 RandomIntegerInRangeFromStream(const FRandomStream& Stream, int32 Min, int32 Max);
 
 	UE_DEPRECATED(5.2, "Use RandomIntegerInRangeFromStream taking the FRandomStream as the first argument.")
-	static UE_INL_API int32 RandomIntegerInRangeFromStream(int32 Min, int32 Max, const FRandomStream& Stream)
+	static int32 RandomIntegerInRangeFromStream(int32 Min, int32 Max, const FRandomStream& Stream)
 	{
 		return RandomIntegerInRangeFromStream(Stream, Min, Max);
 	}
@@ -4191,7 +4191,7 @@ class UKismetMathLibrary : public UBlueprintFunctionLibrary
 	static ENGINE_API float RandomFloatInRangeFromStream(const FRandomStream& Stream, float Min, float Max);
 
 	UE_DEPRECATED(5.2, "Use RandomFloatInRangeFromStream taking the FRandomStream as the first argument.")
-	static UE_INL_API float RandomFloatInRangeFromStream(float Min, float Max, const FRandomStream& Stream)
+	static float RandomFloatInRangeFromStream(float Min, float Max, const FRandomStream& Stream)
 	{
 		return RandomFloatInRangeFromStream(Stream, Min, Max);
 	}
@@ -4205,7 +4205,7 @@ class UKismetMathLibrary : public UBlueprintFunctionLibrary
 	static ENGINE_API FVector RandomPointInBoundingBoxFromStream(const FRandomStream& Stream, const FVector Center, const FVector HalfSize);
 
 	UE_DEPRECATED(5.2, "Use RandomPointInBoundingBoxFromStream taking the FRandomStream as the first argument.")
-	static UE_INL_API FVector RandomPointInBoundingBoxFromStream(const FVector Center, const FVector HalfSize, const FRandomStream& Stream)
+	static FVector RandomPointInBoundingBoxFromStream(const FVector Center, const FVector HalfSize, const FRandomStream& Stream)
 	{
 		return RandomPointInBoundingBoxFromStream(Stream, Center, HalfSize);
 	}
@@ -4215,7 +4215,7 @@ class UKismetMathLibrary : public UBlueprintFunctionLibrary
 	static ENGINE_API FVector RandomPointInBoundingBoxFromStream_Box(const FRandomStream& Stream, const FBox Box);
 
 	UE_DEPRECATED(5.2, "Use RandomPointInBoundingBoxFromStream_Box taking the FRandomStream as the first argument.")
-	static UE_INL_API FVector RandomPointInBoundingBoxFromStream_Box(const FBox Box, const FRandomStream& Stream)
+	static FVector RandomPointInBoundingBoxFromStream_Box(const FBox Box, const FRandomStream& Stream)
 	{
 		return RandomPointInBoundingBoxFromStream_Box(Stream, Box);
 	}
@@ -4225,7 +4225,7 @@ class UKismetMathLibrary : public UBlueprintFunctionLibrary
 	static ENGINE_API FRotator RandomRotatorFromStream(const FRandomStream& Stream, bool bRoll);
 
 	UE_DEPRECATED(5.2, "Use RandomRotatorFromStream taking the FRandomStream as the first argument.")
-	static UE_INL_API FRotator RandomRotatorFromStream(bool bRoll, const FRandomStream& Stream)
+	static FRotator RandomRotatorFromStream(bool bRoll, const FRandomStream& Stream)
 	{
 		return RandomRotatorFromStream(Stream, bRoll);
 	}
@@ -4252,7 +4252,7 @@ class UKismetMathLibrary : public UBlueprintFunctionLibrary
 	static ENGINE_API FVector RandomUnitVectorInConeInRadiansFromStream(const FRandomStream& Stream, const FVector& ConeDir, float ConeHalfAngleInRadians);
 
 	UE_DEPRECATED(5.2, "Use RandomUnitVectorInConeInRadiansFromStream taking the FRandomStream as the first argument.")
-	static UE_INL_API FVector RandomUnitVectorInConeInRadiansFromStream(const FVector& ConeDir, float ConeHalfAngleInRadians, const FRandomStream& Stream)
+	static FVector RandomUnitVectorInConeInRadiansFromStream(const FVector& ConeDir, float ConeHalfAngleInRadians, const FRandomStream& Stream)
 	{
 		return RandomUnitVectorInConeInRadiansFromStream(Stream, ConeDir, ConeHalfAngleInRadians);
 	}
@@ -4264,13 +4264,13 @@ class UKismetMathLibrary : public UBlueprintFunctionLibrary
 	 * @param Stream					The random stream from which to obtain the vector.
 	 */
 	UFUNCTION(BlueprintPure, Category="Math|Random", meta = (Keywords = "RandomVector", ScriptMethod = "RandomUnitVectorInConeInDegrees", ScriptMethodMutable))
-	static inline FVector RandomUnitVectorInConeInDegreesFromStream(const FRandomStream& Stream, const FVector& ConeDir, float ConeHalfAngleInDegrees)
+	static FVector RandomUnitVectorInConeInDegreesFromStream(const FRandomStream& Stream, const FVector& ConeDir, float ConeHalfAngleInDegrees)
 	{
 		return RandomUnitVectorInConeInRadiansFromStream(Stream, ConeDir, FMath::DegreesToRadians(ConeHalfAngleInDegrees));
 	}
 
 	UE_DEPRECATED(5.2, "Use RandomUnitVectorInConeInDegreesFromStream taking the FRandomStream as the first argument.")
-	static inline FVector RandomUnitVectorInConeInDegreesFromStream(const FVector& ConeDir, float ConeHalfAngleInDegrees, const FRandomStream& Stream)
+	static FVector RandomUnitVectorInConeInDegreesFromStream(const FVector& ConeDir, float ConeHalfAngleInDegrees, const FRandomStream& Stream)
 	{
 		return RandomUnitVectorInConeInDegreesFromStream(Stream, ConeDir, ConeHalfAngleInDegrees);
 	}
@@ -4287,7 +4287,7 @@ class UKismetMathLibrary : public UBlueprintFunctionLibrary
 	static ENGINE_API FVector RandomUnitVectorInEllipticalConeInRadiansFromStream(const FRandomStream& Stream, const FVector& ConeDir, float MaxYawInRadians, float MaxPitchInRadians);
 
 	UE_DEPRECATED(5.2, "Use RandomUnitVectorInEllipticalConeInRadiansFromStream taking the FRandomStream as the first argument.")
-	static UE_INL_API FVector RandomUnitVectorInEllipticalConeInRadiansFromStream(const FVector& ConeDir, float MaxYawInRadians, float MaxPitchInRadians, const FRandomStream& Stream)
+	static FVector RandomUnitVectorInEllipticalConeInRadiansFromStream(const FVector& ConeDir, float MaxYawInRadians, float MaxPitchInRadians, const FRandomStream& Stream)
 	{
 		return RandomUnitVectorInEllipticalConeInRadiansFromStream(Stream, ConeDir, MaxYawInRadians, MaxPitchInRadians);
 	}
@@ -4301,13 +4301,13 @@ class UKismetMathLibrary : public UBlueprintFunctionLibrary
 	* @param Stream				The random stream from which to obtain the vector.
 	*/
 	UFUNCTION(BlueprintPure, Category = "Math|Random", meta = (Keywords = "RandomVector", ScriptMethod = "RandomUnitVectorInEllipticalConeInDegrees", ScriptMethodMutable))
-	static inline FVector RandomUnitVectorInEllipticalConeInDegreesFromStream(const FRandomStream& Stream, const FVector& ConeDir, float MaxYawInDegrees, float MaxPitchInDegrees)
+	static FVector RandomUnitVectorInEllipticalConeInDegreesFromStream(const FRandomStream& Stream, const FVector& ConeDir, float MaxYawInDegrees, float MaxPitchInDegrees)
 	{
 		return RandomUnitVectorInEllipticalConeInRadiansFromStream(Stream, ConeDir, FMath::DegreesToRadians(MaxYawInDegrees), FMath::DegreesToRadians(MaxPitchInDegrees));
 	}
 
 	UE_DEPRECATED(5.2, "Use RandomUnitVectorInEllipticalConeInDegreesFromStream taking the FRandomStream as the first argument.")
-	static inline FVector RandomUnitVectorInEllipticalConeInDegreesFromStream(const FVector& ConeDir, float MaxYawInDegrees, float MaxPitchInDegrees, const FRandomStream& Stream)
+	static FVector RandomUnitVectorInEllipticalConeInDegreesFromStream(const FVector& ConeDir, float MaxYawInDegrees, float MaxPitchInDegrees, const FRandomStream& Stream)
 	{
 		return RandomUnitVectorInEllipticalConeInDegreesFromStream(Stream, ConeDir, MaxYawInDegrees, MaxPitchInDegrees);
 	}
@@ -4595,21 +4595,21 @@ class UKismetMathLibrary : public UBlueprintFunctionLibrary
 
 	// NetQuantized vector make/breaks
 	UFUNCTION(BlueprintPure, Category = "Math|Vector", meta = (NativeMakeFunc))
-	static UE_INL_API FVector_NetQuantize MakeVector_NetQuantize(double X, double Y, double Z) { return FVector_NetQuantize(X, Y, Z); }
+	static FVector_NetQuantize MakeVector_NetQuantize(double X, double Y, double Z) { return FVector_NetQuantize(X, Y, Z); }
 	UFUNCTION(BlueprintPure, Category = "Math|Vector", meta = (NativeMakeFunc))
-	static UE_INL_API FVector_NetQuantize10 MakeVector_NetQuantize10(double X, double Y, double Z) { return FVector_NetQuantize10(X, Y, Z); }
+	static FVector_NetQuantize10 MakeVector_NetQuantize10(double X, double Y, double Z) { return FVector_NetQuantize10(X, Y, Z); }
 	UFUNCTION(BlueprintPure, Category = "Math|Vector", meta = (NativeMakeFunc))
-	static UE_INL_API FVector_NetQuantize100 MakeVector_NetQuantize100(double X, double Y, double Z) { return FVector_NetQuantize100(X, Y, Z); }
+	static FVector_NetQuantize100 MakeVector_NetQuantize100(double X, double Y, double Z) { return FVector_NetQuantize100(X, Y, Z); }
 	UFUNCTION(BlueprintPure, Category = "Math|Vector", meta = (NativeMakeFunc))
-	static UE_INL_API FVector_NetQuantizeNormal MakeVector_NetQuantizeNormal(double X, double Y, double Z) { return FVector_NetQuantizeNormal(X, Y, Z); }
+	static FVector_NetQuantizeNormal MakeVector_NetQuantizeNormal(double X, double Y, double Z) { return FVector_NetQuantizeNormal(X, Y, Z); }
 	UFUNCTION(BlueprintPure, Category = "Math|Vector", meta = (NativeBreakFunc))
-	static UE_INL_API void BreakVector_NetQuantize(FVector_NetQuantize InVec, double& X, double& Y, double& Z) { BreakVector((FVector)InVec, X, Y, Z); }
+	static void BreakVector_NetQuantize(FVector_NetQuantize InVec, double& X, double& Y, double& Z) { BreakVector((FVector)InVec, X, Y, Z); }
 	UFUNCTION(BlueprintPure, Category = "Math|Vector", meta = (NativeBreakFunc))
-	static UE_INL_API void BreakVector_NetQuantize10(FVector_NetQuantize10 InVec, double& X, double& Y, double& Z) { BreakVector((FVector)InVec, X, Y, Z); }
+	static void BreakVector_NetQuantize10(FVector_NetQuantize10 InVec, double& X, double& Y, double& Z) { BreakVector((FVector)InVec, X, Y, Z); }
 	UFUNCTION(BlueprintPure, Category = "Math|Vector", meta = (NativeBreakFunc))
-	static UE_INL_API void BreakVector_NetQuantize100(FVector_NetQuantize100 InVec, double& X, double& Y, double& Z) { BreakVector((FVector)InVec, X, Y, Z); }
+	static void BreakVector_NetQuantize100(FVector_NetQuantize100 InVec, double& X, double& Y, double& Z) { BreakVector((FVector)InVec, X, Y, Z); }
 	UFUNCTION(BlueprintPure, Category = "Math|Vector", meta = (NativeBreakFunc))
-	static UE_INL_API void BreakVector_NetQuantizeNormal(FVector_NetQuantizeNormal InVec, double& X, double& Y, double& Z) { BreakVector((FVector)InVec, X, Y, Z); }
+	static void BreakVector_NetQuantizeNormal(FVector_NetQuantizeNormal InVec, double& X, double& Y, double& Z) { BreakVector((FVector)InVec, X, Y, Z); }
 
 private:
 
