Fixed FWorldPartitionActorDescInstances not getting added for newly saved actors to the UActorDescContainerInstance when a new instanced world is saved.

diff --git a/Engine/Source/Runtime/Engine/Private/WorldPartition/ActorDescContainerInstance.cpp b/Engine/Source/Runtime/Engine/Private/WorldPartition/ActorDescContainerInstance.cpp
--- a/Engine/Source/Runtime/Engine/Private/WorldPartition/ActorDescContainerInstance.cpp
+++ b/Engine/Source/Runtime/Engine/Private/WorldPartition/ActorDescContainerInstance.cpp
@@ -379,7 +379,16 @@ FWorldPartitionActorDesc* UActorDescContainerInstance::GetActorDescChecked(const
 
 void UActorDescContainerInstance::SetContainerPackage(FName InContainerPackageName)
 {
-	InstancingContext.Reset();
+	if (GetInstancingContext())
+	{
+		InstancingContext.Reset();
+
+		// Now that the world is no longer instanced, we need to register the delegates we previously skipped in RegisterDelegates.
+		if (bRegisteredDelegates)
+		{
+			Container->OnActorDescAddedEvent.AddUObject(this, &UActorDescContainerInstance::OnActorDescAdded);
+		}
+	}
 
 	UActorDescContainerSubsystem::GetChecked().SetContainerPackage(Container, InContainerPackageName);
 }
