[materials] fix line directive in generated material code; the code calculating the line number offset of the substituted #line directive in MaterialTemplate.ush was incorrect

diff --git a/Engine/Source/Runtime/Engine/Private/Materials/MaterialSourceTemplate.cpp b/Engine/Source/Runtime/Engine/Private/Materials/MaterialSourceTemplate.cpp
--- a/Engine/Source/Runtime/Engine/Private/Materials/MaterialSourceTemplate.cpp
+++ b/Engine/Source/Runtime/Engine/Private/Materials/MaterialSourceTemplate.cpp
@@ -110,20 +110,17 @@ bool FMaterialSourceTemplate::Preload(EShaderPlatform ShaderPlatform)
 	check(LineIndex != INDEX_NONE);
 
 	// Count line endings before the '#line' statement
-	int32 TemplateLineNumber = INDEX_NONE;
+	int32 TemplateLineNumber = 1;
 	int32 StartPosition = LineIndex + 1;
-	do
+	const TCHAR* Cur = *MaterialTemplateString;
+	const TCHAR* End = Cur + LineIndex;
+	while (Cur != End)
 	{
-		TemplateLineNumber++;
-		// Subtract one from the last found line ending index to make sure we skip over it
-		StartPosition = MaterialTemplateString.Find(TEXT("\n"), ESearchCase::CaseSensitive, ESearchDir::FromEnd, StartPosition - 1);
-	} while (StartPosition != INDEX_NONE);
-
-	check(TemplateLineNumber != INDEX_NONE);
-
-	// At this point MaterialTemplateLineNumber is one less than the line number of the '#line' statement
-	// For some reason we have to add 2 more to the #line value to get correct error line numbers from D3DXCompileShader
-	TemplateLineNumber += 3;
+		if (*Cur++ == TEXT('\n'))
+		{
+			TemplateLineNumber++;
+		}
+	}
 
 	// Save the material template line numbers for this shader platform
 	MaterialTemplateLineNumbers[ShaderPlatform] = TemplateLineNumber;
