Change the default animation blend option from Linear to Hermite Cubic (aka SmoothStep). Note that old data will keep the previous default of Linear so existing data should not be affected.

#tests animsandbox test data conversion
[RN]

diff --git a/Engine/Plugins/Animation/BlendStack/Source/Editor/Private/AnimGraphNode_BlendStack.cpp b/Engine/Plugins/Animation/BlendStack/Source/Editor/Private/AnimGraphNode_BlendStack.cpp
--- a/Engine/Plugins/Animation/BlendStack/Source/Editor/Private/AnimGraphNode_BlendStack.cpp
+++ b/Engine/Plugins/Animation/BlendStack/Source/Editor/Private/AnimGraphNode_BlendStack.cpp
@@ -11,9 +11,27 @@
 #include "Animation/AnimRootMotionProvider.h"
 #include "Kismet2/BlueprintEditorUtils.h"
 #include "Kismet2/KismetEditorUtilities.h"
+#include "UObject/FortniteMainBranchObjectVersion.h"
 
 #define LOCTEXT_NAMESPACE "AnimGraphNode_BlendStack"
 
+void UAnimGraphNode_BlendStack::Serialize(FArchive& Ar)
+{
+	// Handle change of default blend type
+	Ar.UsingCustomVersion(FFortniteMainBranchObjectVersion::GUID);
+	if (Ar.IsLoading())
+	{
+		const int32 CustomVersion = Ar.CustomVer(FFortniteMainBranchObjectVersion::GUID);
+		if (CustomVersion < FFortniteMainBranchObjectVersion::ChangeDefaultAlphaBlendType)
+		{
+			// Switch the default back to Linear so old data remains the same
+			// Important: this is done before loading so if data has changed from default it still works
+			Node.BlendOption = EAlphaBlendOption::Linear;
+		}
+	}
+	
+	Super::Serialize(Ar);
+}
 
 FLinearColor UAnimGraphNode_BlendStack::GetNodeTitleColor() const
 {
diff --git a/Engine/Plugins/Animation/BlendStack/Source/Editor/Public/AnimGraphNode_BlendStack.h b/Engine/Plugins/Animation/BlendStack/Source/Editor/Public/AnimGraphNode_BlendStack.h
--- a/Engine/Plugins/Animation/BlendStack/Source/Editor/Public/AnimGraphNode_BlendStack.h
+++ b/Engine/Plugins/Animation/BlendStack/Source/Editor/Public/AnimGraphNode_BlendStack.h
@@ -62,6 +62,7 @@ class UAnimGraphNode_BlendStack : public UAnimGraphNode_BlendStack_Base
 	UPROPERTY(EditAnywhere, Category = Settings)
 	FAnimNode_BlendStack Node;
 
+	virtual void Serialize(FArchive& Ar) override;
 	virtual FLinearColor GetNodeTitleColor() const override;
 	virtual FText GetTooltipText() const override;
 	virtual FText GetNodeTitle(ENodeTitleType::Type TitleType) const override;
diff --git a/Engine/Plugins/Animation/BlendStack/Source/Runtime/Public/BlendStack/AnimNode_BlendStack.h b/Engine/Plugins/Animation/BlendStack/Source/Runtime/Public/BlendStack/AnimNode_BlendStack.h
--- a/Engine/Plugins/Animation/BlendStack/Source/Runtime/Public/BlendStack/AnimNode_BlendStack.h
+++ b/Engine/Plugins/Animation/BlendStack/Source/Runtime/Public/BlendStack/AnimNode_BlendStack.h
@@ -283,7 +283,7 @@ struct FAnimNode_BlendStack : public FAnimNode_BlendStack_Standalone
 	TObjectPtr<UBlendProfile> BlendProfile;
 
 	UPROPERTY(EditAnywhere, BlueprintReadWrite, Category = Settings, meta = (PinHiddenByDefault))
-	EAlphaBlendOption BlendOption = EAlphaBlendOption::Linear;
+	EAlphaBlendOption BlendOption = UE::Anim::DefaultBlendOption;
 
 	// How we should update individual blend space parameters. See dropdown options tooltips.
 	UPROPERTY(EditAnywhere, Category = Blendspace)
diff --git a/Engine/Plugins/Animation/ControlRig/Source/ControlRigDeveloper/Private/AnimGraphNode_ControlRig.cpp b/Engine/Plugins/Animation/ControlRig/Source/ControlRigDeveloper/Private/AnimGraphNode_ControlRig.cpp
--- a/Engine/Plugins/Animation/ControlRig/Source/ControlRigDeveloper/Private/AnimGraphNode_ControlRig.cpp
+++ b/Engine/Plugins/Animation/ControlRig/Source/ControlRigDeveloper/Private/AnimGraphNode_ControlRig.cpp
@@ -290,6 +290,24 @@ void UAnimGraphNode_ControlRig::ValidateAnimNodeDuringCompilation(USkeleton* For
 	Super::ValidateAnimNodeDuringCompilation(ForSkeleton, MessageLog);
 }
 
+void UAnimGraphNode_ControlRig::Serialize(FArchive& Ar)
+{
+	// Handle change of default blend type
+	Ar.UsingCustomVersion(FFortniteMainBranchObjectVersion::GUID);
+	if (Ar.IsLoading())
+	{
+		const int32 CustomVersion = Ar.CustomVer(FFortniteMainBranchObjectVersion::GUID);
+		if (CustomVersion < FFortniteMainBranchObjectVersion::ChangeDefaultAlphaBlendType)
+		{
+			// Switch the default back to Linear so old data remains the same
+			// Important: this is done before loading so if data has changed from default it still works
+			Node.AlphaBoolBlend.BlendOption = EAlphaBlendOption::Linear;
+		}
+	}
+	
+	Super::Serialize(Ar);
+}
+
 void UAnimGraphNode_ControlRig::CustomizeDetails(IDetailLayoutBuilder& DetailBuilder)
 {
 	DECLARE_SCOPE_HIERARCHICAL_COUNTER_FUNC()
diff --git a/Engine/Plugins/Animation/ControlRig/Source/ControlRigDeveloper/Public/AnimGraphNode_ControlRig.h b/Engine/Plugins/Animation/ControlRig/Source/ControlRigDeveloper/Public/AnimGraphNode_ControlRig.h
--- a/Engine/Plugins/Animation/ControlRig/Source/ControlRigDeveloper/Public/AnimGraphNode_ControlRig.h
+++ b/Engine/Plugins/Animation/ControlRig/Source/ControlRigDeveloper/Public/AnimGraphNode_ControlRig.h
@@ -17,6 +17,10 @@ class UAnimGraphNode_ControlRig : public UAnimGraphNode_CustomProperty
 
 	UPROPERTY(EditAnywhere, Category = Settings)
 	FAnimNode_ControlRig Node;
+	
+	// UObject interface
+	virtual void Serialize(FArchive& Ar) override;
+	// End of UObject
 
 	virtual void CustomizeDetails(IDetailLayoutBuilder& DetailBuilder) override;
 	virtual void PostEditChangeProperty(FPropertyChangedEvent& PropertyChangedEvent) override;
diff --git a/Engine/Plugins/Animation/IKRig/Source/IKRigDeveloper/Private/AnimGraphNode_IKRig.cpp b/Engine/Plugins/Animation/IKRig/Source/IKRigDeveloper/Private/AnimGraphNode_IKRig.cpp
--- a/Engine/Plugins/Animation/IKRig/Source/IKRigDeveloper/Private/AnimGraphNode_IKRig.cpp
+++ b/Engine/Plugins/Animation/IKRig/Source/IKRigDeveloper/Private/AnimGraphNode_IKRig.cpp
@@ -11,6 +11,7 @@
 #include "DetailLayoutBuilder.h"
 #include "DetailWidgetRow.h"
 #include "IDetailChildrenBuilder.h"
+#include "UObject/FortniteMainBranchObjectVersion.h"
 
 #include "BoneSelectionWidget.h"
 
@@ -305,6 +306,24 @@ UAnimGraphNode_IKRig::~UAnimGraphNode_IKRig()
 	}
 }
 
+void UAnimGraphNode_IKRig::Serialize(FArchive& Ar)
+{
+	// Handle change of default blend type
+	Ar.UsingCustomVersion(FFortniteMainBranchObjectVersion::GUID);
+	if (Ar.IsLoading())
+	{
+		const int32 CustomVersion = Ar.CustomVer(FFortniteMainBranchObjectVersion::GUID);
+		if (CustomVersion < FFortniteMainBranchObjectVersion::ChangeDefaultAlphaBlendType)
+		{
+			// Switch the default back to Linear so old data remains the same
+			// Important: this is done before loading so if data has changed from default it still works
+			Node.AlphaBoolBlend.BlendOption = EAlphaBlendOption::Linear;
+		}
+	}
+	
+	Super::Serialize(Ar);
+}
+
 void UAnimGraphNode_IKRig::Draw(FPrimitiveDrawInterface* PDI, USkeletalMeshComponent* PreviewSkelMeshComp) const
 {
 	if(PreviewSkelMeshComp)
diff --git a/Engine/Plugins/Animation/IKRig/Source/IKRigDeveloper/Public/AnimGraphNode_IKRig.h b/Engine/Plugins/Animation/IKRig/Source/IKRigDeveloper/Public/AnimGraphNode_IKRig.h
--- a/Engine/Plugins/Animation/IKRig/Source/IKRigDeveloper/Public/AnimGraphNode_IKRig.h
+++ b/Engine/Plugins/Animation/IKRig/Source/IKRigDeveloper/Public/AnimGraphNode_IKRig.h
@@ -97,7 +97,11 @@ public:
 
 	UAnimGraphNode_IKRig() = default;
 	UE_API virtual ~UAnimGraphNode_IKRig();
-	
+
+	// UObject interface
+	virtual void Serialize(FArchive& Ar) override;
+	// End of UObject
+
 	// UEdGraphNode interface
 	UE_API virtual FText GetNodeTitle(ENodeTitleType::Type TitleType) const override;
 	UE_API virtual void PostEditChangeProperty(struct FPropertyChangedEvent& PropertyChangedEvent) override;
diff --git a/Engine/Plugins/Animation/MotionWarping/Source/MotionWarping/Private/RootMotionModifier.cpp b/Engine/Plugins/Animation/MotionWarping/Source/MotionWarping/Private/RootMotionModifier.cpp
--- a/Engine/Plugins/Animation/MotionWarping/Source/MotionWarping/Private/RootMotionModifier.cpp
+++ b/Engine/Plugins/Animation/MotionWarping/Source/MotionWarping/Private/RootMotionModifier.cpp
@@ -8,6 +8,7 @@
 #include "MotionWarpingComponent.h"
 #include "MotionWarpingAdapter.h"
 #include "DrawDebugHelpers.h"
+#include "UObject/FortniteMainBranchObjectVersion.h"
 
 #include UE_INLINE_GENERATED_CPP_BY_NAME(RootMotionModifier)
 
@@ -69,6 +70,24 @@ FTransform FMotionWarpingTarget::GetTargetTransformFromComponent(const USceneCom
 	return Comp->GetSocketTransform(BoneName);
 }
 
+void URootMotionModifier_Warp::Serialize(FArchive& Ar)
+{
+	// Handle change of default blend type
+	Ar.UsingCustomVersion(FFortniteMainBranchObjectVersion::GUID);
+	if (Ar.IsLoading())
+	{
+		const int32 CustomVersion = Ar.CustomVer(FFortniteMainBranchObjectVersion::GUID);
+		if (CustomVersion < FFortniteMainBranchObjectVersion::ChangeDefaultAlphaBlendType)
+		{
+			// Switch the default back to Linear so old data remains the same
+			// Important: this is done before loading so if data has changed from default it still works
+			AddTranslationEasingFunc = EAlphaBlendOption::Linear;
+		}
+	}
+	
+	Super::Serialize(Ar);
+}
+
 FTransform FMotionWarpingTarget::GetTargetTrasform() const
 {
 	if (Component.IsValid() && bFollowComponent)
diff --git a/Engine/Plugins/Animation/MotionWarping/Source/MotionWarping/Public/RootMotionModifier.h b/Engine/Plugins/Animation/MotionWarping/Source/MotionWarping/Public/RootMotionModifier.h
--- a/Engine/Plugins/Animation/MotionWarping/Source/MotionWarping/Public/RootMotionModifier.h
+++ b/Engine/Plugins/Animation/MotionWarping/Source/MotionWarping/Public/RootMotionModifier.h
@@ -317,6 +317,10 @@ class URootMotionModifier_Warp : public URootMotionModifier
 
 public:
 
+	// UObject interface
+	UE_API virtual void Serialize(FArchive& Ar) override;
+	// End of UObject
+
 	/** Name used to find the warp target for this modifier */
 	UPROPERTY(EditAnywhere, BlueprintReadWrite, Category = "Config", meta = (ExposeOnSpawn))
 	FName WarpTargetName = NAME_None;
@@ -346,7 +350,7 @@ public:
 
 	/** Easing function used when adding translation. Only relevant when there is no translation in the animation */
 	UPROPERTY(EditAnywhere, BlueprintReadWrite, Category = "Config", meta = (EditCondition = "bWarpTranslation"))
-	EAlphaBlendOption AddTranslationEasingFunc = EAlphaBlendOption::Linear;
+	EAlphaBlendOption AddTranslationEasingFunc = UE::Anim::DefaultBlendOption;
 
 	/** Custom curve used to add translation when there is none to warp. Only relevant when AddTranslationEasingFunc is set to Custom */
 	UPROPERTY(EditAnywhere, BlueprintReadWrite, Category = "Blend", meta = (EditCondition = "AddTranslationEasingFunc==EAlphaBlendOption::Custom", EditConditionHides))
diff --git a/Engine/Plugins/Animation/PoseSearch/Source/Editor/Private/AnimGraphNode_MotionMatchingInteraction.cpp b/Engine/Plugins/Animation/PoseSearch/Source/Editor/Private/AnimGraphNode_MotionMatchingInteraction.cpp
--- a/Engine/Plugins/Animation/PoseSearch/Source/Editor/Private/AnimGraphNode_MotionMatchingInteraction.cpp
+++ b/Engine/Plugins/Animation/PoseSearch/Source/Editor/Private/AnimGraphNode_MotionMatchingInteraction.cpp
@@ -13,6 +13,24 @@ FText UAnimGraphNode_MotionMatchingInteraction::GetTooltipText() const
 	return LOCTEXT("MotionMatchingInteractionTooltip", "Motion Matching Interaction Node");
 }
 
+void UAnimGraphNode_MotionMatchingInteraction::Serialize(FArchive& Ar)
+{
+	// Handle change of default blend type
+	Ar.UsingCustomVersion(FFortniteMainBranchObjectVersion::GUID);
+	if (Ar.IsLoading())
+	{
+		const int32 CustomVersion = Ar.CustomVer(FFortniteMainBranchObjectVersion::GUID);
+		if (CustomVersion < FFortniteMainBranchObjectVersion::ChangeDefaultAlphaBlendType)
+		{
+			// Switch the default back to Linear so old data remains the same
+			// Important: this is done before loading so if data has changed from default it still works
+			Node.BlendOption = EAlphaBlendOption::Linear;
+		}
+	}
+	
+	Super::Serialize(Ar);
+}
+
 FText UAnimGraphNode_MotionMatchingInteraction::GetNodeTitle(ENodeTitleType::Type TitleType) const
 {
 	return LOCTEXT("MotionMatchingInteraction", "MotionMatchingInteraction");
diff --git a/Engine/Plugins/Animation/PoseSearch/Source/Editor/Public/AnimGraphNode_MotionMatchingInteraction.h b/Engine/Plugins/Animation/PoseSearch/Source/Editor/Public/AnimGraphNode_MotionMatchingInteraction.h
--- a/Engine/Plugins/Animation/PoseSearch/Source/Editor/Public/AnimGraphNode_MotionMatchingInteraction.h
+++ b/Engine/Plugins/Animation/PoseSearch/Source/Editor/Public/AnimGraphNode_MotionMatchingInteraction.h
@@ -4,6 +4,7 @@
 
 #include "AnimGraphNode_BlendStack.h"
 #include "PoseSearch/AnimNode_MotionMatchingInteraction.h"
+#include "UObject/FortniteMainBranchObjectVersion.h"
 #include "AnimGraphNode_MotionMatchingInteraction.generated.h"
 
 UCLASS(Experimental)
@@ -15,6 +16,7 @@ class UAnimGraphNode_MotionMatchingInteraction : public UAnimGraphNode_BlendStac
 	FAnimNode_MotionMatchingInteraction Node;
 
 public:
+	virtual void Serialize(FArchive& Ar) override;
 	virtual FText GetNodeTitle(ENodeTitleType::Type TitleType) const override;
 	virtual FLinearColor GetNodeTitleColor() const override;
 	virtual FText GetTooltipText() const override;
diff --git a/Engine/Plugins/Animation/PoseSearch/Source/Runtime/Private/MotionMatchingAnimNodeLibrary.cpp b/Engine/Plugins/Animation/PoseSearch/Source/Runtime/Private/MotionMatchingAnimNodeLibrary.cpp
--- a/Engine/Plugins/Animation/PoseSearch/Source/Runtime/Private/MotionMatchingAnimNodeLibrary.cpp
+++ b/Engine/Plugins/Animation/PoseSearch/Source/Runtime/Private/MotionMatchingAnimNodeLibrary.cpp
@@ -4,17 +4,49 @@
 #include "AlphaBlend.h"
 #include "PoseSearch/AnimNode_MotionMatching.h"
 #include "PoseSearch/PoseSearchDatabase.h"
-#include "PoseSearch/PoseSearchDerivedData.h"
-#include "PoseSearch/PoseSearchSchema.h"
+#include "UObject/FortniteMainBranchObjectVersion.h"
 
 FMotionMatchingBlueprintBlendSettings::FMotionMatchingBlueprintBlendSettings()
 	: BlendTime(0.2f)
 	, BlendProfile(nullptr)
-	, BlendOption(EAlphaBlendOption::Linear)
+	, BlendOption(UE::Anim::DefaultBlendOption)
 	, bUseInertialBlend(false)
 {
 }
 
+bool FMotionMatchingBlueprintBlendSettings::Serialize(FArchive& Ar)
+{
+	Ar.UsingCustomVersion(FFortniteMainBranchObjectVersion::GUID);
+	if (Ar.IsLoading() && Ar.IsSerializingDefaults())
+	{
+		const int32 CustomVersion = Ar.CustomVer(FFortniteMainBranchObjectVersion::GUID);
+		if (CustomVersion < FFortniteMainBranchObjectVersion::ChangeDefaultAlphaBlendType)
+		{
+			// Switch the default back to Linear so old data remains the same
+			BlendOption = EAlphaBlendOption::Linear;
+		}
+	}
+
+	return false;
+}
+
+bool FMotionMatchingBlueprintBlendSettings::Serialize(FStructuredArchive::FSlot Slot)
+{
+	FArchive& Ar = Slot.GetUnderlyingArchive();
+	Ar.UsingCustomVersion(FFortniteMainBranchObjectVersion::GUID);
+
+	if (Ar.IsLoading() && Ar.IsSerializingDefaults())
+	{
+		const int32 CustomVersion = Ar.CustomVer(FFortniteMainBranchObjectVersion::GUID);
+		if (CustomVersion < FFortniteMainBranchObjectVersion::ChangeDefaultAlphaBlendType)
+		{
+			// Switch the default back to Linear so old data remains the same
+			BlendOption = EAlphaBlendOption::Linear;
+		}
+	}
+	return false;
+}
+
 FMotionMatchingAnimNodeReference UMotionMatchingAnimNodeLibrary::ConvertToMotionMatchingNode(const FAnimNodeReference& Node, EAnimNodeReferenceConversionResult& Result)
 {
 	return FAnimNodeReference::ConvertToType<FMotionMatchingAnimNodeReference>(Node, Result);
diff --git a/Engine/Plugins/Animation/PoseSearch/Source/Runtime/Public/PoseSearch/AnimNode_MotionMatchingInteraction.h b/Engine/Plugins/Animation/PoseSearch/Source/Runtime/Public/PoseSearch/AnimNode_MotionMatchingInteraction.h
--- a/Engine/Plugins/Animation/PoseSearch/Source/Runtime/Public/PoseSearch/AnimNode_MotionMatchingInteraction.h
+++ b/Engine/Plugins/Animation/PoseSearch/Source/Runtime/Public/PoseSearch/AnimNode_MotionMatchingInteraction.h
@@ -46,7 +46,7 @@ struct FAnimNode_MotionMatchingInteraction : public FAnimNode_BlendStack_Standal
 	TObjectPtr<UBlendProfile> BlendProfile;
 
 	UPROPERTY(EditAnywhere, BlueprintReadWrite, Category = Settings, meta = (PinHiddenByDefault))
-	EAlphaBlendOption BlendOption = EAlphaBlendOption::Linear;
+	EAlphaBlendOption BlendOption = UE::Anim::DefaultBlendOption;
 
 	// How we should update individual blend space parameters. See dropdown options tooltips.
 	UPROPERTY(EditAnywhere, Category = Settings)
diff --git a/Engine/Plugins/Animation/PoseSearch/Source/Runtime/Public/PoseSearch/MotionMatchingAnimNodeLibrary.h b/Engine/Plugins/Animation/PoseSearch/Source/Runtime/Public/PoseSearch/MotionMatchingAnimNodeLibrary.h
--- a/Engine/Plugins/Animation/PoseSearch/Source/Runtime/Public/PoseSearch/MotionMatchingAnimNodeLibrary.h
+++ b/Engine/Plugins/Animation/PoseSearch/Source/Runtime/Public/PoseSearch/MotionMatchingAnimNodeLibrary.h
@@ -21,6 +21,9 @@ struct FMotionMatchingBlueprintBlendSettings
 
 	FMotionMatchingBlueprintBlendSettings();
 
+	bool Serialize(FArchive& Ar);
+	bool Serialize(FStructuredArchive::FSlot Slot);
+
 	UPROPERTY(EditAnywhere, BlueprintReadWrite, Category = Settings)
 	float BlendTime;
 
@@ -34,6 +37,15 @@ struct FMotionMatchingBlueprintBlendSettings
 	bool bUseInertialBlend;
 };
 
+template<> struct TStructOpsTypeTraits<FMotionMatchingBlueprintBlendSettings> : public TStructOpsTypeTraitsBase2<FMotionMatchingBlueprintBlendSettings>
+{
+	enum
+	{
+		WithSerializer = true,
+		WithStructuredSerializer = true,
+	};
+};
+
 USTRUCT(BlueprintType)
 struct FMotionMatchingAnimNodeReference : public FAnimNodeReference
 {
diff --git a/Engine/Plugins/Chooser/Source/Chooser/Private/ChooserPlayerSettings.cpp b/Engine/Plugins/Chooser/Source/Chooser/Private/ChooserPlayerSettings.cpp
new file mode 100644
--- /dev/null
+++ b/Engine/Plugins/Chooser/Source/Chooser/Private/ChooserPlayerSettings.cpp
@@ -0,0 +1,37 @@
+// Copyright Epic Games, Inc. All Rights Reserved.
+
+#include "ChooserTypes.h"
+#include "UObject/FortniteMainBranchObjectVersion.h"
+
+bool FChooserPlayerSettings::Serialize(FArchive& Ar)
+{
+	Ar.UsingCustomVersion(FFortniteMainBranchObjectVersion::GUID);
+	if (Ar.IsLoading() && Ar.IsSerializingDefaults())
+	{
+		const int32 CustomVersion = Ar.CustomVer(FFortniteMainBranchObjectVersion::GUID);
+		if (CustomVersion < FFortniteMainBranchObjectVersion::ChangeDefaultAlphaBlendType)
+		{
+			// Switch the default back to Linear so old data remains the same
+			BlendOption = EAlphaBlendOption::Linear;
+		}
+	}
+
+	return false;
+}
+
+bool FChooserPlayerSettings::Serialize(FStructuredArchive::FSlot Slot)
+{
+	FArchive& Ar = Slot.GetUnderlyingArchive();
+	Ar.UsingCustomVersion(FFortniteMainBranchObjectVersion::GUID);
+
+	if (Ar.IsLoading() && Ar.IsSerializingDefaults())
+	{
+		const int32 CustomVersion = Ar.CustomVer(FFortniteMainBranchObjectVersion::GUID);
+		if (CustomVersion < FFortniteMainBranchObjectVersion::ChangeDefaultAlphaBlendType)
+		{
+			// Switch the default back to Linear so old data remains the same
+			BlendOption = EAlphaBlendOption::Linear;
+		}
+	}
+	return false;
+}
diff --git a/Engine/Plugins/Chooser/Source/Chooser/Public/ChooserTypes.h b/Engine/Plugins/Chooser/Source/Chooser/Public/ChooserTypes.h
--- a/Engine/Plugins/Chooser/Source/Chooser/Public/ChooserTypes.h
+++ b/Engine/Plugins/Chooser/Source/Chooser/Public/ChooserTypes.h
@@ -43,6 +43,9 @@ struct FChooserPlayerSettings
 {
 	GENERATED_BODY()
 
+	bool Serialize(FArchive& Ar);
+	bool Serialize(FStructuredArchive::FSlot Slot);
+
 	// Set this value to mirror animations - the MirrorDataTable must also be set on the AnimNode
 	UPROPERTY(EditAnywhere, BlueprintReadWrite, Category = "Settings")
 	bool bMirror = false;
@@ -72,8 +75,17 @@ struct FChooserPlayerSettings
 
 	// How the blend is applied over time to the bones. Common selections are linear, ease in, ease out, and ease in and out.
 	UPROPERTY(EditAnywhere, BlueprintReadWrite, Category = Blending)
-	EAlphaBlendOption BlendOption = EAlphaBlendOption::Linear;
+	EAlphaBlendOption BlendOption = UE::Anim::DefaultBlendOption;
 
 	UPROPERTY(EditAnywhere, BlueprintReadWrite, Category = Blending)
 	bool bUseInertialBlend = false;
+};
+
+template<> struct TStructOpsTypeTraits<FChooserPlayerSettings> : public TStructOpsTypeTraitsBase2<FChooserPlayerSettings>
+{
+	enum
+	{
+		WithSerializer = true,
+		WithStructuredSerializer = true,
+	};
 };
\ No newline at end of file
diff --git a/Engine/Plugins/Chooser/Source/ChooserUncooked/Private/AnimGraphNode_ChooserPlayer.cpp b/Engine/Plugins/Chooser/Source/ChooserUncooked/Private/AnimGraphNode_ChooserPlayer.cpp
--- a/Engine/Plugins/Chooser/Source/ChooserUncooked/Private/AnimGraphNode_ChooserPlayer.cpp
+++ b/Engine/Plugins/Chooser/Source/ChooserUncooked/Private/AnimGraphNode_ChooserPlayer.cpp
@@ -8,6 +8,7 @@
 #include "Animation/AnimPoseSearchProvider.h"
 #include "Animation/AnimRootMotionProvider.h"
 #include "ChooserPropertyAccess.h"
+#include "UObject/FortniteMainBranchObjectVersion.h"
 
 #define LOCTEXT_NAMESPACE "AnimGraphNode_ChooserPlayer"
 
@@ -27,6 +28,24 @@ void UAnimGraphNode_ChooserPlayer::UpdateContextData()
 	}
 }
 
+void UAnimGraphNode_ChooserPlayer::Serialize(FArchive& Ar)
+{
+	// Handle change of default blend type
+	Ar.UsingCustomVersion(FFortniteMainBranchObjectVersion::GUID);
+	if (Ar.IsLoading())
+	{
+		const int32 CustomVersion = Ar.CustomVer(FFortniteMainBranchObjectVersion::GUID);
+		if (CustomVersion < FFortniteMainBranchObjectVersion::ChangeDefaultAlphaBlendType)
+		{
+			// Switch the default back to Linear so old data remains the same
+			// Important: this is done before loading so if data has changed from default it still works
+			Node.DefaultSettings.BlendOption = EAlphaBlendOption::Linear;
+		}
+	}
+	
+	Super::Serialize(Ar);
+}
+
 void UAnimGraphNode_ChooserPlayer::PostPlacedNewNode()
 {
 	Super::PostPlacedNewNode();
diff --git a/Engine/Plugins/Chooser/Source/ChooserUncooked/Public/AnimGraphNode_ChooserPlayer.h b/Engine/Plugins/Chooser/Source/ChooserUncooked/Public/AnimGraphNode_ChooserPlayer.h
--- a/Engine/Plugins/Chooser/Source/ChooserUncooked/Public/AnimGraphNode_ChooserPlayer.h
+++ b/Engine/Plugins/Chooser/Source/ChooserUncooked/Public/AnimGraphNode_ChooserPlayer.h
@@ -17,6 +17,10 @@ class UAnimGraphNode_ChooserPlayer : public UAnimGraphNode_BlendStack_Base, publ
 
 public:
 
+	// UObject interface
+	virtual void Serialize(FArchive& Ar) override;
+	// End of UObject
+
 	UPROPERTY(EditAnywhere, Category=Settings)
 	FAnimNode_ChooserPlayer Node;
 
diff --git a/Engine/Plugins/Experimental/Animation/ContextualAnimation/Source/ContextualAnimation/Private/AnimNotifyState_IKWindow.cpp b/Engine/Plugins/Experimental/Animation/ContextualAnimation/Source/ContextualAnimation/Private/AnimNotifyState_IKWindow.cpp
--- a/Engine/Plugins/Experimental/Animation/ContextualAnimation/Source/ContextualAnimation/Private/AnimNotifyState_IKWindow.cpp
+++ b/Engine/Plugins/Experimental/Animation/ContextualAnimation/Source/ContextualAnimation/Private/AnimNotifyState_IKWindow.cpp
@@ -2,9 +2,28 @@
 
 #include "AnimNotifyState_IKWindow.h"
 #include "Animation/AnimMontage.h"
+#include "UObject/FortniteMainBranchObjectVersion.h"
 
 #include UE_INLINE_GENERATED_CPP_BY_NAME(AnimNotifyState_IKWindow)
 
+void UAnimNotifyState_IKWindow::Serialize(FArchive& Ar)
+{
+	Ar.UsingCustomVersion(FFortniteMainBranchObjectVersion::GUID);
+	if (Ar.IsLoading())
+	{
+		const int32 CustomVersion = Ar.CustomVer(FFortniteMainBranchObjectVersion::GUID);
+		if (CustomVersion < FFortniteMainBranchObjectVersion::ChangeDefaultAlphaBlendType)
+		{
+			// Switch the default back to Linear so old data remains the same
+			// Note this happens before serialization
+			BlendIn.SetBlendOption(EAlphaBlendOption::Linear);
+			BlendOut.SetBlendOption(EAlphaBlendOption::Linear);
+		}
+	}
+	
+	Super::Serialize(Ar);
+}
+
 UAnimNotifyState_IKWindow::UAnimNotifyState_IKWindow(const FObjectInitializer& ObjectInitializer)
 	: Super(ObjectInitializer)
 {
diff --git a/Engine/Plugins/Experimental/Animation/ContextualAnimation/Source/ContextualAnimation/Public/AnimNotifyState_IKWindow.h b/Engine/Plugins/Experimental/Animation/ContextualAnimation/Source/ContextualAnimation/Public/AnimNotifyState_IKWindow.h
--- a/Engine/Plugins/Experimental/Animation/ContextualAnimation/Source/ContextualAnimation/Public/AnimNotifyState_IKWindow.h
+++ b/Engine/Plugins/Experimental/Animation/ContextualAnimation/Source/ContextualAnimation/Public/AnimNotifyState_IKWindow.h
@@ -16,6 +16,8 @@ class UAnimNotifyState_IKWindow : public UAnimNotifyState
 
 public:
 
+	virtual void Serialize(FArchive& Ar) override;
+
 	UPROPERTY(EditAnywhere, Category = Config)
 	FName GoalName;
 
diff --git a/Engine/Source/Editor/AnimGraph/Private/AnimGraphNode_ApplyAdditive.cpp b/Engine/Source/Editor/AnimGraph/Private/AnimGraphNode_ApplyAdditive.cpp
--- a/Engine/Source/Editor/AnimGraph/Private/AnimGraphNode_ApplyAdditive.cpp
+++ b/Engine/Source/Editor/AnimGraph/Private/AnimGraphNode_ApplyAdditive.cpp
@@ -7,6 +7,7 @@
 #include "ScopedTransaction.h"
 #include "Kismet2/BlueprintEditorUtils.h"
 #include "DetailLayoutBuilder.h"
+#include "UObject/FortniteMainBranchObjectVersion.h"
 
 /////////////////////////////////////////////////////
 // UAnimGraphNode_ApplyAdditive
@@ -23,6 +24,24 @@ FLinearColor UAnimGraphNode_ApplyAdditive::GetNodeTitleColor() const
 	return FLinearColor(0.75f, 0.75f, 0.75f);
 }
 
+void UAnimGraphNode_ApplyAdditive::Serialize(FArchive& Ar)
+{
+	// Handle change of default blend type
+	Ar.UsingCustomVersion(FFortniteMainBranchObjectVersion::GUID);
+	if (Ar.IsLoading())
+	{
+		const int32 CustomVersion = Ar.CustomVer(FFortniteMainBranchObjectVersion::GUID);
+		if (CustomVersion < FFortniteMainBranchObjectVersion::ChangeDefaultAlphaBlendType)
+		{
+			// Switch the default back to Linear so old data remains the same
+			// Important: this is done before loading so if data has changed from default it still works
+			Node.AlphaBoolBlend.BlendOption = EAlphaBlendOption::Linear;
+		}
+	}
+	
+	Super::Serialize(Ar);
+}
+
 FText UAnimGraphNode_ApplyAdditive::GetTooltipText() const
 {
 	return LOCTEXT("AnimGraphNode_ApplyAdditive_Tooltip", "Apply additive animation to normal pose");
diff --git a/Engine/Source/Editor/AnimGraph/Private/AnimGraphNode_ApplyMeshSpaceAdditive.cpp b/Engine/Source/Editor/AnimGraph/Private/AnimGraphNode_ApplyMeshSpaceAdditive.cpp
--- a/Engine/Source/Editor/AnimGraph/Private/AnimGraphNode_ApplyMeshSpaceAdditive.cpp
+++ b/Engine/Source/Editor/AnimGraph/Private/AnimGraphNode_ApplyMeshSpaceAdditive.cpp
@@ -8,7 +8,7 @@
 #include "Animation/AnimationSettings.h"
 #include "Kismet2/BlueprintEditorUtils.h"
 #include "Kismet2/CompilerResultsLog.h"
-
+#include "UObject/FortniteMainBranchObjectVersion.h"
 
 /////////////////////////////////////////////////////
 // UAnimGraphNode_ApplyAdditive
@@ -25,6 +25,24 @@ FLinearColor UAnimGraphNode_ApplyMeshSpaceAdditive::GetNodeTitleColor() const
 	return FLinearColor(0.75f, 0.75f, 0.75f);
 }
 
+void UAnimGraphNode_ApplyMeshSpaceAdditive::Serialize(FArchive& Ar)
+{
+	// Handle change of default blend type
+	Ar.UsingCustomVersion(FFortniteMainBranchObjectVersion::GUID);
+	if (Ar.IsLoading())
+	{
+		const int32 CustomVersion = Ar.CustomVer(FFortniteMainBranchObjectVersion::GUID);
+		if (CustomVersion < FFortniteMainBranchObjectVersion::ChangeDefaultAlphaBlendType)
+		{
+			// Switch the default back to Linear so old data remains the same
+			// Important: this is done before loading so if data has changed from default it still works
+			Node.AlphaBoolBlend.BlendOption = EAlphaBlendOption::Linear;
+		}
+	}
+	
+	Super::Serialize(Ar);
+}
+
 FText UAnimGraphNode_ApplyMeshSpaceAdditive::GetTooltipText() const
 {
 	return LOCTEXT("AnimGraphNode_ApplyMeshSpaceAdditive_Tooltip", "Apply mesh space additive animation to normal pose");
diff --git a/Engine/Source/Editor/AnimGraph/Private/AnimGraphNode_BlendListBase.cpp b/Engine/Source/Editor/AnimGraph/Private/AnimGraphNode_BlendListBase.cpp
--- a/Engine/Source/Editor/AnimGraph/Private/AnimGraphNode_BlendListBase.cpp
+++ b/Engine/Source/Editor/AnimGraph/Private/AnimGraphNode_BlendListBase.cpp
@@ -3,6 +3,7 @@
 #include "AnimGraphNode_BlendListBase.h"
 #include "AnimNodes/AnimNode_BlendListBase.h"
 #include "Animation/AnimNode_Inertialization.h"
+#include "UObject/FortniteMainBranchObjectVersion.h"
 
 /////////////////////////////////////////////////////
 // UAnimGraphNode_BlendListBase
@@ -15,6 +16,20 @@ UAnimGraphNode_BlendListBase::UAnimGraphNode_BlendListBase(const FObjectInitiali
 
 void UAnimGraphNode_BlendListBase::Serialize(FArchive& Ar)
 {
+	// Handle change of default blend type
+	Ar.UsingCustomVersion(FFortniteMainBranchObjectVersion::GUID);
+	if (Ar.IsLoading())
+	{
+		FAnimNode_BlendListBase* Node = static_cast<FAnimNode_BlendListBase*>(GetFNode());
+		const int32 CustomVersion = Ar.CustomVer(FFortniteMainBranchObjectVersion::GUID);
+		if (CustomVersion < FFortniteMainBranchObjectVersion::ChangeDefaultAlphaBlendType)
+		{
+			// Switch the default back to Linear so old data remains the same
+			// Important: this is done before loading so if data has changed from default it still works
+			Node->BlendType = EAlphaBlendOption::Linear;
+		}
+	}
+	
 	Super::Serialize(Ar);
 
 	if (Ar.IsLoading())
diff --git a/Engine/Source/Editor/AnimGraph/Private/AnimGraphNode_DeadBlending.cpp b/Engine/Source/Editor/AnimGraph/Private/AnimGraphNode_DeadBlending.cpp
--- a/Engine/Source/Editor/AnimGraph/Private/AnimGraphNode_DeadBlending.cpp
+++ b/Engine/Source/Editor/AnimGraph/Private/AnimGraphNode_DeadBlending.cpp
@@ -2,10 +2,29 @@
 
 #include "AnimGraphNode_DeadBlending.h"
 #include "IAnimBlueprintCopyTermDefaultsContext.h"
+#include "UObject/FortniteMainBranchObjectVersion.h"
 
 #define LOCTEXT_NAMESPACE "AnimGraphNode_DeadBlending"
 
 
+void UAnimGraphNode_DeadBlending::Serialize(FArchive& Ar)
+{
+	// Handle change of default blend type
+	Ar.UsingCustomVersion(FFortniteMainBranchObjectVersion::GUID);
+	if (Ar.IsLoading())
+	{
+		const int32 CustomVersion = Ar.CustomVer(FFortniteMainBranchObjectVersion::GUID);
+		if (CustomVersion < FFortniteMainBranchObjectVersion::ChangeDefaultAlphaBlendType)
+		{
+			// Switch the default back to Linear so old data remains the same
+			// Important: this is done before loading so if data has changed from default it still works
+			Node.DefaultBlendMode = EAlphaBlendOption::Linear;
+		}
+	}
+	
+	Super::Serialize(Ar);
+}
+
 FLinearColor UAnimGraphNode_DeadBlending::GetNodeTitleColor() const
 {
 	return FLinearColor(0.0f, 0.1f, 0.2f);
diff --git a/Engine/Source/Editor/AnimGraph/Private/AnimGraphNode_PoseBlendNode.cpp b/Engine/Source/Editor/AnimGraph/Private/AnimGraphNode_PoseBlendNode.cpp
--- a/Engine/Source/Editor/AnimGraph/Private/AnimGraphNode_PoseBlendNode.cpp
+++ b/Engine/Source/Editor/AnimGraph/Private/AnimGraphNode_PoseBlendNode.cpp
@@ -11,6 +11,7 @@
 #include "BlueprintActionFilter.h"
 #include "BlueprintActionDatabaseRegistrar.h"
 #include "BlueprintNodeSpawner.h"
+#include "UObject/FortniteMainBranchObjectVersion.h"
 
 #define LOCTEXT_NAMESPACE "PoseBlendNode"
 
@@ -72,6 +73,24 @@ FText UAnimGraphNode_PoseBlendNode::GetNodeTitleForPoseAsset(ENodeTitleType::Typ
 	return FText::Format(LOCTEXT("PoseByName_Title", "{PoseAssetName}"), Args);
 }
 
+void UAnimGraphNode_PoseBlendNode::Serialize(FArchive& Ar)
+{
+	// Handle change of default blend type
+	Ar.UsingCustomVersion(FFortniteMainBranchObjectVersion::GUID);
+	if (Ar.IsLoading())
+	{
+		const int32 CustomVersion = Ar.CustomVer(FFortniteMainBranchObjectVersion::GUID);
+		if (CustomVersion < FFortniteMainBranchObjectVersion::ChangeDefaultAlphaBlendType)
+		{
+			// Switch the default back to Linear so old data remains the same
+			// Important: this is done before loading so if data has changed from default it still works
+			Node.BlendOption = EAlphaBlendOption::Linear;
+		}
+	}
+	
+	Super::Serialize(Ar);
+}
+
 FText UAnimGraphNode_PoseBlendNode::GetNodeTitle(ENodeTitleType::Type TitleType) const
 {
 	if (Node.PoseAsset == nullptr)
diff --git a/Engine/Source/Editor/AnimGraph/Private/AnimGraphNode_RotationOffsetBlendSpace.cpp b/Engine/Source/Editor/AnimGraph/Private/AnimGraphNode_RotationOffsetBlendSpace.cpp
--- a/Engine/Source/Editor/AnimGraph/Private/AnimGraphNode_RotationOffsetBlendSpace.cpp
+++ b/Engine/Source/Editor/AnimGraph/Private/AnimGraphNode_RotationOffsetBlendSpace.cpp
@@ -18,6 +18,7 @@
 #include "BlueprintNodeTemplateCache.h"
 #include "AssetRegistry/AssetRegistryModule.h"
 #include "IAnimBlueprintNodeOverrideAssetsContext.h"
+#include "UObject/FortniteMainBranchObjectVersion.h"
 
 /////////////////////////////////////////////////////
 // UAnimGraphNode_RotationOffsetBlendSpace
@@ -29,6 +30,24 @@ UAnimGraphNode_RotationOffsetBlendSpace::UAnimGraphNode_RotationOffsetBlendSpace
 {
 }
 
+void UAnimGraphNode_RotationOffsetBlendSpace::Serialize(FArchive& Ar)
+{
+	// Handle change of default blend type
+	Ar.UsingCustomVersion(FFortniteMainBranchObjectVersion::GUID);
+	if (Ar.IsLoading())
+	{
+		const int32 CustomVersion = Ar.CustomVer(FFortniteMainBranchObjectVersion::GUID);
+		if (CustomVersion < FFortniteMainBranchObjectVersion::ChangeDefaultAlphaBlendType)
+		{
+			// Switch the default back to Linear so old data remains the same
+			// Important: this is done before loading so if data has changed from default it still works
+			Node.AlphaBoolBlend.BlendOption = EAlphaBlendOption::Linear;
+		}
+	}
+	
+	Super::Serialize(Ar);
+}
+
 FText UAnimGraphNode_RotationOffsetBlendSpace::GetNodeTitle(ENodeTitleType::Type TitleType) const
 {
 	UBlendSpace* BlendSpaceToCheck = Node.GetBlendSpace();
diff --git a/Engine/Source/Editor/AnimGraph/Private/AnimGraphNode_SkeletalControlBase.cpp b/Engine/Source/Editor/AnimGraph/Private/AnimGraphNode_SkeletalControlBase.cpp
--- a/Engine/Source/Editor/AnimGraph/Private/AnimGraphNode_SkeletalControlBase.cpp
+++ b/Engine/Source/Editor/AnimGraph/Private/AnimGraphNode_SkeletalControlBase.cpp
@@ -55,6 +55,25 @@ FName UAnimGraphNode_SkeletalControlBase::FindSelectedBone()
 	return NAME_None;
 }
 
+void UAnimGraphNode_SkeletalControlBase::Serialize(FArchive& Ar)
+{
+	// Handle change of default blend type
+	Ar.UsingCustomVersion(FFortniteMainBranchObjectVersion::GUID);
+	if (Ar.IsLoading())
+	{
+		const int32 CustomVersion = Ar.CustomVer(FFortniteMainBranchObjectVersion::GUID);
+		if (CustomVersion < FFortniteMainBranchObjectVersion::ChangeDefaultAlphaBlendType)
+		{
+			// Switch the default back to Linear so old data remains the same
+			// Important: this is done before loading so if data has changed from default it still works
+			FAnimNode_SkeletalControlBase* SkeletalControlBase = static_cast<FAnimNode_SkeletalControlBase*>(GetFNode());
+			SkeletalControlBase->AlphaBoolBlend.BlendOption = EAlphaBlendOption::Linear;
+		}
+	}
+	
+	Super::Serialize(Ar);
+}
+
 FLinearColor UAnimGraphNode_SkeletalControlBase::GetNodeTitleColor() const
 {
 	return FLinearColor(0.75f, 0.75f, 0.10f);
diff --git a/Engine/Source/Editor/AnimGraph/Private/AnimGraphNode_SplineIK.cpp b/Engine/Source/Editor/AnimGraph/Private/AnimGraphNode_SplineIK.cpp
--- a/Engine/Source/Editor/AnimGraph/Private/AnimGraphNode_SplineIK.cpp
+++ b/Engine/Source/Editor/AnimGraph/Private/AnimGraphNode_SplineIK.cpp
@@ -4,9 +4,28 @@
 #include "AnimNodeEditModes.h"
 #include "Animation/AnimBlueprint.h"
 #include "Animation/Skeleton.h"
+#include "UObject/FortniteMainBranchObjectVersion.h"
 
 #define LOCTEXT_NAMESPACE "A3Nodes"
 
+void UAnimGraphNode_SplineIK::Serialize(FArchive& Ar)
+{
+	// Handle change of default blend type
+	Ar.UsingCustomVersion(FFortniteMainBranchObjectVersion::GUID);
+	if (Ar.IsLoading())
+	{
+		const int32 CustomVersion = Ar.CustomVer(FFortniteMainBranchObjectVersion::GUID);
+		if (CustomVersion < FFortniteMainBranchObjectVersion::ChangeDefaultAlphaBlendType)
+		{
+			// Switch the default back to Linear so old data remains the same
+			// Important: this is done before loading so if data has changed from default it still works
+			Node.TwistBlend.SetBlendOption(EAlphaBlendOption::Linear);
+		}
+	}
+	
+	Super::Serialize(Ar);
+}
+
 UAnimGraphNode_SplineIK::UAnimGraphNode_SplineIK(const FObjectInitializer& ObjectInitializer)
 	: Super(ObjectInitializer)
 {
diff --git a/Engine/Source/Editor/AnimGraph/Private/AnimGraphNode_TwoWayBlend.cpp b/Engine/Source/Editor/AnimGraph/Private/AnimGraphNode_TwoWayBlend.cpp
--- a/Engine/Source/Editor/AnimGraph/Private/AnimGraphNode_TwoWayBlend.cpp
+++ b/Engine/Source/Editor/AnimGraph/Private/AnimGraphNode_TwoWayBlend.cpp
@@ -4,6 +4,7 @@
 #include "DetailLayoutBuilder.h"
 #include "ScopedTransaction.h"
 #include "Kismet2/BlueprintEditorUtils.h"
+#include "UObject/FortniteMainBranchObjectVersion.h"
 
 /////////////////////////////////////////////////////
 // UAnimGraphNode_TwoWayBlend
@@ -20,6 +21,24 @@ FText UAnimGraphNode_TwoWayBlend::GetMenuCategory() const
 	return LOCTEXT("TwoWayBlendCategory", "Animation|Blends");
 }
 
+void UAnimGraphNode_TwoWayBlend::Serialize(FArchive& Ar)
+{
+	// Handle change of default blend type
+	Ar.UsingCustomVersion(FFortniteMainBranchObjectVersion::GUID);
+	if (Ar.IsLoading())
+	{
+		const int32 CustomVersion = Ar.CustomVer(FFortniteMainBranchObjectVersion::GUID);
+		if (CustomVersion < FFortniteMainBranchObjectVersion::ChangeDefaultAlphaBlendType)
+		{
+			// Switch the default back to Linear so old data remains the same
+			// Important: this is done before loading so if data has changed from default it still works
+			BlendNode.AlphaBoolBlend.BlendOption = EAlphaBlendOption::Linear;
+		}
+	}
+	
+	Super::Serialize(Ar);
+}
+
 FLinearColor UAnimGraphNode_TwoWayBlend::GetNodeTitleColor() const
 {
 	return FLinearColor(0.75f, 0.75f, 0.75f);
diff --git a/Engine/Source/Editor/AnimGraph/Public/AnimGraphNode_ApplyAdditive.h b/Engine/Source/Editor/AnimGraph/Public/AnimGraphNode_ApplyAdditive.h
--- a/Engine/Source/Editor/AnimGraph/Public/AnimGraphNode_ApplyAdditive.h
+++ b/Engine/Source/Editor/AnimGraph/Public/AnimGraphNode_ApplyAdditive.h
@@ -16,6 +16,10 @@ class UAnimGraphNode_ApplyAdditive : public UAnimGraphNode_Base
 	UPROPERTY(EditAnywhere, Category=Settings)
 	FAnimNode_ApplyAdditive Node;
 
+	// UObject interface
+	virtual void Serialize(FArchive& Ar) override;
+	// End of UObject
+
 	// UEdGraphNode interface
 	virtual FText GetTooltipText() const override;
 	virtual FText GetNodeTitle(ENodeTitleType::Type TitleType) const override;
diff --git a/Engine/Source/Editor/AnimGraph/Public/AnimGraphNode_ApplyMeshSpaceAdditive.h b/Engine/Source/Editor/AnimGraph/Public/AnimGraphNode_ApplyMeshSpaceAdditive.h
--- a/Engine/Source/Editor/AnimGraph/Public/AnimGraphNode_ApplyMeshSpaceAdditive.h
+++ b/Engine/Source/Editor/AnimGraph/Public/AnimGraphNode_ApplyMeshSpaceAdditive.h
@@ -16,6 +16,10 @@ class UAnimGraphNode_ApplyMeshSpaceAdditive : public UAnimGraphNode_Base
 	UPROPERTY(EditAnywhere, Category = Settings)
 	FAnimNode_ApplyMeshSpaceAdditive Node;
 
+	// UObject interface
+	virtual void Serialize(FArchive& Ar) override;
+	// End of UObject
+
 	// UEdGraphNode interface
 	FText GetTooltipText() const override;
 	FText GetNodeTitle(ENodeTitleType::Type TitleType) const override;
diff --git a/Engine/Source/Editor/AnimGraph/Public/AnimGraphNode_DeadBlending.h b/Engine/Source/Editor/AnimGraph/Public/AnimGraphNode_DeadBlending.h
--- a/Engine/Source/Editor/AnimGraph/Public/AnimGraphNode_DeadBlending.h
+++ b/Engine/Source/Editor/AnimGraph/Public/AnimGraphNode_DeadBlending.h
@@ -16,6 +16,10 @@ class UAnimGraphNode_DeadBlending : public UAnimGraphNode_Base
 	UPROPERTY(EditAnywhere, Category = Settings)
 	FAnimNode_DeadBlending Node;
 
+	// UObject interface
+	virtual void Serialize(FArchive& Ar) override;
+	// End of UObject
+
 	virtual FLinearColor GetNodeTitleColor() const override;
 	virtual FText GetTooltipText() const override;
 	virtual FText GetNodeTitle(ENodeTitleType::Type TitleType) const override;
diff --git a/Engine/Source/Editor/AnimGraph/Public/AnimGraphNode_PoseBlendNode.h b/Engine/Source/Editor/AnimGraph/Public/AnimGraphNode_PoseBlendNode.h
--- a/Engine/Source/Editor/AnimGraph/Public/AnimGraphNode_PoseBlendNode.h
+++ b/Engine/Source/Editor/AnimGraph/Public/AnimGraphNode_PoseBlendNode.h
@@ -18,6 +18,10 @@ class UAnimGraphNode_PoseBlendNode : public UAnimGraphNode_PoseHandler
 	UPROPERTY(EditAnywhere, Category=Settings)
 	FAnimNode_PoseBlendNode Node;
 
+	// UObject interface
+	virtual void Serialize(FArchive& Ar) override;
+	// End of UObject
+
 	// UEdGraphNode interface
 	virtual FText GetNodeTitle(ENodeTitleType::Type TitleType) const override;
 	virtual FText GetTooltipText() const override;
diff --git a/Engine/Source/Editor/AnimGraph/Public/AnimGraphNode_RotationOffsetBlendSpace.h b/Engine/Source/Editor/AnimGraph/Public/AnimGraphNode_RotationOffsetBlendSpace.h
--- a/Engine/Source/Editor/AnimGraph/Public/AnimGraphNode_RotationOffsetBlendSpace.h
+++ b/Engine/Source/Editor/AnimGraph/Public/AnimGraphNode_RotationOffsetBlendSpace.h
@@ -20,6 +20,10 @@ class UAnimGraphNode_RotationOffsetBlendSpace: public UAnimGraphNode_BlendSpaceB
 	UPROPERTY(EditAnywhere, Category=Settings)
 	FAnimNode_RotationOffsetBlendSpace Node;
 
+	// UObject interface
+	virtual void Serialize(FArchive& Ar) override;
+	// End of UObject
+
 	// UEdGraphNode interface
 	virtual FText GetNodeTitle(ENodeTitleType::Type TitleType) const override;
 	virtual void PostEditChangeProperty(struct FPropertyChangedEvent& PropertyChangedEvent) override;
diff --git a/Engine/Source/Editor/AnimGraph/Public/AnimGraphNode_SkeletalControlBase.h b/Engine/Source/Editor/AnimGraph/Public/AnimGraphNode_SkeletalControlBase.h
--- a/Engine/Source/Editor/AnimGraph/Public/AnimGraphNode_SkeletalControlBase.h
+++ b/Engine/Source/Editor/AnimGraph/Public/AnimGraphNode_SkeletalControlBase.h
@@ -27,6 +27,10 @@ class UAnimGraphNode_SkeletalControlBase : public UAnimGraphNode_Base
 	GENERATED_UCLASS_BODY()
 
 public:
+	// UObject interface
+	UE_API virtual void Serialize(FArchive& Ar) override;
+	// End of UObject
+
 	// UEdGraphNode interface
 	UE_API virtual FLinearColor GetNodeTitleColor() const override;
 	UE_API virtual FText GetTooltipText() const override;
diff --git a/Engine/Source/Editor/AnimGraph/Public/AnimGraphNode_SplineIK.h b/Engine/Source/Editor/AnimGraph/Public/AnimGraphNode_SplineIK.h
--- a/Engine/Source/Editor/AnimGraph/Public/AnimGraphNode_SplineIK.h
+++ b/Engine/Source/Editor/AnimGraph/Public/AnimGraphNode_SplineIK.h
@@ -16,7 +16,9 @@ class UAnimGraphNode_SplineIK : public UAnimGraphNode_SkeletalControlBase
 	FAnimNode_SplineIK Node;
 
 public:
+
 	// UObject interface
+	virtual void Serialize(FArchive& Ar) override;
 	virtual void PostEditChangeProperty(struct FPropertyChangedEvent& PropertyChangedEvent) override;
 
 	// UEdGraphNode interface
diff --git a/Engine/Source/Editor/AnimGraph/Public/AnimGraphNode_TwoWayBlend.h b/Engine/Source/Editor/AnimGraph/Public/AnimGraphNode_TwoWayBlend.h
--- a/Engine/Source/Editor/AnimGraph/Public/AnimGraphNode_TwoWayBlend.h
+++ b/Engine/Source/Editor/AnimGraph/Public/AnimGraphNode_TwoWayBlend.h
@@ -16,6 +16,10 @@ class UAnimGraphNode_TwoWayBlend : public UAnimGraphNode_Base
 	UPROPERTY(EditAnywhere, Category=Settings)
 	FAnimNode_TwoWayBlend BlendNode;
 
+	// UObject interface
+	virtual void Serialize(FArchive& Ar) override;
+	// End of UObject
+
 	//~ Begin UEdGraphNode Interface.
 	virtual FLinearColor GetNodeTitleColor() const override;
 	virtual FText GetTooltipText() const override;
diff --git a/Engine/Source/Runtime/AnimGraphRuntime/Private/AnimNodes/AnimNode_PoseBlendNode.cpp b/Engine/Source/Runtime/AnimGraphRuntime/Private/AnimNodes/AnimNode_PoseBlendNode.cpp
--- a/Engine/Source/Runtime/AnimGraphRuntime/Private/AnimNodes/AnimNode_PoseBlendNode.cpp
+++ b/Engine/Source/Runtime/AnimGraphRuntime/Private/AnimNodes/AnimNode_PoseBlendNode.cpp
@@ -13,7 +13,7 @@
 FAnimNode_PoseBlendNode::FAnimNode_PoseBlendNode()
 	: CustomCurve(nullptr)
 {
-	BlendOption = EAlphaBlendOption::Linear;
+	BlendOption = UE::Anim::DefaultBlendOption;
 }
 
 void FAnimNode_PoseBlendNode::Initialize_AnyThread(const FAnimationInitializeContext& Context)
diff --git a/Engine/Source/Runtime/AnimGraphRuntime/Public/AnimNodes/AnimNode_BlendListBase.h b/Engine/Source/Runtime/AnimGraphRuntime/Public/AnimNodes/AnimNode_BlendListBase.h
--- a/Engine/Source/Runtime/AnimGraphRuntime/Public/AnimNodes/AnimNode_BlendListBase.h
+++ b/Engine/Source/Runtime/AnimGraphRuntime/Public/AnimNodes/AnimNode_BlendListBase.h
@@ -51,7 +51,7 @@ private:
 	EBlendListTransitionType TransitionType = EBlendListTransitionType::StandardBlend;
 
 	UPROPERTY(EditAnywhere, Category=BlendType, meta=(FoldProperty))
-	EAlphaBlendOption BlendType = EAlphaBlendOption::Linear;
+	EAlphaBlendOption BlendType = UE::Anim::DefaultBlendOption;
 	
 protected:
 	UE_DEPRECATED(5.6, "Use ChildUpateMode instead.")
diff --git a/Engine/Source/Runtime/Core/Public/UObject/FortniteMainBranchObjectVersions.inl b/Engine/Source/Runtime/Core/Public/UObject/FortniteMainBranchObjectVersions.inl
--- a/Engine/Source/Runtime/Core/Public/UObject/FortniteMainBranchObjectVersions.inl
+++ b/Engine/Source/Runtime/Core/Public/UObject/FortniteMainBranchObjectVersions.inl
@@ -652,5 +652,8 @@ FORTNITE_MAIN_VERSION(AnimNextVariableReferences, 212)
 // The default distortion rendering mode used by the Lens Component is now the Lens Distortion Scene View Extension
 FORTNITE_MAIN_VERSION(LensComponentDefaultToDistortionSVE, 213)
 
+// Animation default blend option changed from Linear to HermiteCubic (aka SmoothStep, ease in / ease out)
+FORTNITE_MAIN_VERSION(ChangeDefaultAlphaBlendType, 214)
+
 // -----<new versions can be added above this line>-------------------------------------------------
 #undef FORTNITE_MAIN_VERSION
\ No newline at end of file
diff --git a/Engine/Source/Runtime/Engine/Classes/Animation/AnimMontage.h b/Engine/Source/Runtime/Engine/Classes/Animation/AnimMontage.h
--- a/Engine/Source/Runtime/Engine/Classes/Animation/AnimMontage.h
+++ b/Engine/Source/Runtime/Engine/Classes/Animation/AnimMontage.h
@@ -737,6 +737,9 @@ class UAnimMontage : public UAnimCompositeBase
 
 public:
 	//~ Begin UObject Interface
+#if WITH_EDITORONLY_DATA
+	virtual void Serialize(FArchive& Ar) override;
+#endif // WITH_EDITORONLY_DATA
 	virtual void PostLoad() override;
 	virtual void PreSave(FObjectPreSaveContext ObjectSaveContext) override;
 
diff --git a/Engine/Source/Runtime/Engine/Classes/Animation/AnimNode_DeadBlending.h b/Engine/Source/Runtime/Engine/Classes/Animation/AnimNode_DeadBlending.h
--- a/Engine/Source/Runtime/Engine/Classes/Animation/AnimNode_DeadBlending.h
+++ b/Engine/Source/Runtime/Engine/Classes/Animation/AnimNode_DeadBlending.h
@@ -34,6 +34,7 @@ struct FAnimNode_DeadBlending : public FAnimNode_Base
 #endif
 {
 	GENERATED_BODY()
+	friend class UAnimGraphNode_DeadBlending;
 
 	// Input Pose
 	UPROPERTY(EditAnywhere, Category = Links)
@@ -60,7 +61,7 @@ private:
 
 	// Default blend mode to use when no blend mode is supplied with the inertialization request.
 	UPROPERTY(EditAnywhere, Category = Blending, DisplayName = "Default Blend Mode")
-	EAlphaBlendOption DefaultBlendMode = EAlphaBlendOption::Linear;
+	EAlphaBlendOption DefaultBlendMode = UE::Anim::DefaultBlendOption;
 
 	// Default custom blend curve to use along with the default blend mode.
 	UPROPERTY(EditAnywhere, Category = Blending)
diff --git a/Engine/Source/Runtime/Engine/Classes/Animation/InputScaleBias.h b/Engine/Source/Runtime/Engine/Classes/Animation/InputScaleBias.h
--- a/Engine/Source/Runtime/Engine/Classes/Animation/InputScaleBias.h
+++ b/Engine/Source/Runtime/Engine/Classes/Animation/InputScaleBias.h
@@ -286,6 +286,11 @@ struct FInputAlphaBoolBlend
 {
 	GENERATED_USTRUCT_BODY()
 
+#if WITH_EDITORONLY_DATA
+	bool Serialize(FArchive& Ar);
+	bool Serialize(FStructuredArchive::FSlot Slot);
+#endif // WITH_EDITORONLY_DATA
+
 	UPROPERTY(EditAnywhere, BlueprintReadWrite, Category = Settings)
 	float BlendInTime;
 
@@ -307,7 +312,7 @@ struct FInputAlphaBoolBlend
 	FInputAlphaBoolBlend()
 		: BlendInTime(0.0f)
 		, BlendOutTime(0.0f)
-		, BlendOption(EAlphaBlendOption::Linear)
+		, BlendOption(UE::Anim::DefaultBlendOption)
 		, bInitialized(false)
 		, CustomCurve(nullptr)
 	{}
@@ -316,3 +321,14 @@ struct FInputAlphaBoolBlend
 
 	void Reinitialize() { bInitialized = false; }
 };
+
+#if WITH_EDITORONLY_DATA
+template<> struct TStructOpsTypeTraits<FInputAlphaBoolBlend> : public TStructOpsTypeTraitsBase2<FInputAlphaBoolBlend>
+{
+	enum
+	{
+		WithSerializer = true,
+		WithStructuredSerializer = true,
+	};
+};
+#endif // WITH_EDITORONLY_DATA
diff --git a/Engine/Source/Runtime/Engine/Private/AlphaBlend.cpp b/Engine/Source/Runtime/Engine/Private/AlphaBlend.cpp
--- a/Engine/Source/Runtime/Engine/Private/AlphaBlend.cpp
+++ b/Engine/Source/Runtime/Engine/Private/AlphaBlend.cpp
@@ -2,6 +2,7 @@
 
 #include "AlphaBlend.h"
 #include "Curves/CurveFloat.h"
+#include "UObject/FortniteMainBranchObjectVersion.h"
 
 #include UE_INLINE_GENERATED_CPP_BY_NAME(AlphaBlend)
 
@@ -10,7 +11,7 @@ FAlphaBlend::FAlphaBlend(float NewBlendTime)
 	, BlendTime(NewBlendTime)
 	, BeginValue(0.0f)
 	, DesiredValue(1.0f)
-	, BlendOption(EAlphaBlendOption::Linear)
+	, BlendOption(UE::Anim::DefaultBlendOption)
 {
 	Reset();
 }
@@ -25,6 +26,75 @@ FAlphaBlend::FAlphaBlend(const FAlphaBlend& Other, float NewBlendTime)
 	Reset();
 }
 
+#if WITH_EDITORONLY_DATA
+bool FAlphaBlendArgs::Serialize(FArchive& Ar)
+{
+	Ar.UsingCustomVersion(FFortniteMainBranchObjectVersion::GUID);
+	if (Ar.IsLoading() && Ar.IsSerializingDefaults())
+	{
+		const int32 CustomVersion = Ar.CustomVer(FFortniteMainBranchObjectVersion::GUID);
+		if (CustomVersion < FFortniteMainBranchObjectVersion::ChangeDefaultAlphaBlendType)
+		{
+			// Switch the default back to Linear so old data remains the same
+			BlendOption = EAlphaBlendOption::Linear;
+		}
+	}
+
+	return false;
+}
+
+bool FAlphaBlendArgs::Serialize(FStructuredArchive::FSlot Slot)
+{
+	FArchive& Ar = Slot.GetUnderlyingArchive();
+	Ar.UsingCustomVersion(FFortniteMainBranchObjectVersion::GUID);
+
+	if (Ar.IsLoading() && Ar.IsSerializingDefaults())
+	{
+		const int32 CustomVersion = Ar.CustomVer(FFortniteMainBranchObjectVersion::GUID);
+		if (CustomVersion < FFortniteMainBranchObjectVersion::ChangeDefaultAlphaBlendType)
+		{
+			// Switch the default back to Linear so old data remains the same
+			BlendOption = EAlphaBlendOption::Linear;
+		}
+	}
+	return false;
+}
+
+bool FAlphaBlend::Serialize(FArchive& Ar)
+{
+	Ar.UsingCustomVersion(FFortniteMainBranchObjectVersion::GUID);
+	if (Ar.IsLoading() && Ar.IsSerializingDefaults())
+	{
+		const int32 CustomVersion = Ar.CustomVer(FFortniteMainBranchObjectVersion::GUID);
+		if (CustomVersion < FFortniteMainBranchObjectVersion::ChangeDefaultAlphaBlendType)
+		{
+			// Switch the default back to Linear so old data remains the same
+			BlendOption = EAlphaBlendOption::Linear;
+		}
+	}
+
+	return false;
+}
+
+bool FAlphaBlend::Serialize(FStructuredArchive::FSlot Slot)
+{
+	FArchive& Ar = Slot.GetUnderlyingArchive();
+	Ar.UsingCustomVersion(FFortniteMainBranchObjectVersion::GUID);
+
+	if (Ar.IsLoading() && Ar.IsSerializingDefaults())
+	{
+		const int32 CustomVersion = Ar.CustomVer(FFortniteMainBranchObjectVersion::GUID);
+		if (CustomVersion < FFortniteMainBranchObjectVersion::ChangeDefaultAlphaBlendType)
+		{
+			// Switch the default back to Linear so old data remains the same
+			BlendOption = EAlphaBlendOption::Linear;
+		}
+	}
+	return false;
+}
+
+#endif // WITH_EDITORONLY_DATA
+
 FAlphaBlend::FAlphaBlend(const FAlphaBlendArgs& InArgs)
 	: CustomCurve(InArgs.CustomCurve)
 	, BlendTime(InArgs.BlendTime)
@@ -232,7 +302,6 @@ bool FAlphaBlend::IsComplete() const
 FAlphaBlendArgs::FAlphaBlendArgs()
 	: CustomCurve(nullptr)
 	, BlendTime(0.2f)
-	, BlendOption(EAlphaBlendOption::Linear)
 {
 
 }
@@ -240,7 +309,6 @@ FAlphaBlendArgs::FAlphaBlendArgs()
 FAlphaBlendArgs::FAlphaBlendArgs(float InBlendTime)
 	: CustomCurve(nullptr)
 	, BlendTime(InBlendTime)
-	, BlendOption(EAlphaBlendOption::Linear)
 {
 
 }
diff --git a/Engine/Source/Runtime/Engine/Private/Animation/AnimMontage.cpp b/Engine/Source/Runtime/Engine/Private/Animation/AnimMontage.cpp
--- a/Engine/Source/Runtime/Engine/Private/Animation/AnimMontage.cpp
+++ b/Engine/Source/Runtime/Engine/Private/Animation/AnimMontage.cpp
@@ -23,6 +23,7 @@
 #include "Animation/ActiveMontageInstanceScope.h"
 #include "Animation/AnimationSettings.h"
 #include "Logging/StructuredLog.h"
+#include "UObject/FortniteMainBranchObjectVersion.h"
 
 #include UE_INLINE_GENERATED_CPP_BY_NAME(AnimMontage)
 
@@ -114,6 +115,26 @@ UAnimSequenceBase* UAnimMontage::GetFirstAnimReference() const
 	return nullptr;
 }
 
+#if WITH_EDITORONLY_DATA
+void UAnimMontage::Serialize(FArchive& Ar)
+{
+	Ar.UsingCustomVersion(FFortniteMainBranchObjectVersion::GUID);
+	if (Ar.IsLoading())
+	{
+		const int32 CustomVersion = Ar.CustomVer(FFortniteMainBranchObjectVersion::GUID);
+		if (CustomVersion < FFortniteMainBranchObjectVersion::ChangeDefaultAlphaBlendType)
+		{
+			// Switch the default back to Linear so old data remains the same
+			// Note this happens before serialization
+			BlendIn.SetBlendOption(EAlphaBlendOption::Linear);
+			BlendOut.SetBlendOption(EAlphaBlendOption::Linear);
+		}
+	}
+	
+	Super::Serialize(Ar);
+}
+#endif // WITH_EDITORONLY_DATA
+
 const FAnimTrack* UAnimMontage::GetAnimationData(FName InSlotName) const
 {
 	for (int32 I=0; I<SlotAnimTracks.Num(); ++I)
diff --git a/Engine/Source/Runtime/Engine/Private/Animation/InputScaleBias.cpp b/Engine/Source/Runtime/Engine/Private/Animation/InputScaleBias.cpp
--- a/Engine/Source/Runtime/Engine/Private/Animation/InputScaleBias.cpp
+++ b/Engine/Source/Runtime/Engine/Private/Animation/InputScaleBias.cpp
@@ -1,6 +1,7 @@
 // Copyright Epic Games, Inc. All Rights Reserved.
 
 #include "Animation/InputScaleBias.h"
+#include "UObject/FortniteMainBranchObjectVersion.h"
 
 #include UE_INLINE_GENERATED_CPP_BY_NAME(InputScaleBias)
 
@@ -270,6 +271,41 @@ float FInputScaleBiasClampState::ApplyTo(const FInputScaleBiasClampConstants& In
 	return Result;
 }
 
+#if WITH_EDITORONLY_DATA
+bool FInputAlphaBoolBlend::Serialize(FArchive& Ar)
+{
+	Ar.UsingCustomVersion(FFortniteMainBranchObjectVersion::GUID);
+	if (Ar.IsLoading() && Ar.IsSerializingDefaults())
+	{
+		const int32 CustomVersion = Ar.CustomVer(FFortniteMainBranchObjectVersion::GUID);
+		if (CustomVersion < FFortniteMainBranchObjectVersion::ChangeDefaultAlphaBlendType)
+		{
+			// Switch the default back to Linear so old data remains the same
+			BlendOption = EAlphaBlendOption::Linear;
+		}
+	}
+
+	return false;
+}
+
+bool FInputAlphaBoolBlend::Serialize(FStructuredArchive::FSlot Slot)
+{
+	FArchive& Ar = Slot.GetUnderlyingArchive();
+	Ar.UsingCustomVersion(FFortniteMainBranchObjectVersion::GUID);
+
+	if (Ar.IsLoading() && Ar.IsSerializingDefaults())
+	{
+		const int32 CustomVersion = Ar.CustomVer(FFortniteMainBranchObjectVersion::GUID);
+		if (CustomVersion < FFortniteMainBranchObjectVersion::ChangeDefaultAlphaBlendType)
+		{
+			// Switch the default back to Linear so old data remains the same
+			BlendOption = EAlphaBlendOption::Linear;
+		}
+	}
+	return false;
+}
+#endif // #if WITH_EDITORONLY_DATA
+
 #if WITH_EDITOR
 FText FInputScaleBiasClampConstants::GetFriendlyName(FText InFriendlyName) const
 {
diff --git a/Engine/Source/Runtime/Engine/Public/AlphaBlend.h b/Engine/Source/Runtime/Engine/Public/AlphaBlend.h
--- a/Engine/Source/Runtime/Engine/Public/AlphaBlend.h
+++ b/Engine/Source/Runtime/Engine/Public/AlphaBlend.h
@@ -28,6 +28,11 @@ enum class EAlphaBlendOption : uint8
 	Custom UMETA(Grouping = Custom, DisplayName = "Custom", ToolTip = "Custom interpolation, will use custom curve inside an FAlphaBlend or linear if none has been set"),
 };
 
+namespace UE::Anim
+{
+static constexpr EAlphaBlendOption DefaultBlendOption = EAlphaBlendOption::HermiteCubic;
+}
+
 /**
  * Alpha Blend construction arguments. Used for creation of an AlphaBlend.
  */
@@ -40,6 +45,11 @@ struct FAlphaBlendArgs
 	ENGINE_API FAlphaBlendArgs(float InBlendTime);
 	ENGINE_API FAlphaBlendArgs(const struct FAlphaBlend& InAlphaBlend);
 
+#if WITH_EDITORONLY_DATA
+	bool Serialize(FArchive& Ar);
+	bool Serialize(FStructuredArchive::FSlot Slot);
+#endif // WITH_EDITORONLY_DATA 
+
 	/** If you're using Custom BlendOption, you can specify curve */
 	UPROPERTY(EditAnywhere, BlueprintReadWrite, Category = "Blend", meta=(DisplayAfter="BlendOption"))
 	TObjectPtr<UCurveFloat> CustomCurve;
@@ -77,6 +87,9 @@ public:
 	/* Constructor */
 	ENGINE_API FAlphaBlend(const FAlphaBlend& Other, float NewBlendTime);
 
+	bool Serialize(FArchive& Ar);
+	bool Serialize(FStructuredArchive::FSlot Slot);
+
 	/* Constructor */
 	ENGINE_API explicit FAlphaBlend(const FAlphaBlendArgs& InArgs);
 
@@ -203,7 +216,7 @@ private:
 
 	/** Type of blending used (Linear, Cubic, etc.) */
 	UPROPERTY(EditAnywhere, Category = "Blend")
-	EAlphaBlendOption BlendOption;
+	EAlphaBlendOption BlendOption = UE::Anim::DefaultBlendOption;
 
 	/** internal flag to reset the alpha value */
 	bool bNeedsToResetAlpha;
@@ -216,3 +229,24 @@ private:
 	/** Cached Desired Value with Alpha 1 so that we can check if reached or not */
 	mutable float CachedDesiredBlendedValue;
 };
+
+#if WITH_EDITORONLY_DATA
+template<> struct TStructOpsTypeTraits<FAlphaBlendArgs> : public TStructOpsTypeTraitsBase2<FAlphaBlendArgs>
+{
+	enum
+	{
+		WithSerializer = true,
+		WithStructuredSerializer = true,
+	};
+};
+
+
+template<> struct TStructOpsTypeTraits<FAlphaBlend> : public TStructOpsTypeTraitsBase2<FAlphaBlend>
+{
+	enum
+	{
+		WithSerializer = true,
+		WithStructuredSerializer = true,
+	};
+};
+#endif //#if WITH_EDITORONLY_DATA
\ No newline at end of file
