Don't replicate updates to the properties in ANoPawnPlayerController.

diff --git a/Engine/Source/Runtime/Engine/Private/PlayerController.cpp b/Engine/Source/Runtime/Engine/Private/PlayerController.cpp
--- a/Engine/Source/Runtime/Engine/Private/PlayerController.cpp
+++ b/Engine/Source/Runtime/Engine/Private/PlayerController.cpp
@@ -5050,6 +5050,11 @@ void APlayerController::PostMigrate(const FUObjectMigrationContext& MigrationCon
 			{
 				ensure(NetConnection->PlayerController);
 				ensure(NetConnection->PlayerController->IsA(ANoPawnPlayerController::StaticClass()));
+				
+				// The ANoPawnPlayerController actor has a low NetUpdateFrequency, so ForceNetUpdate() must be called because
+				// it's going to be detached from a connection.
+				NetConnection->PlayerController->ForceNetUpdate();
+				
 				const int32 PrevNetPlayerIndex = NetConnection->PlayerController->NetPlayerIndex;
 				const int32 PrevHandshakeTransactionId = NetConnection->PlayerController->GetClientHandshakeId();
 				NetConnection->PlayerController = this;
@@ -5103,6 +5108,10 @@ void APlayerController::PostMigrate(const FUObjectMigrationContext& MigrationCon
 						if (OtherPlayerController->IsA(ANoPawnPlayerController::StaticClass()) &&
 							OtherPlayerController->CachedConnectionPlayerId == CachedConnectionPlayerId)
 						{
+							// The ANoPawnPlayerController actor has a low NetUpdateFrequency, so ForceNetUpdate() must be called because
+							// it's going to be attached to a connection.
+							OtherPlayerController->ForceNetUpdate();
+
 							PrevNetConnection->PlayerController = OtherPlayerController;
 							PrevNetConnection->OwningActor = OtherPlayerController;
 						}
@@ -6636,6 +6645,14 @@ ANoPawnPlayerController::ANoPawnPlayerController(const FObjectInitializer& Objec
 {
 	bReplicates = true;
 	SetRole(ROLE_Authority);
+
+	// This controller's updates don't need to be replicated to the client very frequently because
+	// it's only intended to be used by the server for replication, relevancy and establishing the
+	// initial connection on the client.
+	//
+	// This low update frequency means that ForceNetUpdate() must be called when any explicit changes
+	// are made to the controller, including being attached or detached from a UNetConnection.
+	SetNetUpdateFrequency(0.1f);
 }
 
 void ANoPawnPlayerController::PostInitializeComponents()
