[shaders] if preprocessing fails, we also need to remove the failed jobs from the pending submit list when finishing them. prevents cycles from being added to the pending list when retries are enabled.

diff --git a/Engine/Source/Runtime/Engine/Private/ShaderCompiler/ShaderCompilerJobCache.cpp b/Engine/Source/Runtime/Engine/Private/ShaderCompiler/ShaderCompilerJobCache.cpp
--- a/Engine/Source/Runtime/Engine/Private/ShaderCompiler/ShaderCompilerJobCache.cpp
+++ b/Engine/Source/Runtime/Engine/Private/ShaderCompiler/ShaderCompilerJobCache.cpp
@@ -1547,6 +1547,9 @@ void FShaderJobCache::SubmitJobs(const TArray<FShaderCommonCompileJobPtr>& InJob
 					}
 					else // if preprocessing ran and failed, finish the job immediately
 					{
+						JobLock.WriteLock();
+						Unlink(*Job); // from PendingSubmitJobTaskJobs
+						JobLock.WriteUnlock();
 						ProcessFinishedJob(Job, EShaderCompileJobStatus::Skipped);
 					}
 
