Implementing SlideAlongNavMesh option for CMC NavWalking mode. This allows pawns in NavWalking mode to move along a navmesh rather than just moving towards a projected point on the navmesh.

diff --git a/Engine/Source/Runtime/AIModule/Classes/Navigation/PathFollowingComponent.h b/Engine/Source/Runtime/AIModule/Classes/Navigation/PathFollowingComponent.h
--- a/Engine/Source/Runtime/AIModule/Classes/Navigation/PathFollowingComponent.h
+++ b/Engine/Source/Runtime/AIModule/Classes/Navigation/PathFollowingComponent.h
@@ -402,6 +402,7 @@ class UPathFollowingComponent : public UActorComponent, public IAIResourceInterf
 	AIMODULE_API virtual void OnUnableToMove(const UObject& Instigator) override;
 	AIMODULE_API virtual void OnStartedFalling() override;
 	virtual void OnLanded() override {}
+	AIMODULE_API virtual bool IsFollowingNavLink() const override;
 	//~ IPathFollowingAgentInterface end
 
 	/** Check if path following can be activated */
diff --git a/Engine/Source/Runtime/AIModule/Private/Navigation/PathFollowingComponent.cpp b/Engine/Source/Runtime/AIModule/Private/Navigation/PathFollowingComponent.cpp
--- a/Engine/Source/Runtime/AIModule/Private/Navigation/PathFollowingComponent.cpp
+++ b/Engine/Source/Runtime/AIModule/Private/Navigation/PathFollowingComponent.cpp
@@ -1912,6 +1912,11 @@ FString UPathFollowingComponent::GetDebugString() const
 	return Desc;
 }
 
+bool UPathFollowingComponent::IsFollowingNavLink() const
+{
+	return IsCurrentSegmentNavigationLink();
+}
+
 bool UPathFollowingComponent::IsPathFollowingAllowed() const
 {
 	return NavMovementInterface.IsValid() && NavMovementInterface->CanStartPathFollowing();
diff --git a/Engine/Source/Runtime/Engine/Classes/AI/Navigation/PathFollowingAgentInterface.h b/Engine/Source/Runtime/Engine/Classes/AI/Navigation/PathFollowingAgentInterface.h
--- a/Engine/Source/Runtime/Engine/Classes/AI/Navigation/PathFollowingAgentInterface.h
+++ b/Engine/Source/Runtime/Engine/Classes/AI/Navigation/PathFollowingAgentInterface.h
@@ -32,4 +32,10 @@ class IPathFollowingAgentInterface
 
 	/** Called when movement is blocked by a collision with another actor.  */
 	virtual void OnMoveBlockedBy(const FHitResult& BlockingImpact) {}
+
+	/** Returns true if the agent is currently following a NavLink on their path */
+	virtual bool IsFollowingNavLink() const
+	{
+		return false;
+	} 
 };
diff --git a/Engine/Source/Runtime/Engine/Classes/GameFramework/CharacterMovementComponent.h b/Engine/Source/Runtime/Engine/Classes/GameFramework/CharacterMovementComponent.h
--- a/Engine/Source/Runtime/Engine/Classes/GameFramework/CharacterMovementComponent.h
+++ b/Engine/Source/Runtime/Engine/Classes/GameFramework/CharacterMovementComponent.h
@@ -1086,7 +1086,11 @@ protected:
 
 	/** Whether to raycast to underlying geometry to better conform navmesh-walking characters */
 	UPROPERTY(Category="Character Movement: NavMesh Movement", EditAnywhere, BlueprintReadOnly)
-	uint8 bProjectNavMeshWalking : 1;
+	uint8 bProjectNavMeshWalking : 1 = false;
+
+	/** Should we attempt to slide along a navmesh instead of projecting a point on the navmesh and moving to that */
+	UPROPERTY(Category="Character Movement: NavMesh Movement", EditAnywhere, BlueprintReadOnly)
+	uint8 bSlideAlongNavMeshEdge : 1 = false;
 
 	/** Use both WorldStatic and WorldDynamic channels for NavWalking geometry conforming */
 	UPROPERTY(Category = "Character Movement: NavMesh Movement", EditAnywhere, BlueprintReadOnly, AdvancedDisplay)
diff --git a/Engine/Source/Runtime/Engine/Private/Components/CharacterMovementComponent.cpp b/Engine/Source/Runtime/Engine/Private/Components/CharacterMovementComponent.cpp
--- a/Engine/Source/Runtime/Engine/Private/Components/CharacterMovementComponent.cpp
+++ b/Engine/Source/Runtime/Engine/Private/Components/CharacterMovementComponent.cpp
@@ -5883,9 +5883,51 @@ void UCharacterMovementComponent::PhysNavWalking(float deltaTime, int32 Iteratio
 			SetGravitySpaceZ(AdjustedDest, GetGravitySpaceZ(CachedNavLocation.Location));
 		}
 
-		// Find the point on the NavMesh
-		const bool bHasNavigationData = FindNavFloor(AdjustedDest, DestNavLocation);
-		if (!bHasNavigationData)
+		bool bFoundPointOnNavMesh = false;
+		if (bSlideAlongNavMeshEdge)
+		{
+			if (const INavigationDataInterface* NavDataInterface = GetNavData())
+			{
+				const IPathFollowingAgentInterface* PathFollowingAgent = GetPathFollowingAgent();
+				const bool bIsOnNavLink = PathFollowingAgent && PathFollowingAgent->IsFollowingNavLink();
+
+				if (!bIsOnNavLink)
+				{
+					FNavLocation StartingNavFloorLocation;
+					bool bHasValidCachedNavLocation = NavDataInterface->IsNodeRefValid(CachedNavLocation.NodeRef);
+
+					// If we don't have a valid CachedNavLocation lets try finding the NavFloor where we're currently at and use that
+					if (!bHasValidCachedNavLocation)
+					{
+						bHasValidCachedNavLocation = FindNavFloor(OldLocation, OUT StartingNavFloorLocation);
+					}
+					else
+					{
+						StartingNavFloorLocation = CachedNavLocation;
+					}
+
+					if (bHasValidCachedNavLocation)
+					{
+						bFoundPointOnNavMesh = NavDataInterface->FindMoveAlongSurface(StartingNavFloorLocation, AdjustedDest, OUT DestNavLocation);
+
+						if (bFoundPointOnNavMesh)
+						{
+							AdjustedDest = ProjectToGravityFloor(DestNavLocation.Location) + GetGravitySpaceComponentZ(AdjustedDest);
+						}
+					}
+				}
+				else
+				{
+					bFoundPointOnNavMesh = FindNavFloor(AdjustedDest, DestNavLocation);
+				}
+			}
+		}
+		else
+		{
+			bFoundPointOnNavMesh = FindNavFloor(AdjustedDest, DestNavLocation);
+		}
+		
+		if (!bFoundPointOnNavMesh)
 		{
 			SetMovementMode(MOVE_Walking);
 			return;
