NetDriver
* Fixed ping updates discarding higher ping values when multiple Acks were received in one packet.
* For exemple this caused ping values to stay low during BufferFloat occurences
* Added RawPing metric that tracks the highgest ping value received every frame.
* RawPing will not change if no packets are received or the incoming packets contain no new Acks.

diff --git a/Engine/Config/BaseEngine.ini b/Engine/Config/BaseEngine.ini
--- a/Engine/Config/BaseEngine.ini
+++ b/Engine/Config/BaseEngine.ini
@@ -3837,6 +3837,7 @@ IntervalSeconds=1
 +Listeners=(MetricName=PingBucketInt7, Class=/Script/Engine.NetworkMetricsPerfCounters)
 +Listeners=(MetricName=ClosedConnectionsDueToReliableBufferOverflow, Class=/Script/Engine.NetworkMetricsPerfCounters)
 +Listeners=(MetricName=AddedConnections, Class=/Script/Engine.NetworkMetricsPerfCounters)
++Listeners=(MetricName=RawPing, Class=/Script/Engine.NetworkMetricsCSV_Replication, EnableMode=ENetworkMetricEnableMode::ClientOnly)
 
 ; Studio Telemetry Settings
 [StudioTelemetry.Config] 
diff --git a/Engine/Source/Runtime/Engine/Classes/Engine/NetConnection.h b/Engine/Source/Runtime/Engine/Classes/Engine/NetConnection.h
--- a/Engine/Source/Runtime/Engine/Classes/Engine/NetConnection.h
+++ b/Engine/Source/Runtime/Engine/Classes/Engine/NetConnection.h
@@ -623,13 +623,17 @@ public:
 
 	// Packet.
 	FBitWriter		SendBuffer;						// Queued up bits waiting to send
-	double			OutLagTime[256];				// For lag measuring.
-	int32			OutLagPacketId[256];			// For lag measuring.
+	double			OutLagTime[256];				// The sent time of a packetId to calculate RTT from
+	int32			OutLagPacketId[256];			// Keep track of the packedId we stored the sent time for
 	uint8			OutBytesPerSecondHistory[256];	// For saturation measuring.
 	int32			InPacketId;						// Full incoming packet index.
 	int32			OutPacketId;					// Most recently sent packet.
-	int32 			OutAckPacketId;					// Most recently acked outgoing packet.
-	int32			OutFirstAckPacketId = -1;		// The id of the first ACK we processed each frame.
+	int32 			OutAckPacketId;					// The id of the last ACK included in the current packet
+
+private:
+	int32			OutFirstAckPacketId = -1;		// The id of the first ACK included in the current packet
+
+public:
 
 	bool			bLastHasServerFrameTime;
 
diff --git a/Engine/Source/Runtime/Engine/Classes/Engine/NetDriver.h b/Engine/Source/Runtime/Engine/Classes/Engine/NetDriver.h
--- a/Engine/Source/Runtime/Engine/Classes/Engine/NetDriver.h
+++ b/Engine/Source/Runtime/Engine/Classes/Engine/NetDriver.h
@@ -1078,7 +1078,7 @@ private:
 	TMap<FName, TObjectPtr<UNetworkMetricsBaseListener>> NetworkMetricsListeners;
 
 	/** Register metric listeners provided by the *.ini file. */
-	void SetupNetworkMetricsListeners();
+	void SetupNetworkMetricsListeners(bool bIsClient);
 
 	/** Create an instance of UNetworkMetricsStats that is associated with a given Stat and cached with other listeners in NetworkMetricsListeners. */
 	void RegisterStatsListener(const FName MetricName, const FName StatName);
@@ -2150,6 +2150,12 @@ protected:
 	/** Update stats related to networking. */
 	void UpdateNetworkStats();
 
+private:
+	/** Reset all the per-frame net connection stats */
+	void ResetNetConnectionStats();
+
+protected:
+
 #if WITH_SERVER_CODE
 	/**
 	* Helper functions for ServerReplicateActors
diff --git a/Engine/Source/Runtime/Engine/Private/Net/NetworkMetricsDefs.cpp b/Engine/Source/Runtime/Engine/Private/Net/NetworkMetricsDefs.cpp
--- a/Engine/Source/Runtime/Engine/Private/Net/NetworkMetricsDefs.cpp
+++ b/Engine/Source/Runtime/Engine/Private/Net/NetworkMetricsDefs.cpp
@@ -100,6 +100,7 @@ const FName PingBucketInt6("PingBucketInt6");
 const FName PingBucketInt7("PingBucketInt7");
 const FName OutgoingReliableMessageQueueMaxSize("OutgoingReliableMessageQueueMaxSize");
 const FName IncomingReliableMessageQueueMaxSize("IncomingReliableMessageQueueMaxSize");
+const FName RawPing("RawPing");
 
 }
 
diff --git a/Engine/Source/Runtime/Engine/Private/NetConnection.cpp b/Engine/Source/Runtime/Engine/Private/NetConnection.cpp
--- a/Engine/Source/Runtime/Engine/Private/NetConnection.cpp
+++ b/Engine/Source/Runtime/Engine/Private/NetConnection.cpp
@@ -2640,10 +2640,16 @@ void UNetConnection::ReceivedAck(int32 AckPacketId, FChannelsToClose& OutChannel
 	// Advance OutAckPacketId
 	OutAckPacketId = AckPacketId;
 
-	// Store the first Ack we received each frame
+	// Store the first Ack we received in the packet
 	if (OutFirstAckPacketId == -1)
 	{
 		OutFirstAckPacketId = AckPacketId;
+
+		// On the first ACK of the first packet of a frame, reset the RawPing
+		if (InPacketsThisFrame == 1)
+		{
+			RawPingInSeconds = 0.0;
+		}
 	}
 
 	// Process the bunch.
@@ -2912,6 +2918,13 @@ bool UNetConnection::ReadPacketInfo(FBitReader& Reader, bool bHasPacketInfoPaylo
 		return false;
 	}
 
+	if (OutFirstAckPacketId == -1)
+	{
+		// This means we received a packet without new ACKs in it.
+		// Continue reading it but we don't need to update the Ping.
+		return true;
+	}
+
 	// Find the received time of the packet
 	double PacketReceiveTime = 0.0;
 	FTimespan& RecvTimespan = LastOSReceiveTime.Timestamp;
@@ -2933,59 +2946,47 @@ bool UNetConnection::ReadPacketInfo(FBitReader& Reader, bool bHasPacketInfoPaylo
 
 	// Update ping
 
-	// Find the ping of the first packet we acked (highest ping)
-	if (OutFirstAckPacketId != -1)
-	{
-		const int32 FirstAckIndex = OutFirstAckPacketId & (UE_ARRAY_COUNT(OutLagPacketId) - 1);
+	const bool bExcludeFrameTime = !!CVarPingExcludeFrameTime.GetValueOnAnyThread();
 
-		if (OutLagPacketId[FirstAckIndex] == OutFirstAckPacketId)
+	// Go trought all the ACKs we received
+	for (int32 AckId = OutFirstAckPacketId; AckId <= OutAckPacketId; ++AckId)
+	{
+		// Map the Id to the Index(0-255) storing the sent time 
+		const int32 AckIndex = AckId & (UE_ARRAY_COUNT(OutLagPacketId) - 1);
+		
+		if (OutLagPacketId[AckIndex] != AckId)
 		{
-			if (OutFirstAckPacketId != OutAckPacketId)
-			{
-				OutLagPacketId[FirstAckIndex] = -1;
-			}
-
-			const double PacketSendTime = OutLagTime[FirstAckIndex];
-			const double RTT = CurrentTime - PacketSendTime;
-			const double RTTExclFrame = FMath::Max(RTT - ServerFrameTime, 0.0);
-
-			RawPingInSeconds = RTTExclFrame;
+			// The Ids were already processed, nothing more to do here
+			return true;
 		}
 
-		OutFirstAckPacketId = -1;
-	}
-
-	// Find the ping of the last packet we acked (lowest ping)
-	const int32 LastAckIndex =  OutAckPacketId & (UE_ARRAY_COUNT(OutLagPacketId) - 1);
-
-	if (OutLagPacketId[LastAckIndex] == OutAckPacketId)
-	{
-		OutLagPacketId[LastAckIndex] = -1;		// Only use the ack once
+		// Clear the id to tell it was processed.
+		OutLagPacketId[AckIndex] = -1;
 
-		const bool bExcludeFrameTime = !!CVarPingExcludeFrameTime.GetValueOnAnyThread();
-		const double PacketSendTime = OutLagTime[LastAckIndex];
+		const double PacketSendTime = OutLagTime[AckIndex];
 		const double RTT			= (CurrentTime - PacketSendTime);
 		const double RTTExclFrame	= RTT - (bExcludeFrameTime ? ServerFrameTime : 0.0);
-		const double NewLag			= FMath::Max(RTTExclFrame, 0.0);
+		const double PacketPing		= FMath::Max(RTTExclFrame, 0.0);
 
-		//UE_LOG( LogNet, Warning, TEXT( "Out: %i, InRemote: %i, Saturation: %f" ), OutBytesPerSecondHistory[Index], RemoteInKBytesPerSecond, RemoteSaturation );
-
-		LagAcc += NewLag;
+		LagAcc += PacketPing;
 		LagCount++;
 
+		// Keep the highest ping as the RawPing value for this frame
+		RawPingInSeconds = FMath::Max(RawPingInSeconds, PacketPing);
+
 		if (PlayerController)
 		{
 			// PlayerController is valid and local (as opposed to currently loaned out)
 			UE_CLOG(PlayerController.IsRemote(), LogNet, Warning, TEXT("%hs: PlayerController is unexpectedly a RemoteObject"), __func__);
 			if (!PlayerController.IsRemote())
 			{
-				PlayerController->UpdatePing(NewLag);
+				PlayerController->UpdatePing(PacketPing);
 			}
 		}
 
 		if (NetPing.IsValid())
 		{
-			NetPing->UpdatePing(EPingType::RoundTrip, CurrentTime, NewLag);
+			NetPing->UpdatePing(EPingType::RoundTrip, CurrentTime, PacketPing);
 			NetPing->UpdatePing(EPingType::RoundTripExclFrame, CurrentTime, FMath::Max(RTT, 0.0));
 		}
 
@@ -2995,6 +2996,8 @@ bool UNetConnection::ReadPacketInfo(FBitReader& Reader, bool bHasPacketInfoPaylo
 		}
 	}
 
+	OutFirstAckPacketId = -1;
+
 	return true;
 }
 
@@ -3195,6 +3198,7 @@ void UNetConnection::ReceivedPacket( FBitReader& Reader, bool bIsReinjectedPacke
 
 		const int32 PacketSequenceDelta = PacketNotify.GetSequenceDelta(Header);
 
+		// If the packetID is higher then the last one
 		if (PacketSequenceDelta > 0)
 		{
 			const bool bPacketOrderCacheActive = !bFlushingPacketOrderCache && PacketOrderCache.IsSet();
@@ -3204,7 +3208,7 @@ void UNetConnection::ReceivedPacket( FBitReader& Reader, bool bIsReinjectedPacke
 
 			const int32 MissingPacketCount = PacketSequenceDelta - 1;
 
-			// Cache the packet if we are already caching, and begin caching if we just encountered a missing sequence, within range
+			// If there is a gap in the Id sequence, cache the packet so they can be processed in the right order
 			if (bFillingPacketOrderCache || (bCheckForMissingSequence && MissingPacketCount > 0 && MissingPacketCount <= MaxMissingPackets))
 			{
 				int32 LinearCacheIdx = PacketSequenceDelta - 1;
@@ -3300,6 +3304,7 @@ void UNetConnection::ReceivedPacket( FBitReader& Reader, bool bIsReinjectedPacke
 
 			PacketAnalytics.TrackInPacket(InPacketId, MissingPacketCount);
 		}
+		// PacketID was earlier then a previously read packet, drop it entirely because it is now obsolete
 		else
 		{
 			TotalOutOfOrderPacketsLost++;
diff --git a/Engine/Source/Runtime/Engine/Private/NetDriver.cpp b/Engine/Source/Runtime/Engine/Private/NetDriver.cpp
--- a/Engine/Source/Runtime/Engine/Private/NetDriver.cpp
+++ b/Engine/Source/Runtime/Engine/Private/NetDriver.cpp
@@ -1150,6 +1150,8 @@ void UNetDriver::TickFlush(float DeltaSeconds)
 
 	UpdateNetworkStats();
 
+	ResetNetConnectionStats();
+
 	// Send the current values of metrics to all of the metrics listeners.
 	GetMetrics()->ProcessListeners();
 
@@ -1722,7 +1724,7 @@ bool UNetDriver::InitBase(bool bInitAsClient, FNetworkNotify* InNotify, const FU
 
 	if (ShouldRegisterMetricsDatabaseListeners())
 	{
-		SetupNetworkMetricsListeners();
+		SetupNetworkMetricsListeners(bInitAsClient);
 	}
 
 	return bSuccess;
@@ -1955,18 +1957,57 @@ void UNetDriver::SetupNetworkMetrics()
 	GetMetrics()->CreateInt(UE::Net::Metric::ImportedNetGuids, 0);
 	GetMetrics()->CreateInt(UE::Net::Metric::PendingOuterNetGuids, 0);
 	GetMetrics()->CreateInt(UE::Net::Metric::NetInBunchTimeOvershootPercent, 0);
+
+	// (Client only) The highest maximum ping processed every frame. If no pings are processed in the frame it keeps the last value set.
+	GetMetrics()->CreateInt(UE::Net::Metric::RawPing, 0);
 }
 
-void UNetDriver::SetupNetworkMetricsListeners()
+void UNetDriver::SetupNetworkMetricsListeners(bool bIsClient)
 {
 	// Register metrics in the ini configuration file with a listener.
 	const UNetworkMetricsConfig* Config = GetDefault<UNetworkMetricsConfig>();
 
+	const bool bIrisEnabled = IsUsingIrisReplication();
+
 	for (const FNetworkMetricConfig& MetricConfig : Config->Listeners)
 	{
 		// Allow config file to selectively enable/disable listeners
-		if ((MetricConfig.EnableMode == ENetworkMetricEnableMode::EnableForIrisOnly && !IsUsingIrisReplication())
-			|| (MetricConfig.EnableMode == ENetworkMetricEnableMode::EnableForNonIrisOnly && IsUsingIrisReplication()))
+		bool bEnabled = true;
+
+		switch(MetricConfig.EnableMode)
+		{
+			case ENetworkMetricEnableMode::EnableForAllReplication:
+			{
+				bEnabled = true;
+			} break;
+
+			case ENetworkMetricEnableMode::EnableForIrisOnly:
+			{
+				bEnabled = bIrisEnabled;
+			} break;
+
+			case ENetworkMetricEnableMode::EnableForNonIrisOnly:
+			{
+				bEnabled = !bIsUsingIris;
+			} break;
+
+			case ENetworkMetricEnableMode::ServerOnly:
+			{
+				bEnabled = !bIsClient;
+			} break;
+
+			case ENetworkMetricEnableMode::ClientOnly:
+			{
+				bEnabled = bIsClient;
+			} break;
+
+			default:
+			{
+				ensureMsgf(false, TEXT("Missing implementation"));
+			} break;
+		}
+
+		if (!bEnabled)
 		{
 			continue;
 		}
@@ -8034,7 +8075,7 @@ void UNetDriver::UpdateNetworkStats()
 			int32 InBunchTimeOvershootPercent = 0;
 
 			// these need to be updated even if we are not collecting stats, since they get reported to analytics/QoS
-			for (UNetConnection* Client : ClientConnections)
+			for (const UNetConnection* Client : ClientConnections)
 			{
 				if (Client)
 				{
@@ -8059,9 +8100,6 @@ void UNetDriver::UpdateNetworkStats()
 					ClientsInPacketsThisFrameMax = FMath::Max(ClientsInPacketsThisFrameMax, Client->InPacketsThisFrame);
 					ClientsOutPacketsThisFrameMax = FMath::Max(ClientsOutPacketsThisFrameMax, Client->OutPacketsThisFrame);
 
-					Client->InPacketsThisFrame = 0;
-					Client->OutPacketsThisFrame = 0;
-
 					++NumClients;
 				}
 			}
@@ -8121,12 +8159,12 @@ void UNetDriver::UpdateNetworkStats()
 				GetMetrics()->SetInt(UE::Net::Metric::OutPacketsLost, OutPacketsLost);
 				GetMetrics()->SetInt(UE::Net::Metric::InPacketsLost, InPacketsLost);
 
-				if (ServerConnection != nullptr && ServerConnection->PlayerController != nullptr && ServerConnection->PlayerController->PlayerState != nullptr)
+				if (ServerConnection && ServerConnection->PlayerController && ServerConnection->PlayerController->PlayerState)
 				{
 					Ping = FMath::TruncToInt(ServerConnection->PlayerController->PlayerState->ExactPing);
 				}
 
-				if (ServerConnection != nullptr)
+				if (ServerConnection)
 				{
 					NumOpenChannels = ServerConnection->OpenChannels.Num();
 					GetMetrics()->SetFloat(UE::Net::Metric::AverageJitterInMS, ServerConnection->GetAverageJitterInMS());
@@ -8446,7 +8484,7 @@ void UNetDriver::UpdateNetworkStats()
 
 			int32 NumClients = 0;
 
-			for (UNetConnection* Client : ClientConnections)
+			for (const UNetConnection* Client : ClientConnections)
 			{
 				if (Client)
 				{
@@ -8456,9 +8494,6 @@ void UNetDriver::UpdateNetworkStats()
 					ClientsOutPacketsThisFrameAvg += Client->OutPacketsThisFrame;
 					ClientsInPacketsThisFrameMax = FMath::Max(ClientsInPacketsThisFrameMax, Client->InPacketsThisFrame);
 					ClientsOutPacketsThisFrameMax = FMath::Max(ClientsOutPacketsThisFrameMax, Client->OutPacketsThisFrame);
-
-					Client->InPacketsThisFrame = 0;
-					Client->OutPacketsThisFrame = 0;
 				}
 			}
 
@@ -8471,26 +8506,27 @@ void UNetDriver::UpdateNetworkStats()
 			if (ServerConnection)
 			{
 				const int32 RawPingValueInMS = (int32)(ServerConnection->RawPingInSeconds * 1000.f);
-				FCsvProfiler::RecordCustomStat(StatNameRawPing, CSV_CATEGORY_INDEX(Networking), RawPingValueInMS, ECsvCustomStatOp::Set);
+				GetMetrics()->SetInt(UE::Net::Metric::RawPing, RawPingValueInMS);
 			}
 
 			GetMetrics()->SetInt(UE::Net::Metric::InPacketsClientAvg, ClientsInPacketsThisFrameAvg);
 			GetMetrics()->SetInt(UE::Net::Metric::InPacketsClientMax, ClientsInPacketsThisFrameMax);
 			GetMetrics()->SetInt(UE::Net::Metric::OutPacketsClientAvg, ClientsOutPacketsThisFrameAvg);
 			GetMetrics()->SetInt(UE::Net::Metric::OutPacketsClientMax, ClientsOutPacketsThisFrameMax);
-#else // #if CSV_PROFILER_STATS
-			// Reset the per-frame stats here
-			for (UNetConnection* Client : ClientConnections)
-			{
-				if (Client)
-				{
-					Client->InPacketsThisFrame = 0;
-					Client->OutPacketsThisFrame = 0;
-				}
-			}
-#endif // #else CSV_PROFILER_STATS
+#endif  //#if CSV_PROFILER_STATS
 		}
 	} // bCollectNetStats ||(USE_SERVER_PERF_COUNTERS) || STATS
+
+}
+
+
+void UNetDriver::ResetNetConnectionStats()
+{
+	if (ServerConnection)
+	{
+		ServerConnection->InPacketsThisFrame = 0;
+		ServerConnection->OutPacketsThisFrame = 0;
+	}
 	else
 	{
 		for (UNetConnection* Client : ClientConnections)
diff --git a/Engine/Source/Runtime/Engine/Public/Net/NetworkMetricsConfig.h b/Engine/Source/Runtime/Engine/Public/Net/NetworkMetricsConfig.h
--- a/Engine/Source/Runtime/Engine/Public/Net/NetworkMetricsConfig.h
+++ b/Engine/Source/Runtime/Engine/Public/Net/NetworkMetricsConfig.h
@@ -9,7 +9,9 @@ enum class ENetworkMetricEnableMode : uint8
 {
 	EnableForAllReplication,
 	EnableForIrisOnly,
-	EnableForNonIrisOnly
+	EnableForNonIrisOnly,
+	ServerOnly,
+	ClientOnly,
 };
 
 USTRUCT()
diff --git a/Engine/Source/Runtime/Engine/Public/Net/NetworkMetricsDefs.h b/Engine/Source/Runtime/Engine/Public/Net/NetworkMetricsDefs.h
--- a/Engine/Source/Runtime/Engine/Public/Net/NetworkMetricsDefs.h
+++ b/Engine/Source/Runtime/Engine/Public/Net/NetworkMetricsDefs.h
@@ -102,5 +102,5 @@ extern ENGINE_API const FName PingBucketInt6;
 extern ENGINE_API const FName PingBucketInt7;
 extern ENGINE_API const FName OutgoingReliableMessageQueueMaxSize;
 extern ENGINE_API const FName IncomingReliableMessageQueueMaxSize;
-
+extern ENGINE_API const FName RawPing;
 }
