Fix Seek Restart Pitch Behavior

diff --git a/Engine/Source/Runtime/AudioMixer/Private/AudioMixerSource.cpp b/Engine/Source/Runtime/AudioMixer/Private/AudioMixerSource.cpp
--- a/Engine/Source/Runtime/AudioMixer/Private/AudioMixerSource.cpp
+++ b/Engine/Source/Runtime/AudioMixer/Private/AudioMixerSource.cpp
@@ -1333,6 +1333,8 @@ namespace Audio
 		check(WaveData);
 		const float ModPitchBase = ModulationUtils::GetRoutedPitch(*WaveInstance, *WaveData, *ActiveSound);
 		MixerSourceVoice->SetModPitch(ModPitchBase);
+
+		WaveInstance->PlaybackPercent = GetPlaybackPercent();
 	}
 
 	float FMixerSource::GetInheritedSubmixVolumeModulation() const
diff --git a/Engine/Source/Runtime/Engine/Private/AudioDevice.cpp b/Engine/Source/Runtime/Engine/Private/AudioDevice.cpp
--- a/Engine/Source/Runtime/Engine/Private/AudioDevice.cpp
+++ b/Engine/Source/Runtime/Engine/Private/AudioDevice.cpp
@@ -5455,7 +5455,7 @@ void FAudioDevice::AddNewActiveSoundInternal(const FActiveSound& InNewActiveSoun
 		// Attempt to seek the ActiveSound based on virtualization time
 		if (ActiveSound->GetVirtualizationMode() == EVirtualizationMode::SeekRestart)
 		{
-			ActiveSound->RequestedStartTime = ActiveSound->PlaybackTime;
+			ActiveSound->SeekRestartTimes = InVirtualLoopToRetrigger->GetEffectivePlaybackTimes();
 		}
 	}
 
diff --git a/Engine/Source/Runtime/Engine/Private/AudioVirtualLoop.cpp b/Engine/Source/Runtime/Engine/Private/AudioVirtualLoop.cpp
--- a/Engine/Source/Runtime/Engine/Private/AudioVirtualLoop.cpp
+++ b/Engine/Source/Runtime/Engine/Private/AudioVirtualLoop.cpp
@@ -71,6 +71,7 @@ UE_TRACE_EVENT_END()
 namespace AudioVirtualLoopPrivate
 {
 	static const FLazyName VolumeModulationParameterName("Volume");
+	static const FLazyName PitchModulationParameterName("Pitch");
 }
 
 FAudioVirtualLoop::FAudioVirtualLoop()
@@ -123,6 +124,36 @@ bool FAudioVirtualLoop::Virtualize(const FActiveSound& InActiveSound, FAudioDevi
 	OutVirtualLoop.ActiveSound = ActiveSound;
 	OutVirtualLoop.CalculateUpdateInterval();
 
+	if (InActiveSound.GetVirtualizationMode() == EVirtualizationMode::SeekRestart)
+	{
+		for (const TPair<UPTRINT, FWaveInstance*>& Pair : InActiveSound.GetWaveInstances())
+		{
+			if (const FWaveInstance* WaveInstance = Pair.Value)
+			{
+				FWaveInstanceData NewWaveInstanceData;
+				NewWaveInstanceData.EffectivePitch = WaveInstance->Pitch;
+
+				// Each Wave Instance can have its own assigned sound class, which can modify the pitch while the sound is virtualized.
+				// This has an impact on the effective playback time, so we have to track it for Seek Restart.
+				if (WaveInstance->SoundClass)
+				{
+					NewWaveInstanceData.SoundClass = WaveInstance->SoundClass;
+				}
+			
+				// Same with Pitch modulation
+				const FSoundModulationDestinationSettings& PitchModulationSettings = WaveInstance->GetEffectiveModulationSettings(EModulationDestination::Pitch);
+				if (!PitchModulationSettings.Modulators.IsEmpty())
+				{
+					NewWaveInstanceData.PitchModulationDestination.Init(InAudioDevice.DeviceID, AudioVirtualLoopPrivate::PitchModulationParameterName);
+					NewWaveInstanceData.PitchModulationDestination.UpdateModulators(PitchModulationSettings.Modulators);
+				}
+				
+				OutVirtualLoop.WaveInstanceData.Add(Pair.Key, NewWaveInstanceData);
+				OutVirtualLoop.EffectivePlaybackTimes.Add(Pair.Key, WaveInstance->PlaybackPercent * WaveInstance->WaveData->Duration);
+			}
+		}
+	}
+
 	// Once virtualized, Modulation stops getting updated at the Mixer Source level. Because of this, we need to grab
 	// Modulation data now so that Volume modulation can be used to determine if a sound is loud enough to realize.
 	// It's possible that the ActiveSound contains multiple WaveInstances.
@@ -170,6 +201,11 @@ float FAudioVirtualLoop::GetTimeVirtualized() const
 	return TimeVirtualized;
 }
 
+const TSortedMap<UPTRINT, float>& FAudioVirtualLoop::GetEffectivePlaybackTimes()
+{
+	return EffectivePlaybackTimes;
+}
+
 float FAudioVirtualLoop::GetUpdateInterval() const
 {
 	return UpdateInterval;
@@ -262,6 +298,32 @@ bool FAudioVirtualLoop::Update(float DeltaTime, bool bForceUpdate)
 	ActiveSound->PlaybackTime += DeltaTimePitchCorrected;
 	TimeVirtualized += DeltaTimePitchCorrected;
 
+	// Update playback time estimates for each wave instance for Seek Restart, factoring in pitch.
+	// This set only has elements if the ActiveSound was set to Seek Restart upon initial virtualization.
+	for (TPair<UPTRINT, float>& Pair : EffectivePlaybackTimes)
+	{
+		FWaveInstanceData& Data = WaveInstanceData[Pair.Key];
+		Data.EffectivePitch = ActiveSound->GetPitch();
+
+		if (USoundClass* SoundClass = Data.SoundClass)
+		{
+			// We must also walk through all the soundclass's parents, just in case
+			while (SoundClass)
+			{
+				Data.EffectivePitch *= SoundClass->Properties.Pitch;
+				SoundClass = SoundClass->ParentClass;
+			}
+		}
+		float PitchModulationValue = 1.0f;
+		if (Data.PitchModulationDestination.IsActive())
+		{
+			Data.PitchModulationDestination.ProcessControl(0.0f);
+			PitchModulationValue = Audio::GetFrequencyMultiplier(Data.PitchModulationDestination.GetValue());
+		}
+		
+		EffectivePlaybackTimes[Pair.Key] += DeltaTime * Data.EffectivePitch * PitchModulationValue;
+	}
+
 	const float UpdateDelta = TimeSinceLastUpdate + DeltaTime;
 	if (bForceUpdate)
 	{
diff --git a/Engine/Source/Runtime/Engine/Private/SoundWave.cpp b/Engine/Source/Runtime/Engine/Private/SoundWave.cpp
--- a/Engine/Source/Runtime/Engine/Private/SoundWave.cpp
+++ b/Engine/Source/Runtime/Engine/Private/SoundWave.cpp
@@ -3601,14 +3601,22 @@ void USoundWave::Parse(FAudioDevice* AudioDevice, const UPTRINT NodeWaveInstance
 	WaveInstance->AmbientZoneFilterFrequency = ParseParams.AmbientZoneFilterFrequency;
 	WaveInstance->bApplyRadioFilter = ActiveSound.bApplyRadioFilter;
 
-	// Procedural sounds do not have a duration that  makes sense to fmod with
-	if (!bProcedural && IsLooping())
+	WaveInstance->StartTime = ParseParams.StartTime;
+	// Override the start time for seek restart virtualization mode
+	if (ActiveSound.SeekRestartTimes.Contains(WaveInstance->WaveInstanceHash))
 	{
-		WaveInstance->StartTime = FMath::Fmod(ParseParams.StartTime, Duration);
-	}
-	else
-	{
-		WaveInstance->StartTime = ParseParams.StartTime;
+		float SeekRestartTime = ActiveSound.SeekRestartTimes[WaveInstance->WaveInstanceHash];
+		// For sounds with a defined duration, we wrap the seek time so that it sounds like it properly looped.
+		// Procedural sounds will get the unwrapped time, which they can use how they see fit.
+		const bool bWrapSeekTime = !bProcedural && (ParseParams.bLooping || IsLooping());
+		if (bWrapSeekTime)
+		{
+			WaveInstance->StartTime = FMath::Fmod(SeekRestartTime, Duration);
+		}
+		else
+		{
+			WaveInstance->StartTime = SeekRestartTime;
+		}
 	}
 
 	WaveInstance->UserIndex = ActiveSound.UserIndex;
diff --git a/Engine/Source/Runtime/Engine/Public/ActiveSound.h b/Engine/Source/Runtime/Engine/Public/ActiveSound.h
--- a/Engine/Source/Runtime/Engine/Public/ActiveSound.h
+++ b/Engine/Source/Runtime/Engine/Public/ActiveSound.h
@@ -585,6 +585,8 @@ public:
 	float MinCurrentPitch;
 	float RequestedStartTime;
 
+	TSortedMap<UPTRINT, float> SeekRestartTimes;
+	
 	float VolumeMultiplier;
 	float PitchMultiplier;
 
diff --git a/Engine/Source/Runtime/Engine/Public/Audio.h b/Engine/Source/Runtime/Engine/Public/Audio.h
--- a/Engine/Source/Runtime/Engine/Public/Audio.h
+++ b/Engine/Source/Runtime/Engine/Public/Audio.h
@@ -406,6 +406,9 @@ public:
 	/** The playback time of the wave instance. Updated from active sound. */
 	float PlaybackTime;
 
+	/** The linear percentage of the Wave which has been played based on its duration. Updated from Mixer Source. */
+	float PlaybackPercent;
+	
 	/** The output reverb send level to use for tje wave instance. */
 	float ReverbSendLevel;
 
diff --git a/Engine/Source/Runtime/Engine/Public/AudioVirtualLoop.h b/Engine/Source/Runtime/Engine/Public/AudioVirtualLoop.h
--- a/Engine/Source/Runtime/Engine/Public/AudioVirtualLoop.h
+++ b/Engine/Source/Runtime/Engine/Public/AudioVirtualLoop.h
@@ -24,6 +24,17 @@ private:
 	 */
 	TArray<Audio::FModulationDestination> VolumeConcurrencyDestinations;
 
+	struct FWaveInstanceData
+	{
+		float EffectivePitch;
+		Audio::FModulationDestination PitchModulationDestination;
+		USoundClass* SoundClass;
+		
+	};
+	
+	TSortedMap<UPTRINT, float> EffectivePlaybackTimes;
+	TSortedMap<UPTRINT, FWaveInstanceData> WaveInstanceData;
+
 	/**
 	  * Check if provided active sound is in audible range.
 	  */
@@ -59,6 +70,12 @@ public:
 	 */
 	ENGINE_API float GetTimeVirtualized() const;
 
+	/**
+	 * Returns the effective start time for each wave instance that was virtualized, factoring in pitch multipliers.
+	 * This map will be empty if the sound was not virtualized with Seek Restart.
+	 */
+	const TSortedMap<UPTRINT, float>& GetEffectivePlaybackTimes();
+	
 	/**
 	 * Returns the wait interval being observed before next update
 	 */
