New material translator: fixed bug in ValidateOperatorAndGetResultType.

#rnx

diff --git a/Engine/Source/Runtime/Engine/Private/Materials/MaterialIREmitter.cpp b/Engine/Source/Runtime/Engine/Private/Materials/MaterialIREmitter.cpp
--- a/Engine/Source/Runtime/Engine/Private/Materials/MaterialIREmitter.cpp
+++ b/Engine/Source/Runtime/Engine/Private/Materials/MaterialIREmitter.cpp
@@ -1821,9 +1821,10 @@ static FType ValidateOperatorAndGetResultType(FEmitter& Emitter, EOperator Op, F
 		
 		if (Filter & OPTF_CheckIsBoolean)
 		{
-			bValid &= Emitter.CheckIsArithmetic(Arguments[i]).IsValid(); 
+			bValid &= Emitter.CheckIsBoolean(Arguments[i]).IsValid(); 
 		}
-		else if (Filter & OPTF_CheckIsInteger)
+		
+		if (Filter & OPTF_CheckIsInteger)
 		{
 			bValid &= Emitter.CheckIsInteger(Arguments[i]).IsValid(); 
 		}
