Expose UAnimInstance::GetLODLevel to blueprint.

diff --git a/Engine/Source/Runtime/Engine/Classes/Animation/AnimInstance.h b/Engine/Source/Runtime/Engine/Classes/Animation/AnimInstance.h
--- a/Engine/Source/Runtime/Engine/Classes/Animation/AnimInstance.h
+++ b/Engine/Source/Runtime/Engine/Classes/Animation/AnimInstance.h
@@ -1427,6 +1427,7 @@ public:
 	 * This function is used by the anim graph to determine the LOD level at which to run.
 	 * @return the current LOD level
 	 */
+	UFUNCTION(BlueprintPure, Category = "Animation")
 	ENGINE_API virtual int32 GetLODLevel() const;
 
 public:
