EntityStreaming:
- Implement streaming interface flush to block until streaming is completed.
- Disable streaming on the server by loading the whole world during the first flush.
- Perform an additional streaming flush after initializing the world when loading maps to allow streaming interfaces to properly flush their requests that could have been pushed during world initialization.

diff --git a/Engine/Source/Runtime/Engine/Classes/Engine/EngineBaseTypes.h b/Engine/Source/Runtime/Engine/Classes/Engine/EngineBaseTypes.h
--- a/Engine/Source/Runtime/Engine/Classes/Engine/EngineBaseTypes.h
+++ b/Engine/Source/Runtime/Engine/Classes/Engine/EngineBaseTypes.h
@@ -933,7 +933,7 @@ struct FURL
  * The network mode the game is currently running.
  * @see https://docs.unrealengine.com/latest/INT/Gameplay/Networking/Overview/
  */
-enum ENetMode
+enum ENetMode : int
 {
 	/** Standalone: a game without networking, with one or more local players. Still considered a server because it has all server functionality. */
 	NM_Standalone,
diff --git a/Engine/Source/Runtime/Engine/Private/UnrealEngine.cpp b/Engine/Source/Runtime/Engine/Private/UnrealEngine.cpp
--- a/Engine/Source/Runtime/Engine/Private/UnrealEngine.cpp
+++ b/Engine/Source/Runtime/Engine/Private/UnrealEngine.cpp
@@ -16171,6 +16171,9 @@ bool UEngine::LoadMap( FWorldContext& WorldContext, FURL URL, class UPendingNetG
 	// calling it after InitializeActorsForPlay has been called to have all potential bounding boxed initialized
 	FNavigationSystem::AddNavigationSystemToWorld(*WorldContext.World(), FNavigationSystemRunMode::GameMode);
 
+	// Flush streaming again in case registered world subsystems streaming interfaces have some more work to do
+	WorldContext.World()->FlushLevelStreaming(EFlushLevelStreamingType::Visibility);
+
 	// Remember the URL. Put this before spawning player controllers so that
 	// a player controller can get the map name during initialization and
 	// have it be correct
diff --git a/Engine/Source/Runtime/Engine/Private/World.cpp b/Engine/Source/Runtime/Engine/Private/World.cpp
--- a/Engine/Source/Runtime/Engine/Private/World.cpp
+++ b/Engine/Source/Runtime/Engine/Private/World.cpp
@@ -2886,7 +2886,8 @@ void UWorld::UpdateWorldComponents(bool bRerunConstructionScripts, bool bCurrent
 		}
 	}
 
-	SubsystemCollection.ForEachSubsystem([this](UWorldSubsystem* WorldSubsystem){
+	SubsystemCollection.ForEachSubsystem([this](UWorldSubsystem* WorldSubsystem)
+	{
 		WorldSubsystem->OnWorldComponentsUpdated(*this);
 	});
 
@@ -5266,6 +5267,12 @@ void UWorld::FlushLevelStreaming(EFlushLevelStreamingType FlushType)
 		// the flush type, we'll need to do at least one additional loop to be certain all processing is complete
 		bLevelsPendingVisibility = ((LastStreamingType != FlushLevelStreamingType) || IsVisibilityRequestPending());
 	}
+
+	// Notify World Subsystems of the flush
+	SubsystemCollection.ForEachSubsystemWithInterface<UStreamingWorldSubsystemInterface>([](UWorldSubsystem* WorldSubsystem)
+	{
+		CastChecked<IStreamingWorldSubsystemInterface>(WorldSubsystem)->OnFlushStreaming();
+	});
 	
 	check(!IsVisibilityRequestPending());
 
@@ -10210,7 +10217,8 @@ void UWorld::PostInitializeSubsystems()
 {
 	check(bIsWorldInitialized);
 
-	SubsystemCollection.ForEachSubsystem([](UWorldSubsystem* WorldSubsystem){
+	SubsystemCollection.ForEachSubsystem([](UWorldSubsystem* WorldSubsystem)
+	{
 		WorldSubsystem->PostInitialize();
 	});
 }
diff --git a/Engine/Source/Runtime/Engine/Public/Streaming/StreamingWorldSubsystemInterface.h b/Engine/Source/Runtime/Engine/Public/Streaming/StreamingWorldSubsystemInterface.h
--- a/Engine/Source/Runtime/Engine/Public/Streaming/StreamingWorldSubsystemInterface.h
+++ b/Engine/Source/Runtime/Engine/Public/Streaming/StreamingWorldSubsystemInterface.h
@@ -18,5 +18,6 @@ class IStreamingWorldSubsystemInterface
 
 public:
 
-	virtual void OnUpdateStreamingState() {}
+	virtual void OnUpdateStreamingState() { }
+	virtual void OnFlushStreaming() {}
 };
