Added support for printing the address of a property in verbose mode in Obj.DumpProperties

#rnx

diff --git a/Engine/Source/Runtime/Engine/Private/PrintObjectUtils.cpp b/Engine/Source/Runtime/Engine/Private/PrintObjectUtils.cpp
--- a/Engine/Source/Runtime/Engine/Private/PrintObjectUtils.cpp
+++ b/Engine/Source/Runtime/Engine/Private/PrintObjectUtils.cpp
@@ -217,8 +217,9 @@ void PrintObjectProperties(UObject* Object, EPrintObjectFlag Flags, FOutputDevic
 	if (Object)
 	{
 		bool bVerbose = EnumHasAnyFlags(Flags, EPrintObjectFlag::Verbose);
-		bool bIncludeinitializationState = EnumHasAnyFlags(Flags, EPrintObjectFlag::PropertyInitializationState);
-
+		bool bIncludeInitializationState = EnumHasAnyFlags(Flags, EPrintObjectFlag::PropertyInitializationState);
+		bool bIncludePropertyAddresses = bVerbose;
+			
 		LogObjectState(Object, OutputDevice, bVerbose, TEXT("Printing properties of object "));
 
 		// Keep track of visited property-owner pairs to avoid referencing cycles.
@@ -226,7 +227,7 @@ void PrintObjectProperties(UObject* Object, EPrintObjectFlag Flags, FOutputDevic
 
 		FString ObjectPath = Object->GetPathName();
 
-		Object->GetClass()->Visit(Object, [&VisitedPropOwners, Object, &ObjectPath, bIncludeinitializationState, OutputDevice](const FPropertyVisitorContext& Context)->EPropertyVisitorControlFlow
+		Object->GetClass()->Visit(Object, [&VisitedPropOwners, Object, &ObjectPath, bIncludeInitializationState, bIncludePropertyAddresses, OutputDevice](const FPropertyVisitorContext& Context)->EPropertyVisitorControlFlow
 		{
 			const FPropertyVisitorPath& PropertyPath = Context.Path;
 			const FPropertyVisitorData& Data = Context.Data;
@@ -257,8 +258,7 @@ void PrintObjectProperties(UObject* Object, EPrintObjectFlag Flags, FOutputDevic
 					FString OwnerRelPath = OwnerPath.Contains(ObjectPath) ? OwnerPath.RightChop(ObjectPath.Len()) : OwnerPath;
 
 					FString InitializationStateText;
-
-					if (bIncludeinitializationState)
+					if (bIncludeInitializationState)
 					{
 						#if WITH_EDITORONLY_DATA
 						UE::FInitializedPropertyValueState InitializedState(OwnerObject);
@@ -279,7 +279,14 @@ void PrintObjectProperties(UObject* Object, EPrintObjectFlag Flags, FOutputDevic
 						#endif
 					}
 
-					OutputDevice->Logf(TEXT("%s.%s: [%s] %s%s"), *OwnerRelPath, *Property->GetName(), *Property->GetClass()->GetName(), *PropValueAsString, *InitializationStateText);
+					FString PropAddrText;
+					if (bIncludePropertyAddresses)
+					{
+						PropAddrText = FString::Printf(TEXT(", (0x%p)"), PropData);
+					}
+					
+					OutputDevice->Logf(TEXT("%s.%s: [%s] %s%s%s"), 
+						*OwnerRelPath, *Property->GetName(), *Property->GetClass()->GetName(), *PropValueAsString, *InitializationStateText, *PropAddrText);
 
 					bIsInRootObject = OwnerObject->IsInOuter(Object);
 				}
