Fix for missing initializer in FAlphaBlendArgs constructor

diff --git a/Engine/Source/Runtime/Engine/Private/AlphaBlend.cpp b/Engine/Source/Runtime/Engine/Private/AlphaBlend.cpp
--- a/Engine/Source/Runtime/Engine/Private/AlphaBlend.cpp
+++ b/Engine/Source/Runtime/Engine/Private/AlphaBlend.cpp
@@ -302,6 +302,7 @@ bool FAlphaBlend::IsComplete() const
 FAlphaBlendArgs::FAlphaBlendArgs()
 	: CustomCurve(nullptr)
 	, BlendTime(0.2f)
+	, BlendOption(UE::Anim::DefaultBlendOption)
 {
 
 }
@@ -309,6 +310,7 @@ FAlphaBlendArgs::FAlphaBlendArgs()
 FAlphaBlendArgs::FAlphaBlendArgs(float InBlendTime)
 	: CustomCurve(nullptr)
 	, BlendTime(InBlendTime)
+	, BlendOption(UE::Anim::DefaultBlendOption)
 {
 
 }
