From Max: PVS 7.36 fix in AudioDerivedData
#rnx

diff --git a/Engine/Source/Runtime/Engine/Private/AudioDerivedData.cpp b/Engine/Source/Runtime/Engine/Private/AudioDerivedData.cpp
--- a/Engine/Source/Runtime/Engine/Private/AudioDerivedData.cpp
+++ b/Engine/Source/Runtime/Engine/Private/AudioDerivedData.cpp
@@ -1501,7 +1501,7 @@ void FStreamedAudioPlatformData::Serialize(FArchive& Ar, USoundWave* Owner)
 	Ar << AudioFormat;
 
 #if WITH_EDITOR
-	if (Ar.IsCooking() && Ar.IsSaving() && Ar.GetCookContext() && Ar.GetCookContext()->GetCookTagList())
+	if (Ar.IsCooking() && Ar.IsSaving() && Ar.GetCookContext() && Ar.GetCookContext()->GetCookTagList() && Owner)
 	{
 		FCookTagList* CookTags = Ar.GetCookContext()->GetCookTagList();
 		CookTags->Add(Owner, "StreamingFormat", LexToString(AudioFormat));
@@ -1520,18 +1520,18 @@ void FStreamedAudioPlatformData::Serialize(FArchive& Ar, USoundWave* Owner)
 		}
 	}
 
-	for (int32 ChunkIndex = 0; ChunkIndex < Chunks.Num(); ++ChunkIndex)
+	if (Owner)
 	{
-		Chunks[ChunkIndex].Serialize(Ar, Owner, ChunkIndex);
-	}
+		for (int32 ChunkIndex = 0; ChunkIndex < Chunks.Num(); ++ChunkIndex)
+		{
+			Chunks[ChunkIndex].Serialize(Ar, Owner, ChunkIndex);
+		}
 	
 
 #if WITH_EDITORONLY_DATA
-	if (Owner)
-	{
 		Owner->RawDataCriticalSection.Unlock();
-	}
 #endif
+	}
 }
 
 /**
