Fix Anim Notify States triggering on the dedicated server in PIE even if TriggerOnDedicatedServer is set to false

diff --git a/Engine/Source/Runtime/Engine/Private/Animation/AnimNotifyQueue.cpp b/Engine/Source/Runtime/Engine/Private/Animation/AnimNotifyQueue.cpp
--- a/Engine/Source/Runtime/Engine/Private/Animation/AnimNotifyQueue.cpp
+++ b/Engine/Source/Runtime/Engine/Private/Animation/AnimNotifyQueue.cpp
@@ -51,6 +51,7 @@ bool FAnimNotifyQueue::PassesChanceOfTriggering(const FAnimNotifyEvent* Event) c
 
 void FAnimNotifyQueue::AddAnimNotifiesToDest(bool bSrcIsLeader, const TArray<FAnimNotifyEventReference>& NewNotifies, TArray<FAnimNotifyEventReference>& DestArray, const float InstanceWeight)
 {
+	bool bIsDedicatedServer = World ? (World->GetNetMode() == NM_DedicatedServer) : false;
 	// for now there is no filter whatsoever, it just adds everything requested
 	for (const FAnimNotifyEventReference& NotifyRef : NewNotifies)
 	{
@@ -58,7 +59,7 @@ void FAnimNotifyQueue::AddAnimNotifiesToDest(bool bSrcIsLeader, const TArray<FAn
 		if( Notify && (bSrcIsLeader || Notify->bTriggerOnFollower))
 		{
 			// only add if it is over TriggerWeightThreshold
-			const bool bPassesDedicatedServerCheck = Notify->bTriggerOnDedicatedServer || !IsRunningDedicatedServer();
+			const bool bPassesDedicatedServerCheck = Notify->bTriggerOnDedicatedServer || !bIsDedicatedServer;
 			if (bPassesDedicatedServerCheck && Notify->TriggerWeightThreshold <= InstanceWeight && PassesFiltering(Notify) && PassesChanceOfTriggering(Notify))
 			{
 				const UE::Anim::IAnimEventsFilterContext* ContextData = Notify->bCanBeFilteredViaRequest ? NotifyRef.GetContextData<UE::Anim::IAnimEventsFilterContext>() : nullptr;
@@ -109,7 +110,16 @@ void FAnimNotifyQueue::Reset(USkeletalMeshComponent* Component)
 {
 	AnimNotifies.Reset();
 	UnfilteredMontageAnimNotifies.Reset();
-	PredictedLODLevel = Component ? Component->GetPredictedLODLevel() : -1;
+	if (Component)
+	{
+		PredictedLODLevel = Component->GetPredictedLODLevel();
+		World = Component->GetWorld();
+	}
+	else
+	{
+		PredictedLODLevel =  -1;
+		World = nullptr;
+	}
 }
 
 void FAnimNotifyQueue::Append(const FAnimNotifyQueue& Queue)
diff --git a/Engine/Source/Runtime/Engine/Public/Animation/AnimNotifyQueue.h b/Engine/Source/Runtime/Engine/Public/Animation/AnimNotifyQueue.h
--- a/Engine/Source/Runtime/Engine/Public/Animation/AnimNotifyQueue.h
+++ b/Engine/Source/Runtime/Engine/Public/Animation/AnimNotifyQueue.h
@@ -14,7 +14,7 @@ struct FAnimInstanceProxy;
 struct FAnimNotifyEvent;
 class UMirrorDataTable;
 struct FAnimTickRecord;
-
+class UWorld;
 
 USTRUCT(BlueprintType)
 struct FAnimNotifyEventReference
@@ -216,4 +216,7 @@ private:
 
 	/** Adds the contents of the NewNotifies array to the DestArray (maintaining uniqueness of notify states*/
 	void AddAnimNotifiesToDestNoFiltering(const TArray<FAnimNotifyEventReference>& NewNotifies, TArray<FAnimNotifyEventReference>& DestArray) const;
+
+	UPROPERTY(Transient)
+	TObjectPtr<UWorld> World;
 };
