Mark deprecated virtual functions from MaterialInterface and FMaterial as final to avoid silent errors
Modify all derived classes to override the right functions

diff --git a/Engine/Source/Developer/MaterialBaking/Private/ExportMaterialProxy.h b/Engine/Source/Developer/MaterialBaking/Private/ExportMaterialProxy.h
--- a/Engine/Source/Developer/MaterialBaking/Private/ExportMaterialProxy.h
+++ b/Engine/Source/Developer/MaterialBaking/Private/ExportMaterialProxy.h
@@ -287,16 +287,15 @@ public:
 	* Should the shader for this material with the given platform, shader type and vertex
 	* factory type combination be compiled
 	*
-	* @param Platform		The platform currently being compiled for
 	* @param ShaderType	Which shader is being compiled
 	* @param VertexFactory	Which vertex factory is being compiled (can be NULL)
 	*
 	* @return true if the shader should be compiled
 	*/
-	virtual bool ShouldCache(EShaderPlatform Platform, const FShaderType* ShaderType, const FVertexFactoryType* VertexFactoryType) const override
+	virtual bool ShouldCache(const FShaderType* ShaderType, const FVertexFactoryType* VertexFactoryType) const override
 	{
 		const bool bCorrectVertexFactory = VertexFactoryType == FindVertexFactoryType(FName(TEXT("FLocalVertexFactory"), FNAME_Find));
-		const bool bPCPlatform = !IsConsolePlatform(Platform);
+		const bool bPCPlatform = !IsConsolePlatform(GetShaderPlatform());
 		const bool bCorrectFrequency = ShaderType->GetFrequency() == SF_Vertex || ShaderType->GetFrequency() == SF_Pixel;
 		return bCorrectVertexFactory && bPCPlatform && bCorrectFrequency;
 	}
@@ -311,7 +310,7 @@ public:
 		return ReferencedTextureCollections;
 	}
 
-	virtual void GetStaticParameterSet(EShaderPlatform Platform, FStaticParameterSet& OutSet) const override
+	virtual void GetStaticParameterSet(FStaticParameterSet& OutSet) const override
 	{
 		if (const FMaterialResource* Resource = MaterialInterface->GetMaterialResource(GMaxRHIShaderPlatform))
 		{
diff --git a/Engine/Source/Developer/MaterialUtilities/Private/MaterialUtilities.cpp b/Engine/Source/Developer/MaterialUtilities/Private/MaterialUtilities.cpp
--- a/Engine/Source/Developer/MaterialUtilities/Private/MaterialUtilities.cpp
+++ b/Engine/Source/Developer/MaterialUtilities/Private/MaterialUtilities.cpp
@@ -521,13 +521,12 @@ public:
 	* Should the shader for this material with the given platform, shader type and vertex 
 	* factory type combination be compiled
 	*
-	* @param Platform		The platform currently being compiled for
 	* @param ShaderType	Which shader is being compiled
 	* @param VertexFactory	Which vertex factory is being compiled (can be NULL)
 	*
 	* @return true if the shader should be compiled
 	*/
-	virtual bool ShouldCache(EShaderPlatform Platform, const FShaderType* ShaderType, const FVertexFactoryType* VertexFactoryType) const override
+	virtual bool ShouldCache(const FShaderType* ShaderType, const FVertexFactoryType* VertexFactoryType) const override
 	{
 		// Always cache - decreases performance but avoids missing shaders during exports.
 		return true;
diff --git a/Engine/Source/Editor/MaterialEditor/Private/MaterialEditor.cpp b/Engine/Source/Editor/MaterialEditor/Private/MaterialEditor.cpp
--- a/Engine/Source/Editor/MaterialEditor/Private/MaterialEditor.cpp
+++ b/Engine/Source/Editor/MaterialEditor/Private/MaterialEditor.cpp
@@ -283,14 +283,14 @@ void FMatExpressionPreview::AddReferencedObjects(FReferenceCollector& Collector)
 	}
 }
 
-bool FMatExpressionPreview::ShouldCache(EShaderPlatform Platform, const FShaderType* ShaderType, const FVertexFactoryType* VertexFactoryType) const
+bool FMatExpressionPreview::ShouldCache(const FShaderType* ShaderType, const FVertexFactoryType* VertexFactoryType) const
 {
 	if(VertexFactoryType == FindVertexFactoryType(FName(TEXT("FLocalVertexFactory"), FNAME_Find)))
 	{
 		// we only need the non-light-mapped, base pass, local vertex factory shaders for drawing an opaque Material Tile
 		// @todo: Added a FindShaderType by fname or something"
 
-		if (IsMobilePlatform(Platform))
+		if (IsMobilePlatform(GetShaderPlatform()))
 		{
 			if (FCString::Stristr(ShaderType->GetName(), TEXT("BasePassForForwardShadingVSFNoLightMapPolicy")) ||
 				FCString::Stristr(ShaderType->GetName(), TEXT("BasePassForForwardShadingPSFNoLightMapPolicy")))
diff --git a/Engine/Source/Editor/MaterialEditor/Private/MaterialEditor.h b/Engine/Source/Editor/MaterialEditor/Private/MaterialEditor.h
--- a/Engine/Source/Editor/MaterialEditor/Private/MaterialEditor.h
+++ b/Engine/Source/Editor/MaterialEditor/Private/MaterialEditor.h
@@ -78,13 +78,12 @@ public:
 	 * Should the shader for this material with the given platform, shader type and vertex 
 	 * factory type combination be compiled
 	 *
-	 * @param Platform		The platform currently being compiled for
 	 * @param ShaderType	Which shader is being compiled
 	 * @param VertexFactory	Which vertex factory is being compiled (can be NULL)
 	 *
 	 * @return true if the shader should be compiled
 	 */
-	virtual bool ShouldCache(EShaderPlatform Platform, const FShaderType* ShaderType, const FVertexFactoryType* VertexFactoryType) const override;
+	virtual bool ShouldCache(const FShaderType* ShaderType, const FVertexFactoryType* VertexFactoryType) const override;
 
 	virtual TArrayView<const TObjectPtr<UObject>> GetReferencedTextures() const override;
 	virtual TConstArrayView<TObjectPtr<UTextureCollection>> GetReferencedTextureCollections() const override;
diff --git a/Engine/Source/Editor/MaterialEditor/Private/MaterialStatsCommon.cpp b/Engine/Source/Editor/MaterialEditor/Private/MaterialStatsCommon.cpp
--- a/Engine/Source/Editor/MaterialEditor/Private/MaterialStatsCommon.cpp
+++ b/Engine/Source/Editor/MaterialEditor/Private/MaterialStatsCommon.cpp
@@ -15,12 +15,6 @@
 /***********************************************************************************************************************/
 /*begin FMaterialResourceStats functions*/
 
-// Deprecated 5.7
-void FMaterialResourceStats::SetupExtraCompilationSettings(const EShaderPlatform Platform, FExtraShaderCompilerSettings& Settings) const
-{
-	SetupExtraCompilationSettings(Settings);
-}
-
 void FMaterialResourceStats::SetupExtraCompilationSettings(FExtraShaderCompilerSettings& Settings) const
 {
 	Settings.bExtractShaderSource = true;
diff --git a/Engine/Source/Editor/MaterialEditor/Public/MaterialStatsCommon.h b/Engine/Source/Editor/MaterialEditor/Public/MaterialStatsCommon.h
--- a/Engine/Source/Editor/MaterialEditor/Public/MaterialStatsCommon.h
+++ b/Engine/Source/Editor/MaterialEditor/Public/MaterialStatsCommon.h
@@ -32,8 +32,6 @@ public:
 	virtual bool RequiresSynchronousCompilation() const override { return false; }
 
 	/** this will enable shader source extraction and pass paths to (eventual) offline shader compilers */
-	UE_DEPRECATED(5.7, "Call SetupExtraCompilationSettings without EShaderPlatform as its now cached inside the FMaterial")
-	UE_API virtual void SetupExtraCompilationSettings(const EShaderPlatform Platform, FExtraShaderCompilerSettings& Settings) const override;
 	UE_API virtual void SetupExtraCompilationSettings(FExtraShaderCompilerSettings& Settings) const override;
 };
 
diff --git a/Engine/Source/Editor/UnrealEd/Private/Lightmass/LightmassRender.cpp b/Engine/Source/Editor/UnrealEd/Private/Lightmass/LightmassRender.cpp
--- a/Engine/Source/Editor/UnrealEd/Private/Lightmass/LightmassRender.cpp
+++ b/Engine/Source/Editor/UnrealEd/Private/Lightmass/LightmassRender.cpp
@@ -301,13 +301,12 @@ public:
 	 * Should the shader for this material with the given platform, shader type and vertex 
 	 * factory type combination be compiled
 	 *
-	 * @param Platform		The platform currently being compiled for
 	 * @param ShaderType	Which shader is being compiled
 	 * @param VertexFactory	Which vertex factory is being compiled (can be NULL)
 	 *
 	 * @return true if the shader should be compiled
 	 */
-	virtual bool ShouldCache(EShaderPlatform Platform, const FShaderType* ShaderType, const FVertexFactoryType* VertexFactoryType) const override
+	virtual bool ShouldCache(const FShaderType* ShaderType, const FVertexFactoryType* VertexFactoryType) const override
 	{
 		if (VertexFactoryType == FindVertexFactoryType(FName(TEXT("FLocalVertexFactory"), FNAME_Find)))
 		{
@@ -989,7 +988,7 @@ public:
 		}
 	}
 
-	virtual void GetStaticParameterSet(EShaderPlatform Platform, FStaticParameterSet& OutSet) const override
+	virtual void GetStaticParameterSet(FStaticParameterSet& OutSet) const override
 	{
 		if (const FMaterialResource* Resource = MaterialInterface->GetMaterialResource(GMaxRHIShaderPlatform))
 		{
diff --git a/Engine/Source/Editor/UnrealEd/Private/PreviewMaterial.cpp b/Engine/Source/Editor/UnrealEd/Private/PreviewMaterial.cpp
--- a/Engine/Source/Editor/UnrealEd/Private/PreviewMaterial.cpp
+++ b/Engine/Source/Editor/UnrealEd/Private/PreviewMaterial.cpp
@@ -53,13 +53,12 @@ public:
 	 * Should the shader for this material with the given platform, shader type and vertex 
 	 * factory type combination be compiled
 	 *
-	 * @param Platform		The platform currently being compiled for
 	 * @param ShaderType	Which shader is being compiled
 	 * @param VertexFactory	Which vertex factory is being compiled (can be NULL)
 	 *
 	 * @return true if the shader should be compiled
 	 */
-	virtual bool ShouldCache(EShaderPlatform Platform, const FShaderType* ShaderType, const FVertexFactoryType* VertexFactoryType) const override
+	virtual bool ShouldCache(const FShaderType* ShaderType, const FVertexFactoryType* VertexFactoryType) const override
 	{
 		// only generate the needed shaders (which should be very restrictive for fast recompiling during editing)
 		// @todo: Add a FindShaderType by fname or something
diff --git a/Engine/Source/Runtime/Engine/Private/Materials/Material.cpp b/Engine/Source/Runtime/Engine/Private/Materials/Material.cpp
--- a/Engine/Source/Runtime/Engine/Private/Materials/Material.cpp
+++ b/Engine/Source/Runtime/Engine/Private/Materials/Material.cpp
@@ -379,11 +379,6 @@ void FMaterialResource::BuildShaderMapIdOverride(const FBuildShaderMapIdArgs& Ar
 }
 
 #if WITH_EDITORONLY_DATA
-// Deprecated 5.7
-void FMaterialResource::GetStaticParameterSet(EShaderPlatform Platform, FStaticParameterSet& OutSet) const
-{
-	GetStaticParameterSet(OutSet);
-}
 void FMaterialResource::GetStaticParameterSet(FStaticParameterSet& OutSet) const
 {
 	FMaterial::GetStaticParameterSet(OutSet);
@@ -1215,12 +1210,6 @@ FMaterialResource* UMaterial::AllocateResource()
 	return new FMaterialResource();
 }
 
-// Deprecated in 5.7
-void UMaterial::GetUsedTextures(TArray<UTexture*>& OutTextures, EMaterialQualityLevel::Type QualityLevel, bool bAllQualityLevels, ERHIFeatureLevel::Type FeatureLevel, bool bAllFeatureLevels) const
-{
-	GetUsedTextures(OutTextures, QualityLevel, bAllQualityLevels, GetFeatureLevelShaderPlatform_Checked(FeatureLevel), bAllFeatureLevels);
-}
-
 void UMaterial::GetUsedTextures(TArray<UTexture*>& OutTextures, EMaterialQualityLevel::Type QualityLevel, bool bAllQualityLevels, EShaderPlatform ShaderPlatform, bool bAllShaderPlatforms) const
 {
 	OutTextures.Empty();
@@ -1334,12 +1323,6 @@ void UMaterial::GetUsedTextures(TArray<UTexture*>& OutTextures, EMaterialQuality
 	}
 }
 
-// Deprecated in 5.7
-void UMaterial::GetUsedTexturesAndIndices(TArray<UTexture*>& OutTextures, TArray< TArray<int32> >& OutIndices, EMaterialQualityLevel::Type QualityLevel, ERHIFeatureLevel::Type FeatureLevel) const
-{
-	GetUsedTexturesAndIndices(OutTextures, OutIndices, QualityLevel, GetFeatureLevelShaderPlatform_Checked(FeatureLevel));
-}
-
 void UMaterial::GetUsedTexturesAndIndices(TArray<UTexture*>& OutTextures, TArray< TArray<int32> >& OutIndices, EMaterialQualityLevel::Type QualityLevel, EShaderPlatform InShaderPlatform) const
 {
 	OutTextures.Empty();
@@ -1486,12 +1469,6 @@ void UMaterial::OverrideTexture(const UTexture* InTextureToOverride, UTexture* O
 #endif // WITH_EDITOR
 }
 
-// Deprecated in 5.7
-void UMaterial::OverrideNumericParameterDefault(EMaterialParameterType Type, const FHashedMaterialParameterInfo& ParameterInfo, const UE::Shader::FValue& Value, bool bOverride, ERHIFeatureLevel::Type InFeatureLevel)
-{
-	return OverrideNumericParameterDefault(Type, ParameterInfo, Value, bOverride, GetFeatureLevelShaderPlatform_Checked(InFeatureLevel));
-}
-
 void UMaterial::OverrideNumericParameterDefault(EMaterialParameterType Type, const FHashedMaterialParameterInfo& ParameterInfo, const UE::Shader::FValue& Value, bool bOverride, EShaderPlatform InShaderPlatform)
 {
 #if WITH_EDITOR
@@ -2993,12 +2970,6 @@ bool UMaterial::AttemptInsertNewGroupName(const FString & InNewName)
 	return false;
 }
 
-// Deprecated in 5.7
-FMaterialResource* UMaterial::GetMaterialResource(ERHIFeatureLevel::Type InFeatureLevel, EMaterialQualityLevel::Type QualityLevel)
-{
-	return GetMaterialResource(GetFeatureLevelShaderPlatform_Checked(InFeatureLevel), QualityLevel);
-}
-
 FMaterialResource* UMaterial::GetMaterialResource(EShaderPlatform InShaderPlatform, EMaterialQualityLevel::Type QualityLevel)
 {
 	if (QualityLevel == EMaterialQualityLevel::Num)
@@ -3008,12 +2979,6 @@ FMaterialResource* UMaterial::GetMaterialResource(EShaderPlatform InShaderPlatfo
 	return FindMaterialResource(MaterialResources, InShaderPlatform, QualityLevel, true);
 }
 
-// Deprecated in 5.7
-const FMaterialResource* UMaterial::GetMaterialResource(ERHIFeatureLevel::Type InFeatureLevel, EMaterialQualityLevel::Type QualityLevel) const
-{
-	return GetMaterialResource(GetFeatureLevelShaderPlatform_Checked(InFeatureLevel), QualityLevel);
-}
-
 const FMaterialResource* UMaterial::GetMaterialResource(EShaderPlatform InShaderPlatform, EMaterialQualityLevel::Type QualityLevel) const
 {
 	if (QualityLevel == EMaterialQualityLevel::Num)
diff --git a/Engine/Source/Runtime/Engine/Private/Materials/MaterialInstance.cpp b/Engine/Source/Runtime/Engine/Private/Materials/MaterialInstance.cpp
--- a/Engine/Source/Runtime/Engine/Private/Materials/MaterialInstance.cpp
+++ b/Engine/Source/Runtime/Engine/Private/Materials/MaterialInstance.cpp
@@ -1403,12 +1403,6 @@ void UMaterialInstance::GetTextureCollectionExpressionValues(const FMaterialReso
 	}
 }
 
-// Deprecated in 5.7
-void UMaterialInstance::GetUsedTextures(TArray<UTexture*>& OutTextures, EMaterialQualityLevel::Type QualityLevel, bool bAllQualityLevels, ERHIFeatureLevel::Type FeatureLevel, bool bAllFeatureLevels) const
-{
-	GetUsedTextures(OutTextures, QualityLevel, bAllQualityLevels, GetFeatureLevelShaderPlatform_Checked(FeatureLevel), bAllFeatureLevels);
-}
-
 void UMaterialInstance::GetUsedTextures(TArray<UTexture*>& OutTextures, EMaterialQualityLevel::Type QualityLevel, bool bAllQualityLevels, EShaderPlatform ShaderPlatform, bool bAllShaderPlatforms) const
 {
 	OutTextures.Empty();
@@ -1673,12 +1667,6 @@ void UMaterialInstance::ValidateTextureOverrides(EShaderPlatform InShaderPlatfor
 	}
 }
 
-// Deprecated in 5.7
-void UMaterialInstance::GetUsedTexturesAndIndices(TArray<UTexture*>& OutTextures, TArray< TArray<int32> >& OutIndices, EMaterialQualityLevel::Type QualityLevel, ERHIFeatureLevel::Type FeatureLevel) const
-{
-	GetUsedTexturesAndIndices(OutTextures, OutIndices, QualityLevel, GetFeatureLevelShaderPlatform_Checked(FeatureLevel));
-}
-
 void UMaterialInstance::GetUsedTexturesAndIndices(TArray<UTexture*>& OutTextures, TArray< TArray<int32> >& OutIndices, EMaterialQualityLevel::Type QualityLevel, EShaderPlatform InShaderPlatform) const
 {
 	OutTextures.Empty();
@@ -1769,12 +1757,6 @@ void UMaterialInstance::OverrideTexture(const UTexture* InTextureToOverride, UTe
 #endif // #if WITH_EDITOR
 }
 
-//Deprecated in 5.7
-void UMaterialInstance::OverrideNumericParameterDefault(EMaterialParameterType Type, const FHashedMaterialParameterInfo& ParameterInfo, const UE::Shader::FValue& Value, bool bOverride, ERHIFeatureLevel::Type FeatureLevel)
-{
-	return OverrideNumericParameterDefault(Type, ParameterInfo, Value, bOverride, GetFeatureLevelShaderPlatform_Checked(FeatureLevel));
-}
-
 void UMaterialInstance::OverrideNumericParameterDefault(EMaterialParameterType Type, const FHashedMaterialParameterInfo& ParameterInfo, const UE::Shader::FValue& Value, bool bOverride, EShaderPlatform InShaderPlatform)
 {
 #if WITH_EDITOR
@@ -2119,12 +2101,6 @@ void UMaterialInstance::CopyMaterialInstanceParameters(UMaterialInterface* Sourc
 	}
 }
 
-// Deprecated in 5.7
-FMaterialResource* UMaterialInstance::GetMaterialResource(ERHIFeatureLevel::Type InFeatureLevel, EMaterialQualityLevel::Type QualityLevel)
-{
-	return GetMaterialResource(GetFeatureLevelShaderPlatform_Checked(InFeatureLevel), QualityLevel);
-}
-
 FMaterialResource* UMaterialInstance::GetMaterialResource(EShaderPlatform InShaderPlatform, EMaterialQualityLevel::Type QualityLevel)
 {
 	if (bHasStaticPermutationResource)
@@ -2165,12 +2141,6 @@ bool UMaterialInstance::HasCustomPrimitiveData() const
 	return Parent ? Parent->HasCustomPrimitiveData() : false;
 }
 
-// Deprecated in 5.7
-const FMaterialResource* UMaterialInstance::GetMaterialResource(ERHIFeatureLevel::Type InFeatureLevel, EMaterialQualityLevel::Type QualityLevel) const
-{
-	return GetMaterialResource(GetFeatureLevelShaderPlatform_Checked(InFeatureLevel), QualityLevel);
-}
-
 const FMaterialResource* UMaterialInstance::GetMaterialResource(EShaderPlatform InShaderPlatform, EMaterialQualityLevel::Type QualityLevel) const
 {
 	return const_cast<UMaterialInstance*>(this)->GetMaterialResource(InShaderPlatform, QualityLevel);
diff --git a/Engine/Source/Runtime/Engine/Private/Materials/MaterialShared.cpp b/Engine/Source/Runtime/Engine/Private/Materials/MaterialShared.cpp
--- a/Engine/Source/Runtime/Engine/Private/Materials/MaterialShared.cpp
+++ b/Engine/Source/Runtime/Engine/Private/Materials/MaterialShared.cpp
@@ -1821,12 +1821,6 @@ bool FMaterialResource::IsTranslucencyAfterMotionBlurEnabled() const
 		&& !IsDeferredDecal();
 }
 
-// Deprecated in 5.7
-bool FMaterialResource::IsDualBlendingEnabled(EShaderPlatform Platform) const
-{
-	return IsDualBlendingEnabled();
-}
-
 bool FMaterialResource::IsDualBlendingEnabled() const
 {
 	bool bMaterialRequestsDualSourceBlending = Material->ShadingModel == MSM_ThinTranslucent;
diff --git a/Engine/Source/Runtime/Engine/Public/MaterialShared.h b/Engine/Source/Runtime/Engine/Public/MaterialShared.h
--- a/Engine/Source/Runtime/Engine/Public/MaterialShared.h
+++ b/Engine/Source/Runtime/Engine/Public/MaterialShared.h
@@ -2174,7 +2174,7 @@ public:
 	 * @return true if the shader should be compiled
 	 */
 	UE_DEPRECATED(5.7, "Call ShouldCachePipeline without EShaderPlatform as its now cached inside the FMaterial")
-	ENGINE_API virtual bool ShouldCache(EShaderPlatform Platform, const FShaderType* ShaderType, const FVertexFactoryType* VertexFactoryType) const;
+	ENGINE_API virtual bool ShouldCache(EShaderPlatform Platform, const FShaderType* ShaderType, const FVertexFactoryType* VertexFactoryType) const final;
 	ENGINE_API virtual bool ShouldCache(const FShaderType* ShaderType, const FVertexFactoryType* VertexFactoryType) const;
 
 	UE_DEPRECATED(5.7, "Call ShouldCachePipeline without EShaderPlatform as its now cached inside the FMaterial")
@@ -2240,7 +2240,7 @@ public:
 
 #if WITH_EDITORONLY_DATA
 	UE_DEPRECATED(5.7, "Call GetStaticParameterSet without EShaderPlatform as its now cached inside the FMaterial")
-	ENGINE_API virtual void GetStaticParameterSet(EShaderPlatform Platform, FStaticParameterSet& OutSet) const;
+	ENGINE_API	virtual void GetStaticParameterSet(EShaderPlatform Platform, FStaticParameterSet& OutSet) const final;
 	ENGINE_API virtual void GetStaticParameterSet(FStaticParameterSet& OutSet) const;
 #endif // WITH_EDITORONLY_DATA
 	virtual EMaterialDomain GetMaterialDomain() const = 0; // See EMaterialDomain.
@@ -2331,7 +2331,7 @@ public:
 	virtual bool IsTranslucencyAfterDOFEnabled() const { return false; }
 	virtual bool IsTranslucencyAfterMotionBlurEnabled() const { return false; }
 	UE_DEPRECATED(5.7, "Call IsDualBlendingEnabled without EShaderPlatform as its now cached inside the FMaterial")
-	virtual bool IsDualBlendingEnabled(EShaderPlatform Platform) const { return false; }
+	virtual bool IsDualBlendingEnabled(EShaderPlatform Platform) const final { return false; }
 	virtual bool IsDualBlendingEnabled() const { return false; }
 	virtual bool IsMobileSeparateTranslucencyEnabled() const { return false; }
 	virtual FLinearColor GetTranslucentMultipleScatteringExtinction() const { return FLinearColor::White; }
@@ -2671,7 +2671,7 @@ public:
 
 	/** call during shader compilation jobs setup to fill additional settings that may be required by classes who inherit from this */
 	UE_DEPRECATED(5.7, "Call SetupExtraCompilationSettings without EShaderPlatform as its now cached inside the FMaterial")
-	virtual void SetupExtraCompilationSettings(const EShaderPlatform Platform, FExtraShaderCompilerSettings& Settings) const { SetupExtraCompilationSettings(Settings); }
+	virtual void SetupExtraCompilationSettings(const EShaderPlatform Platform, FExtraShaderCompilerSettings& Settings) const final { SetupExtraCompilationSettings(Settings); }
 	virtual void SetupExtraCompilationSettings(FExtraShaderCompilerSettings& Settings) const {}
 
 	void DumpDebugInfo(FOutputDevice& OutputDevice);
@@ -3039,8 +3039,6 @@ protected:
 	ENGINE_API virtual void BuildShaderMapIdOverride(const FBuildShaderMapIdArgs& Args) const override;
 public:
 #if WITH_EDITORONLY_DATA
-	UE_DEPRECATED(5.7, "Call GetStaticParameterSet without EShaderPlatform as its now cached inside the FMaterial")
-	ENGINE_API virtual void GetStaticParameterSet(EShaderPlatform Platform, FStaticParameterSet& OutSet) const override;
 	ENGINE_API virtual void GetStaticParameterSet(FStaticParameterSet& OutSet) const override;
 #endif
 #if WITH_EDITOR
@@ -3141,8 +3139,6 @@ public:
 	ENGINE_API virtual float GetTranslucentBackscatteringExponent() const override;
 	ENGINE_API virtual bool IsTranslucencyAfterDOFEnabled() const override;
 	ENGINE_API virtual bool IsTranslucencyAfterMotionBlurEnabled() const override;
-	UE_DEPRECATED(5.7, "Call IsDualBlendingEnabled without EShaderPlatform as its now cached inside the FMaterial")
-	ENGINE_API virtual bool IsDualBlendingEnabled(EShaderPlatform Platform) const override;
 	ENGINE_API virtual bool IsDualBlendingEnabled() const override;
 	ENGINE_API virtual bool IsMobileSeparateTranslucencyEnabled() const override;
 	ENGINE_API virtual FDisplacementScaling GetDisplacementScaling() const override;
diff --git a/Engine/Source/Runtime/Engine/Public/Materials/Material.h b/Engine/Source/Runtime/Engine/Public/Materials/Material.h
--- a/Engine/Source/Runtime/Engine/Public/Materials/Material.h
+++ b/Engine/Source/Runtime/Engine/Public/Materials/Material.h
@@ -1233,24 +1233,14 @@ public:
 	ENGINE_API virtual UPhysicalMaterialMask* GetPhysicalMaterialMask() const override;
 	ENGINE_API virtual UPhysicalMaterial* GetPhysicalMaterialFromMap(int32 Index) const override;
 	ENGINE_API virtual UMaterialInterface* GetNaniteOverride(TMicRecursionGuard RecursionGuard = TMicRecursionGuard()) const override;
-	UE_DEPRECATED(5.7, "Please use GetUsedTextures with EShaderPlatform argument and not ERHIFeatureLevel::Type")
-	ENGINE_API virtual void GetUsedTextures(TArray<UTexture*>& OutTextures, EMaterialQualityLevel::Type QualityLevel, bool bAllQualityLevels, ERHIFeatureLevel::Type FeatureLevel, bool bAllShaderPlatforms) const override;
 	ENGINE_API virtual void GetUsedTextures(TArray<UTexture*>& OutTextures, EMaterialQualityLevel::Type QualityLevel, bool bAllQualityLevels, EShaderPlatform ShaderPlatform, bool bAllFeatureLevels) const override;
-	UE_DEPRECATED(5.7, "Please use GetUsedTexturesAndIndices with EShaderPlatform argument and not ERHIFeatureLevel::Type")
-	ENGINE_API virtual void GetUsedTexturesAndIndices(TArray<UTexture*>& OutTextures, TArray< TArray<int32> >& OutIndices, EMaterialQualityLevel::Type QualityLevel, ERHIFeatureLevel::Type FeatureLevel) const override;
 	ENGINE_API virtual void GetUsedTexturesAndIndices(TArray<UTexture*>& OutTextures, TArray< TArray<int32> >& OutIndices, EMaterialQualityLevel::Type QualityLevel, EShaderPlatform InShaderPlatform) const override;
 	ENGINE_API virtual void OverrideTexture(const UTexture* InTextureToOverride, UTexture* OverrideTexture, ERHIFeatureLevel::Type InFeatureLevel) override;
-	UE_DEPRECATED(5.7, "Please use OverrideNumericParameterDefault with EShaderPlatform argument and not ERHIFeatureLevel::Type")
-	ENGINE_API virtual void OverrideNumericParameterDefault(EMaterialParameterType Type, const FHashedMaterialParameterInfo& ParameterInfo, const UE::Shader::FValue& Value, bool bOverride, ERHIFeatureLevel::Type FeatureLevel) override;
 	ENGINE_API virtual void OverrideNumericParameterDefault(EMaterialParameterType Type, const FHashedMaterialParameterInfo& ParameterInfo, const UE::Shader::FValue& Value, bool bOverride, EShaderPlatform InShaderPlatform) override;
 	ENGINE_API virtual bool CheckMaterialUsage(const EMaterialUsage Usage) override;
 	ENGINE_API virtual bool CheckMaterialUsage_Concurrent(const EMaterialUsage Usage) const override;
 	ENGINE_API virtual FMaterialResource* AllocateResource();
-	UE_DEPRECATED(5.7, "Please use GetMaterialResource with EShaderPlatform argument and not ERHIFeatureLevel::Type")
-	ENGINE_API virtual FMaterialResource* GetMaterialResource(ERHIFeatureLevel::Type InFeatureLevel, EMaterialQualityLevel::Type QualityLevel = EMaterialQualityLevel::Num) override;
 	ENGINE_API virtual FMaterialResource* GetMaterialResource(EShaderPlatform InShaderPlatform, EMaterialQualityLevel::Type QualityLevel = EMaterialQualityLevel::Num) override;
-	UE_DEPRECATED(5.7, "Please use GetMaterialResource with EShaderPlatform argument and not ERHIFeatureLevel::Type")
-	ENGINE_API virtual const FMaterialResource* GetMaterialResource(ERHIFeatureLevel::Type InFeatureLevel, EMaterialQualityLevel::Type QualityLevel = EMaterialQualityLevel::Num) const override;
 	ENGINE_API virtual const FMaterialResource* GetMaterialResource(EShaderPlatform InShaderPlatform, EMaterialQualityLevel::Type QualityLevel = EMaterialQualityLevel::Num) const override;
 	ENGINE_API virtual bool GetMaterialLayers(FMaterialLayersFunctions& OutLayers, TMicRecursionGuard RecursionGuard = TMicRecursionGuard()) const override;
 	ENGINE_API virtual bool UpdateLightmassTextureTracking() override;
diff --git a/Engine/Source/Runtime/Engine/Public/Materials/MaterialInstance.h b/Engine/Source/Runtime/Engine/Public/Materials/MaterialInstance.h
--- a/Engine/Source/Runtime/Engine/Public/Materials/MaterialInstance.h
+++ b/Engine/Source/Runtime/Engine/Public/Materials/MaterialInstance.h
@@ -948,27 +948,17 @@ public:
 	virtual ENGINE_API void GetMaterialInheritanceChain(FMaterialInheritanceChain& OutChain) const override;
 	virtual ENGINE_API const FMaterialCachedExpressionData& GetCachedExpressionData(TMicRecursionGuard RecursionGuard = TMicRecursionGuard()) const override;
 	virtual ENGINE_API FMaterialResource* AllocatePermutationResource();
-	UE_DEPRECATED(5.7, "Please use GetMaterialResource with EShaderPlatform argument and not ERHIFeatureLevel::Type")
-	virtual ENGINE_API FMaterialResource* GetMaterialResource(ERHIFeatureLevel::Type InFeatureLevel, EMaterialQualityLevel::Type QualityLevel = EMaterialQualityLevel::Num) override;
 	virtual ENGINE_API FMaterialResource* GetMaterialResource(EShaderPlatform InShaderPlatform, EMaterialQualityLevel::Type QualityLevel = EMaterialQualityLevel::Num) override;
-	UE_DEPRECATED(5.7, "Please use GetMaterialResource with EShaderPlatform argument and not ERHIFeatureLevel::Type")
-	virtual ENGINE_API const FMaterialResource* GetMaterialResource(ERHIFeatureLevel::Type InFeatureLevel, EMaterialQualityLevel::Type QualityLevel = EMaterialQualityLevel::Num) const override;
 	virtual ENGINE_API const FMaterialResource* GetMaterialResource(EShaderPlatform InShaderPlatform, EMaterialQualityLevel::Type QualityLevel = EMaterialQualityLevel::Num) const override;
 
 	ENGINE_API bool GetParameterOverrideValue(EMaterialParameterType Type, const FMemoryImageMaterialParameterInfo& ParameterInfo, FMaterialParameterMetadata& OutValue) const;
 	virtual ENGINE_API bool GetParameterValue(EMaterialParameterType Type, const FMemoryImageMaterialParameterInfo& ParameterInfo, FMaterialParameterMetadata& OutValue, EMaterialGetParameterValueFlags Flags = EMaterialGetParameterValueFlags::Default) const override;
 
-	UE_DEPRECATED(5.7, "Please use GetUsedTextures with EShaderPlatform argument and not ERHIFeatureLevel::Type")
-	virtual ENGINE_API void GetUsedTextures(TArray<UTexture*>& OutTextures, EMaterialQualityLevel::Type QualityLevel, bool bAllQualityLevels, ERHIFeatureLevel::Type FeatureLevel, bool bAllFeatureLevels) const override;
 	virtual ENGINE_API void GetUsedTextures(TArray<UTexture*>& OutTextures, EMaterialQualityLevel::Type QualityLevel, bool bAllQualityLevels, EShaderPlatform ShaderPlatform, bool bAllShaderPlatform) const override;
-	UE_DEPRECATED(5.7, "Please use GetUsedTexturesAndIndices with EShaderPlatform argument and not ERHIFeatureLevel::Type")
-	virtual ENGINE_API void GetUsedTexturesAndIndices(TArray<UTexture*>& OutTextures, TArray< TArray<int32> >& OutIndices, EMaterialQualityLevel::Type QualityLevel, ERHIFeatureLevel::Type FeatureLevel) const;
 	virtual ENGINE_API void GetUsedTexturesAndIndices(TArray<UTexture*>& OutTextures, TArray< TArray<int32> >& OutIndices, EMaterialQualityLevel::Type QualityLevel, EShaderPlatform InShaderPlatform) const;
 	virtual ENGINE_API bool GetTextureParameterValue(const FHashedMaterialParameterInfo& ParameterInfo, class UTexture*& OutValue, bool bOveriddenOnly = false) const override;
 	virtual ENGINE_API bool GetTextureCollectionParameterValue(const FHashedMaterialParameterInfo& ParameterInfo, class UTextureCollection*& OutValue, bool bOveriddenOnly = false) const override;
 	virtual ENGINE_API void OverrideTexture(const UTexture* InTextureToOverride, UTexture* OverrideTexture, ERHIFeatureLevel::Type InFeatureLevel) override;
-	UE_DEPRECATED(5.7, "Please use OverrideNumericParameterDefault with EShaderPlatform argument and not ERHIFeatureLevel::Type")
-	virtual ENGINE_API void OverrideNumericParameterDefault(EMaterialParameterType Type, const FHashedMaterialParameterInfo& ParameterInfo, const UE::Shader::FValue& Value, bool bOverride, ERHIFeatureLevel::Type InFeatureLevel) override;
 	virtual ENGINE_API void OverrideNumericParameterDefault(EMaterialParameterType Type, const FHashedMaterialParameterInfo& ParameterInfo, const UE::Shader::FValue& Value, bool bOverride, EShaderPlatform InShaderPlatform) override;
 	virtual ENGINE_API bool CheckMaterialUsage(const EMaterialUsage Usage) override;
 	virtual ENGINE_API bool CheckMaterialUsage_Concurrent(const EMaterialUsage Usage) const override;
diff --git a/Engine/Source/Runtime/Engine/Public/Materials/MaterialInterface.h b/Engine/Source/Runtime/Engine/Public/Materials/MaterialInterface.h
--- a/Engine/Source/Runtime/Engine/Public/Materials/MaterialInterface.h
+++ b/Engine/Source/Runtime/Engine/Public/Materials/MaterialInterface.h
@@ -560,7 +560,7 @@ public:
 
 	/** Return the textures used to render this material. */
 	UE_DEPRECATED(5.7, "Please use GetUsedTextures with EShaderPlatform argument and not ERHIFeatureLevel::Type")
-	virtual void GetUsedTextures(TArray<UTexture*>& OutTextures, EMaterialQualityLevel::Type QualityLevel, bool bAllQualityLevels, ERHIFeatureLevel::Type FeatureLevel, bool bAllFeatureLevels) const {}
+	virtual void GetUsedTextures(TArray<UTexture*>& OutTextures, EMaterialQualityLevel::Type QualityLevel, bool bAllQualityLevels, ERHIFeatureLevel::Type FeatureLevel, bool bAllFeatureLevels) const final {}
 	virtual void GetUsedTextures(TArray<UTexture*>& OutTextures, EMaterialQualityLevel::Type QualityLevel, bool bAllQualityLevels, EShaderPlatform ShaderPlatform, bool bAllShaderPlatforms) const
 		PURE_VIRTUAL(UMaterialInterface::GetUsedTextures, );
 	/** 
@@ -569,7 +569,7 @@ public:
 	* An empty array in OutIndices means the index is undefined.
 	*/
 	UE_DEPRECATED(5.7, "Please use GetUsedTexturesAndIndices with EShaderPlatform argument and not ERHIFeatureLevel::Type")
-	ENGINE_API virtual void GetUsedTexturesAndIndices(TArray<UTexture*>& OutTextures, TArray< TArray<int32> >& OutIndices, EMaterialQualityLevel::Type QualityLevel, ERHIFeatureLevel::Type FeatureLevel) const;
+	ENGINE_API virtual void GetUsedTexturesAndIndices(TArray<UTexture*>& OutTextures, TArray< TArray<int32> >& OutIndices, EMaterialQualityLevel::Type QualityLevel, ERHIFeatureLevel::Type FeatureLevel) const final;
 	ENGINE_API virtual void GetUsedTexturesAndIndices(TArray<UTexture*>& OutTextures, TArray< TArray<int32> >& OutIndices, EMaterialQualityLevel::Type QualityLevel, EShaderPlatform InShaderPlatform) const;
 
 	/**
@@ -586,7 +586,7 @@ public:
 	 * Handles updating dependent MI's and cached uniform expressions.
 	 */
 	UE_DEPRECATED(5.7, "Please use OverrideNumericParameterDefault with EShaderPlatform argument and not ERHIFeatureLevel::Type")
-	virtual void OverrideNumericParameterDefault(EMaterialParameterType Type, const FHashedMaterialParameterInfo& ParameterInfo, const UE::Shader::FValue& Value, bool bOverride, ERHIFeatureLevel::Type FeatureLevel) {}
+	virtual void OverrideNumericParameterDefault(EMaterialParameterType Type, const FHashedMaterialParameterInfo& ParameterInfo, const UE::Shader::FValue& Value, bool bOverride, ERHIFeatureLevel::Type FeatureLevel) final {}
 	virtual void OverrideNumericParameterDefault(EMaterialParameterType Type, const FHashedMaterialParameterInfo& ParameterInfo, const UE::Shader::FValue& Value, bool bOverride, EShaderPlatform InShaderPlatform) PURE_VIRTUAL(UMaterialInterface::OverrideNumericParameterDefault, return;);
 
 	/**
@@ -606,7 +606,7 @@ public:
 	 * @return - the appropriate FMaterialResource if one exists, otherwise NULL
 	 */
 	UE_DEPRECATED(5.7, "Please use GetMaterialResource with EShaderPlatform argument and not ERHIFeatureLevel::Type")
-	virtual FMaterialResource* GetMaterialResource(ERHIFeatureLevel::Type InFeatureLevel, EMaterialQualityLevel::Type QualityLevel = EMaterialQualityLevel::Num) { return NULL; }
+	virtual FMaterialResource* GetMaterialResource(ERHIFeatureLevel::Type InFeatureLevel, EMaterialQualityLevel::Type QualityLevel = EMaterialQualityLevel::Num) final { return NULL; }
 	virtual FMaterialResource* GetMaterialResource(EShaderPlatform InShaderPlatform, EMaterialQualityLevel::Type QualityLevel = EMaterialQualityLevel::Num) { return NULL; }
 
 	/**
@@ -614,7 +614,7 @@ public:
 	 * @return - the appropriate FMaterialResource if one exists, otherwise NULL
 	 */
 	UE_DEPRECATED(5.7, "Please use GetMaterialResource with EShaderPlatform argument and not ERHIFeatureLevel::Type")
-	virtual const FMaterialResource* GetMaterialResource(ERHIFeatureLevel::Type InFeatureLevel, EMaterialQualityLevel::Type QualityLevel = EMaterialQualityLevel::Num) const { return NULL; }
+	virtual const FMaterialResource* GetMaterialResource(ERHIFeatureLevel::Type InFeatureLevel, EMaterialQualityLevel::Type QualityLevel = EMaterialQualityLevel::Num) const final { return NULL; }
 	virtual const FMaterialResource* GetMaterialResource(EShaderPlatform InShaderPlatform, EMaterialQualityLevel::Type QualityLevel = EMaterialQualityLevel::Num) const { return NULL; }
 
 	/**
diff --git a/Engine/Source/Runtime/Landscape/Private/LandscapeRender.cpp b/Engine/Source/Runtime/Landscape/Private/LandscapeRender.cpp
--- a/Engine/Source/Runtime/Landscape/Private/LandscapeRender.cpp
+++ b/Engine/Source/Runtime/Landscape/Private/LandscapeRender.cpp
@@ -4011,14 +4011,14 @@ public:
 	bool IsUsedWithGeometryCollections() const override { return false; }
 	bool IsUsedWithGeometryCache() const override { return false; }
 
-	bool ShouldCache(EShaderPlatform Platform, const FShaderType* ShaderType, const FVertexFactoryType* VertexFactoryType) const override
+	bool ShouldCache(const FShaderType* ShaderType, const FVertexFactoryType* VertexFactoryType) const override
 	{
 		if (bIsLayerThumbnail)
 		{
 			// Always check against FLocalVertexFactory in editor builds as it is required to render thumbnails
 			// Thumbnail MICs are only rendered in the preview scene using a simple LocalVertexFactory
 			static const FName LocalVertexFactory = FName(TEXT("FLocalVertexFactory"));
-			if (!IsMobilePlatform(Platform) && (VertexFactoryType != nullptr) && (VertexFactoryType->GetFName() == LocalVertexFactory))
+			if (!IsMobilePlatform(GetShaderPlatform()) && (VertexFactoryType != nullptr) && (VertexFactoryType->GetFName() == LocalVertexFactory))
 			{
 				if (Algo::Find(GetAllowedShaderTypesInThumbnailRender(), ShaderType->GetFName()))
 				{
