Don't register an object being destroyed in the NetGUID cache on migration.

diff --git a/Engine/Source/Runtime/Engine/Private/Actor.cpp b/Engine/Source/Runtime/Engine/Private/Actor.cpp
--- a/Engine/Source/Runtime/Engine/Private/Actor.cpp
+++ b/Engine/Source/Runtime/Engine/Private/Actor.cpp
@@ -1459,7 +1459,14 @@ void AActor::PostMigrate(const FUObjectMigrationContext& MigrationContext)
 
 			if (!GuidCache->GetNetGUID(this).IsValid())
 			{
-				GuidCache->RegisterNetGUID_Server(NetGUID, this);
+				if (!bActorIsBeingDestroyed)
+				{
+					GuidCache->RegisterNetGUID_Server(NetGUID, this);
+				}
+				else
+				{
+					UE_LOG(LogActor, Verbose, TEXT("Borrowed object %s cannot be added to the NetGUID cache because it's being destroyed."), *GetName());
+				}
 			}
 			else
 			{
