Update anim next graph instance tracing to use child element tracing

diff --git a/Engine/Plugins/Animation/GameplayInsights/Source/GameplayInsights/Private/GameplayAnalyzer.cpp b/Engine/Plugins/Animation/GameplayInsights/Source/GameplayInsights/Private/GameplayAnalyzer.cpp
--- a/Engine/Plugins/Animation/GameplayInsights/Source/GameplayInsights/Private/GameplayAnalyzer.cpp
+++ b/Engine/Plugins/Animation/GameplayInsights/Source/GameplayInsights/Private/GameplayAnalyzer.cpp
@@ -236,11 +236,12 @@ bool FGameplayAnalyzer::OnEvent(uint16 RouteId, EStyle Style, const FOnEventCont
 	{
 		uint64 Id = EventData.GetValue<uint64>("Id");
 		uint64 ChildElementId = EventData.GetValue<uint64>("ChildElementId");
+		uint64 StructTypeId = EventData.GetValue<uint64>("StructTypeId", 0);
 		uint64 OuterId = EventData.GetValue<uint64>("OuterId");
 		uint64 OuterChildElementId = EventData.GetValue<uint64>("OuterChildElementId");
 		FString ElementName;
 		EventData.GetString("Name", ElementName);
-		GameplayProvider.AppendObject(RewindDebugger::FObjectId(Id, ChildElementId), RewindDebugger::FObjectId(OuterId, OuterChildElementId), /*ClassId*/0, *ElementName, /*PathName*/{});
+		GameplayProvider.AppendObject(RewindDebugger::FObjectId(Id, ChildElementId), RewindDebugger::FObjectId(OuterId, OuterChildElementId), StructTypeId, *ElementName, /*PathName*/{});
 		break;
 	}
 	case RouteId_ChildElementLifetimeBegin:
diff --git a/Engine/Plugins/Animation/GameplayInsights/Source/RewindDebugger/Private/RewindDebuggerObjectTrack.cpp b/Engine/Plugins/Animation/GameplayInsights/Source/RewindDebugger/Private/RewindDebuggerObjectTrack.cpp
--- a/Engine/Plugins/Animation/GameplayInsights/Source/RewindDebugger/Private/RewindDebuggerObjectTrack.cpp
+++ b/Engine/Plugins/Animation/GameplayInsights/Source/RewindDebugger/Private/RewindDebuggerObjectTrack.cpp
@@ -18,14 +18,21 @@ namespace RewindDebugger
 {
 
 // check if an object is or is a subclass of a type by name, based on Insights traced type info
-static bool GetTypeHierarchyNames(const uint64 InObjectId, const TraceServices::IAnalysisSession& InSession, TArray<const TCHAR*>& OutTypeNames)
+static bool GetTypeHierarchyNames(const RewindDebugger::FObjectId InObjectId, const TraceServices::IAnalysisSession& InSession, TArray<const TCHAR*>& OutTypeNames)
 {
 	TraceServices::FAnalysisSessionReadScope SessionReadScope(InSession);
 
 	const IGameplayProvider* GameplayProvider = InSession.ReadProvider<IGameplayProvider>("GameplayProvider");
 	const FObjectInfo& ObjectInfo = GameplayProvider->GetObjectInfo(InObjectId);
 	uint64 ClassId = ObjectInfo.ClassId;
-
+	
+	if (InObjectId.IsChildElement() && ClassId == FObjectId::InvalidId)
+	{
+		// for child elements with no traced ClassId (actually a StructId in this case... pending rename), use the main object ClassId (backwards compatibility with previous behavior)
+		const FObjectInfo& BaseObjectInfo = GameplayProvider->GetObjectInfo(InObjectId.GetMainId());
+		ClassId = BaseObjectInfo.ClassId;
+	}
+	
 	while (ClassId != FObjectId::InvalidId)
 	{
 		const FClassInfo& ClassInfo = GameplayProvider->GetClassInfo(ClassId);
@@ -47,7 +54,7 @@ FRewindDebuggerObjectTrack::FRewindDebuggerObjectTrack(const FObjectId& InObject
 	
 	const TraceServices::IAnalysisSession* Session = IRewindDebugger::Instance()->GetAnalysisSession();
 
-	if (TArray<const TCHAR*> TypeNames; GetTypeHierarchyNames(InObjectId.GetMainId(), *Session, TypeNames))
+	if (TArray<const TCHAR*> TypeNames; GetTypeHierarchyNames(InObjectId, *Session, TypeNames))
 	{
 		FRewindDebuggerTrackCreators::EnumerateCreators([&Creators = ChildTrackCreators, &TypeNames, IsChildElementTrack = InObjectId.IsChildElement()](const IRewindDebuggerTrackCreator* Creator)
 			{
@@ -346,7 +353,7 @@ bool FRewindDebuggerObjectTrack::UpdateInternal()
 	{
 		TRACE_CPUPROFILER_EVENT_SCOPE(FRewindDebuggerObjectTrack::UpdateInternal_AddViews);
 
-		if (TArray<const TCHAR*> TypeNames; GetTypeHierarchyNames(ObjectId.GetMainId(), *Session, TypeNames))
+		if (TArray<const TCHAR*> TypeNames; GetTypeHierarchyNames(ObjectId, *Session, TypeNames))
 		{
 			FRewindDebuggerViewCreators::EnumerateCreators([ObjectId = ObjectId.GetMainId(), &Children = Children, &bChanged, &TypeNames](const IRewindDebuggerViewCreator* Creator)
 				{
diff --git a/Engine/Plugins/Experimental/AnimNext/Source/AnimNext/Private/RewindDebugger/AnimNextTrace.cpp b/Engine/Plugins/Experimental/AnimNext/Source/AnimNext/Private/RewindDebugger/AnimNextTrace.cpp
--- a/Engine/Plugins/Experimental/AnimNext/Source/AnimNext/Private/RewindDebugger/AnimNextTrace.cpp
+++ b/Engine/Plugins/Experimental/AnimNext/Source/AnimNext/Private/RewindDebugger/AnimNextTrace.cpp
@@ -112,15 +112,17 @@ void FAnimNextTrace::OutputAnimNextInstance(const FUAFAssetInstance* AnimNextIns
 			OutputAnimNextInstance(AnimNextInstance->GetHost(), OuterObject);
 		}
 
+		const FUAFAssetInstance* HostInstance = AnimNextInstance->GetHost();
+		uint64 HostInstanceId = HostInstance ? HostInstance->GetUniqueId() : FObjectTrace::InvalidObjectId;
+
 		TRACE_OBJECT(OuterObject);
 		TRACE_OBJECT(AnimNextInstance->GetAsset<UAnimNextRigVMAsset>());
-		
-		const FUAFAssetInstance* HostInstance = AnimNextInstance->GetHost();
+		TRACE_CHILD_ELEMENT_WITH_OUTER(OuterObject, AnimNextInstance->GetUniqueId(), OuterObject, HostInstanceId, FUAFAssetInstance::StaticStruct(), AnimNextInstance->GetAssetName().ToString());
 
 		UE_TRACE_LOG(UAF, Instance, UAFChannel)
 			<< Instance.InstanceId(InstanceId)
 			<< Instance.OuterObjectId(FObjectTrace::GetObjectId(OuterObject))
-			<< Instance.HostInstanceId(HostInstance ? HostInstance->GetUniqueId() : 0)
+			<< Instance.HostInstanceId(HostInstanceId)
 			<< Instance.AssetId(FObjectTrace::GetObjectId(AnimNextInstance->GetAsset<UAnimNextRigVMAsset>()));
 	}
 }
diff --git a/Engine/Plugins/Experimental/AnimNext/Source/AnimNextEditor/Private/RewindDebugger/AnimNextModuleTrack.cpp b/Engine/Plugins/Experimental/AnimNext/Source/AnimNextEditor/Private/RewindDebugger/AnimNextModuleTrack.cpp
--- a/Engine/Plugins/Experimental/AnimNext/Source/AnimNextEditor/Private/RewindDebugger/AnimNextModuleTrack.cpp
+++ b/Engine/Plugins/Experimental/AnimNext/Source/AnimNextEditor/Private/RewindDebugger/AnimNextModuleTrack.cpp
@@ -23,7 +23,7 @@ static const FName AnimNextModulesName("AnimNextModules");
 
  FName FAnimNextModuleTrackCreator::GetTargetTypeNameInternal() const
  {
- 	static const FName ObjectName("AnimNextComponent");
+ 	static const FName ObjectName("UAFAssetInstance");
 	return ObjectName;
 }
 
@@ -60,16 +60,9 @@ void FAnimNextModuleTrackCreator::GetTrackTypesInternal(TArray<RewindDebugger::F
 
 TSharedPtr<RewindDebugger::FRewindDebuggerTrack> FAnimNextModuleTrackCreator::CreateTrackInternal(const RewindDebugger::FObjectId& ObjectId) const
 {
-	return MakeShared<FAnimNextModuleTrack>(ObjectId.GetMainId());
+	return MakeShared<FAnimNextModuleTrack>(ObjectId.GetMainId(), ObjectId.GetChildElementId());
 }
 
-		
-FAnimNextModuleTrack::FAnimNextModuleTrack(uint64 InObjectId) :
-	ObjectId(InObjectId)
-{
-	Initialize();
-}
-	
 FAnimNextModuleTrack::FAnimNextModuleTrack(uint64 InObjectId, uint64 InInstanceId) :
    	ObjectId(InObjectId),
 	InstanceId(InInstanceId)
@@ -212,47 +205,12 @@ bool FAnimNextModuleTrack::UpdateInternal()
 				ExistenceRange->Windows[0].TimeStart = Data->StartTime;
 				ExistenceRange->Windows[0].TimeEnd = Data->EndTime;
 			}
-
-			TArray<uint64, TInlineAllocator<32>> CurrentChildren;
-		
-			// update/create child tracks
-			AnimNextProvider->EnumerateChildInstances(InstanceId, [this, &bChanged, &CurrentChildren, &ViewRange](const FDataInterfaceData& ChildData)
-			{
-				if (ChildData.StartTime < ViewRange.GetUpperBoundValue() && ChildData.EndTime > ViewRange.GetLowerBoundValue())
-				{
-					CurrentChildren.Add(ChildData.InstanceId);
-					if (!Children.ContainsByPredicate([&ChildData](TSharedPtr<FAnimNextModuleTrack>& Child) { return Child->InstanceId == ChildData.InstanceId; } ))
-					{
-						Children.Add(MakeShared<FAnimNextModuleTrack>(ObjectId, ChildData.InstanceId));
-						bChanged = true;
-					}
-				}
-			});
-			
-			int32 NumRemoved = Children.RemoveAll([&CurrentChildren](TSharedPtr<FAnimNextModuleTrack>& Child)
-			{
-				return !CurrentChildren.Contains(Child->InstanceId);
-			});
-			bChanged |= (NumRemoved > 0);
-		}
-
-		for (auto& Child : Children)
-		{
-			bChanged |= Child->Update();
 		}
 	}
 	
  	return bChanged;
 }
 
-void FAnimNextModuleTrack::IterateSubTracksInternal(TFunction<void(TSharedPtr<FRewindDebuggerTrack> SubTrack)> IteratorFunction)
-{
-	for(TSharedPtr<FAnimNextModuleTrack>& Track : Children)
-	{
-		IteratorFunction(Track);
-	}
-};
-
 bool FAnimNextModuleTrack::HandleDoubleClickInternal()
 {
 	const IRewindDebugger* RewindDebugger = IRewindDebugger::Instance();
diff --git a/Engine/Plugins/Experimental/AnimNext/Source/AnimNextEditor/Private/RewindDebugger/AnimNextModuleTrack.h b/Engine/Plugins/Experimental/AnimNext/Source/AnimNextEditor/Private/RewindDebugger/AnimNextModuleTrack.h
--- a/Engine/Plugins/Experimental/AnimNext/Source/AnimNextEditor/Private/RewindDebugger/AnimNextModuleTrack.h
+++ b/Engine/Plugins/Experimental/AnimNext/Source/AnimNextEditor/Private/RewindDebugger/AnimNextModuleTrack.h
@@ -34,7 +34,6 @@ namespace UE::AnimNextEditor
 class FAnimNextModuleTrack : public RewindDebugger::FRewindDebuggerTrack
 {
 public:
-	FAnimNextModuleTrack(uint64 InObjectId);
 	FAnimNextModuleTrack(uint64 InObjectId, uint64 InstanceId);
 	virtual ~FAnimNextModuleTrack();
 	
@@ -42,13 +41,28 @@ public:
 private:
 	virtual TSharedPtr<SWidget> GetTimelineViewInternal() override;
 	virtual bool UpdateInternal() override;
-	virtual TSharedPtr<SWidget> GetDetailsViewInternal() override { return DetailsView; }
-
-	virtual FSlateIcon GetIconInternal() override { return Icon; }
-	virtual FName GetNameInternal() const override { return "AnimNextModule"; }
+	
+	virtual TSharedPtr<SWidget> GetDetailsViewInternal() override
+	{
+		return DetailsView;
+	}
+	
+	virtual FSlateIcon GetIconInternal() override
+	{
+		return Icon;
+	}
+	
+	virtual FName GetNameInternal() const override
+	{
+		return "AnimNextModule";
+	}
+	
+	virtual uint64 GetObjectIdInternal() const override
+	{
+		return ObjectId;
+	}
+	
 	virtual FText GetDisplayNameInternal() const override;
-	virtual uint64 GetObjectIdInternal() const override { return ObjectId; }
-	virtual void IterateSubTracksInternal(TFunction<void(TSharedPtr<FRewindDebuggerTrack> SubTrack)> IteratorFunction) override;
 	virtual bool HandleDoubleClickInternal() override;
 
 	void Initialize();
@@ -63,7 +77,6 @@ private:
 	TSharedPtr<SEventTimelineView::FTimelineEventData> ExistenceRange;
 	mutable FText DisplayNameCache;
 
-	TArray<TSharedPtr<FAnimNextModuleTrack>> Children;
 	TArray<FPropertyBagPropertyDesc> PropertyDescriptions;
 };
 
@@ -72,10 +85,25 @@ class FAnimNextModuleTrackCreator : public RewindDebugger::IRewindDebuggerTrackC
 {
 private:
 	virtual FName GetTargetTypeNameInternal() const;
-	virtual FName GetNameInternal() const override { return "AnimNextModule"; }
 	virtual void GetTrackTypesInternal(TArray<RewindDebugger::FRewindDebuggerTrackType>& Types) const override;
 	virtual TSharedPtr<RewindDebugger::FRewindDebuggerTrack> CreateTrackInternal(const RewindDebugger::FObjectId& InObjectId) const override;
 	virtual bool HasDebugInfoInternal(const RewindDebugger::FObjectId& InObjectId) const override;
+	
+	
+	virtual FName GetNameInternal() const override
+	{
+		return "AnimNextModule";
+	}
+	
+	virtual bool IsChildElementSupportedInternal() const override
+	{
+		return true;
+	}
+	
+	virtual bool IsCreatingPrimaryChildTrackInternal() const override
+	{
+		return true;
+	}
 };
 
 }
diff --git a/Engine/Plugins/Experimental/AnimNextAnimGraph/Source/AnimNextAnimGraph/Private/Graph/AnimNextGraphInstance.cpp b/Engine/Plugins/Experimental/AnimNextAnimGraph/Source/AnimNextAnimGraph/Private/Graph/AnimNextGraphInstance.cpp
--- a/Engine/Plugins/Experimental/AnimNextAnimGraph/Source/AnimNextAnimGraph/Private/Graph/AnimNextGraphInstance.cpp
+++ b/Engine/Plugins/Experimental/AnimNextAnimGraph/Source/AnimNextAnimGraph/Private/Graph/AnimNextGraphInstance.cpp
@@ -3,6 +3,7 @@
 #include "Graph/AnimNextGraphInstance.h"
 
 #include "AnimNextAnimGraphStats.h"
+#include "ObjectTrace.h"
 #include "UAFRigVMComponent.h"
 #include "Graph/AnimNextAnimationGraph.h"
 #include "Graph/AnimNextGraphContextData.h"
@@ -23,6 +24,7 @@ FAnimNextGraphInstance::FAnimNextGraphInstance()
 
 FAnimNextGraphInstance::~FAnimNextGraphInstance()
 {
+	TRACE_CHILD_ELEMENT_LIFETIME_END(GetModuleInstance() ? GetModuleInstance()->GetObject() : nullptr, GetUniqueId());
 	Release();
 }
 
diff --git a/Engine/Source/Runtime/Engine/Private/ObjectTrace.cpp b/Engine/Source/Runtime/Engine/Private/ObjectTrace.cpp
--- a/Engine/Source/Runtime/Engine/Private/ObjectTrace.cpp
+++ b/Engine/Source/Runtime/Engine/Private/ObjectTrace.cpp
@@ -100,6 +100,7 @@ UE_TRACE_EVENT_BEGIN(Object, ChildElement)
 	UE_TRACE_EVENT_FIELD(uint64, ChildElementId)
 	UE_TRACE_EVENT_FIELD(uint64, OuterId)
 	UE_TRACE_EVENT_FIELD(uint64, OuterChildElementId)
+	UE_TRACE_EVENT_FIELD(uint64, StructTypeId)
 	UE_TRACE_EVENT_FIELD(UE::Trace::WideString, Name)
 UE_TRACE_EVENT_END()
 
@@ -400,7 +401,7 @@ uint16 FObjectTrace::GetObjectWorldTickCounter(const UObject* InObject)
 	return 0;
 }
 
-void FObjectTrace::OutputClass(const UClass* InClass)
+void FObjectTrace::OutputClass(const UStruct* InClass)
 {
 	const bool bChannelEnabled = UE_TRACE_CHANNELEXPR_IS_ENABLED(ObjectChannel);
 	if (!bChannelEnabled || InClass == nullptr)
@@ -413,7 +414,7 @@ void FObjectTrace::OutputClass(const UClass* InClass)
 		return;
 	}
 
-	OutputClass(InClass->GetSuperClass());
+	OutputClass(InClass->GetSuperStruct());
 
 	const FString ClassPathName = InClass->GetPathName();
 	TCHAR ClassName[FName::StringBufferSize];
@@ -421,7 +422,7 @@ void FObjectTrace::OutputClass(const UClass* InClass)
 
 	UE_TRACE_LOG(Object, Class, ObjectChannel)
 		<< Class.Id(ObjectId)
-		<< Class.SuperId(GetObjectId(InClass->GetSuperClass()))
+		<< Class.SuperId(GetObjectId(InClass->GetSuperStruct()))
 		<< Class.Name(ClassName, ClassNameLength)
 		<< Class.Path(*ClassPathName, ClassPathName.Len());
 }
@@ -588,38 +589,41 @@ void FObjectTrace::OutputWorld(const UWorld* InWorld)
 	TRACE_OBJECT(InWorld);
 }
 
-void FObjectTrace::OutputChildElement(const UObject* InObject, const uint64 InChildElementId, const UObject* InOuterObject, const uint64 InOuterChildElementId)
+void FObjectTrace::OutputChildElement(const UObject* InObject, const uint64 InChildElementId, const UObject* InOuterObject, const uint64 InOuterChildElementId, const UStruct* InStructType, const FString& InName)
 {
 	if (UE::ObjectTrace::IsTracingDisabledForObject(InObject))
 	{
 		return;
 	}
-
+	
 	const uint64 ObjectId = GetObjectId(InObject);
 	if (!AddObjectToTracedSet(UE::ObjectTrace::FObjectId{ObjectId, InChildElementId}))
 	{
 		return;
 	}
 
-	if (InOuterChildElementId == InvalidObjectId)
-	{
-		OutputObject(InOuterObject);
-	}
-	else
+	OutputObject(InOuterObject);
+	// intentionally not tracing the outer child element from here because we don't know what its parent should be
+	//   - game systems must ensure they trace the whole child element hierarchy
+	// 	OutputChildElement(InOuterObject, InOuterChildElementId);
+
+	FString Name = InName;
+	if (Name.IsEmpty())
 	{
-		OutputChildElement(InOuterObject, InOuterChildElementId);
+		check(InObject);
+		Name = InObject->GetName();
+		Name += FString::Printf(TEXT(" : [0x%llX]"), InChildElementId);
 	}
-
-	check(InObject);
-	FString ObjectName = InObject->GetName();
-	ObjectName += FString::Printf(TEXT(" : [0x%llX]"), InChildElementId);
+	
+	TRACE_CLASS(InStructType);
 
 	UE_TRACE_LOG(Object, ChildElement, ObjectChannel)
 		<< ChildElement.Id(ObjectId)
 		<< ChildElement.ChildElementId(InChildElementId)
 		<< ChildElement.OuterId(GetObjectId(InOuterObject))
 		<< ChildElement.OuterChildElementId(InOuterChildElementId)
-		<< ChildElement.Name(*ObjectName, ObjectName.Len());
+		<< ChildElement.StructTypeId(GetObjectId(InStructType))
+		<< ChildElement.Name(*Name, Name.Len());
 
 	UE_TRACE_LOG(Object, ChildElementLifetimeBegin, ObjectChannel)
 		<< ChildElementLifetimeBegin.Cycle(FPlatformTime::Cycles64())
@@ -635,8 +639,7 @@ void FObjectTrace::OutputChildElementLifetimeEnd(const UObject* InObject, const
 		return;
 	}
 
-	check(InObject);
-	TRACE_CHILD_ELEMENT(InObject, InChildElementId);
+	TRACE_OBJECT(InObject);
 
 	UE_TRACE_LOG(Object, ChildElementLifetimeEnd, ObjectChannel)
 		<< ChildElementLifetimeEnd.Cycle(FPlatformTime::Cycles64())
@@ -652,7 +655,7 @@ void FObjectTrace::OutputChildElementEvent(const UObject* InObject, const uint64
 		return;
 	}
 
-	TRACE_CHILD_ELEMENT(InObject, InChildElementId);
+	TRACE_OBJECT(InObject);
 
 	UE_TRACE_LOG(Object, ChildElementEvent, ObjectChannel)
 		<< ChildElementEvent.Cycle(FPlatformTime::Cycles64())
diff --git a/Engine/Source/Runtime/Engine/Public/ObjectTrace.h b/Engine/Source/Runtime/Engine/Public/ObjectTrace.h
--- a/Engine/Source/Runtime/Engine/Public/ObjectTrace.h
+++ b/Engine/Source/Runtime/Engine/Public/ObjectTrace.h
@@ -57,7 +57,7 @@ struct FObjectTrace
 	ENGINE_API static void Reset();
 
 	/** Helper function to output an object */
-	ENGINE_API static void OutputClass(const UClass* InClass);
+	ENGINE_API static void OutputClass(const UStruct* InClass);
 
 	/** Helper function to output an object */
 	ENGINE_API static void OutputObject(const UObject* InObject);
@@ -111,7 +111,8 @@ struct FObjectTrace
 	ENGINE_API static void OutputWorld(const UWorld* InWorld);
 
 	/** Helper function to output a child element of a UObject */
-	ENGINE_API static void OutputChildElement(const UObject* InObject, uint64 InChildElementId, const UObject* InOuterObject = nullptr, uint64 InOuterChildElementId = InvalidObjectId);
+	ENGINE_API static void OutputChildElement(const UObject* InObject, uint64 InChildElementId,
+			const UObject* InOuterObject = nullptr, uint64 InOuterChildElementId = InvalidObjectId, const UStruct* InStructType = nullptr, const FString& InName = "");
 
 	/** Helper function to output destruction event of a child element */
 	ENGINE_API static void OutputChildElementLifetimeEnd(const UObject* InObject, uint64 InChildElementId);
@@ -147,10 +148,10 @@ struct FObjectTrace
 	if (CAN_TRACE_OBJECT(Controller) && (Pawn==nullptr || (CAN_TRACE_OBJECT(Pawn)))) { FObjectTrace::OutputPawnPossess(Controller, Pawn); }
 
 #define TRACE_CHILD_ELEMENT(Object, ChildElementId) \
-	if (CAN_TRACE_OBJECT(Object)) { FObjectTrace::OutputChildElement(Object, ChildElementId, Object, FObjectTrace::InvalidObjectId); }
+	if (CAN_TRACE_OBJECT(Object)) { FObjectTrace::OutputChildElement(Object, ChildElementId); }
 
-#define TRACE_CHILD_ELEMENT_WITH_OUTER(Object, ChildElementId, OuterObject, OuterChildElementId) \
-	if (CAN_TRACE_OBJECT(Object)) { FObjectTrace::OutputChildElement(Object, ChildElementId, OuterObject, OuterChildElementId); }
+#define TRACE_CHILD_ELEMENT_WITH_OUTER(Object, ChildElementId, OuterObject, OuterChildElementId, StructType, Name) \
+	if (CAN_TRACE_OBJECT(Object)) { FObjectTrace::OutputChildElement(Object, ChildElementId, OuterObject, OuterChildElementId, StructType, Name); }
 
 #define TRACE_CHILD_ELEMENT_LIFETIME_BEGIN(Object, ChildElementId) \
 	if (CAN_TRACE_OBJECT(Object)) { FObjectTrace::OutputChildElement(Object, ChildElementId, Object, FObjectTrace::InvalidObjectId); }
@@ -184,6 +185,10 @@ struct FObjectTrace
 #define TRACE_CHILD_ELEMENT(Object, ChildElementId) \
 	FObjectTrace::OutputChildElement(Object, ChildElementId, Object, FObjectTrace::InvalidObjectId);
 
+
+#define TRACE_CHILD_ELEMENT_WITH_OUTER_(Object, ChildElementId, OuterObject, OuterChildElementId, Type, Name) \
+	FObjectTrace::OutputChildElement(Object, ChildElementId, OuterObject, OuterChildElementId, Type, Name);
+
 #define TRACE_CHILD_ELEMENT_WITH_OUTER(Object, ChildElementId, OuterObject, OuterChildElementId) \
 	FObjectTrace::OutputChildElement(Object, ChildElementId, OuterObject, OuterChildElementId);
 
