PVS 7.36 fix in AudioDerivedData fix-up round 2
#rnx

diff --git a/Engine/Source/Runtime/Engine/Private/AudioDerivedData.cpp b/Engine/Source/Runtime/Engine/Private/AudioDerivedData.cpp
--- a/Engine/Source/Runtime/Engine/Private/AudioDerivedData.cpp
+++ b/Engine/Source/Runtime/Engine/Private/AudioDerivedData.cpp
@@ -1501,7 +1501,7 @@ void FStreamedAudioPlatformData::Serialize(FArchive& Ar, USoundWave* Owner)
 	Ar << AudioFormat;
 
 #if WITH_EDITOR
-	if (Ar.IsCooking() && Ar.IsSaving() && Ar.GetCookContext() && Ar.GetCookContext()->GetCookTagList())
+	if (Ar.IsCooking() && Ar.IsSaving() && Ar.GetCookContext() && Ar.GetCookContext()->GetCookTagList() && Owner)
 	{
 		FCookTagList* CookTags = Ar.GetCookContext()->GetCookTagList();
 		CookTags->Add(Owner, "StreamingFormat", LexToString(AudioFormat));
