Skeletal Mesh: Fix compressed morph target decoding

#rnx

diff --git a/Engine/Source/Runtime/Engine/Private/Rendering/MorphTargetVertexCodec.cpp b/Engine/Source/Runtime/Engine/Private/Rendering/MorphTargetVertexCodec.cpp
--- a/Engine/Source/Runtime/Engine/Private/Rendering/MorphTargetVertexCodec.cpp
+++ b/Engine/Source/Runtime/Engine/Private/Rendering/MorphTargetVertexCodec.cpp
@@ -361,20 +361,23 @@ bool IterativeDecode(
 		return false;
 	}
 
-	// We assume the token is intact.
-	uint32 HeaderIndex = (static_cast<uint32>(InOutNextItemToken) >> BatchSizeBits) & ((32 - BatchSize) - 1);
+	// We trust the token is intact.
+	uint32 HeaderIndex = static_cast<uint32>(InOutNextItemToken) >> BatchSizeBits;
 	uint32 ElementIndex = InOutNextItemToken & (BatchSize - 1);
 	uint32 DataOffset =  InOutNextItemToken >> 32;
-
-	check(HeaderIndex < static_cast<uint32>(InBatchHeaders.Num()));
+	
+	// Less-than-or-equal, since the remaining data could be all zero number of bits.  
+	checkSlow(DataOffset <= static_cast<uint32>(InCompressedVertices.Num() * 32));
+	checkSlow(HeaderIndex < static_cast<uint32>(InBatchHeaders.Num()));
 	
 	const FDeltaBatchHeader& BatchHeader = InBatchHeaders[HeaderIndex];
-	check(ElementIndex < BatchHeader.NumElements);
-
+	checkSlow(ElementIndex < BatchHeader.NumElements);
+	
 	FDwordBitReader Reader(InCompressedVertices, DataOffset);
 	OutMorphTargetDelta = DecodeMorphTargetDelta(BatchHeader, Reader, ElementIndex, InPositionPrecision, InTangentZPrecision);
 	
 	DataOffset = Reader.GetOffset();
+	
 	ElementIndex++;
 	if (ElementIndex == BatchHeader.NumElements)
 	{
