Remove an object from the NetGUID cache when it's destroyed while loaning it out.

diff --git a/Engine/Source/Runtime/Engine/Private/Actor.cpp b/Engine/Source/Runtime/Engine/Private/Actor.cpp
--- a/Engine/Source/Runtime/Engine/Private/Actor.cpp
+++ b/Engine/Source/Runtime/Engine/Private/Actor.cpp
@@ -1414,6 +1414,16 @@ void AActor::PostMigrate(const FUObjectMigrationContext& MigrationContext)
 					ensureMsgf(IsPendingKillPending(), TEXT("%hs: We thought %s was destroyed while on loan, but our received state does not match a destroyed Actor"), __func__, *GetName());
 					UE_LOG(LogActor, Warning, TEXT("%hs: Destroying %s back on owning server (it was destroyed while loaned-out)"), __func__, *GetName());
 
+					// Remove the object from the NetGUID cache (if it's there) before destroying the actor below.
+					if (GetWorld()->GetNetDriver()->GuidCache->RemoveNetGUID(this))
+					{
+						UE_LOG(LogActor, Warning, TEXT("Removed destroyed object %s from the NetGUID cache."), *GetName());
+					}
+					else
+					{
+						UE_LOG(LogActor, Warning, TEXT("Destroyed object %s is not in the NetGUID cache."), *GetName());
+					}
+
 					// Destroy it again on the owning server, as it needs to be deregistered from our systems
 					bActorIsBeingDestroyed = false;
 					ClearGarbage();
