Change to the LayerStack function callers logic to instantiate FunctionCalls for re-use instead of new ones each time, and only reset the associated material function if there's been a change (helps speed up local compiles with editor previews).

diff --git a/Engine/Source/Runtime/Engine/Private/Materials/MaterialExpressionLayerStack.cpp b/Engine/Source/Runtime/Engine/Private/Materials/MaterialExpressionLayerStack.cpp
--- a/Engine/Source/Runtime/Engine/Private/Materials/MaterialExpressionLayerStack.cpp
+++ b/Engine/Source/Runtime/Engine/Private/Materials/MaterialExpressionLayerStack.cpp
@@ -538,17 +538,42 @@ void UMaterialExpressionLayerStack::RebuildLayerGraph(bool bReportErrors)
 		//Similar to the parent implementation, this helper function creates a new Function call for each referenced function to create the new graph
 		auto ProcessFunctionCallers = [this](const TArray<UMaterialFunctionInterface*>& Functions, TArray<TObjectPtr<UMaterialExpressionMaterialFunctionCall>>& CallerArray, int32& CallerCount, EMaterialParameterAssociation InAssociation) -> void
 			{
-				CallerArray.Empty();
 				CallerCount = 0;
+				if (Functions.IsEmpty())
+				{
+					CallerArray.Empty();
+					return;
+				}
+
+				int32 NumFunctions = Functions.Num();
+				int32 NumCallers = CallerArray.Num();
+
+				if (NumFunctions > NumCallers)
+				{
+					while (CallerArray.Num() < NumFunctions)
+					{
+						CallerArray.Push(NewObject<UMaterialExpressionMaterialFunctionCall>(GetTransientPackage()));
+					}
+				}
+				else if (NumFunctions < NumCallers)
+				{
+					CallerArray.SetNum(NumFunctions, EAllowShrinking::Yes);
+				}
+
 				for (int32 LayerIndex = 0; LayerIndex < Functions.Num(); LayerIndex++)
 				{
 					//Create the new function call
-					TObjectPtr<UMaterialExpressionMaterialFunctionCall> NewFunctionCall = NewObject<UMaterialExpressionMaterialFunctionCall>(GetTransientPackage());
-					NewFunctionCall->SetMaterialFunction(Functions[LayerIndex]);
-					NewFunctionCall->FunctionParameterInfo.Association = InAssociation;
-					NewFunctionCall->FunctionParameterInfo.Index = LayerIndex;
-					NewFunctionCall->UpdateFromFunctionResource();
-					if (!this->LayerInputs.IsEmpty() && !NewFunctionCall->FunctionInputs.IsEmpty())
+					TObjectPtr<UMaterialExpressionMaterialFunctionCall>& FunctionCaller = CallerArray[LayerIndex];
+					UMaterialFunctionInterface* CurrentFunction = Functions[LayerIndex];
+					FunctionCaller->FunctionParameterInfo.Index = LayerIndex;
+					if (FunctionCaller->MaterialFunction != CurrentFunction)
+					{
+						FunctionCaller->SetMaterialFunction(CurrentFunction);
+						FunctionCaller->FunctionParameterInfo.Association = InAssociation;						
+						FunctionCaller->UpdateFromFunctionResource();
+					}
+
+					if (!this->LayerInputs.IsEmpty() && !FunctionCaller->FunctionInputs.IsEmpty())
 					{
 						for (FMaterialLayerInput& LayerInput : this->LayerInputs)
 						{
@@ -557,9 +582,9 @@ void UMaterialExpressionLayerStack::RebuildLayerGraph(bool bReportErrors)
 								continue;
 							}
 
-							for (int32 CallInputIndex = 0; CallInputIndex < NewFunctionCall->FunctionInputs.Num(); CallInputIndex++)
+							for (int32 CallInputIndex = 0; CallInputIndex < FunctionCaller->FunctionInputs.Num(); CallInputIndex++)
 							{
-								FFunctionExpressionInput& ThisInput = NewFunctionCall->FunctionInputs[CallInputIndex];
+								FFunctionExpressionInput& ThisInput = FunctionCaller->FunctionInputs[CallInputIndex];
 								if (!ThisInput.ExpressionInput)
 								{
 									continue;
@@ -575,8 +600,8 @@ void UMaterialExpressionLayerStack::RebuildLayerGraph(bool bReportErrors)
 							}
 						}
 					}
-					CallerArray.Add(NewFunctionCall);
-					if (NewFunctionCall->MaterialFunction)
+
+					if (FunctionCaller->MaterialFunction)
 					{
 						CallerCount++;
 					}
