[AutoRTFM} Remove AutoRTFM/AutoRTFM.h

This header has been deprecated for months, and the last few usages have now been cleaned up.

diff --git a/Engine/Source/Runtime/AutoRTFM/Public/AutoRTFM/AutoRTFM.h b/Engine/Source/Runtime/AutoRTFM/Public/AutoRTFM/AutoRTFM.h
deleted file mode 100644
diff --git a/Engine/Source/Runtime/Engine/Private/PlayerController.cpp b/Engine/Source/Runtime/Engine/Private/PlayerController.cpp
--- a/Engine/Source/Runtime/Engine/Private/PlayerController.cpp
+++ b/Engine/Source/Runtime/Engine/Private/PlayerController.cpp
@@ -76,7 +76,7 @@
 #endif // UE_WITH_IRIS
 
 #if UE_WITH_REMOTE_OBJECT_HANDLE
-#include "AutoRTFM/AutoRTFM.h"
+#include "AutoRTFM.h"
 #include "UObject/UObjectMigrationContext.h"
 #include "UObject/RemoteExecutor.h"
 #endif
