[ZenLoader]
- Fix recursive loads to work under ALT
- Add a delay in postloads to make sure we can detect partial loads coming from ALT
- Avoid that delay when ALT is not active
- Avoid using thread-safe delegates since calling them currently is an exclusive lock and can cause deadlocks in the test when serialize goes wide

diff --git a/Engine/Source/Runtime/Engine/Private/Tests/Loading/AsyncLoadingTests_RecursiveLoads.cpp b/Engine/Source/Runtime/Engine/Private/Tests/Loading/AsyncLoadingTests_RecursiveLoads.cpp
--- a/Engine/Source/Runtime/Engine/Private/Tests/Loading/AsyncLoadingTests_RecursiveLoads.cpp
+++ b/Engine/Source/Runtime/Engine/Private/Tests/Loading/AsyncLoadingTests_RecursiveLoads.cpp
@@ -24,17 +24,45 @@ bool FLoadingTests_RecursiveLoads_FromSerialize::RunTest(const FString& Paramete
 {
 	FLoadingTestsScope LoadingTestScope(this);
 
+	UE::FManualResetEvent PostLoadEvent;
+	std::atomic<bool> bHasBeenPostloaded = false;
+	std::atomic<bool> bIsMultithreadedLoading = false;
+
+	UAsyncLoadingTests_Shared::OnPostLoad.BindLambda(
+		[this, &PostLoadEvent, &bHasBeenPostloaded, &bIsMultithreadedLoading](UAsyncLoadingTests_Shared* Object)
+		{
+			if (Object->GetPathName() == FLoadingTestsScope::ObjectPath2)
+			{
+				check(IsInGameThread());
+
+				if (bIsMultithreadedLoading)
+				{
+					// If we run with ALT enabled, we want just enough time to make sure the flush
+					// can return before this function exits, proving that we're not waiting on postload
+					// during serialize.
+					PostLoadEvent.WaitFor(UE::FMonotonicTimeSpan::FromSeconds(1.0));
+				}
+
+				bHasBeenPostloaded = true;
+			}
+		}
+	);
+
 	UAsyncLoadingTests_Shared::OnSerialize.BindLambda(
-		[this](FArchive& Ar, UAsyncLoadingTests_Shared* Object)
+		[this, &bHasBeenPostloaded, &PostLoadEvent, &bIsMultithreadedLoading](FArchive& Ar, UAsyncLoadingTests_Shared* Object)
 		{
 			if (Ar.IsLoading())
 			{
+				bIsMultithreadedLoading = !IsInGameThread();
+
 				if (UObject* Obj = Object->SoftReference.LoadSynchronous())
 				{
 					TestTrue(TEXT("Recursive loads in serialize should be deserialized"), !Obj->HasAnyFlags(RF_NeedLoad));
 					if (!IsInGameThread())
 					{
-						TestTrue(TEXT("Recursive loads in serialize skip thread-unsafe postloads when run from the ALT"), Obj->HasAnyFlags(RF_NeedPostLoad));
+						TestFalse(TEXT("Recursive loads in serialize skip thread-unsafe postloads when run from the ALT"), bHasBeenPostloaded);
+
+						PostLoadEvent.Notify();
 					}
 				}
 			}
@@ -180,21 +208,49 @@ bool FLoadingTests_RecursiveLoads_FullFlushFrom_Serialize::RunTest(const FString
 
 	AddExpectedMessage(TEXT("will result in a partially loaded package to avoid a deadlock."), EAutomationExpectedErrorFlags::Contains);
 
+	UE::FManualResetEvent PostLoadEvent;
+	std::atomic<bool> bHasBeenPostloaded = false;
+	std::atomic<bool> bIsMultithreadedLoading = false;
+
+	UAsyncLoadingTests_Shared::OnPostLoad.BindLambda(
+		[this, &PostLoadEvent, &bHasBeenPostloaded, &bIsMultithreadedLoading](UAsyncLoadingTests_Shared* Object)
+		{
+			if (Object->GetPathName() == FLoadingTestsScope::ObjectPath2)
+			{
+				check(IsInGameThread());
+
+				if (bIsMultithreadedLoading)
+				{
+					// If we run with ALT enabled, we want just enough time to make sure the flush
+					// can return before this function exits, proving that we're not waiting on postload
+					// during serialize.
+					PostLoadEvent.WaitFor(UE::FMonotonicTimeSpan::FromSeconds(1.0));
+				}
+
+				bHasBeenPostloaded = true;
+			}
+		}
+	);
+
 	int32 RequestId = 0;
 	UAsyncLoadingTests_Shared::OnSerialize.BindLambda(
-		[this, &RequestId](FArchive& Ar, UAsyncLoadingTests_Shared* Object)
+		[this, &RequestId, &PostLoadEvent, &bHasBeenPostloaded, &bIsMultithreadedLoading](FArchive& Ar, UAsyncLoadingTests_Shared* Object)
 		{
 			// Do not try to flush ourself as this would lead to a fatal error :)
 			// Just flush Package2 when we're in Package1
 			if (Ar.IsLoading() && Object->GetPathName() == FLoadingTestsScope::ObjectPath1)
 			{
+				bIsMultithreadedLoading = !IsInGameThread();
+
 				// Flush the requestId that has been created outside of the recursive load. This request
 				// should be a full request and flushing it should result in an error being reported.
 				FlushAsyncLoading(RequestId);
 
 				UAsyncLoadingTests_Shared* Object2 = FindObject<UAsyncLoadingTests_Shared>(nullptr, FLoadingTestsScope::ObjectPath2);
 				TestFalse(TEXT("The object should be serialized"), Object2->HasAnyFlags(RF_NeedLoad));
-				TestTrue(TEXT("The object should not have been postloaded"), Object2->HasAnyFlags(RF_NeedPostLoad));
+				TestFalse(TEXT("The object should not have been postloaded"), bHasBeenPostloaded);
+
+				PostLoadEvent.Notify();
 			}
 		}
 	);
@@ -262,14 +318,20 @@ bool FLoadingTests_RecursiveLoads_FromBothSerializeAndPostLoad::RunTest(const FS
 	std::atomic<bool> PartialLoadAchieved = false;
 	UE::FManualResetEvent Event;
 	int32 SerializeCount = 0;
+
+	UE::FManualResetEvent PostLoadEvent;
+	std::atomic<bool> bHasBeenPostloaded = false;
+	std::atomic<bool> bIsMultithreadedLoading = false;
+	
 	// On serialize we try to force load, this should add the newly loaded package as a dynamic import of the package that requested them.
 	UAsyncLoadingTests_Shared::OnSerialize.BindLambda(
-		[this, &LoadingTestScope, &PartialLoadAchieved, &Event, &SerializeCount](FArchive& Ar, UAsyncLoadingTests_Shared* Object)
+		[this, &LoadingTestScope, &PartialLoadAchieved, &Event, &SerializeCount, &bHasBeenPostloaded, &PostLoadEvent, &bIsMultithreadedLoading](FArchive& Ar, UAsyncLoadingTests_Shared* Object)
 		{
 			if (Ar.IsLoading() && Object->GetPathName() == FLoadingTestsScope::ObjectPath1)
 			{
 				SerializeCount++;
 				check(SerializeCount == 1);
+				bIsMultithreadedLoading = !IsInGameThread();
 
 				// Loading Object2 while being in Object1 serialize should trigger partial load of Object2
 				if (UObject* Obj = Object->SoftReference.LoadSynchronous())
@@ -277,9 +339,11 @@ bool FLoadingTests_RecursiveLoads_FromBothSerializeAndPostLoad::RunTest(const FS
 					TestTrue(TEXT("Recursive loads in serialize should be deserialized"), !Obj->HasAnyFlags(RF_NeedLoad));
 					if (!IsInGameThread())
 					{
-						TestTrue(TEXT("Recursive loads in serialize skip thread-unsafe postloads when run from the ALT"), Obj->HasAnyFlags(RF_NeedPostLoad));
+						TestFalse(TEXT("Recursive loads in serialize skip thread-unsafe postloads when run from the ALT"), bHasBeenPostloaded);
 					}
 
+					PostLoadEvent.Notify();
+
 					PartialLoadAchieved = true;
 
 					// When everything runs on GT, we will have no choice but to use the whole time
@@ -293,8 +357,23 @@ bool FLoadingTests_RecursiveLoads_FromBothSerializeAndPostLoad::RunTest(const FS
 
 	// Once in postload of object3, we now force load the same softref a second time, expecting to be able to postload it without deadlocking because of the merged postload groups.
 	UAsyncLoadingTests_Shared::OnPostLoad.BindLambda(
-		[this, &LoadingTestScope](UAsyncLoadingTests_Shared* Object)
+		[this, &LoadingTestScope, &bHasBeenPostloaded, &PostLoadEvent, &bIsMultithreadedLoading](UAsyncLoadingTests_Shared* Object)
 		{
+			if (Object->GetPathName() == FLoadingTestsScope::ObjectPath2)
+			{
+				check(IsInGameThread());
+
+				if (bIsMultithreadedLoading)
+				{
+					// If we run with ALT enabled, we want just enough time to make sure the flush
+					// can return before this function exits, proving that we're not waiting on postload
+					// during serialize.
+					PostLoadEvent.WaitFor(UE::FMonotonicTimeSpan::FromSeconds(1.0));
+				}
+
+				bHasBeenPostloaded = true;
+			}
+
 			if (Object->GetPathName() == FLoadingTestsScope::ObjectPath3)
 			{
 				// Trying to load Package2 completely while being in package 3 postload should be able to succeed
diff --git a/Engine/Source/Runtime/Engine/Private/Tests/Loading/AsyncLoadingTests_Shared.h b/Engine/Source/Runtime/Engine/Private/Tests/Loading/AsyncLoadingTests_Shared.h
--- a/Engine/Source/Runtime/Engine/Private/Tests/Loading/AsyncLoadingTests_Shared.h
+++ b/Engine/Source/Runtime/Engine/Private/Tests/Loading/AsyncLoadingTests_Shared.h
@@ -22,16 +22,18 @@ public:
 	TObjectPtr<UObject> HardReference;
 
 	// Those delegates allow to easily change behavior between tests
-	using FOnPostLoadDelegate = TDelegate<void(UAsyncLoadingTests_Shared*), FDefaultTSDelegateUserPolicy>;
+	// Avoid using thread-safe delegates since calling them currently is an exclusive lock and can cause deadlocks in the test when serialize is multithreaded.
+	// Registration to these delegates happens safely when the test begins before threads are involved.
+	using FOnPostLoadDelegate = TDelegate<void(UAsyncLoadingTests_Shared*), FNotThreadSafeNotCheckedDelegateUserPolicy>;
 	static FOnPostLoadDelegate OnPostLoad;
 
-	using FOnSerializeDelegate = TDelegate<void(FArchive& Ar, UAsyncLoadingTests_Shared*), FDefaultTSDelegateUserPolicy>;
+	using FOnSerializeDelegate = TDelegate<void(FArchive& Ar, UAsyncLoadingTests_Shared*), FNotThreadSafeNotCheckedDelegateUserPolicy>;
 	static FOnSerializeDelegate OnSerialize;
 
-	using FOnIsReadyForAsyncPostLoadDelegate = TDelegate<bool(const UAsyncLoadingTests_Shared*), FDefaultTSDelegateUserPolicy>;
+	using FOnIsReadyForAsyncPostLoadDelegate = TDelegate<bool(const UAsyncLoadingTests_Shared*), FNotThreadSafeNotCheckedDelegateUserPolicy>;
 	static FOnIsReadyForAsyncPostLoadDelegate OnIsReadyForAsyncPostLoad;
 
-	using FOnIsPostLoadThreadSafeDelegate = TDelegate<bool(const UAsyncLoadingTests_Shared*), FDefaultTSDelegateUserPolicy>;
+	using FOnIsPostLoadThreadSafeDelegate = TDelegate<bool(const UAsyncLoadingTests_Shared*), FNotThreadSafeNotCheckedDelegateUserPolicy>;
 	static FOnIsReadyForAsyncPostLoadDelegate OnIsPostLoadThreadSafe;
 
 	virtual void PostLoad() override
