Level Instance: don't report invalid level instance map settings for cooked level instances referenced by uncooked maps.

#preflkight 6855ab30be75c7640225c438

diff --git a/Engine/Source/Runtime/Engine/Classes/Engine/Level.h b/Engine/Source/Runtime/Engine/Classes/Engine/Level.h
--- a/Engine/Source/Runtime/Engine/Classes/Engine/Level.h
+++ b/Engine/Source/Runtime/Engine/Classes/Engine/Level.h
@@ -768,6 +768,9 @@ public:
 	ENGINE_API static bool GetIsStreamingDisabledFromAsset(const FAssetData& Asset);	
 	ENGINE_API static bool GetIsStreamingDisabledFromPackage(FName LevelPackage);
 
+	ENGINE_API static bool GetIsLevelCookedFromAsset(const FAssetData& Asset);
+	ENGINE_API static bool GetIsLevelCookedFromPackage(FName LevelPackage);
+
 	UE_DEPRECATED(5.3, "GetPartitionedLevelCanBeUsedByLevelInstanceFromAsset is deprecated.")
 	static bool GetPartitionedLevelCanBeUsedByLevelInstanceFromAsset(const FAssetData& Asset) { return true; }
 	UE_DEPRECATED(5.3, "GetPartitionedLevelCanBeUsedByLevelInstanceFromPackage is deprecated.")
diff --git a/Engine/Source/Runtime/Engine/Private/Level.cpp b/Engine/Source/Runtime/Engine/Private/Level.cpp
--- a/Engine/Source/Runtime/Engine/Private/Level.cpp
+++ b/Engine/Source/Runtime/Engine/Private/Level.cpp
@@ -2944,6 +2944,19 @@ bool ULevel::GetIsStreamingDisabledFromPackage(FName LevelPackage)
 	});
 }
 
+bool ULevel::GetIsLevelCookedFromAsset(const FAssetData& Asset)
+{
+	return !!(Asset.PackageFlags & PKG_Cooked);
+}
+
+bool ULevel::GetIsLevelCookedFromPackage(FName LevelPackage)
+{
+	return LevelAssetRegistryHelper::GetLevelInfoFromAssetRegistry(LevelPackage, [](const FAssetData& Asset)
+	{
+		return GetIsLevelCookedFromAsset(Asset);
+	});
+}
+
 bool ULevel::GetWorldExternalActorsReferencesFromPackage(FName LevelPackage, TArray<FGuid>& OutWorldExternalActorsReferences)
 {
 	return LevelAssetRegistryHelper::GetLevelInfoFromAssetRegistry(LevelPackage, [&OutWorldExternalActorsReferences](const FAssetData& Asset)
diff --git a/Engine/Source/Runtime/Engine/Private/WorldPartition/LevelInstance/LevelInstanceActorDesc.cpp b/Engine/Source/Runtime/Engine/Private/WorldPartition/LevelInstance/LevelInstanceActorDesc.cpp
--- a/Engine/Source/Runtime/Engine/Private/WorldPartition/LevelInstance/LevelInstanceActorDesc.cpp
+++ b/Engine/Source/Runtime/Engine/Private/WorldPartition/LevelInstance/LevelInstanceActorDesc.cpp
@@ -259,7 +259,7 @@ void FLevelInstanceActorDesc::CheckForErrors(const IWorldPartitionActorDescInsta
 	}
 	else
 	{
-		if (!ULevel::GetIsLevelUsingExternalActorsFromPackage(ChildContainerPackage))
+		if (!ULevel::GetIsLevelUsingExternalActorsFromPackage(ChildContainerPackage) && !ULevel::GetIsLevelCookedFromPackage(ChildContainerPackage))
 		{
 			if (DesiredRuntimeBehavior != ELevelInstanceRuntimeBehavior::LevelStreaming)
 			{
