A NetDriver command for printing out all of the actors registered for replication (generic replication system only).

diff --git a/Engine/Source/Runtime/Engine/Private/NetDriver.cpp b/Engine/Source/Runtime/Engine/Private/NetDriver.cpp
--- a/Engine/Source/Runtime/Engine/Private/NetDriver.cpp
+++ b/Engine/Source/Runtime/Engine/Private/NetDriver.cpp
@@ -336,6 +336,53 @@ namespace UE::Net::Private
 		SendDormantDestructionOnRemoval,
 		TEXT("If true and net.SerializeNewActorOverrideLevel = 0, servers will send destruction info for dormant destroyed actors in invisible levels, regardless of how they were removed. This ensures they destroy on the client, in-case the they (or the levels they are in) were removed unexpectedly."));
 
+	static void DumpActiveNetActors(UWorld* InWorld)
+	{
+		UNetDriver *NetDriver = InWorld->NetDriver;
+		if (!NetDriver)
+		{
+			return;
+		}
+
+		UE_LOG(LogConsoleResponse, Log, TEXT("----------------------------"));
+		UE_LOG(LogConsoleResponse, Log, TEXT("Active network actors in %s:"), *NetDriver->GetName());
+
+		TMap<UClass*, int32> ActorClassCount;
+
+		for (const TSharedPtr<FNetworkObjectInfo>& ObjectInfo : NetDriver->GetNetworkObjectList().GetActiveObjects())
+		{
+			AActor* Actor = ObjectInfo->WeakActor.Get();
+			if (!Actor)
+			{
+				continue;
+			}
+
+			UE_LOG(LogConsoleResponse, Log, TEXT(" * %s (%s)"), *Actor->GetName(), *Actor->GetClass()->GetName());
+
+			int32& ClassCount = ActorClassCount.FindOrAdd(Actor->GetClass(), 0);
+			ClassCount++;
+		}
+
+		UE_LOG(LogConsoleResponse, Log, TEXT("----------------------------"));
+
+		UE_LOG(LogConsoleResponse, Log, TEXT("Number of each active actor class type:"));
+
+		for (TPair<UClass*, int32> ClassCount : ActorClassCount)
+		{
+			UE_LOG(LogConsoleResponse, Log, TEXT(" %s: %d"), *ClassCount.Key->GetName(), ClassCount.Value);
+		}
+
+		UE_LOG(LogConsoleResponse, Log, TEXT("----------------------------"));
+		UE_LOG(LogConsoleResponse, Log, TEXT("Total: %d"), NetDriver->GetNetworkObjectList().GetActiveObjects().Num());
+		UE_LOG(LogConsoleResponse, Log, TEXT("----------------------------"));
+	}
+
+	FAutoConsoleCommandWithWorld DumpActiveNetActorsCommand(
+		TEXT("net.DumpActiveNetActors"), 
+		TEXT("Dumps all of the active net actors in the replication system."),
+		FConsoleCommandWithWorldDelegate::CreateStatic(DumpActiveNetActors)
+	);
+
 } //namespace UE::Net::Private
 
 namespace UE::Net
