FWorldCookPackageSplitter:
- PopulateGeneratedPackage will now be called on the subsplitter that issued the generated package instead of all subsplitters.
- We could do the same for PopulateGeneratorPackage but the existing code keeps internal data instead of relying on the passed FPopulateContext so it's not an issue.
- Fixed assert in FWorldPartitionCookPackageSplitter when cooking worlds with more than one cook package subsplitter.

diff --git a/Engine/Source/Runtime/Engine/Private/WorldCookPackageSplitter.cpp b/Engine/Source/Runtime/Engine/Private/WorldCookPackageSplitter.cpp
--- a/Engine/Source/Runtime/Engine/Private/WorldCookPackageSplitter.cpp
+++ b/Engine/Source/Runtime/Engine/Private/WorldCookPackageSplitter.cpp
@@ -94,9 +94,26 @@ TArray<ICookPackageSplitter::FGeneratedPackage> FWorldCookPackageSplitter::GetGe
 
 	TArray<ICookPackageSplitter::FGeneratedPackage> PackagesToGenerate;
 	ForEachCookPackageSubSplitters(
-		[&PackagesToGenerate, OwnerPackage, OwnerWorld](FWorldCookPackageSplitter::ISubSplitter* WorldCookPackageSplitter)
+		[this, &PackagesToGenerate, OwnerPackage, OwnerWorld](FWorldCookPackageSplitter::ISubSplitter* WorldCookPackageSplitter)
 		{
-			PackagesToGenerate.Append(WorldCookPackageSplitter->GetGenerateList(OwnerPackage));
+			const TArray<ICookPackageSplitter::FGeneratedPackage> GenerateList = WorldCookPackageSplitter->GetGenerateList(OwnerPackage);
+			
+			PackagesToGenerate.Reserve(PackagesToGenerate.Num() + GenerateList.Num());
+			for (const ICookPackageSplitter::FGeneratedPackage& GeneratePackage : GenerateList)
+			{
+				ICookPackageSplitter::FGeneratedPackage LocalGeneratePackage(GeneratePackage);
+
+				if (LocalGeneratePackage.GeneratedRootPath.IsEmpty())
+				{
+					LocalGeneratePackage.GeneratedRootPath = OwnerPackage->GetName();
+				}
+
+				PackagesToGenerate.Add(LocalGeneratePackage);
+
+				const TPair<FName, FName> FullPackageName = { *LocalGeneratePackage.GeneratedRootPath, *LocalGeneratePackage.RelativePath };
+				check(!SplittersGenerateListMap.Contains(FullPackageName));
+				SplittersGenerateListMap.Add(FullPackageName, WorldCookPackageSplitter);
+			}
 		});
 
 	FWorldDelegates::OnWorldCleanup.AddRaw(this, &FWorldCookPackageSplitter::OnWorldCleanup);
@@ -106,14 +123,10 @@ TArray<ICookPackageSplitter::FGeneratedPackage> FWorldCookPackageSplitter::GetGe
 
 bool FWorldCookPackageSplitter::PopulateGeneratedPackage(FPopulateContext& PopulateContext)
 {
-	bool bResult = false;
-	ForEachCookPackageSubSplitters(
-		[&bResult, &PopulateContext](FWorldCookPackageSplitter::ISubSplitter* WorldCookPackageSplitter)
-		{
-			bResult = WorldCookPackageSplitter->PopulateGeneratedPackage(PopulateContext);
-			return bResult;
-		});
-	return bResult;
+	const ICookPackageSplitter::FGeneratedPackageForPopulate* TargetGeneratedPackage = PopulateContext.GetTargetGeneratedPackage();
+	const TPair<FName, FName> FullPackageName = { *TargetGeneratedPackage->GeneratedRootPath, *TargetGeneratedPackage->RelativePath };
+	ISubSplitter* SubSplitter = SplittersGenerateListMap.FindChecked(FullPackageName);
+	return SubSplitter->PopulateGeneratedPackage(PopulateContext);
 }
 
 bool FWorldCookPackageSplitter::PopulateGeneratorPackage(FPopulateContext& PopulateContext)
@@ -143,6 +156,8 @@ void FWorldCookPackageSplitter::Teardown(ICookPackageSplitter::ETeardown Status)
 			WorldCookPackageSplitter->Teardown(Status);
 		});
 
+	SplittersGenerateListMap.Empty();
+
 	for (const TPair<const FWorldCookPackageSplitter::FSubSplitterFactory*, ISubSplitter*>& CookPackageSubSplitterPair : CookPackageSubSplitters)
 	{
 		CookPackageSubSplitterPair.Key->ReleaseInstance(ReferencedWorld, CookPackageSubSplitterPair.Value);
@@ -158,5 +173,4 @@ void FWorldCookPackageSplitter::Teardown(ICookPackageSplitter::ETeardown Status)
 
 	ReferencedWorld = nullptr;
 }
-
 #endif // WITH_EDITOR
\ No newline at end of file
diff --git a/Engine/Source/Runtime/Engine/Private/WorldPartition/WorldPartitionCookPackageSplitter.cpp b/Engine/Source/Runtime/Engine/Private/WorldPartition/WorldPartitionCookPackageSplitter.cpp
--- a/Engine/Source/Runtime/Engine/Private/WorldPartition/WorldPartitionCookPackageSplitter.cpp
+++ b/Engine/Source/Runtime/Engine/Private/WorldPartition/WorldPartitionCookPackageSplitter.cpp
@@ -77,11 +77,6 @@ bool FWorldPartitionCookPackageSplitter::PopulateGeneratorPackage(ICookPackageSp
 	UE_LOG(LogWorldPartition, Display, TEXT("[Cook][PopulateGeneratorPackage] Processing %u packages"), GeneratedPackages.Num());
 
 	bool bIsSuccess = true;
-	if (GeneratedPackages.Num() != CookContext.NumPackageToGenerate())
-	{
-		UE_LOG(LogWorldPartition, Error, TEXT("[Cook][PopulateGeneratorPackage] Receieved %u generated packages. Was expecting %u"), GeneratedPackages.Num(), CookContext.NumPackageToGenerate());
-		bIsSuccess = false;
-	}
 
 	TArray<UPackage*> ModifiedPackages;
 	for (IWorldPartitionCookPackageGenerator* CookPackageGenerator : CookContext.GetCookPackageGenerators())
diff --git a/Engine/Source/Runtime/Engine/Public/WorldCookPackageSplitter.h b/Engine/Source/Runtime/Engine/Public/WorldCookPackageSplitter.h
--- a/Engine/Source/Runtime/Engine/Public/WorldCookPackageSplitter.h
+++ b/Engine/Source/Runtime/Engine/Public/WorldCookPackageSplitter.h
@@ -75,6 +75,7 @@ protected:
 	//~ End of ICookPackageSplitter
 
 	TArray<TPair<const FWorldCookPackageSplitter::FSubSplitterFactory*, ISubSplitter*>> CookPackageSubSplitters;
+	TMap<TPair<FName, FName>, ISubSplitter*> SplittersGenerateListMap;
 
 	template<typename FuncType>
 	bool ForEachCookPackageSubSplitters(FuncType Func)
