Remove restriction on non-layer function assets when managing parameter owners in the layering system by checking whether they're associated with layers rather than checking if they're a layer type function asset specifically.

diff --git a/Engine/Source/Runtime/Engine/Private/Materials/MaterialExpressions.cpp b/Engine/Source/Runtime/Engine/Private/Materials/MaterialExpressions.cpp
--- a/Engine/Source/Runtime/Engine/Private/Materials/MaterialExpressions.cpp
+++ b/Engine/Source/Runtime/Engine/Private/Materials/MaterialExpressions.cpp
@@ -17352,7 +17352,7 @@ void UMaterialExpressionMaterialFunctionCall::PushParameterOwner(FMaterialCompil
 {
 	// Update parameter owner when stepping into layer functions.
 	// This is an optional step when we only want to march the material graph (e.g. to gather Substrate material topology)
-	if (Compiler && MaterialFunction->GetMaterialFunctionUsage() != EMaterialFunctionUsage::Default)
+	if (Compiler && FunctionParameterInfo.Association != EMaterialParameterAssociation::GlobalParameter)
 	{
 		Compiler->PushParameterOwner(FunctionParameterInfo);
 	}
@@ -17360,7 +17360,7 @@ void UMaterialExpressionMaterialFunctionCall::PushParameterOwner(FMaterialCompil
 
 void UMaterialExpressionMaterialFunctionCall::PopParameterOwner(FMaterialCompiler* Compiler)
 {
-	if (Compiler && MaterialFunction->GetMaterialFunctionUsage() != EMaterialFunctionUsage::Default)
+	if (Compiler && FunctionParameterInfo.Association != EMaterialParameterAssociation::GlobalParameter)
 	{
 		const FMaterialParameterInfo PoppedParameterInfo = Compiler->PopParameterOwner();
 		check(PoppedParameterInfo == FunctionParameterInfo);
