Bump streamed audio DDC key to address cache poisoning stemming from CL43678646.

diff --git a/Engine/Source/Runtime/Engine/Private/AudioDerivedData.cpp b/Engine/Source/Runtime/Engine/Private/AudioDerivedData.cpp
--- a/Engine/Source/Runtime/Engine/Private/AudioDerivedData.cpp
+++ b/Engine/Source/Runtime/Engine/Private/AudioDerivedData.cpp
@@ -178,7 +178,7 @@ Derived data key generation.
 // If you want to bump this version, generate a new guid using
 // VS->Tools->Create GUID and paste it here. https://www.guidgen.com works too.
 #define AUDIO_DERIVEDDATA_VER				TEXT("64e45415311549acbee28698853aa3a4")				// Change this if you want to bump all audio formats at once
-#define STREAMEDAUDIO_DERIVEDDATA_VER		TEXT("adfedd669bf247aaa06e5f0d25b9f4ce")				// This depends on the above key, but will regenerate all streaming chunk data derived from the compressed audio
+#define STREAMEDAUDIO_DERIVEDDATA_VER		TEXT("479371f6537b4fc1be91786813e83a16")				// This depends on the above key, but will regenerate all streaming chunk data derived from the compressed audio
 
 #if WITH_EDITORONLY_DATA
 
