Avoid iterating over net connections if the actor we are evaluating has the correct world, and if it doesn't check if the actor is a player controller first.

diff --git a/Engine/Source/Runtime/Engine/Private/NetDriver.cpp b/Engine/Source/Runtime/Engine/Private/NetDriver.cpp
--- a/Engine/Source/Runtime/Engine/Private/NetDriver.cpp
+++ b/Engine/Source/Runtime/Engine/Private/NetDriver.cpp
@@ -2755,10 +2755,22 @@ bool UNetDriver::IsLevelInitializedForActor(const AActor* InActor, const UNetCon
 	check(World == InActor->GetWorld());
 #endif
 
+	UPackage* WorldPackageRawPtr = GetWorldPackage();
+	const bool bWorldHasSamePackageName = WorldPackageRawPtr != nullptr && InConnection->GetClientWorldPackageName() == WorldPackageRawPtr->GetFName();
+
 	// we can't create channels while the client is in the wrong world
-	const bool bCorrectWorld = GetWorldPackage() != nullptr && (InConnection->GetClientWorldPackageName() == GetWorldPackage()->GetFName()) && InConnection->ClientHasInitializedLevel(InActor->GetLevel());
-	
+	if (bWorldHasSamePackageName && InConnection->ClientHasInitializedLevel(InActor->GetLevel()))
+	{
+		return true;
+	}
+
 	// exception: Special case for PlayerControllers as they are required for the client to travel to the new world correctly			
+	const APlayerController* AsPlayerController = Cast<APlayerController>(InActor);
+	if (!AsPlayerController)
+	{
+		return false;
+	}
+
 	bool bIsConnectionPC = false;
 	// Make sure we have Parent Connection
 	const UNetConnection* const ParentConnection = InConnection->IsA<UChildConnection>() ? Cast<UChildConnection>(InConnection)->Parent : InConnection;
@@ -2770,7 +2782,7 @@ bool UNetDriver::IsLevelInitializedForActor(const AActor* InActor, const UNetCon
 	}
 	// Check Child Connections
 	else
-	{	
+	{
 		for (const UChildConnection* const ChildConnection : ParentConnection->Children)
 		{
 			if (InActor == ChildConnection->PlayerController)
@@ -2781,7 +2793,7 @@ bool UNetDriver::IsLevelInitializedForActor(const AActor* InActor, const UNetCon
 		}
 	}
 
-	return bCorrectWorld || bIsConnectionPC;
+	return bIsConnectionPC;
 }
 
 //
