Potential fix for deadlock in CancellAllPendingLoads

diff --git a/Engine/Source/Runtime/Engine/Private/AudioStreamingCache.cpp b/Engine/Source/Runtime/Engine/Private/AudioStreamingCache.cpp
--- a/Engine/Source/Runtime/Engine/Private/AudioStreamingCache.cpp
+++ b/Engine/Source/Runtime/Engine/Private/AudioStreamingCache.cpp
@@ -1354,12 +1354,22 @@ void FAudioChunkCache::BlockForAllPendingLoads() const
 
 void FAudioChunkCache::CancelAllPendingLoads()
 {
-	FScopeLock ScopeLock(&CacheMutationCriticalSection.Get());
-	FCacheElement* CurrentElement = MostRecentElement;
-	while (CurrentElement != nullptr)
+	TArray<FCacheElement*> ToCancel;
 	{
-		CurrentElement->WaitForAsyncLoadCompletion(true);
-		CurrentElement = CurrentElement->LessRecentElement;
+		// Lock and collect all elements.
+		FScopeLock ScopeLock(&CacheMutationCriticalSection.Get());
+		FCacheElement* Current = MostRecentElement;
+		while (Current != nullptr)
+		{
+			ToCancel.Add(Current);
+			Current = Current->LessRecentElement;
+		}
+	}
+
+	// Wait for completion outside the lock.
+	for (FCacheElement* Elem : ToCancel)
+	{
+		Elem->WaitForAsyncLoadCompletion(/*bCancel=*/true);
 	}
 }
 
