Subtitles: Make default priority 1, clean up tooltips/comments.

diff --git a/Engine/Source/Runtime/Engine/Public/Subtitles/SubtitlesAndClosedCaptionsDelegates.h b/Engine/Source/Runtime/Engine/Public/Subtitles/SubtitlesAndClosedCaptionsDelegates.h
--- a/Engine/Source/Runtime/Engine/Public/Subtitles/SubtitlesAndClosedCaptionsDelegates.h
+++ b/Engine/Source/Runtime/Engine/Public/Subtitles/SubtitlesAndClosedCaptionsDelegates.h
@@ -58,13 +58,13 @@ class USubtitleAssetUserData : public UAssetUserData
 public:
 
 	// #SUBTITLES_PRD: carried over from FSubtitleCue, still required
-	//
+
 	// The text to appear in the subtitle.
 	UPROPERTY(EditAnywhere, BlueprintReadWrite, Category = Subtitles, meta = (Multiline = true))
 	FText Text;
 
 	// #SUBTITLES_PRD: carried over from FSubtitleCue, still required
-	//
+
 	// Time to display in seconds. 
 	// Defaulted to 3 seconds so when adding new subtitles it's not required to enter a placeholder Duration.
 	// Duration can be be set by ingestion pipelines when importing Subtitles in bulk
@@ -79,11 +79,10 @@ public:
 
 	// #SUBTITLES_PRD: Priority comes from USoundBase::GetSubtitlePriority, USoundCue::GetSubtitlePriority and USoundWave::GetSubtitlePriority
 	// Consolidate various subtitle properties throughout sound/audio code into this new subtitles plugin.
-	// Kept the 10000 default in case that's what users are already used to.
-	// 
-	// The priority of the subtitle.  Defaults to 10000.  Higher values will play instead of lower values.
+
+	// The priority of the subtitle.  Defaults to 1.  Higher values will play instead of lower values.
 	UPROPERTY(EditAnywhere, BlueprintReadWrite, Category = Subtitles)
-	float Priority = 10000.f;
+	float Priority = 1.f;
 
 	// ESRB rating category.
 	UPROPERTY(EditAnywhere, BlueprintReadWrite, Category = Subtitles)
