RemoteObjects:  Missed a case for forward-loaning a borrowed object.  We assumed that if we were loaning we're an owner (but that's not always the case).

#rnx

diff --git a/Engine/Source/Runtime/Engine/Private/Actor.cpp b/Engine/Source/Runtime/Engine/Private/Actor.cpp
--- a/Engine/Source/Runtime/Engine/Private/Actor.cpp
+++ b/Engine/Source/Runtime/Engine/Private/Actor.cpp
@@ -1306,7 +1306,7 @@ void AActor::PostMigrate(const FUObjectMigrationContext& MigrationContext)
 		UE_LOG(LogActor, Verbose, TEXT("%hs: Send %s (LocalSrvId: %s. DstSrvId: %s. HasOwnership: %d. MigrationType: %s)"), __func__, *GetName(), *LocalServerId.ToString(), *MigrationContext.RemoteServerId.ToString(), bHasOwnership, ToString(MigrationSendType));
 		ensureMsgf(MigrationContext.RemoteServerId != LocalServerId, TEXT("%hs: Sending Object %s from %s to %s but destination is the same as the Local Server Id."), __func__, *GetName(), *LocalServerId.ToString(), *MigrationContext.RemoteServerId.ToString());
 
-		const bool bKeepOwnership = (MigrationSendType == EObjectMigrationSendType::Loan);
+		const bool bKeepOwnership = (MigrationSendType == EObjectMigrationSendType::Loan) && UE::RemoteObject::Handle::IsOwned(this);
 		if (bKeepOwnership)
 		{
 			RemoteObjects_ActorLoanedStates.AddAnnotation(this, { .bActorWasDestroyed = (bool)bActorIsBeingDestroyed, .bTickFunctionsRegistered = (bool)bTickFunctionsRegistered });
diff --git a/Engine/Source/Runtime/Engine/Private/Components/ActorComponent.cpp b/Engine/Source/Runtime/Engine/Private/Components/ActorComponent.cpp
--- a/Engine/Source/Runtime/Engine/Private/Components/ActorComponent.cpp
+++ b/Engine/Source/Runtime/Engine/Private/Components/ActorComponent.cpp
@@ -717,7 +717,7 @@ void UActorComponent::PostMigrate(const struct FUObjectMigrationContext& Migrati
 	{
 		const EObjectMigrationSendType SendType = MigrationContext.GetObjectMigrationSendType(this);
 
-		const bool bKeepOwnership = (SendType == EObjectMigrationSendType::Loan);
+		const bool bKeepOwnership = (SendType == EObjectMigrationSendType::Loan) && UE::RemoteObject::Handle::IsOwned(this);
 		if (bKeepOwnership)
 		{
 			checkfSlow(!RemoteObjects_LoanedStates.GetAnnotationMap().Find(this), TEXT("%hs: Expected no existing Loaned State for %s (it should have been consumed when received, or not loaned out twice in a row)"), __func__, *GetName());
