[NetPhysics] Iris support in the NetworkPhysicsComponent
- Setup a new base for custom inputs and states that handle Iris serialization, without a NetSerialize function
- Setup a new struct to network an array of the derived custom inputs and states, bsed on the new base struct.
- Setup an interface to communicate inputs and states with the implementation, i.e. apply to the implementation and and build from the implementation
- Implemented templated helper functions to copy data between the rewind history and the iris replicated array

diff --git a/Engine/Source/Runtime/Engine/Private/PhysicsEngine/NetworkPhysicsComponent.cpp b/Engine/Source/Runtime/Engine/Private/PhysicsEngine/NetworkPhysicsComponent.cpp
--- a/Engine/Source/Runtime/Engine/Private/PhysicsEngine/NetworkPhysicsComponent.cpp
+++ b/Engine/Source/Runtime/Engine/Private/PhysicsEngine/NetworkPhysicsComponent.cpp
@@ -338,9 +338,10 @@ int32 FNetworkPhysicsCallback::TriggerRewindIfNeeded_Internal(int32 LatestStepCo
 
 	if (RewindData)
 	{
+		int32 TargetStateComparisonFrame = INDEX_NONE;
 		if (!PhysicsReplicationCVars::ResimulationCVars::bApplyPredictiveInterpolationWhenBehindServer)
 		{
-			const int32 TargetStateComparisonFrame = RewindData->CompareTargetsToLastFrame();
+			TargetStateComparisonFrame = RewindData->CompareTargetsToLastFrame();
 			ResimFrame = (ResimFrame == INDEX_NONE) ? TargetStateComparisonFrame : (TargetStateComparisonFrame == INDEX_NONE) ? ResimFrame : FMath::Min(TargetStateComparisonFrame, ResimFrame);
 		}
 
@@ -574,6 +575,7 @@ void UNetworkPhysicsComponent::UninitializeComponent()
 		{
 			AsyncInput->ActorComponent = nullptr;
 			AsyncInput->PhysicsObject = nullptr;
+			AsyncInput->ImplementationInterface = nullptr;
 		}
 	}
 
@@ -623,12 +625,19 @@ void UNetworkPhysicsComponent::GetLifetimeReplicatedProps(TArray< FLifetimePrope
 	ReplicatedParamsAll.RepNotifyCondition = REPNOTIFY_Always;
 	ReplicatedParamsAll.bIsPushBased = true;
 
+	DOREPLIFETIME_CONDITION(UNetworkPhysicsComponent, InputsToNetwork_Owner, COND_OwnerOnly);
+
+	// Legacy
 	DOREPLIFETIME_WITH_PARAMS_FAST(UNetworkPhysicsComponent, ReplicatedDeltaSourceInput, ReplicatedParamsAll);
 	DOREPLIFETIME_WITH_PARAMS_FAST(UNetworkPhysicsComponent, ReplicatedDeltaSourceState, ReplicatedParamsAll);
-	DOREPLIFETIME_CONDITION(UNetworkPhysicsComponent, InputsToNetwork_Owner, COND_OwnerOnly);
 	DOREPLIFETIME_WITH_PARAMS_FAST(UNetworkPhysicsComponent, ReplicatedInputs, ReplicatedParamsOwner);
 	DOREPLIFETIME_WITH_PARAMS_FAST(UNetworkPhysicsComponent, ReplicatedRemoteInputs, ReplicatedParamsRemote);
 	DOREPLIFETIME_WITH_PARAMS_FAST(UNetworkPhysicsComponent, ReplicatedStates, ReplicatedParamsAll);
+
+	// Iris
+	DOREPLIFETIME_WITH_PARAMS_FAST(UNetworkPhysicsComponent, ReplicatedInputCollection, ReplicatedParamsOwner);
+	DOREPLIFETIME_WITH_PARAMS_FAST(UNetworkPhysicsComponent, ReplicatedRemoteInputCollection, ReplicatedParamsRemote);
+	DOREPLIFETIME_WITH_PARAMS_FAST(UNetworkPhysicsComponent, ReplicatedStateCollection, ReplicatedParamsAll);
 }
 
 // Called every Game Thread frame
@@ -642,78 +651,85 @@ void UNetworkPhysicsComponent::TickComponent(float DeltaTime, ELevelTick TickTyp
 
 void UNetworkPhysicsComponent::NetworkMarshaledData()
 {
+	if (!NetworkPhysicsComponent_Internal)
+	{
+		return;
+	}
+
 	UWorld* World = GetWorld();
 	if (!World)
 	{
 		return;
 	}
 
-
-	if (NetworkPhysicsComponent_Internal)
+	const bool bIsServer = HasServerWorld();
+	if (!bIsServer && !IsNetworkPhysicsTickOffsetAssigned())
 	{
-		const bool bIsServer = HasServerWorld();
-		if (!bIsServer && !IsNetworkPhysicsTickOffsetAssigned())
-		{
-			// Don't replicate data to the server until networked physics is setup with a synchronized physics tick offset
-			return;
-		}
+		// Don't replicate data to the server until networked physics is setup with a synchronized physics tick offset
+		return;
+	}
 
-		const bool bShouldSyncDeltaSourceInput = bIsServer && PhysicsReplicationCVars::ResimulationCVars::bEnableStatefulDeltaSerialization && World->GetRealTimeSeconds() > TimeToSyncDeltaSourceInput;
-		const bool bShouldSyncDeltaSourceState = bIsServer && PhysicsReplicationCVars::ResimulationCVars::bEnableStatefulDeltaSerialization && World->GetRealTimeSeconds() > TimeToSyncDeltaSourceState;
-		bool bHasSyncedDeltaSourceInput = false;
-		bool bHasSyncedDeltaSourceState = false;
+	const bool bShouldSyncDeltaSourceInput = bIsServer && PhysicsReplicationCVars::ResimulationCVars::bEnableStatefulDeltaSerialization && World->GetRealTimeSeconds() > TimeToSyncDeltaSourceInput;
+	const bool bShouldSyncDeltaSourceState = bIsServer && PhysicsReplicationCVars::ResimulationCVars::bEnableStatefulDeltaSerialization && World->GetRealTimeSeconds() > TimeToSyncDeltaSourceState;
+	bool bHasSyncedDeltaSourceInput = false;
+	bool bHasSyncedDeltaSourceState = false;
 
-		// Replicate source data for input delta serialization
-		auto DeltaSourceInputSyncHelper = [&](const TUniquePtr<Chaos::FBaseRewindHistory>& InputData)
+	// Replicate source data for input delta serialization
+	auto DeltaSourceInputSyncHelper = [&](const TUniquePtr<Chaos::FBaseRewindHistory>& InputData)
+	{
+		if (bShouldSyncDeltaSourceInput && !bHasSyncedDeltaSourceInput && IsValidNextDeltaSourceInput(InputData->GetLatestFrame()))
 		{
-			if (bShouldSyncDeltaSourceInput && !bHasSyncedDeltaSourceInput && IsValidNextDeltaSourceInput(InputData->GetLatestFrame()))
+			ReplicatedDeltaSourceInput.History->SetRecordDataIncremental(true); // Only record data that is newer than already cached data
+			if (InputData->CopyAllData(*ReplicatedDeltaSourceInput.History, /*bIncludeUnimportant*/ true, /*bIncludeImportant*/ true))
 			{
-				ReplicatedDeltaSourceInput.History->SetRecordDataIncremental(true); // Only record data that is newer than already cached data
-				if (InputData->CopyAllData(*ReplicatedDeltaSourceInput.History, /*bIncludeUnimportant*/ true, /*bIncludeImportant*/ true))
-				{
-					MARK_PROPERTY_DIRTY_FROM_NAME(UNetworkPhysicsComponent, ReplicatedDeltaSourceInput, this);
-					bHasSyncedDeltaSourceInput = true;
-					TimeToSyncDeltaSourceInput = World->GetRealTimeSeconds() + PhysicsReplicationCVars::ResimulationCVars::TimeToSyncStatefulDeltaSource;
-					AddDeltaSourceInput();
-				}
+				MARK_PROPERTY_DIRTY_FROM_NAME(UNetworkPhysicsComponent, ReplicatedDeltaSourceInput, this);
+				bHasSyncedDeltaSourceInput = true;
+				TimeToSyncDeltaSourceInput = World->GetRealTimeSeconds() + PhysicsReplicationCVars::ResimulationCVars::TimeToSyncStatefulDeltaSource;
+				AddDeltaSourceInput();
 			}
-		};
+		}
+	};
 
-		// Replicate source data for state delta serialization
-		auto DeltaSourceStateSyncHelper = [&](const TUniquePtr<Chaos::FBaseRewindHistory>& StateData)
+	// Replicate source data for state delta serialization
+	auto DeltaSourceStateSyncHelper = [&](const TUniquePtr<Chaos::FBaseRewindHistory>& StateData)
+	{
+		if (bShouldSyncDeltaSourceState && !bHasSyncedDeltaSourceState && IsValidNextDeltaSourceState(StateData->GetLatestFrame()))
 		{
-			if (bShouldSyncDeltaSourceState && !bHasSyncedDeltaSourceState && IsValidNextDeltaSourceState(StateData->GetLatestFrame()))
+			ReplicatedDeltaSourceState.History->SetRecordDataIncremental(true); // Only record data that is newer than already cached data
+			if (StateData->CopyAllData(*ReplicatedDeltaSourceState.History, /*bIncludeUnimportant*/ true, /*bIncludeImportant*/ true))
 			{
-				ReplicatedDeltaSourceState.History->SetRecordDataIncremental(true); // Only record data that is newer than already cached data
-				if (StateData->CopyAllData(*ReplicatedDeltaSourceState.History, /*bIncludeUnimportant*/ true, /*bIncludeImportant*/ true))
-				{
-					MARK_PROPERTY_DIRTY_FROM_NAME(UNetworkPhysicsComponent, ReplicatedDeltaSourceState, this);
-					bHasSyncedDeltaSourceState = true;
-					TimeToSyncDeltaSourceState = World->GetRealTimeSeconds() + PhysicsReplicationCVars::ResimulationCVars::TimeToSyncStatefulDeltaSource;
-					AddDeltaSourceState();
-				}
+				MARK_PROPERTY_DIRTY_FROM_NAME(UNetworkPhysicsComponent, ReplicatedDeltaSourceState, this);
+				bHasSyncedDeltaSourceState = true;
+				TimeToSyncDeltaSourceState = World->GetRealTimeSeconds() + PhysicsReplicationCVars::ResimulationCVars::TimeToSyncStatefulDeltaSource;
+				AddDeltaSourceState();
 			}
-		};
+		}
+	};
 
-		while (Chaos::TSimCallbackOutputHandle<FAsyncNetworkPhysicsComponentOutput> AsyncOutput = NetworkPhysicsComponent_Internal->PopFutureOutputData_External())
+	while (Chaos::TSimCallbackOutputHandle<FAsyncNetworkPhysicsComponentOutput> AsyncOutput = NetworkPhysicsComponent_Internal->PopFutureOutputData_External())
+	{
+		if (AsyncOutput->InputsToNetwork_Owner.IsSet())
 		{
-			if (AsyncOutput->InputsToNetwork_Owner.IsSet())
-			{
-				// Only marshaled from PT to GT on the server, InputsToNetwork_Owner is a replicated property towards the owner
-				InputsToNetwork_Owner = *AsyncOutput->InputsToNetwork_Owner;
-			}
+			// Only marshaled from PT to GT on the server, InputsToNetwork_Owner is a replicated property towards the owner
+			InputsToNetwork_Owner = *AsyncOutput->InputsToNetwork_Owner;
+		}
 
-			// Unimportant / Unreliable
-			if (bEnableUnreliableFlow
-				&& AsyncOutput->InputData
-				&& AsyncOutput->InputData->HasDataInHistory())
+		// Unimportant / Unreliable
+		if (bEnableUnreliableFlow
+			&& AsyncOutput->InputData
+			&& AsyncOutput->InputData->HasDataInHistory())
+		{
+			if (bIsServer)
 			{
-				if (bIsServer)
+				if (bIsUsingLegacyData)
 				{
 					// Replicate source data for delta serialization
 					DeltaSourceInputSyncHelper(AsyncOutput->InputData);
+				}
 
-					if (IsLocallyControlled())
+				if (IsLocallyControlled())
+				{
+					if (bIsUsingLegacyData)
 					{
 						// Send inputs to remote clients after getting marshaled from PT if server is the one controlling the component
 						ReplicatedRemoteInputs.History->SetRecordDataIncremental(true); // Only record data that is newer than already cached data
@@ -722,7 +738,17 @@ void UNetworkPhysicsComponent::NetworkMarshaledData()
 							MARK_PROPERTY_DIRTY_FROM_NAME(UNetworkPhysicsComponent, ReplicatedRemoteInputs, this);
 						}
 					}
+					else
+					{
+						// Send inputs to remote clients after getting marshaled from PT if server is the one controlling the component
+						ReplicatedRemoteInputCollection.DataArray.SetNum(InputsToNetwork_Simulated, EAllowShrinking::Yes);
+						InputHelper->CopyIncrementalData(AsyncOutput->InputData.Get(), ReplicatedRemoteInputCollection);
+						MARK_PROPERTY_DIRTY_FROM_NAME(UNetworkPhysicsComponent, ReplicatedRemoteInputCollection, this);
+					}
+				}
 
+				if (bIsUsingLegacyData)
+				{
 					// Only replicate data to owning client if bDataAltered is true i.e. the input has been altered by the server
 					ReplicatedInputs.History->ResizeDataHistory(AsyncOutput->InputData->CountAlteredData(/*bIncludeUnimportant*/ true, /*bIncludeImportant*/ bEnableReliableFlow == false), EAllowShrinking::Yes);
 					ReplicatedInputs.History->SetRecordDataIncremental(true); // Only record data that is newer than already cached data
@@ -733,28 +759,50 @@ void UNetworkPhysicsComponent::NetworkMarshaledData()
 						MARK_PROPERTY_DIRTY_FROM_NAME(UNetworkPhysicsComponent, ReplicatedInputs, this);
 					}
 				}
-				else if (IsLocallyControlled()) // Client-side
+				else
+				{
+					// Server sends inputs through property replication to owning client
+					ReplicatedInputCollection.DataArray.SetNum(AsyncOutput->InputData->CountAlteredData(/*bIncludeUnimportant*/ true, /*bIncludeImportant*/ bEnableReliableFlow == false), EAllowShrinking::Yes);
+					if (ReplicatedInputCollection.DataArray.Num() > 0)
+					{
+						InputHelper->CopyAlteredData(AsyncOutput->InputData.Get(), ReplicatedInputCollection);
+					}
+					MARK_PROPERTY_DIRTY_FROM_NAME(UNetworkPhysicsComponent, ReplicatedInputCollection, this);
+				}
+			}
+			else if (IsLocallyControlled()) // Client-side
+			{
+				if (bIsUsingLegacyData)
 				{
 					ReplicatedInputs.History->ResizeDataHistory(AsyncOutput->InputData->GetHistorySize(), EAllowShrinking::Yes);
-
 					if (AsyncOutput->InputData->CopyAllData(*ReplicatedInputs.History, /*bIncludeUnimportant*/ true, /*bIncludeImportant*/ bEnableReliableFlow == false))
 					{
 						// Clients send inputs through an RPC to the server
 						ServerReceiveInputData(ReplicatedInputs);
 					}
 				}
+				else
+				{
+					// Clients send inputs through an RPC to the server
+					ReplicatedInputCollection.DataArray.SetNum(AsyncOutput->InputData->GetHistorySize(), EAllowShrinking::Yes);
+					InputHelper->CopyData(AsyncOutput->InputData.Get(), ReplicatedInputCollection);
+					ServerReceiveInputCollection(ReplicatedInputCollection);
+				}
 			}
+		}
 
-			// Important / Reliable
-			if (bEnableReliableFlow)
+		// Important / Reliable
+		if (bEnableReliableFlow)
+		{
+			for (const TUniquePtr<Chaos::FBaseRewindHistory>& InputImportant : AsyncOutput->InputDataImportant)
 			{
-				for (const TUniquePtr<Chaos::FBaseRewindHistory>& InputImportant : AsyncOutput->InputDataImportant)
+				if (!InputImportant || !InputImportant->HasDataInHistory())
 				{
-					if (!InputImportant || !InputImportant->HasDataInHistory())
-					{
-						continue;
-					}
+					continue;
+				}
 
+				if (bIsUsingLegacyData)
+				{
 					// Replicate source data for delta serialization
 					DeltaSourceInputSyncHelper(InputImportant);
 
@@ -771,14 +819,21 @@ void UNetworkPhysicsComponent::NetworkMarshaledData()
 						}
 					}
 				}
+				else
+				{
+					// TODO
+				}
 			}
+		}
 
-			if (bIsServer)
+		if (bIsServer)
+		{
+			// Unimportant / Unreliable
+			if (bEnableUnreliableFlow
+				&& AsyncOutput->StateData
+				&& AsyncOutput->StateData->HasDataInHistory())
 			{
-				// Unimportant / Unreliable
-				if (bEnableUnreliableFlow
-					&& AsyncOutput->StateData
-					&& AsyncOutput->StateData->HasDataInHistory())
+				if (bIsUsingLegacyData)
 				{
 					// Replicate source data for delta serialization
 					DeltaSourceStateSyncHelper(AsyncOutput->StateData);
@@ -789,17 +844,26 @@ void UNetworkPhysicsComponent::NetworkMarshaledData()
 						MARK_PROPERTY_DIRTY_FROM_NAME(UNetworkPhysicsComponent, ReplicatedStates, this);
 					}
 				}
+				else
+				{
+					ReplicatedStateCollection.DataArray.SetNum(StatesToNetwork, EAllowShrinking::Yes);
+					StateHelper->CopyData(AsyncOutput->StateData.Get(), ReplicatedStateCollection);
+					MARK_PROPERTY_DIRTY_FROM_NAME(UNetworkPhysicsComponent, ReplicatedStateCollection, this);
+				}
+			}
 
-				// Important / Reliable
-				if (bEnableReliableFlow)
+			// Important / Reliable
+			if (bEnableReliableFlow)
+			{
+				for (const TUniquePtr<Chaos::FBaseRewindHistory>& StateImportant : AsyncOutput->StateDataImportant)
 				{
-					for (const TUniquePtr<Chaos::FBaseRewindHistory>& StateImportant : AsyncOutput->StateDataImportant)
+					if (!StateImportant || !StateImportant->HasDataInHistory())
 					{
-						if (!StateImportant || !StateImportant->HasDataInHistory())
-						{
-							continue;
-						}
+						continue;
+					}
 
+					if (bIsUsingLegacyData)
+					{
 						// Replicate source data for delta serialization
 						DeltaSourceStateSyncHelper(StateImportant);
 
@@ -809,14 +873,18 @@ void UNetworkPhysicsComponent::NetworkMarshaledData()
 							MulticastReceiveImportantStateData(ReplicatedImportantState);
 						}
 					}
+					else
+					{
+						// TODO
+					}
 				}
 			}
+		}
 
-			if (bStopRelayingLocalInputsDeferred)
-			{
-				bIsRelayingLocalInputs = false;
-				bStopRelayingLocalInputsDeferred = false;
-			}
+		if (bStopRelayingLocalInputsDeferred)
+		{
+			bIsRelayingLocalInputs = false;
+			bStopRelayingLocalInputsDeferred = false;
 		}
 	}
 }
@@ -833,12 +901,14 @@ void UNetworkPhysicsComponent::OnRep_SetReplicatedDeltaSourceInput()
 	{
 		return;
 	}
+	ensure(bIsUsingLegacyData);
 
 	AddDeltaSourceInput();
 }
 
 void UNetworkPhysicsComponent::ServerReceiveDeltaSourceInputFrame_Implementation(const int32 Frame)
 {
+	ensure(bIsUsingLegacyData);
 	for (int32 I = 0; I < DeltaSourceInputs.Num(); I++)
 	{
 		TUniquePtr<FNetworkPhysicsData>& Data = DeltaSourceInputs[I];
@@ -911,12 +981,14 @@ void UNetworkPhysicsComponent::OnRep_SetReplicatedDeltaSourceState()
 	{
 		return;
 	}
+	ensure(bIsUsingLegacyData);
 
 	AddDeltaSourceState();
 }
 
 void UNetworkPhysicsComponent::ServerReceiveDeltaSourceStateFrame_Implementation(const int32 Frame)
 {
+	ensure(bIsUsingLegacyData);
 	for (int32 I = 0; I < DeltaSourceStates.Num(); I++)
 	{
 		TUniquePtr<FNetworkPhysicsData>& Data = DeltaSourceStates[I];
@@ -987,7 +1059,7 @@ FNetworkPhysicsData* UNetworkPhysicsComponent::GetDeltaSourceInput(const int32 V
 	}
 	else if (Value == -2) // Default
 	{
-		DataPtr = &InputDataDefault;
+		DataPtr = &InputDataDefault_Legacy;
 	}
 	else if(bValueIsIndexElseFrame)
 	{
@@ -1020,7 +1092,7 @@ FNetworkPhysicsData* UNetworkPhysicsComponent::GetDeltaSourceState(const int32 V
 	}
 	else if (Value == -2) // Default
 	{
-		DataPtr = &StateDataDefault;
+		DataPtr = &StateDataDefault_Legacy;
 	}
 	else if (bValueIsIndexElseFrame)
 	{
@@ -1058,6 +1130,7 @@ void UNetworkPhysicsComponent::OnRep_SetReplicatedStates()
 	{
 		return;
 	}
+	ensure(bIsUsingLegacyData);
 
 	if (FAsyncNetworkPhysicsComponentInput* AsyncInput = NetworkPhysicsComponent_Internal->GetProducerInputData_External())
 	{
@@ -1086,6 +1159,7 @@ void UNetworkPhysicsComponent::OnRep_SetReplicatedInputs()
 	{
 		return ;
 	}
+	ensure(bIsUsingLegacyData);
 
 	if (FAsyncNetworkPhysicsComponentInput* AsyncInput = NetworkPhysicsComponent_Internal->GetProducerInputData_External())
 	{
@@ -1114,6 +1188,7 @@ void UNetworkPhysicsComponent::OnRep_SetReplicatedRemoteInputs()
 	{
 		return ;
 	}
+	ensure(bIsUsingLegacyData);
 
 	if (FAsyncNetworkPhysicsComponentInput* AsyncInput = NetworkPhysicsComponent_Internal->GetProducerInputData_External())
 	{
@@ -1140,6 +1215,7 @@ void UNetworkPhysicsComponent::ServerReceiveInputData_Implementation(const FNetw
 	{
 		return;
 	}
+	ensure(bIsUsingLegacyData);
 
 	if (FAsyncNetworkPhysicsComponentInput* AsyncInput = NetworkPhysicsComponent_Internal->GetProducerInputData_External())
 	{
@@ -1178,6 +1254,7 @@ void UNetworkPhysicsComponent::ServerReceiveImportantInputData_Implementation(co
 	{
 		return;
 	}
+	ensure(bIsUsingLegacyData);
 
 	if (FAsyncNetworkPhysicsComponentInput* AsyncInput = NetworkPhysicsComponent_Internal->GetProducerInputData_External())
 	{
@@ -1215,6 +1292,7 @@ void UNetworkPhysicsComponent::MulticastReceiveImportantInputData_Implementation
 	{
 		return;
 	}
+	ensure(bIsUsingLegacyData);
 
 	if (FAsyncNetworkPhysicsComponentInput* AsyncInput = NetworkPhysicsComponent_Internal->GetProducerInputData_External())
 	{
@@ -1246,6 +1324,7 @@ void UNetworkPhysicsComponent::MulticastReceiveImportantStateData_Implementation
 	{
 		return;
 	}
+	ensure(bIsUsingLegacyData);
 
 	if (FAsyncNetworkPhysicsComponentInput* AsyncInput = NetworkPhysicsComponent_Internal->GetProducerInputData_External())
 	{
@@ -1257,41 +1336,118 @@ void UNetworkPhysicsComponent::MulticastReceiveImportantStateData_Implementation
 	}
 }
 
-/** DEPRECATED UE 5.5 */
-const float UNetworkPhysicsComponent::GetCurrentInputDecay(FNetworkPhysicsData* PhysicsData)
+// Server RPC to receive inputs from client
+void UNetworkPhysicsComponent::ServerReceiveInputCollection_Implementation(const FNetworkPhysicsDataCollection& ClientInputCollection)
 {
-	if (!PhysicsData)
+	ensure(bIsUsingLegacyData == false);
+
+	if (!NetworkPhysicsComponent_Internal || !InputHelper || !ClientInputCollection.DataArray.Num())
 	{
-		return 0.0f;
+		return;
 	}
 
-	FPhysScene* PhysScene = GetWorld()->GetPhysicsScene();
-	if (!PhysScene)
+	if (FAsyncNetworkPhysicsComponentInput* AsyncInput = NetworkPhysicsComponent_Internal->GetProducerInputData_External())
 	{
-		return 0.0f;
+		// Temporarily populate the LocalFrame value with ServerFrame in the inputs to have a value to use when marshaling the inputs from GT to PT, the correct LocalFrame value will be set on PT when calling History->ReceiveNewData
+		ClientInputCollection.UpdateLocalFrameFromServerFrame(0);
+
+		if (!AsyncInput->InputData)
+		{
+			AsyncInput->InputData = InputHelper->CreateUniqueRewindHistory(ClientInputCollection.DataArray.Num());
+		}
+
+		AsyncInput->InputData->ResetFast();
+		AsyncInput->InputData->ResizeDataHistory(ClientInputCollection.DataArray.Num(), EAllowShrinking::No);
+		AsyncInput->InputData->SetRecordDataIncremental(true); // Only record data that is newer than already cached data
+
+		InputHelper->CopyData(ClientInputCollection, AsyncInput->InputData.Get());
+
+		// Send received inputs to remote clients
+		ReplicatedRemoteInputCollection.DataArray.SetNum(InputsToNetwork_Simulated);
+		InputHelper->CopyIncrementalData(AsyncInput->InputData.Get(), ReplicatedRemoteInputCollection);
+		MARK_PROPERTY_DIRTY_FROM_NAME(UNetworkPhysicsComponent, ReplicatedRemoteInputCollection, this);
 	}
+}
 
-	Chaos::FPhysicsSolver* Solver = PhysScene->GetSolver();
-	if (!Solver)
+// repnotify for inputs on owner client
+void UNetworkPhysicsComponent::OnRep_SetReplicatedInputCollection()
+{
+	ensure(bIsUsingLegacyData == false);
+	if (!NetworkPhysicsComponent_Internal || !InputHelper || !ReplicatedInputCollection.DataArray.Num())
 	{
-		return 0.0f;
+		return;
 	}
 
-	Chaos::FRewindData* RewindData = Solver->GetRewindData();
-	if (!RewindData)
+	if (FAsyncNetworkPhysicsComponentInput* AsyncInput = NetworkPhysicsComponent_Internal->GetProducerInputData_External())
 	{
-		return 0.0f;
+		// Temporarily populate the LocalFrame value with ServerFrame in the inputs to have a value to use when marshaling the inputs from GT to PT, the correct LocalFrame value will be set on PT when calling History->ReceiveNewData
+		ReplicatedInputCollection.UpdateLocalFrameFromServerFrame(0);
+
+		if (!AsyncInput->InputData)
+		{
+			AsyncInput->InputData = InputHelper->CreateUniqueRewindHistory(ReplicatedInputCollection.DataArray.Num());
+		}
+
+		AsyncInput->InputData->ResetFast();
+		AsyncInput->InputData->ResizeDataHistory(ReplicatedInputCollection.DataArray.Num(), EAllowShrinking::No);
+
+		InputHelper->CopyData(ReplicatedInputCollection, AsyncInput->InputData.Get());
 	}
-	
-	const float NumPredictedInputs = RewindData->CurrentFrame() - PhysicsData->LocalFrame; // Number of frames we have used the same PhysicsData for during resim
-	const float MaxPredictedInputs = RewindData->GetLatestFrame() - 1 - PhysicsData->LocalFrame; // Max number of frames PhysicsData registered frame until end of resim
+}
 
-	// Linear decay
-	const float PredictionAlpha = MaxPredictedInputs > 0 ? (NumPredictedInputs / MaxPredictedInputs) : 0.0f;
+// repnotify for inputs on remote clients
+void UNetworkPhysicsComponent::OnRep_SetReplicatedRemoteInputCollection()
+{
+	ensure(bIsUsingLegacyData == false);
+	if (!NetworkPhysicsComponent_Internal || !InputHelper || !ReplicatedRemoteInputCollection.DataArray.Num())
+	{
+		return;
+	}
+
+	if (FAsyncNetworkPhysicsComponentInput* AsyncInput = NetworkPhysicsComponent_Internal->GetProducerInputData_External())
+	{
+		// Temporarily populate the LocalFrame value with ServerFrame in the inputs to have a value to use when marshaling the inputs from GT to PT, the correct LocalFrame value will be set on PT when calling History->ReceiveNewData
+		ReplicatedRemoteInputCollection.UpdateLocalFrameFromServerFrame(0);
+
+		if (!AsyncInput->InputData)
+		{
+			AsyncInput->InputData = InputHelper->CreateUniqueRewindHistory(ReplicatedRemoteInputCollection.DataArray.Num());
+		}
 
-	return PredictionAlpha;
+		AsyncInput->InputData->ResetFast();
+		AsyncInput->InputData->ResizeDataHistory(ReplicatedRemoteInputCollection.DataArray.Num(), EAllowShrinking::No);
+
+		InputHelper->CopyData(ReplicatedRemoteInputCollection, AsyncInput->InputData.Get());
+	}
+}
+
+// repnotify for the states on the client
+void UNetworkPhysicsComponent::OnRep_SetReplicatedStateCollection()
+{
+	ensure(bIsUsingLegacyData == false);
+	if (!NetworkPhysicsComponent_Internal || !StateHelper || !ReplicatedStateCollection.DataArray.Num())
+	{
+		return;
+	}
+
+	if (FAsyncNetworkPhysicsComponentInput* AsyncInput = NetworkPhysicsComponent_Internal->GetProducerInputData_External())
+	{
+		// Temporarily populate the LocalFrame value with ServerFrame in the states to have a value to use when marshaling the inputs from GT to PT, the correct LocalFrame value will be set on PT when calling History->ReceiveNewData
+		ReplicatedStateCollection.UpdateLocalFrameFromServerFrame(0);
+
+		if (!AsyncInput->StateData)
+		{
+			AsyncInput->StateData = StateHelper->CreateUniqueRewindHistory(ReplicatedStateCollection.DataArray.Num());
+		}
+
+		AsyncInput->StateData->ResetFast();
+		AsyncInput->StateData->ResizeDataHistory(ReplicatedStateCollection.DataArray.Num(), EAllowShrinking::No);
+
+		StateHelper->CopyData(ReplicatedStateCollection, AsyncInput->StateData.Get());
+	}
 }
 
+
 bool UNetworkPhysicsComponent::HasServerWorld() const
 {
 	Chaos::EnsureIsInGameThreadContext();
@@ -1437,7 +1593,11 @@ void UNetworkPhysicsComponent::UpdateAsyncComponent(const bool bFullUpdate)
 				
 				if (ActorComponent.IsValid())
 				{
-					AsyncInput->ActorComponent = ActorComponent.Get();
+					AsyncInput->ActorComponent = ActorComponent;
+				}
+				if (ImplementationInterface)
+				{
+					AsyncInput->ImplementationInterface = ImplementationInterface;
 				}
 			}
 		}
@@ -1450,7 +1610,14 @@ void UNetworkPhysicsComponent::CreateAsyncDataHistory()
 	{
 		if (FAsyncNetworkPhysicsComponentInput* AsyncInput = NetworkPhysicsComponent_Internal->GetProducerInputData_External())
 		{
-			AsyncInput->ActorComponent = ActorComponent.Get();
+			if (ActorComponent.IsValid())
+			{
+				AsyncInput->ActorComponent = ActorComponent;
+			}
+			if (ImplementationInterface)
+			{
+				AsyncInput->ImplementationInterface = ImplementationInterface;
+			}
 
 			if (InputHelper)
 			{
@@ -1521,8 +1688,10 @@ FAsyncNetworkPhysicsComponent::FAsyncNetworkPhysicsComponent() : TSimCallbackObj
 	, NetRole(ENetRole::ROLE_Authority)
 	, NetworkPhysicsTickOffset(0)
 	, PhysicsReplicationMode(EPhysicsReplicationMode::Default)
+	, bIsUsingLegacyData(false)
 	, SettingsComponent(nullptr)
 	, ActorComponent(nullptr)
+	, ImplementationInterface(nullptr)
 	, PhysicsObject(nullptr)
 	, bCompareStateToTriggerRewind(false)
 	, bCompareStateToTriggerRewindIncludeSimProxies(false)
@@ -1623,6 +1792,10 @@ void FAsyncNetworkPhysicsComponent::ConsumeAsyncInput(const int32 PhysicsStep)
 			{
 				ActorComponent = *AsyncInput->ActorComponent;
 			}
+			if (AsyncInput->ImplementationInterface.IsSet())
+			{
+				ImplementationInterface = *AsyncInput->ImplementationInterface;
+			}
 			if (AsyncInput->PhysicsObject.IsSet())
 			{
 				if (PhysicsObject == nullptr || PhysicsObject != *AsyncInput->PhysicsObject)
@@ -1655,26 +1828,32 @@ void FAsyncNetworkPhysicsComponent::ConsumeAsyncInput(const int32 PhysicsStep)
 			{
 				bCompareInputToTriggerRewind = *AsyncInput->bCompareInputToTriggerRewind;
 			}
-			if (AsyncInput->InputHelper.IsSet())
+			if (AsyncInput->StateHelper.IsSet())
 			{
 				// Setup rewind data if not already done, and get history size
 				const int32 NumFrames = SetupRewindData();
 
-				// Create input history and local data properties
-				InputData = (*AsyncInput->InputHelper)->CreateUniqueData();
-				LatestInputReceiveData = (*AsyncInput->InputHelper)->CreateUniqueData();
-				InputHistory = MakeShareable((*AsyncInput->InputHelper)->CreateUniqueRewindHistory(NumFrames).Release());
+				StateHelper = (*AsyncInput->StateHelper)->Clone();
+
+				// Create state history and local property
+				StateData = StateHelper->CreateUniqueData();
+				StateHistory = MakeShareable(StateHelper->CreateUniqueRewindHistory(NumFrames).Release());
 				RegisterDataHistoryInRewindData();
+				bIsUsingLegacyData = StateHelper->IsUsingLegacyData();
 			}
-			if (AsyncInput->StateHelper.IsSet())
+			if (AsyncInput->InputHelper.IsSet())
 			{
 				// Setup rewind data if not already done, and get history size
 				const int32 NumFrames = SetupRewindData();
 
-				// Create state history and local property
-				StateData = (*AsyncInput->StateHelper)->CreateUniqueData();
-				StateHistory = MakeShareable((*AsyncInput->StateHelper)->CreateUniqueRewindHistory(NumFrames).Release());
+				InputHelper = (*AsyncInput->InputHelper)->Clone();
+
+				// Create input history and local data properties
+				InputData = InputHelper->CreateUniqueData();
+				LatestInputReceiveData = InputHelper->CreateUniqueData();
+				InputHistory = MakeShareable(InputHelper->CreateUniqueRewindHistory(NumFrames).Release());
 				RegisterDataHistoryInRewindData();
+				bIsUsingLegacyData = InputHelper->IsUsingLegacyData();
 			}
 		}
 
@@ -1711,7 +1890,7 @@ void FAsyncNetworkPhysicsComponent::ConsumeAsyncInput(const int32 PhysicsStep)
 			{
 				/* Extract latest received input from client on the server, to be used if input buffer runs empty
 				* TODO, improve flow to not require this before ReceiveHelper */
-				FNetworkPhysicsData* PhysicsData = nullptr;
+				FNetworkPhysicsBase* PhysicsData = nullptr;
 				if (bIsServer && LatestInputReceiveData)
 				{
 					PhysicsData = LatestInputReceiveData.Get();
@@ -1728,9 +1907,22 @@ void FAsyncNetworkPhysicsComponent::ConsumeAsyncInput(const int32 PhysicsStep)
 				}
 
 				// Validate data in the received inputs on the server
-				if (!ComponentSettings.GetValidateDataOnGameThread() && bIsServer && ActorComponent.IsValid() && ActorComponent.Get()->IsBeingDestroyed() == false)
+				if (!ComponentSettings.GetValidateDataOnGameThread() && bIsServer)
 				{
-					AsyncInput->InputData->ValidateDataInHistory(ActorComponent.Get());
+					if (bIsUsingLegacyData)
+					{
+						if (ensure(ActorComponent.IsValid() && ActorComponent.Get()->IsBeingDestroyed() == false))
+						{
+							AsyncInput->InputData->ValidateDataInHistory(ActorComponent.Get());
+						}
+					}
+					else
+					{
+						if (InputHelper && ensure(ImplementationInterface))
+						{
+							InputHelper->ValidateDataInHistory(AsyncInput->InputData.Get(), *ImplementationInterface);
+						}
+					}
 				}
 				ReceiveHelper(InputHistory.Get(), AsyncInput->InputData.Get(), /*bImportant*/false, bCompareInput);
 
@@ -1823,18 +2015,31 @@ void FAsyncNetworkPhysicsComponent::OnPreProcessInputs_Internal(const int32 Phys
 	}
 #endif
 
-	if (ActorComponent.IsValid() && ActorComponent.Get()->IsBeingDestroyed() == false)
 	{
 		// Apply replicated state on clients if we are resimulating or on simulated proxies if setting is enabled
 		const bool bApplySimProxyState = ComponentSettings.GetApplySimProxyStateAtRuntime() && !bIsServer && !IsLocallyControlled();
 		if ((bApplySimProxyState || bIsSolverResim) && StateHistory && StateData)
 		{
-			FNetworkPhysicsData* PhysicsData = StateData.Get();
+			FNetworkPhysicsBase* PhysicsData = StateData.Get();
 			PhysicsData->LocalFrame = PhysicsStep;
 			const bool bExactFrame = PhysicsReplicationCVars::ResimulationCVars::bAllowRewindToClosestState ? !bIsSolverReset : true;
 			if (StateHistory->ExtractData(PhysicsStep, bIsSolverReset, PhysicsData, (bExactFrame && bIsSolverResim)) && PhysicsData->bReceivedData)
 			{
-				PhysicsData->ApplyData(ActorComponent.Get());
+				if (bIsUsingLegacyData)
+				{
+					if (ensure(ActorComponent.IsValid() && ActorComponent.Get()->IsBeingDestroyed() == false))
+					{
+						FNetworkPhysicsData* LegacyData = static_cast<FNetworkPhysicsData*>(PhysicsData);
+						LegacyData->ApplyData(ActorComponent.Get());
+					}
+				}
+				else
+				{
+					if (ensure(ImplementationInterface))
+					{
+						ImplementationInterface->ApplyState(*PhysicsData);
+					}
+				}
 #if DEBUG_NETWORK_PHYSICS
 				UE_LOG(LogChaos, Log, TEXT("			Applying extracted state from history | bExactFrame = %d | LocalFrame = %d | ServerFrame = %d | bDataAltered = %d | Data: %s")
 					, bExactFrame, PhysicsData->LocalFrame, PhysicsData->ServerFrame, PhysicsData->bDataAltered, *PhysicsData->DebugData());
@@ -1854,7 +2059,7 @@ void FAsyncNetworkPhysicsComponent::OnPreProcessInputs_Internal(const int32 Phys
 		const bool bApplyServerInput = bIsServer && !IsLocallyControlled();
 		if ((bApplyServerInput || bApplySimProxyInput || bIsSolverResim) && InputHistory && InputData)
 		{
-			FNetworkPhysicsData* PhysicsData = InputData.Get();
+			FNetworkPhysicsBase* PhysicsData = InputData.Get();
 			int32 NextExpectedLocalFrame = PhysicsData->LocalFrame + 1;
 
 			// There are important inputs earlier than upcoming input to apply
@@ -1865,8 +2070,23 @@ void FAsyncNetworkPhysicsComponent::OnPreProcessInputs_Internal(const int32 Phys
 #if DEBUG_NETWORK_PHYSICS
 					UE_LOG(LogChaos, Log, TEXT("		Non-Determinism: Reapplying multiple data due to receiving an important data that was previously missed. FromFrame: %d | ToFrame: %d | IsLocallyControlled = %d"), NewImportantInputFrame, (NextExpectedLocalFrame - 1), IsLocallyControlled());
 #endif
-					// Apply all inputs in range
-					InputHistory->ApplyDataRange(NewImportantInputFrame, NextExpectedLocalFrame - 1, ActorComponent.Get(), /*bOnlyImportant*/false);
+					if (bIsUsingLegacyData)
+					{
+						if (ensure(ActorComponent.IsValid() && ActorComponent.Get()->IsBeingDestroyed() == false))
+						{
+							// Apply all inputs in range
+							InputHistory->ApplyDataRange(NewImportantInputFrame, NextExpectedLocalFrame - 1, ActorComponent.Get(), /*bOnlyImportant*/false);
+						}
+					}
+/* // Importance is not implemented into the new network flow (yet?)
+					else
+					{
+						if (InputHelper && ensure(ImplementationInterface))
+						{
+							InputHelper->ApplyDataRange(InputHistory.Get(), NewImportantInputFrame, NextExpectedLocalFrame - 1, *ImplementationInterface);
+						}
+					}
+*/
 				}
 				else
 				{
@@ -1880,6 +2100,24 @@ void FAsyncNetworkPhysicsComponent::OnPreProcessInputs_Internal(const int32 Phys
 
 			if (InputHistory->ExtractData(PhysicsStep, bIsSolverReset, PhysicsData, /*bExactFrame*/(ComponentSettings.GetAllowInputExtrapolation() == false)))
 			{
+				// Check if the extracted data was altered and if we have a hole in the buffer
+				if (bIsServer && PhysicsData->bDataAltered)
+				{
+					// Check if the buffer is empty
+					if (PhysicsStep <= InputHistory->GetLatestFrame())
+					{
+						// A missing input was detected and buffer is not empty, inform the owning client to send more inputs in each RPC to not get a gaps in the buffer
+						// NOTE: We don't send more extra inputs when the buffer runs empty since that case is corrected via time dilation, not sending extra inputs
+						MissingInputCount++;
+					}
+#if DEBUG_NETWORK_PHYSICS
+					else
+					{
+						UE_LOG(LogChaos, Log, TEXT("		Non-Determinism: Input buffer Empty, input for frame %d was extrapolated"), PhysicsStep);
+					}
+#endif
+				}
+
 				// Calculate input decay if we are resimulating and we don't have up to date inputs
 				if (bIsSolverResim)
 				{
@@ -1903,7 +2141,21 @@ void FAsyncNetworkPhysicsComponent::OnPreProcessInputs_Internal(const int32 Phys
 						{
 							if (InputHistory->ExtractData(NextExpectedLocalFrame, bIsSolverReset, PhysicsData, true) && NextExpectedLocalFrame < LastFrame)
 							{
-								PhysicsData->ApplyData(ActorComponent.Get());
+								if (bIsUsingLegacyData)
+								{
+									if (ensure(ActorComponent.IsValid() && ActorComponent.Get()->IsBeingDestroyed() == false))
+									{
+										FNetworkPhysicsData* LegacyData = static_cast<FNetworkPhysicsData*>(PhysicsData);
+										LegacyData->ApplyData(ActorComponent.Get());
+									}
+								}
+								else
+								{
+									if (ensure(ImplementationInterface))
+									{
+										ImplementationInterface->ApplyInput(*PhysicsData);
+									}
+								}
 							}
 						}
 					}
@@ -1917,33 +2169,29 @@ void FAsyncNetworkPhysicsComponent::OnPreProcessInputs_Internal(const int32 Phys
 					}
 				}
 
-				// If the extracted input data was altered (extrapolated, interpolated) on the server, record it into the history for it to get replicated to clients
+				// If the extracted input data was altered (merged, extrapolated, interpolated) on the server, record it into the history for it to get replicated to clients
 				if (bIsServer && PhysicsData->bDataAltered)
 				{
-					// Check that buffer is not empty
-					if (PhysicsStep <= InputHistory->GetLatestFrame())
-					{
-						if (PhysicsData->LocalFrame < PhysicsStep) 
-						{
-							// A missing input was detected and buffer is not empty, inform the owning client to send more inputs in each RPC to not get a gaps in the buffer
-							// NOTE: We don't send more extra inputs when the buffer runs empty since that case is corrected via time dilation, not sending extra inputs
-							MissingInputCount++;
-						}
-					}
-#if DEBUG_NETWORK_PHYSICS
-					else
-					{
-						UE_LOG(LogChaos, Log, TEXT("		Non-Determinism: Input buffer Empty, input for frame %d was extrapolated from frame: %d"), PhysicsStep, PhysicsData->LocalFrame);
-					}
-#endif
-
 					PhysicsData->bReceivedData = true; // Mark the input data as received so that it doesn't get overwritten by incoming client inputs
 					PhysicsData->bImportant = false;
-					PhysicsData->LocalFrame = PhysicsStep;
 					InputHistory->RecordData(PhysicsStep, PhysicsData);
 				}
 
-				PhysicsData->ApplyData(ActorComponent.Get());
+				if (bIsUsingLegacyData)
+				{
+					if (ensure(ActorComponent.IsValid() && ActorComponent.Get()->IsBeingDestroyed() == false))
+					{
+						FNetworkPhysicsData* LegacyData = static_cast<FNetworkPhysicsData*>(PhysicsData);
+						LegacyData->ApplyData(ActorComponent.Get());
+					}
+				}
+				else
+				{
+					if (ensure(ImplementationInterface))
+					{
+						ImplementationInterface->ApplyInput(*PhysicsData);
+					}
+				}
 
 #if DEBUG_NETWORK_PHYSICS
 				{
@@ -1984,53 +2232,80 @@ void FAsyncNetworkPhysicsComponent::OnPostProcessInputs_Internal(const int32 Phy
 	}
 #endif
 
-	if (ActorComponent.IsValid() && ActorComponent.Get()->IsBeingDestroyed() == false)
+	// Cache current input if we are locally controlled
+	const bool bShouldCacheInputHistory = IsLocallyControlled() && !bIsSolverResim;
+	if (bShouldCacheInputHistory && (InputData != nullptr))
 	{
-		// Cache current input if we are locally controlled
-		const bool bShouldCacheInputHistory = IsLocallyControlled() && !bIsSolverResim;
-		if (bShouldCacheInputHistory && (InputData != nullptr))
-		{
-			// Prepare to gather input data
-			FNetworkPhysicsData* PhysicsData = InputData.Get();
-			PhysicsData->PrepareFrame(PhysicsStep, bIsServer, GetNetworkPhysicsTickOffset());
-
-			// Gather input data from implementation
-			PhysicsData->BuildData(ActorComponent.Get());
+		// Prepare to gather input data
+		FNetworkPhysicsBase* PhysicsData = InputData.Get();
+		PhysicsData->PrepareFrame(PhysicsStep, bIsServer, GetNetworkPhysicsTickOffset());
 
-			// Record input in history
-			InputHistory->RecordData(PhysicsStep, PhysicsData);
-
-#if DEBUG_NETWORK_PHYSICS
+		if (bIsUsingLegacyData)
+		{
+			if (ensure(ActorComponent.IsValid() && ActorComponent.Get()->IsBeingDestroyed() == false))
 			{
-				UE_LOG(LogChaos, Log, TEXT("		Recording input into history | LocalFrame = %d | ServerFrame = %d | bDataAltered = %d | Input: %s ")
-					, PhysicsData->LocalFrame, PhysicsData->ServerFrame, PhysicsData->bDataAltered, *PhysicsData->DebugData());
+				// Gather input data from implementation
+				FNetworkPhysicsData* LegacyData = static_cast<FNetworkPhysicsData*>(PhysicsData);
+				LegacyData->BuildData(ActorComponent.Get());
+			}
+		}
+		else
+		{
+			if (ensure(ImplementationInterface))
+			{
+				// Gather input data from implementation
+				ImplementationInterface->BuildInput(*PhysicsData);
 			}
-#endif
 		}
 
-		// Cache current state if this is the server or we are comparing predicted states clients
-		const bool bShouldCacheStateHistory = bIsServer
-			|| (ComponentSettings.GetCompareStateToTriggerRewind(bCompareStateToTriggerRewind) && (IsLocallyControlled() || ComponentSettings.GetCompareStateToTriggerRewindIncludeSimProxies(bCompareStateToTriggerRewindIncludeSimProxies)));
+		// Record input in history
+		InputHistory->RecordData(PhysicsStep, PhysicsData);
 
-		if (StateHistory && StateData && bShouldCacheStateHistory)
+#if DEBUG_NETWORK_PHYSICS
 		{
-			// Prepare to gather state data
-			FNetworkPhysicsData* PhysicsData = StateData.Get();
-			PhysicsData->PrepareFrame(PhysicsStep, bIsServer, GetNetworkPhysicsTickOffset());
+			UE_LOG(LogChaos, Log, TEXT("		Recording input into history | LocalFrame = %d | ServerFrame = %d | bDataAltered = %d | Input: %s ")
+				, PhysicsData->LocalFrame, PhysicsData->ServerFrame, PhysicsData->bDataAltered, *PhysicsData->DebugData());
+		}
+#endif
+	}
 
-			// Gather input data from implementation
-			PhysicsData->BuildData(ActorComponent.Get());
+	// Cache current state if this is the server or we are comparing predicted states clients
+	const bool bShouldCacheStateHistory = bIsServer
+		|| (ComponentSettings.GetCompareStateToTriggerRewind(bCompareStateToTriggerRewind) && (IsLocallyControlled() || ComponentSettings.GetCompareStateToTriggerRewindIncludeSimProxies(bCompareStateToTriggerRewindIncludeSimProxies)));
 
-			// Record input in history
-			StateHistory->RecordData(PhysicsStep, PhysicsData);
+	if (StateHistory && StateData && bShouldCacheStateHistory)
+	{
+		// Prepare to gather state data
+		FNetworkPhysicsBase* PhysicsData = StateData.Get();
+		PhysicsData->PrepareFrame(PhysicsStep, bIsServer, GetNetworkPhysicsTickOffset());
 
-#if DEBUG_NETWORK_PHYSICS
+		if (bIsUsingLegacyData)
+		{
+			if (ensure(ActorComponent.IsValid() && ActorComponent.Get()->IsBeingDestroyed() == false))
 			{
-				UE_LOG(LogChaos, Log, TEXT("		Recording state into history | LocalFrame = %d | ServerFrame = %d | bDataAltered = %d | State: %s ")
-					, PhysicsData->LocalFrame, PhysicsData->ServerFrame, PhysicsData->bDataAltered, *PhysicsData->DebugData());
+				// Gather state data from implementation
+				FNetworkPhysicsData* LegacyData = static_cast<FNetworkPhysicsData*>(PhysicsData);
+				LegacyData->BuildData(ActorComponent.Get());
 			}
-#endif
 		}
+		else
+		{
+			if (ensure(ImplementationInterface))
+			{
+				// Gather state data from implementation
+				ImplementationInterface->BuildState(*PhysicsData);
+			}
+		}
+
+		// Record state in history
+		StateHistory->RecordData(PhysicsStep, PhysicsData);
+
+#if DEBUG_NETWORK_PHYSICS
+		{
+			UE_LOG(LogChaos, Log, TEXT("		Recording state into history | LocalFrame = %d | ServerFrame = %d | bDataAltered = %d | State: %s ")
+				, PhysicsData->LocalFrame, PhysicsData->ServerFrame, PhysicsData->bDataAltered, *PhysicsData->DebugData());
+		}
+#endif
 	}
 
 	// Marshal inputs and states from PT to GT for networking
@@ -2208,7 +2483,7 @@ void FAsyncNetworkPhysicsComponent::TriggerResimulation(int32 ResimFrame)
 	}
 }
 
-const float FAsyncNetworkPhysicsComponent::GetCurrentInputDecay(const FNetworkPhysicsData* PhysicsData)
+const float FAsyncNetworkPhysicsComponent::GetCurrentInputDecay(const FNetworkPhysicsBase* PhysicsData)
 {
 	if (!PhysicsData)
 	{
diff --git a/Engine/Source/Runtime/Engine/Public/Physics/NetworkPhysicsComponent.h b/Engine/Source/Runtime/Engine/Public/Physics/NetworkPhysicsComponent.h
--- a/Engine/Source/Runtime/Engine/Public/Physics/NetworkPhysicsComponent.h
+++ b/Engine/Source/Runtime/Engine/Public/Physics/NetworkPhysicsComponent.h
@@ -12,6 +12,7 @@
 #include "PhysicsEngine/PhysicsSettings.h"
 #include "Chaos/PhysicsObject.h"
 #include "Physics/NetworkPhysicsSettingsComponent.h"
+#include "StructUtils/InstancedStruct.h"
 
 #include "NetworkPhysicsComponent.generated.h"
 
@@ -26,7 +27,7 @@ DECLARE_MULTICAST_DELEGATE_OneParam(FOnPostProcessInputsInternal, const int32);
 DECLARE_MULTICAST_DELEGATE_TwoParams(FOnInjectInputsExternal, const int32 /* PhysicsStep */, const int32 /* NumSteps */);
 
 /** Templated data history, holding a data buffer */
-template<typename DataType>
+template<typename DataType, bool bLegacyData = false>
 struct TNetRewindHistory : public Chaos::TDataRewindHistory<DataType>
 {
 	using Super = Chaos::TDataRewindHistory<DataType>;
@@ -57,11 +58,14 @@ struct TNetRewindHistory : public Chaos::TDataRewindHistory<DataType>
 
 	virtual void ValidateDataInHistory(const void* ActorComponent) override
 	{
-		const UActorComponent* NetworkComponent = static_cast<const UActorComponent*>(ActorComponent);
-		for (int32 FrameIndex = 0; FrameIndex < Super::NumFrames; ++FrameIndex)
+		if constexpr (bLegacyData)
 		{
-			DataType& FrameData = Super::DataHistory[FrameIndex];
-			FrameData.ValidateData(NetworkComponent);
+			const UActorComponent* NetworkComponent = static_cast<const UActorComponent*>(ActorComponent);
+			for (int32 FrameIndex = 0; FrameIndex < Super::NumFrames; ++FrameIndex)
+			{
+				DataType& FrameData = Super::DataHistory[FrameIndex];
+				FrameData.ValidateData(NetworkComponent);
+			}
 		}
 	}
 
@@ -127,15 +131,18 @@ struct TNetRewindHistory : public Chaos::TDataRewindHistory<DataType>
 
 	virtual void ApplyDataRange(const int32 FromFrame, const int32 ToFrame, void* ActorComponent, const bool bOnlyImportant = false) override
 	{
-		UActorComponent* NetworkComponent = static_cast<UActorComponent*>(ActorComponent);
-
-		for (int32 ApplyFrame = FromFrame; ApplyFrame <= ToFrame; ++ApplyFrame)
+		if constexpr (bLegacyData)
 		{
-			const int32 ApplyIndex = Super::GetFrameIndex(ApplyFrame);
-			DataType& FrameData = Super::DataHistory[ApplyIndex];
-			if (ApplyFrame == FrameData.LocalFrame && (!bOnlyImportant || FrameData.bImportant))
+			UActorComponent* NetworkComponent = static_cast<UActorComponent*>(ActorComponent);
+
+			for (int32 ApplyFrame = FromFrame; ApplyFrame <= ToFrame; ++ApplyFrame)
 			{
-				FrameData.ApplyData(NetworkComponent);
+				const int32 ApplyIndex = Super::GetFrameIndex(ApplyFrame);
+				DataType& FrameData = Super::DataHistory[ApplyIndex];
+				if (ApplyFrame == FrameData.LocalFrame && (!bOnlyImportant || FrameData.bImportant))
+				{
+					FrameData.ApplyData(NetworkComponent);
+				}
 			}
 		}
 	}
@@ -265,7 +272,8 @@ struct TNetRewindHistory : public Chaos::TDataRewindHistory<DataType>
 
 		// Allow received data to get recorded if it's for a newer frame than already stored at the history index or if the data stored is not marked as received (i.e. it's locally predicted)
 		Super::LoadData(ReceivedData.LocalFrame);
-		return (Super::DataHistory[Super::CurrentIndex].LocalFrame < ReceivedData.LocalFrame || (!Super::DataHistory[Super::CurrentIndex].bReceivedData && Super::DataHistory[Super::CurrentIndex].LocalFrame == ReceivedData.LocalFrame));
+		DataType& Data = Super::DataHistory[Super::CurrentIndex];
+		return (Data.LocalFrame < ReceivedData.LocalFrame || (!Data.bReceivedData && Data.LocalFrame == ReceivedData.LocalFrame));
 	}
 
 	/** Compares new received data with local predicted data and returns true if they differ enough to trigger a resimulation  */
@@ -326,7 +334,10 @@ struct TNetRewindHistory : public Chaos::TDataRewindHistory<DataType>
 			// Set the internal delta serialization source (between data entries in the collection)
 			if (I > 0)
 			{
-				Data.SetDeltaSourceData(&Super::DataHistory[I - 1]);
+				if constexpr (bLegacyData)
+				{
+					Data.SetDeltaSourceData(&Super::DataHistory[I - 1]);
+				}
 			}
 		}
 
@@ -335,8 +346,11 @@ struct TNetRewindHistory : public Chaos::TDataRewindHistory<DataType>
 			NetSerializeData(Data, Ar, InPackageMap);
 
 			// Clear delta source and strong pointer to implementation component after serialization
-			Data.ClearImplementationComponent();
-			Data.ClearDeltaSourceData();
+			if constexpr (bLegacyData)
+			{
+				Data.ClearImplementationComponent();
+				Data.ClearDeltaSourceData();
+			}
 		}
 
 		Super::Initialize();
@@ -668,11 +682,181 @@ public:
 };
 
 
+/** START Iris Compliant Data */
+
+USTRUCT()
+struct FNetworkPhysicsBase
+{
+	GENERATED_BODY()
+
+	virtual ~FNetworkPhysicsBase() { }
+
+	UPROPERTY()
+	int32 ServerFrame = 0;
+
+	mutable int32 LocalFrame = 0; // TEMP: Mutable because we need to set this after receiving it via RPC where it's const
+
+	/** Define how to interpolate between two data points if we have a gap between known data.
+	* @param MinData is data from a previous frame.
+	* @param MaxData is data from a future frame.
+	* @param LerpAlpha is the 0.0 - 1.0 value of where *this* data is between MinData and MaxData calculated as so: float LerpAlpha = (LocalFrame - MinData.LocalFrame) / (MaxData.LocalFrame - MinData.LocalFrame)
+ 	* EXAMPLE: We have input data for frame 1 and 4 and we need to interpolate data for frame 2 and 3 based on frame 1 as MinData and frame 4 as MaxData, for frame 2 LerpAlpha will be 0.33 and for frame 3 LerpAlpha will be 0.66
+	*/
+	virtual void InterpolateData(const FNetworkPhysicsBase& MinData, const FNetworkPhysicsBase& MaxData, float LerpAlpha) {}
+
+	/** Define how to merge data together
+	* @param FromData is data from a previous frame that is getting merged into the current data.
+	* EXAMPLE: Simulated proxies might receive two inputs at the same time after having used the same input twice, to not miss any important inputs we need to take both inputs into account
+	* and to not get behind in simulation we need to apply them both at the same simulation tick meaning we merge the two new inputs to one input.
+	*/
+	virtual void MergeData(const FNetworkPhysicsBase& FromData) {}
+
+	/** Use to decay desired data during resimulation if data is forward predicted.
+	* @param DecayAmount = Total amount of decay as a multiplier. 10% decay = 0.1.
+	* NOTE: Decay is not accumulated, the data will be in its original state each time DecayData is called. DecayAmount will increase each time the input is predicted (reused).
+	* EXAMPLE: Use to decay steering inputs to make resimulation not predict too much with a high steering value. Use DecayAmount of 0.1 to turn a steering value of 0.5 into 0.45 for example.
+	*/
+	virtual void DecayData(float DecayAmount) {}
+
+	/** Define how to compare client and server data for the same frame, returning false means the data differ enough to trigger a resimulation.
+	* @param PredictedData is data predicted on the client to compare with the current data received from the server.
+	* NOTE: To use this function, CVars np2.Resim.CompareStateToTriggerRewind and/or np2.Resim.CompareInputToTriggerRewind needs to be set to true
+	* or the equivalent settings overridden on the actor via UNetworkPhysicsSettingsComponent.
+	*/
+	virtual bool CompareData(const FNetworkPhysicsBase& PredictedData) const { return true; }
+
+	/** Return string with custom debug data */
+	virtual const FString DebugData() { return FString(" - DebugData() not implemented - "); } // TODO, Deprecate non-const version
+	virtual const FString DebugData() const { return FString(" - DebugData() not implemented - "); }
+
+	// If this data was altered so that it doesn't correspond the produced source data (from merging, interpolating or extrapolating)
+	bool bDataAltered = false;
+
+	// If this data was received over the network or locally predicted
+	bool bReceivedData = false;
+
+	// If this data is marked as important (replicated reliably)
+	bool bImportant = false;
+
+	void PrepareFrame(int32 CurrentFrame, bool bIsServer, int32 ClientFrameOffset)
+	{
+		LocalFrame = CurrentFrame;
+		ServerFrame = bIsServer ? CurrentFrame : CurrentFrame + ClientFrameOffset;
+		bDataAltered = false;
+		bReceivedData = false;
+		bImportant = false;
+	}
+
+	/** Temporary virtual function for backwards compatibility with FNetworkPhysicsData.Use the InterpolateData function that passes the LerpAlpha instead */
+	UE_DEPRECATED(5.7, "Temporary virtual function for backwards compatibility with FNetworkPhysicsData. Use the InterpolateData function that passes the LerpAlpha instead")
+	virtual void DoInterpolateData(const FNetworkPhysicsBase& MinData, const FNetworkPhysicsBase& MaxData)
+	{
+		const float LerpAlpha = static_cast<float>(LocalFrame - MinData.LocalFrame) / static_cast<float>(MaxData.LocalFrame - MinData.LocalFrame);
+		InterpolateData(MinData, MaxData, LerpAlpha);
+	}
+};
+
+USTRUCT()
+struct FNetworkPhysicsDataCollection
+{
+	GENERATED_BODY()
+
+	UPROPERTY()
+	TArray<TInstancedStruct<FNetworkPhysicsBase>> DataArray;
+
+	// TEMP, handle this via FNetworkPhysicsBaseNetSerializer (but without offset, just set LocalFrame = ServerFrame)
+	void UpdateLocalFrameFromServerFrame(const int32 ClientFrameOffset) const
+	{
+		for (const TInstancedStruct<FNetworkPhysicsBase>& DataInstance : DataArray)
+		{
+			const FNetworkPhysicsBase& Data = DataInstance.Get();
+			Data.LocalFrame = Data.ServerFrame - ClientFrameOffset;
+		}
+	}
+
+	void DebugCollection(const FString& DebugText) const
+	{
+		UE_LOG(LogChaos, Log, TEXT("%s"), *DebugText);
+		UE_LOG(LogChaos, Log, TEXT("	NumFrames in data collection: %d"), DataArray.Num());
+
+		if (DataArray.Num() >= 0)
+		{
+			for (int32 FrameIndex = 0; FrameIndex < DataArray.Num(); ++FrameIndex)
+			{
+				const TInstancedStruct<FNetworkPhysicsBase>& DataInstance = DataArray[FrameIndex];
+				const FNetworkPhysicsBase* Data = DataInstance.GetPtr<FNetworkPhysicsBase>();
+
+				if (Data)
+				{
+					UE_LOG(LogChaos, Log, TEXT("		Index: %d || LocalFrame = %d || ServerFrame = %d || bDataAltered = %d || bReceivedData = %d || bImportant = %d  ||  Data: %s")
+						, FrameIndex
+						, Data->LocalFrame
+						, Data->ServerFrame
+						, Data->bDataAltered
+						, Data->bReceivedData
+						, Data->bImportant
+						, *Data->DebugData());
+				}
+				else
+				{
+					UE_LOG(LogChaos, Log, TEXT("		No data in collection"));
+				}
+			}
+		}
+	}
+
+};
+
+// Physics Thread Input and State Interface API
+class INetworkPhysicsInputStateImplementationBase
+{
+public:
+	virtual void BuildInput(FNetworkPhysicsBase& Input) const = 0;
+	virtual void ValidateInput(FNetworkPhysicsBase& Input) const = 0;
+	virtual void ApplyInput(const FNetworkPhysicsBase& Input) = 0;
+
+	virtual void BuildState(FNetworkPhysicsBase& State) const = 0;
+	virtual void ApplyState(const FNetworkPhysicsBase& State) = 0;
+};
+
+template<typename InputType, typename StateType>
+class TNetworkPhysicsInputStateImplementation : public INetworkPhysicsInputStateImplementationBase
+{
+	// Wrapper for non-templated base
+	void BuildInput(FNetworkPhysicsBase& Input) const override { BuildInput_Internal(static_cast<InputType&>(Input)); }
+	void ValidateInput(FNetworkPhysicsBase& Input) const override { ValidateInput_Internal(static_cast<InputType&>(Input)); }
+	void ApplyInput(const FNetworkPhysicsBase& Input) override { ApplyInput_Internal(static_cast<const InputType&>(Input)); }
+
+	void BuildState(FNetworkPhysicsBase& State) const override { BuildState_Internal(static_cast<StateType&>(State)); }
+	void ApplyState(const FNetworkPhysicsBase& State) override { ApplyState_Internal(static_cast<const StateType&>(State)); }
+
+public:
+	/** Populate the input struct with current input data */
+	virtual void BuildInput_Internal(InputType& Input) const = 0;
+
+	/** Validate data received on the server from clients
+	* EXAMPLE: Validate incoming inputs from clients on the server and correct any invalid input commands.
+	* NOTE: Changes to the data in this callback will be sent from server to clients. */
+	virtual void ValidateInput_Internal(InputType& Input) const = 0;
+
+	/** Apply input struct to implementation */
+	virtual void ApplyInput_Internal(const InputType& Input) = 0;
+
+	/** Populate the state struct with current state data */
+	virtual void BuildState_Internal(StateType& State) const = 0;
+
+	/** Apply state struct to implementation */
+	virtual void ApplyState_Internal(const StateType& State) = 0;
+};
+
+/** END Iris Compliant Data */
+
+
 /**
  * Base network physics data that will be used by physics
  */
 USTRUCT()
-struct FNetworkPhysicsData
+struct FNetworkPhysicsData : public FNetworkPhysicsBase
 {
 	GENERATED_USTRUCT_BODY()
 
@@ -685,14 +869,6 @@ struct FNetworkPhysicsData
 	FNetworkPhysicsData& operator=(FNetworkPhysicsData&&) = default;
 	PRAGMA_ENABLE_DEPRECATION_WARNINGS
 
-	// Server frame at which this data has been generated 
-	UPROPERTY()
-	int32 ServerFrame = INDEX_NONE;
-
-	// Local frame at which this data has been generated (not replicated, calculated from ServerFrame - FrameOffset )
-	UPROPERTY()
-	int32 LocalFrame = INDEX_NONE;
-
 #if WITH_EDITORONLY_DATA
 	// InputFrame is no longer replicated or used. Use bDataAltered to check if an input has been altered
 	UE_DEPRECATED(5.6, "InputFrame is no longer replicated or used. Use bDataAltered to check if an input has been altered.")
@@ -700,15 +876,6 @@ struct FNetworkPhysicsData
 	int32 InputFrame_DEPRECATED = INDEX_NONE;
 #endif
 
-	// If this data was altered so that it doesn't correspond the produced source data (from merging, interpolating or extrapolating)
-	bool bDataAltered = false;
-
-	// If this data was received over the network or locally predicted
-	bool bReceivedData = false;
-
-	// If this data is marked as important (replicated reliably)
-	bool bImportant = false;
-
 	/** Thread safe pointer to the UActorComponent that implements the derived type of this data
 	* Note: This pointer can be accessed on both Game Thread and Physics Thread but you still need to ensure the read and write calls inside the UActorComponent are thread safe. */
 	TStrongObjectPtr<UActorComponent> ImplementationComponent = nullptr;
@@ -790,15 +957,6 @@ struct FNetworkPhysicsData
 		}
 	}
 
-	void PrepareFrame(int32 CurrentFrame, bool bIsServer, int32 ClientFrameOffset)
-	{
-		LocalFrame = CurrentFrame;
-		ServerFrame = bIsServer ? CurrentFrame : CurrentFrame + ClientFrameOffset;
-		bDataAltered = false;
-		bReceivedData = false;
-		bImportant = false;
-	}
-
 	/** Set if this data is important(replicated reliably) or unimportant(replicated unreliably)
 	* NOTE: Default is to handle all inputs as unimportant, while one time events can be marked as important. */
 	void SetImportant(bool bIsImportant)
@@ -824,7 +982,7 @@ struct FNetworkPhysicsData
 	* NOTE: Decay is not accumulated, the data will be in its original state each time DecayData is called. DecayAmount will increase each time the input is predicted (reused).
 	* EXAMPLE: Use to decay steering inputs to make resimulation not predict too much with a high steering value. Use DecayAmount of 0.1 to turn a steering value of 0.5 into 0.45 for example.
 	*/ 
-	virtual void DecayData(float DecayAmount) { }
+	virtual void DecayData(float DecayAmount) override { }
 	
 	/** Define how to merge data together
 	* @param FromData is data from a previous frame that is getting merged into the current data.
@@ -854,7 +1012,12 @@ struct FNetworkPhysicsData
 		return ServerFrame == Other.ServerFrame && LocalFrame == Other.LocalFrame;
 	}
 
-	friend UNetworkPhysicsComponent;
+	// Base API Wrapper
+	UE_DEPRECATED(5.7, "Temporary virtual function for backwards compatibility with FNetworkPhysicsData. Use the InterpolateData function that passes the LerpAlpha instead")
+	void DoInterpolateData(const FNetworkPhysicsBase& MinData, const FNetworkPhysicsBase& MaxData) override { this->InterpolateData(static_cast<const FNetworkPhysicsData&>(MinData), static_cast<const FNetworkPhysicsData&>(MaxData)); }
+	void MergeData(const FNetworkPhysicsBase& FromData) override { this->MergeData(static_cast<const FNetworkPhysicsData&>(FromData)); }
+	bool CompareData(const FNetworkPhysicsBase& PredictedData) const override { return const_cast<FNetworkPhysicsData*>(this)->CompareData(static_cast<const FNetworkPhysicsData&>(PredictedData)); }
+	using FNetworkPhysicsBase::InterpolateData; // Silence warning for FNetworkPhysicsData::InterpolateData not overriding a function in FNetworkPhysicsBase, which is intended
 };
 
 /** Base for helper, to create data and data history */
@@ -863,24 +1026,152 @@ struct FNetworkPhysicsDataHelper
 	virtual ~FNetworkPhysicsDataHelper() = default;
 	
 	virtual TUniquePtr<FNetworkPhysicsDataHelper> Clone() const = 0;
-	virtual TUniquePtr<FNetworkPhysicsData> CreateUniqueData() const = 0;
-	virtual void CopyData(FNetworkPhysicsData* To, FNetworkPhysicsData* From) = 0;
+	virtual TUniquePtr<FNetworkPhysicsBase> CreateUniqueData() const = 0;
 	virtual TUniquePtr<Chaos::FBaseRewindHistory> CreateUniqueRewindHistory(const int32 Size) const = 0;
+	virtual bool IsUsingLegacyData() = 0;
+
+	/** Copy data from the networked data collection to the rewind history */
+	virtual void CopyData(const FNetworkPhysicsDataCollection& From, Chaos::FBaseRewindHistory* To) = 0;
+	
+	/** Copy data from the rewind history to the networked data collection */
+	virtual void CopyData(const Chaos::FBaseRewindHistory* From, FNetworkPhysicsDataCollection& To) = 0;
+	
+	/** Copy data from the rewind history to the networked data collection and only copy data that is newer than the already existing data in the data collection */
+	virtual void CopyIncrementalData(const Chaos::FBaseRewindHistory* From, FNetworkPhysicsDataCollection& To) = 0;
+	
+	/** Copy data from the rewind history to the networked data collection and only copy data that has been altered since it was created (i.e. the client created an input that the server later altered) */
+	virtual bool CopyAlteredData(const Chaos::FBaseRewindHistory* From, FNetworkPhysicsDataCollection& To) = 0;
+	
+	/** Call the ValidateData callback on the implementation interface for data in the provided rewind history */
+	virtual void ValidateDataInHistory(Chaos::FBaseRewindHistory* History, INetworkPhysicsInputStateImplementationBase& Interface) = 0;
+
+/* // Importance is not implemented into the new network flow (yet?)
+	virtual void ApplyDataRange(const Chaos::FBaseRewindHistory* History, const int32 FromFrame, const int32 ToFrame, INetworkPhysicsInputStateImplementationBase& Interface) = 0;
+*/
 };
 
 /** Helper for the creation of state / input data and history with correct derived type */
-template<typename DataType>
+template<typename DataType, bool bLegacyData = false>
 struct TNetworkPhysicsDataHelper : FNetworkPhysicsDataHelper
 {
-	virtual ~TNetworkPhysicsDataHelper() = default;
+	~TNetworkPhysicsDataHelper() = default;
+	TUniquePtr<FNetworkPhysicsDataHelper> Clone() const override { return MakeUnique<TNetworkPhysicsDataHelper>(*this); }
+	TUniquePtr<FNetworkPhysicsBase> CreateUniqueData() const override { return MakeUnique<DataType>(); }
+	TUniquePtr<Chaos::FBaseRewindHistory> CreateUniqueRewindHistory(const int32 Size) const override { return MakeUnique<TNetRewindHistory<DataType, bLegacyData>>(Size); }
+	bool IsUsingLegacyData() { return bLegacyData; }
+
+	/** Copy data from the networked data collection to the rewind history */
+	void CopyData(const FNetworkPhysicsDataCollection& From, Chaos::FBaseRewindHistory* To) override
+	{
+		Chaos::TDataRewindHistory<DataType>* NetHistory = static_cast<Chaos::TDataRewindHistory<DataType>*>(To);
+
+		for (const TInstancedStruct<FNetworkPhysicsBase>& DataInstance : From.DataArray)
+		{
+			const DataType& Data = DataInstance.Get<DataType>();
+			NetHistory->RecordData(Data.LocalFrame, &Data);
+		}
+	}
 
-	virtual TUniquePtr<FNetworkPhysicsDataHelper> Clone() const { return MakeUnique<TNetworkPhysicsDataHelper>(*this); }
-	virtual TUniquePtr<FNetworkPhysicsData> CreateUniqueData() const { return MakeUnique<DataType>(); }
-	virtual void CopyData(FNetworkPhysicsData* To, FNetworkPhysicsData* From)
+	/** Copy data from the rewind history to the networked data collection */
+	void CopyData(const Chaos::FBaseRewindHistory* From, FNetworkPhysicsDataCollection& To) override
 	{
-		*static_cast<DataType*>(To) = *static_cast<DataType*>(From);
-	};
-	virtual TUniquePtr<Chaos::FBaseRewindHistory> CreateUniqueRewindHistory(const int32 Size) const { return MakeUnique<TNetRewindHistory<DataType>>(Size); }
+		if (ensure(To.DataArray.Num() > 0) == false)
+		{
+			return;
+		}
+		
+		const Chaos::TDataRewindHistory<DataType>* NetHistory = static_cast<const Chaos::TDataRewindHistory<DataType>*>(From);
+		const TArray<DataType>& HistoryArray = NetHistory->GetDataHistoryConst();
+		for (int32 FromIdx = 0; FromIdx < NetHistory->GetHistorySize(); FromIdx++)
+		{
+			const int32 ToIdx = FromIdx % To.DataArray.Num();
+			To.DataArray[ToIdx] = TInstancedStruct<DataType>::Make(HistoryArray[FromIdx]);
+		}
+	}
+
+	/** Copy data from the rewind history to the networked data collection and only copy data that is newer than the already existing data in the data collection */
+	void CopyIncrementalData(const Chaos::FBaseRewindHistory* From, FNetworkPhysicsDataCollection& To) override
+	{
+		if (ensure(To.DataArray.Num() > 0) == false)
+		{
+			return;
+		}
+		const Chaos::TDataRewindHistory<DataType>* NetHistory = static_cast<const Chaos::TDataRewindHistory<DataType>*>(From);
+
+		const TArray<DataType>& HistoryArray = NetHistory->GetDataHistoryConst();
+		for (int32 FromIdx = 0; FromIdx < NetHistory->GetHistorySize(); FromIdx++)
+		{
+			const DataType& FromData = HistoryArray[FromIdx];
+
+			const int32 ToIdx = FromIdx % To.DataArray.Num();
+			const FNetworkPhysicsBase* ToData = To.DataArray[ToIdx].GetPtr<FNetworkPhysicsBase>();
+
+			// Only copy the data if it's newer than the already cached data
+			if (!ToData || ToData->LocalFrame < FromData.LocalFrame)
+			{
+				To.DataArray[ToIdx] = TInstancedStruct<DataType>::Make(FromData);
+			}
+		}
+	}
+
+	/** Copy data from the rewind history to the networked data collection and only copy data that has been altered since it was created (i.e. the client created an input that the server later altered) */
+	bool CopyAlteredData(const Chaos::FBaseRewindHistory* From, FNetworkPhysicsDataCollection& To) override
+	{
+		bool bHasCopiedData = false;
+		if (ensure(To.DataArray.Num() > 0) == false)
+		{
+			return bHasCopiedData;
+		}
+		
+		const Chaos::TDataRewindHistory<DataType>* NetHistory = static_cast<const Chaos::TDataRewindHistory<DataType>*>(From);
+
+		const TArray<DataType>& HistoryArray = NetHistory->GetDataHistoryConst();
+		for (int32 FromIdx = 0; FromIdx < NetHistory->GetHistorySize(); FromIdx++)
+		{
+			const DataType& FromData = HistoryArray[FromIdx];
+
+			// Only copy the data that has been altered
+			if (FromData.bDataAltered)
+			{
+				const int32 ToIdx = FromIdx % To.DataArray.Num();
+				To.DataArray[ToIdx] = TInstancedStruct<DataType>::Make(FromData);
+				bHasCopiedData = true;
+			}
+		}
+
+		return bHasCopiedData;
+	}
+
+	/** Call the ValidateData callback on the implementation interface for data in the provided rewind history */
+	void ValidateDataInHistory(Chaos::FBaseRewindHistory* History, INetworkPhysicsInputStateImplementationBase& Interface) override
+	{
+		Chaos::TDataRewindHistory<DataType>* NetHistory = static_cast<Chaos::TDataRewindHistory<DataType>*>(History);
+		TArray<DataType>& HistoryArray = NetHistory->GetDataHistory();
+
+		for (int32 FromIdx = 0; FromIdx < NetHistory->GetHistorySize(); FromIdx++)
+		{
+			FNetworkPhysicsBase& Data = HistoryArray[FromIdx];
+			Interface.ValidateInput(Data);
+		}
+	}
+
+/* // Important is not implemented into the new network flow (yet?)
+	void ApplyDataRange(const Chaos::FBaseRewindHistory* History, const int32 FromFrame, const int32 ToFrame, INetworkPhysicsInputStateImplementationBase& Interface) override
+	{
+		const Chaos::TDataRewindHistory<DataType>* NetHistory = static_cast<const Chaos::TDataRewindHistory<DataType>*>(History);
+		const TArray<DataType>& HistoryArray = NetHistory->GetDataHistoryConst();
+
+		for (int32 ApplyFrame = FromFrame; ApplyFrame <= ToFrame; ++ApplyFrame)
+		{
+			const int32 ApplyIndex = NetHistory->GetFrameIndex(ApplyFrame);
+			const DataType& Data = HistoryArray[ApplyIndex];
+			if (ApplyFrame == Data.LocalFrame)
+			{
+				Interface.ApplyInput(Data);
+			}
+		}
+	}
+*/
 };
 
 /**
@@ -964,11 +1255,11 @@ public:
 	/** Check if this is controlled locally through relayed inputs from autonomous proxy. It's recommended to use IsLocallyControlled() when checking if this is locally controlled. */
 	const bool GetIsRelayingLocalInputs() const { return bIsRelayingLocalInputs; }
 
-	/** Override the initial value set by CVar np2.Resim.CompareStateToTriggerRewind on initialize -- When true, cache the clients FNetworkPhysicsData state in rewind history for autonomous proxies and compare the predicted state with incoming server state to trigger resimulations if they differ, comparison done through FNetworkPhysicsData::CompareData. 
+	/** Override the initial value set by CVar np2.Resim.CompareStateToTriggerRewind on initialize -- When true, cache the clients FNetworkPhysicsBase state in rewind history for autonomous proxies and compare the predicted state with incoming server state to trigger resimulations if they differ, comparison done through FNetworkPhysicsData::CompareData. 
 	* @param bInIncludeSimProxies overrides value set by CVar np2.Resim.CompareStateToTriggerRewind.IncludeSimProxies on initialize -- When true CompareStateToTriggerRewind is also done for simulated proxies. */
 	ENGINE_API void SetCompareStateToTriggerRewind(const bool bInCompareStateToTriggerRewind, const bool bInIncludeSimProxies = false);
 
-	/** Override the initial value set by CVar np2.Resim.CompareInputToTriggerRewind on initialize -- When true, compare autonomous proxies predicted FNetworkPhysicsData inputs with incoming server inputs to trigger resimulations if they differ, comparison done through FNetworkPhysicsData::CompareData.*/
+	/** Override the initial value set by CVar np2.Resim.CompareInputToTriggerRewind on initialize -- When true, compare autonomous proxies predicted FNetworkPhysicsBase inputs with incoming server inputs to trigger resimulations if they differ, comparison done through FNetworkPhysicsData::CompareData.*/
 	ENGINE_API void SetCompareInputToTriggerRewind(const bool bInCompareInputToTriggerRewind);
 
 	/** Get FAsyncNetworkPhysicsComponent on the Physics Thread */
@@ -1053,6 +1344,37 @@ protected:
 	UPROPERTY(Transient, ReplicatedUsing = OnRep_SetReplicatedStates)
 	FNetworkPhysicsRewindDataStateProxy ReplicatedStates;
 
+
+	/** Iris Replication */
+	// Server RPC to receive inputs from client
+	UFUNCTION(Server, unreliable)
+	ENGINE_API void ServerReceiveInputCollection(const FNetworkPhysicsDataCollection& ClientInputCollection);
+
+	// repnotify for inputs on owner client
+	UFUNCTION()
+	ENGINE_API void OnRep_SetReplicatedInputCollection();
+
+	// repnotify for inputs on remote clients
+	UFUNCTION()
+	ENGINE_API void OnRep_SetReplicatedRemoteInputCollection();
+
+	// repnotify for the states on the client
+	UFUNCTION()
+	ENGINE_API void OnRep_SetReplicatedStateCollection();
+
+	// replicated physics inputs for owner client
+	UPROPERTY(Transient, ReplicatedUsing = OnRep_SetReplicatedInputCollection)
+	FNetworkPhysicsDataCollection ReplicatedInputCollection;
+
+	// replicated physics inputs for remote clients
+	UPROPERTY(Transient, ReplicatedUsing = OnRep_SetReplicatedRemoteInputCollection)
+	FNetworkPhysicsDataCollection ReplicatedRemoteInputCollection;
+
+	// replicated physics states 
+	UPROPERTY(Transient, ReplicatedUsing = OnRep_SetReplicatedStateCollection)
+	FNetworkPhysicsDataCollection ReplicatedStateCollection;
+
+
 private:
 	/** Send network data from marshaled data */
 	void NetworkMarshaledData();
@@ -1130,6 +1452,10 @@ public:
 	template<class InputsType>
 	void CreateInputHistory(UActorComponent* HistoryComponent);
 
+	/** Register state and input to be networked and cached in history along with the interface implementation to interact with the input and state */
+	template<typename Input, typename State>
+	void CreateDataHistory(TNetworkPhysicsInputStateImplementation<Input, State>* Implementation);
+
 private:
 
 	friend FNetworkPhysicsCallback;
@@ -1141,6 +1467,10 @@ private:
 	friend struct FNetworkPhysicsRewindDataDeltaSourceInputProxy;
 	friend struct FNetworkPhysicsRewindDataDeltaSourceStateProxy;
 
+	INetworkPhysicsInputStateImplementationBase* ImplementationInterface = nullptr;
+
+	bool bIsUsingLegacyData = false;
+
 	// Network Physics Component data internal to the physics thread
 	FAsyncNetworkPhysicsComponent* NetworkPhysicsComponent_Internal;
 
@@ -1156,11 +1486,11 @@ private:
 	// Helper for the creation of state data and history with correct derived type
 	TUniquePtr<FNetworkPhysicsDataHelper> StateHelper;
 
-	// Local temporary input data
-	TUniquePtr<FNetworkPhysicsData> InputDataDefault;
+	// Local default input data in legacy type
+	TUniquePtr<FNetworkPhysicsData> InputDataDefault_Legacy;
 
-	// Local temporary state data
-	TUniquePtr<FNetworkPhysicsData> StateDataDefault;
+	// Local default state data in legacy type
+	TUniquePtr<FNetworkPhysicsData> StateDataDefault_Legacy;
 
 	// The number of inputs the owning client should send to the server with each RPC, replicated from the server. This is dynamically scaled based on when there are holes in the inputs buffer if np2.Resim.DynamicInputScaling.Enabled is enabled
 	UPROPERTY( Replicated )
@@ -1201,7 +1531,7 @@ private:
 	// If we are currently relaying inputs and will stop after next network send.
 	bool bStopRelayingLocalInputsDeferred = false;
 
-	// Compare state / input to trigger rewind via FNetworkPhysicsData::CompareData
+	// Compare state / input to trigger rewind via FNetworkPhysicsBase::CompareData
 	bool bCompareStateToTriggerRewind = false;
 	bool bCompareStateToTriggerRewindIncludeSimProxies = false; // Include simulated proxies when bCompareStateToTriggerRewind is enabled
 	bool bCompareInputToTriggerRewind = false;
@@ -1211,31 +1541,7 @@ private:
 	bool bEnableReliableFlow = false;
 	bool bValidateDataOnGameThread = false;
 
-
 	/** ----- Deprecated API ----- */
-public:
-	// Replicate input data 
-	UE_DEPRECATED(5.5, "Deprecated, inputs are sent from PT marshaled data and can't be manually triggered on GT.")
-		void SendInputData() {}
-
-	// Replicate state data
-	UE_DEPRECATED(5.5, "Deprecated, states are sent from PT marshaled data and can't be manually triggered on GT.")
-		void SendStateData() {}
-
-	// Enable RewindData history caching and return the history size
-	UE_DEPRECATED(5.5, "Deprecated, don't setup rewind data from GameThread")
-		int32 SetupRewindData() { return 0; }
-
-	UE_DEPRECATED(5.5, "Deprecated, use GetStateHistory_Internal() or GetStateHistory_External() instead")
-		TSharedPtr<Chaos::FBaseRewindHistory>& GetStateHistory() { return GetStateHistory_Internal(); }
-
-	UE_DEPRECATED(5.5, "Deprecated, use GetInputHistory_Internal() or GetInputHistory_External() instead")
-		TSharedPtr<Chaos::FBaseRewindHistory>& GetInputHistory() { return GetInputHistory_Internal(); }
-
-	/** Returns the current amount of input decay during resimulation as a magnitude from 0.0 to 1.0. Returns 0 if not currently resimulating. */
-	UE_DEPRECATED(5.5, "Deprecated, use FAsyncNetworkPhysicsComponent::GetCurrentInputDecay instead.")
-		ENGINE_API const float GetCurrentInputDecay(FNetworkPhysicsData* PhysicsData);
-
 private:
 #if WITH_EDITORONLY_DATA
 	UE_DEPRECATED(5.6, "Deprecated, use InputsToNetwork_Owner instead")
@@ -1245,20 +1551,35 @@ private:
 #endif
 };
 
+template<typename Input, typename State>
+void UNetworkPhysicsComponent::CreateDataHistory(TNetworkPhysicsInputStateImplementation<Input, State>* Implementation)
+{
+	bIsUsingLegacyData = false;
+
+	ImplementationInterface = Implementation;
+
+	InputHelper = MakeUnique<TNetworkPhysicsDataHelper<Input, /*bLegacyData*/ false>>();
+	StateHelper = MakeUnique<TNetworkPhysicsDataHelper<State, /*bLegacyData*/ false>>();
+
+	CreateAsyncDataHistory();
+}
+
 template<typename PhysicsTraits>
 inline void UNetworkPhysicsComponent::CreateDataHistory(UActorComponent* HistoryComponent)
 {
-	InputHelper = MakeUnique<TNetworkPhysicsDataHelper<typename PhysicsTraits::InputsType>>();
-	StateHelper = MakeUnique<TNetworkPhysicsDataHelper<typename PhysicsTraits::StatesType>>();
+	bIsUsingLegacyData = true;
+
+	InputHelper = MakeUnique<TNetworkPhysicsDataHelper<typename PhysicsTraits::InputsType, /*bLegacyData*/ true>>();
+	StateHelper = MakeUnique<TNetworkPhysicsDataHelper<typename PhysicsTraits::StatesType, /*bLegacyData*/ true>>();
 	
-	InputDataDefault = InputHelper->CreateUniqueData();
-	StateDataDefault = StateHelper->CreateUniqueData();
+	InputDataDefault_Legacy = MakeUnique<typename PhysicsTraits::InputsType>();
+	StateDataDefault_Legacy = MakeUnique<typename PhysicsTraits::StatesType>();
 
 	// Initialize delta source arrays
 	for (int32 I = 0; I < DeltaSourceBufferSize; I++)
 	{
-		DeltaSourceInputs.Add(InputHelper->CreateUniqueData());
-		DeltaSourceStates.Add(StateHelper->CreateUniqueData());
+		DeltaSourceInputs.Add(MakeUnique<typename PhysicsTraits::InputsType>());
+		DeltaSourceStates.Add(MakeUnique<typename PhysicsTraits::StatesType>());
 	}
 
 	ReplicatedInputs.History = InputHelper->CreateUniqueRewindHistory(InputsToNetwork_OwnerDefault);
@@ -1290,14 +1611,16 @@ inline void UNetworkPhysicsComponent::CreateDataHistory(UActorComponent* History
 template<class InputsType>
 inline void UNetworkPhysicsComponent::CreateInputHistory(UActorComponent* HistoryComponent)
 {
-	InputHelper = MakeUnique<TNetworkPhysicsDataHelper<InputsType>>();
+	bIsUsingLegacyData = true;
+
+	InputHelper = MakeUnique<TNetworkPhysicsDataHelper<InputsType, /*bLegacyData*/ true>>();
 	
-	InputDataDefault = InputHelper->CreateUniqueData();
+	InputDataDefault_Legacy = MakeUnique<InputsType>();
 
 	// Initialize delta source array
 	for (int32 I = 0; I < DeltaSourceBufferSize; I++)
 	{
-		DeltaSourceInputs.Add(InputHelper->CreateUniqueData());
+		DeltaSourceInputs.Add(MakeUnique<InputsType>());
 	}
 
 	ReplicatedInputs.History = InputHelper->CreateUniqueRewindHistory(InputsToNetwork_OwnerDefault);
@@ -1329,6 +1652,7 @@ struct FAsyncNetworkPhysicsComponentInput : public Chaos::FSimCallbackInput
 	TOptional<uint16> InputsToNetwork_Owner;
 	TOptional<EPhysicsReplicationMode> PhysicsReplicationMode;
 	TOptional<TWeakObjectPtr<UActorComponent>> ActorComponent;
+	TOptional<INetworkPhysicsInputStateImplementationBase*> ImplementationInterface;
 	TOptional<Chaos::FConstPhysicsObjectHandle> PhysicsObject;
 	TOptional<FString> ActorName;
 	TOptional<TUniquePtr<FNetworkPhysicsDataHelper>> InputHelper;
@@ -1355,6 +1679,7 @@ struct FAsyncNetworkPhysicsComponentInput : public Chaos::FSimCallbackInput
 		InputsToNetwork_Owner.Reset();
 		PhysicsReplicationMode.Reset();
 		ActorComponent.Reset();
+		ImplementationInterface.Reset();
 		PhysicsObject.Reset();
 		ActorName.Reset();
 		InputHelper.Reset();
@@ -1482,7 +1807,7 @@ private:
 	void TriggerResimulation(int32 ResimFrame);
 
 	/** Returns the current amount of input decay during resimulation as a magnitude from 0.0 to 1.0. Returns 0 if not currently resimulating. */
-	const float GetCurrentInputDecay(const FNetworkPhysicsData* PhysicsData);
+	const float GetCurrentInputDecay(const FNetworkPhysicsBase* PhysicsData);
 
 	/** Populate input data in AsyncOutput to send over the network */
 	void SendInputData_Internal(FAsyncNetworkPhysicsComponentOutput& AsyncOutput, const int32 PhysicsStep);
@@ -1500,6 +1825,7 @@ private:
 	int32 NetworkPhysicsTickOffset;
 	EPhysicsReplicationMode PhysicsReplicationMode;
 	FString ActorName;
+	bool bIsUsingLegacyData;
 
 	int32 LastInputSendFrame = INDEX_NONE;
 	int32 LastStateSendFrame = INDEX_NONE;
@@ -1512,9 +1838,18 @@ private:
 	// Actor component that will be used to fill the histories
 	TWeakObjectPtr<UActorComponent> ActorComponent;
 
+	// Implementation of input / state interface
+	INetworkPhysicsInputStateImplementationBase* ImplementationInterface;
+
 	// Root components physics object
 	Chaos::FConstPhysicsObjectHandle PhysicsObject;
 
+	// Helper for the creation of input data and history with correct derived type
+	TUniquePtr<FNetworkPhysicsDataHelper> InputHelper;
+
+	// Helper for the creation of state data and history with correct derived type
+	TUniquePtr<FNetworkPhysicsDataHelper> StateHelper;
+
 	// States history uses to rewind simulation 
 	TSharedPtr<Chaos::FBaseRewindHistory> StateHistory;
 
@@ -1522,13 +1857,13 @@ private:
 	TSharedPtr<Chaos::FBaseRewindHistory> InputHistory;
 
 	// Local temporary inputs data used by pre/post process inputs functions
-	TUniquePtr<FNetworkPhysicsData> InputData;
+	TUniquePtr<FNetworkPhysicsBase> InputData;
 	
 	// Local temporary inputs data used by ConsumeAsyncInput
-	TUniquePtr<FNetworkPhysicsData> LatestInputReceiveData;
+	TUniquePtr<FNetworkPhysicsBase> LatestInputReceiveData;
 
 	// Local temporary states data used by pre/post process inputs functions
-	TUniquePtr<FNetworkPhysicsData> StateData;
+	TUniquePtr<FNetworkPhysicsBase> StateData;
 
 	// Send last N number of inputs each replication call
 	uint16 InputsToNetwork_OwnerDefault = 3; // Default value for owning client
diff --git a/Engine/Source/Runtime/Experimental/Chaos/Public/RewindData.h b/Engine/Source/Runtime/Experimental/Chaos/Public/RewindData.h
--- a/Engine/Source/Runtime/Experimental/Chaos/Public/RewindData.h
+++ b/Engine/Source/Runtime/Experimental/Chaos/Public/RewindData.h
@@ -234,16 +234,18 @@ public :
 				DataType& ExtractedData = *static_cast<DataType*>(HistoryData);
 				ExtractedData = DataHistory[MinFrameIndex];
 
-				ExtractedData.InterpolateData(DataHistory[MinFrameIndex], DataHistory[MaxFrameIndex]);
-
 				const int32 DeltaFrame = FMath::Abs(ExtractFrame - DataHistory[MinFrameIndex].LocalFrame);
 
-				// Don't update LocalFrame here so that it can be used to check if this data was grabbed from the correct frame or not
-				ExtractedData.ServerFrame = DataHistory[MinFrameIndex].ServerFrame + DeltaFrame;
+				ExtractedData.LocalFrame = ExtractFrame;
+				ExtractedData.ServerFrame += DeltaFrame;
 				ExtractedData.bDataAltered = true; // This history entry is now altered and doesn't correspond to the source entry
 
+				PRAGMA_DISABLE_DEPRECATION_WARNINGS
+				ExtractedData.DoInterpolateData(DataHistory[MinFrameIndex], DataHistory[MaxFrameIndex]);
+				PRAGMA_ENABLE_DEPRECATION_WARNINGS
+
 #if DEBUG_NETWORK_PHYSICS
-				UE_LOG(LogChaos, Log, TEXT("		Interpolating data between frame %d and %d - > [%d]"), DataHistory[MinFrameIndex].LocalFrame, DataHistory[MaxFrameIndex].LocalFrame, static_cast<DataType*>(HistoryData)->ServerFrame);
+				UE_LOG(LogChaos, Log, TEXT("		Interpolating data between frame %d and %d - > [%d]"), DataHistory[MinFrameIndex].LocalFrame, DataHistory[MaxFrameIndex].LocalFrame, ExtractedData.LocalFrame);
 #endif
 				return true;
 			}
@@ -254,7 +256,7 @@ public :
 
 				const int32 DeltaFrame = FMath::Abs(ExtractFrame - DataHistory[MinFrameIndex].LocalFrame);
 
-				// Don't update LocalFrame here so that it can be used to check how many frames ahead of the last entry we are currently
+				ExtractedData.LocalFrame = ExtractFrame;
 				ExtractedData.ServerFrame = DataHistory[MinFrameIndex].ServerFrame + DeltaFrame;
 				ExtractedData.bDataAltered = true; // This history entry is now altered and doesn't correspond to the source entry
 
@@ -354,6 +356,7 @@ public :
 	}
 
 	TArray<DataType>& GetDataHistory() { return DataHistory; }
+	const TArray<DataType>& GetDataHistoryConst() const { return DataHistory; }
 
 	/** Return the most up to date frame entry in history */
 	virtual const int32 GetLatestFrame() const override
@@ -401,7 +404,7 @@ public :
 
 	FORCEINLINE const uint32 GetFrameIndex(const int32 Frame) const
 	{
-		return FMath::Abs(Frame % NumFrames);
+		return NumFrames == 0 ? 0 : FMath::Abs(Frame % NumFrames);
 	}
 
 	FORCEINLINE virtual void ResetFast()
