New material translator:
- Changed the definition of Vector from column-based to row-based so that operations like "matmul(v, M)" behave the same way in both the translator API and the generated HLSL.

#rnx

diff --git a/Engine/Source/Runtime/Engine/Private/Materials/MaterialExpressionsIR.cpp b/Engine/Source/Runtime/Engine/Private/Materials/MaterialExpressionsIR.cpp
--- a/Engine/Source/Runtime/Engine/Private/Materials/MaterialExpressionsIR.cpp
+++ b/Engine/Source/Runtime/Engine/Private/Materials/MaterialExpressionsIR.cpp
@@ -396,7 +396,7 @@ void UMaterialExpressionVectorParameter::Build(MIR::FEmitter& Em)
 void UMaterialExpressionDoubleVectorParameter::Build(MIR::FEmitter& Em)
 {
 	FValueRef Value = BuildMaterialExpressionParameter(Em, this);
-	Em.Output(0, Em.Cast(Value, MIR::FType::MakeDouble(3)));
+	Em.Output(0, Em.Cast(Value, MIR::FType::MakeDoubleVector(3)));
 	Em.Output(1, Em.Subscript(Value, 0));
 	Em.Output(2, Em.Subscript(Value, 1));
 	Em.Output(3, Em.Subscript(Value, 2));
@@ -499,7 +499,7 @@ void UMaterialExpressionAppendVector::Build(MIR::FEmitter& Em)
 	MIR::FPrimitive AType = AVal->Type.GetPrimitive();
 	MIR::FPrimitive BType = BVal ? BVal->Type.GetPrimitive() : MIR::FPrimitive{};
 
-	int Dimensions = AType.NumRows + (BVal ? BType.NumRows : 0);
+	int Dimensions = AType.NumColumns + (BVal ? BType.NumColumns : 0);
 	if (Dimensions > 4)
 	{
 		Em.Errorf(TEXT("The resulting vector would have %d component (it can have at most 4)."), Dimensions);
@@ -509,20 +509,20 @@ void UMaterialExpressionAppendVector::Build(MIR::FEmitter& Em)
 	check(Dimensions >= 2 && Dimensions <= 4);
 
 	// Construct the output vector type.
-	MIR::EScalarKind ResultKind = (AType.IsLWC() || (BVal && BType.IsLWC())) ? MIR::EScalarKind::Double : MIR::EScalarKind::Float;
+	MIR::EScalarKind ResultKind = (AType.IsDouble() || (BVal && BType.IsDouble())) ? MIR::EScalarKind::Double : MIR::EScalarKind::Float;
 	MIR::FType ResultType = MIR::FType::MakeVector(ResultKind, Dimensions);
 
 	// Set up each output vector component.  These need CastToScalarKind in case we are appending LWC and non-LWC.
 	FValueRef Components[4] = { nullptr, nullptr, nullptr, nullptr };
 	int ComponentIndex = 0;
-	for (int i = 0; i < AType.NumRows; ++i, ++ComponentIndex)
+	for (int i = 0; i < AType.NumColumns; ++i, ++ComponentIndex)
 	{
 		Components[ComponentIndex] = Em.CastToScalarKind(Em.Subscript(AVal, i), ResultKind);
 	}
 
 	if (BVal)
 	{
-		for (int i = 0; i < BType.NumRows; ++i, ++ComponentIndex)
+		for (int i = 0; i < BType.NumColumns; ++i, ++ComponentIndex)
 		{
 			Components[ComponentIndex] = Em.CastToScalarKind(Em.Subscript(BVal, i), ResultKind);
 		}
@@ -972,7 +972,7 @@ static FValueRef SceneTextureExpressionTexCoords(MIR::FEmitter& Em, FValueRef Sc
 	if (DefaultOffset)
 	{
 		TexCoords = Em.InputDefaultFloat2(&Coordinates, { (float)DefaultOffset->X, (float)DefaultOffset->Y });
-		TexCoords = Em.InlineHLSL(MIR::FType::MakeFloat(2), TEXTVIEW("CalcScreenUVFromOffsetFraction(GetScreenPosition(Parameters), $0)"), { TexCoords });
+		TexCoords = Em.InlineHLSL(MIR::FType::MakeFloatVector(2), TEXTVIEW("CalcScreenUVFromOffsetFraction(GetScreenPosition(Parameters), $0)"), { TexCoords });
 	}
 	else
 	{
@@ -980,8 +980,8 @@ static FValueRef SceneTextureExpressionTexCoords(MIR::FEmitter& Em, FValueRef Sc
 		if (TexCoords)
 		{
 			// Convert raw TexCoords expression to Float2, then convert from viewport to scene texture space
-			TexCoords = Em.Cast(TexCoords, MIR::FType::MakeFloat(2));
-			TexCoords = Em.InlineHLSL(MIR::FType::MakeFloat(2), TEXTVIEW("ClampSceneTextureUV(ViewportUVToSceneTextureUV($0, $1), $1)"), { TexCoords, SceneTextureInput });
+			TexCoords = Em.Cast(TexCoords, MIR::FType::MakeFloatVector(2));
+			TexCoords = Em.InlineHLSL(MIR::FType::MakeFloatVector(2), TEXTVIEW("ClampSceneTextureUV(ViewportUVToSceneTextureUV($0, $1), $1)"), { TexCoords, SceneTextureInput });
 		}
 		else
 		{
@@ -989,7 +989,7 @@ static FValueRef SceneTextureExpressionTexCoords(MIR::FEmitter& Em, FValueRef Sc
 				TEXTVIEW("ClampSceneTextureUV(GetDefaultSceneTextureUV(Parameters, $0), $0)") :
 				TEXTVIEW("GetDefaultSceneTextureUV(Parameters, $0)");
 
-			TexCoords = Em.InlineHLSL(MIR::FType::MakeFloat(2), TexCoordsCode, { SceneTextureInput });
+			TexCoords = Em.InlineHLSL(MIR::FType::MakeFloatVector(2), TexCoordsCode, { SceneTextureInput });
 		}
 	}
 	return TexCoords;
@@ -1003,9 +1003,9 @@ static void SceneTextureExpressionBuild(MIR::FEmitter& Em, FValueRef SceneTextur
 		TEXTVIEW("SceneTextureLookup(Parameters, $0, $1, true)") :
 		TEXTVIEW("SceneTextureLookup(Parameters, $0, $1, false)");
 
-	Em.Output(0, Em.InlineHLSL(MIR::FType::MakeFloat(4), SceneTextureLookupCode, { TexCoords, SceneTextureInput }));
-	Em.Output(1, Em.InlineHLSL(MIR::FType::MakeFloat(2), TEXTVIEW("GetSceneTextureViewSize($0).xy"), { SceneTextureInput }));
-	Em.Output(2, Em.InlineHLSL(MIR::FType::MakeFloat(2), TEXTVIEW("GetSceneTextureViewSize($0).zw"), { SceneTextureInput }));
+	Em.Output(0, Em.InlineHLSL(MIR::FType::MakeFloatVector(4), SceneTextureLookupCode, { TexCoords, SceneTextureInput }));
+	Em.Output(1, Em.InlineHLSL(MIR::FType::MakeFloatVector(2), TEXTVIEW("GetSceneTextureViewSize($0).xy"), { SceneTextureInput }));
+	Em.Output(2, Em.InlineHLSL(MIR::FType::MakeFloatVector(2), TEXTVIEW("GetSceneTextureViewSize($0).zw"), { SceneTextureInput }));
 }
 
 void UMaterialExpressionSceneTexture::Build(MIR::FEmitter& Em)
@@ -1025,7 +1025,7 @@ void UMaterialExpressionSceneColor::Build(MIR::FEmitter& Em)
 
 	// We need a dependency on ScreenTexture as a second argument, so the value analyzer can see it, even though it's technically not used in the code.
 	FValueRef ScreenTexture = Em.ScreenTexture(MIR::EScreenTexture::SceneColor);
-	Em.Output(0, Em.InlineHLSL(MIR::FType::MakeFloat(4), TEXTVIEW("DecodeSceneColorAndAlpharForMaterialNode($0)"), { TexCoords, ScreenTexture }));;
+	Em.Output(0, Em.InlineHLSL(MIR::FType::MakeFloatVector(4), TEXTVIEW("DecodeSceneColorAndAlpharForMaterialNode($0)"), { TexCoords, ScreenTexture }));;
 }
 
 void UMaterialExpressionSceneDepth::Build(MIR::FEmitter& Em)
@@ -1034,7 +1034,7 @@ void UMaterialExpressionSceneDepth::Build(MIR::FEmitter& Em)
 
 	// We need a dependency on ScreenTexture as a second argument, so the value analyzer can see it, even though it's technically not used in the code.
 	FValueRef ScreenTexture = Em.ScreenTexture(MIR::EScreenTexture::SceneDepth);
-	Em.Output(0, Em.InlineHLSL(MIR::FType::MakeFloat(1), TEXTVIEW("CalcSceneDepth($0)"), { TexCoords, ScreenTexture }));;
+	Em.Output(0, Em.InlineHLSL(MIR::FType::MakeFloatScalar(), TEXTVIEW("CalcSceneDepth($0)"), { TexCoords, ScreenTexture }));;
 }
 
 void UMaterialExpressionSceneDepthWithoutWater::Build(MIR::FEmitter& Em)
@@ -1043,7 +1043,7 @@ void UMaterialExpressionSceneDepthWithoutWater::Build(MIR::FEmitter& Em)
 
 	// We need a dependency on ScreenTexture as a third argument, so the value analyzer can see it, even though it's technically not used in the code.
 	FValueRef ScreenTexture = Em.ScreenTexture(MIR::EScreenTexture::SceneDepthWithoutWater);
-	Em.Output(0, Em.InlineHLSL(MIR::FType::MakeFloat(1), TEXTVIEW("MaterialExpressionSceneDepthWithoutWater($0, $1)"), { TexCoords, Em.ConstantFloat(FallbackDepth), ScreenTexture }));;
+	Em.Output(0, Em.InlineHLSL(MIR::FType::MakeFloatScalar(), TEXTVIEW("MaterialExpressionSceneDepthWithoutWater($0, $1)"), { TexCoords, Em.ConstantFloat(FallbackDepth), ScreenTexture }));;
 }
 
 void UMaterialExpressionDBufferTexture::Build(MIR::FEmitter& Em)
@@ -1051,7 +1051,7 @@ void UMaterialExpressionDBufferTexture::Build(MIR::FEmitter& Em)
 	FValueRef TexCoords = SceneTextureExpressionTexCoords(Em, Em.ConstantInt(0), Coordinates);
 
 	FValueRef ScreenTexture = Em.DBufferTexture(DBufferTextureId);
-	Em.Output(0, Em.InlineHLSL(MIR::FType::MakeFloat(4), TEXTVIEW("MaterialExpressionDBufferTextureLookup(Parameters, $0, $1)"), { TexCoords, ScreenTexture }));;
+	Em.Output(0, Em.InlineHLSL(MIR::FType::MakeFloatVector(4), TEXTVIEW("MaterialExpressionDBufferTextureLookup(Parameters, $0, $1)"), { TexCoords, ScreenTexture }));;
 }
 
 void UMaterialExpressionSphericalParticleOpacity::Build(MIR::FEmitter& Em)
@@ -1568,7 +1568,7 @@ static FValueRef BuildViewProperty(MIR::FEmitter& Em, EMaterialExposedViewProper
 
 	// CastToNonLWCIfDisabled
 	TOptional<MIR::FPrimitive> PrimitiveType = HLSLCodeType.AsPrimitive();
-	if (PrimitiveType && PrimitiveType->IsLWC() && !UE::MaterialTranslatorUtils::IsLWCEnabled())
+	if (PrimitiveType && PrimitiveType->IsDouble() && !UE::MaterialTranslatorUtils::IsLWCEnabled())
 	{
 		Result = Em.CastToFloatKind(Result);
 	}
@@ -1610,7 +1610,7 @@ void UMaterialExpressionCameraPositionWS::Build(MIR::FEmitter& Em)
 
 void UMaterialExpressionPixelNormalWS::Build(MIR::FEmitter& Em)
 {
-	FValueRef Output = Em.InlineHLSL(MIR::FType::MakeFloat(3), TEXTVIEW("Parameters.WorldNormal"), {}, MIR::EValueFlags::None, MIR::EGraphProperties::ReadsPixelNormal);
+	FValueRef Output = Em.InlineHLSL(MIR::FType::MakeFloatVector(3), TEXTVIEW("Parameters.WorldNormal"), {}, MIR::EValueFlags::None, MIR::EGraphProperties::ReadsPixelNormal);
 	Em.Output(0, Output);
 }
 
@@ -2096,7 +2096,7 @@ void UMaterialExpressionBlend::Build(MIR::FEmitter& Em)
 
 static FValueRef EmitParameterCollectionVectorInlineHLSL(MIR::FEmitter& Em, MIR::FValueRef CollectionValue, int32 ParameterIndex)
 {
-	return Em.InlineHLSL(MIR::FType::MakeFloat(4), TEXTVIEW("MaterialCollection$0.Vectors[$1]"), { CollectionValue, Em.ConstantInt(ParameterIndex) });
+	return Em.InlineHLSL(MIR::FType::MakeFloatVector(4), TEXTVIEW("MaterialCollection$0.Vectors[$1]"), { CollectionValue, Em.ConstantInt(ParameterIndex) });
 }
 
 static bool GetExpressionCollectionParameter(MIR::FEmitter& Em, UMaterialParameterCollection* Collection, FName ParameterName, const FGuid& ParameterId, int32& OutParamIndex, int32& OutComponentIndex)
@@ -2149,16 +2149,13 @@ void UMaterialExpressionCollectionTransform::Build(MIR::FEmitter& Em)
 		return;
 	}
 
-	FValueRef Value = Em.TryInput(&Input);
-	if (!Value)
-	{
-		Em.Error(TEXT("CollectionTransform missing input"));
-		return;
-	}
+	FValueRef Value = Em.CheckIsPrimitive(Em.Input(&Input));
+	
+	UE_MIR_CHECKPOINT(Em);
 
-	if (!Value->Type.IsAnyFloat() || Value->Type.GetPrimitive().NumRows < 3 || Value->Type.GetPrimitive().NumColumns != 1)
+	if (!Value->Type.IsAnyFloat() || Value->Type.GetPrimitive().NumRows != 1 || Value->Type.GetPrimitive().NumColumns < 3)
 	{
-		Em.Error(TEXT("CollectionTransform requires float vector input"));
+		Em.Error(TEXT("CollectionTransform requires float3 vector input"));
 		return;
 	}
 
@@ -2184,7 +2181,6 @@ void UMaterialExpressionCollectionTransform::Build(MIR::FEmitter& Em)
 	else
 	{
 		check(TransformType == EParameterCollectionTransformType::LocalToWorld || TransformType == EParameterCollectionTransformType::WorldToLocal);
-
 		if (ParameterIndex + 5 > Collection->GetTotalVectorStorage())
 		{
 			Em.Errorf(TEXT("CollectionTransform parameter %s requires 5 vectors for LWC Matrix"), *ParameterName.ToString());
@@ -2207,7 +2203,7 @@ void UMaterialExpressionCollectionTransform::Build(MIR::FEmitter& Em)
 	if (TransformType == EParameterCollectionTransformType::Vector)
 	{
 		// Treat input as a direction vector (w = 0)
-		Value = Em.Cast(Value, MIR::FType::MakeFloat(3));
+		Value = Em.Cast(Value, MIR::FType::MakeFloatVector(3));
 
 		Result =        Em.Multiply(Em.Subscript(Value, 0), Em.Swizzle(CollectionParameters[0], MIR::FSwizzleMask::XYZ()));
 		Result = Em.Add(Em.Multiply(Em.Subscript(Value, 1), Em.Swizzle(CollectionParameters[1], MIR::FSwizzleMask::XYZ())), Result);
@@ -2235,7 +2231,7 @@ void UMaterialExpressionCollectionTransform::Build(MIR::FEmitter& Em)
 		Value = Em.CastToFloatKind(Value);
 
 		// If a 3-element vector is provided as input, we want to generate a 3-element vector as output.  Swizzle the collection parameters to achieve this.
-		if (Value->Type.GetPrimitive().NumRows == 3)
+		if (Value->Type.GetPrimitive().NumColumns == 3)
 		{
 			for (int32 i = 0; i < 4; i++)
 			{
@@ -2247,7 +2243,7 @@ void UMaterialExpressionCollectionTransform::Build(MIR::FEmitter& Em)
 		Result = Em.Add(Em.Multiply(Em.Subscript(Value, 1), CollectionParameters[1]), Result);
 		Result = Em.Add(Em.Multiply(Em.Subscript(Value, 2), CollectionParameters[2]), Result);
 
-		if (Value->Type.GetPrimitive().NumRows == 3)
+		if (Value->Type.GetPrimitive().NumColumns == 3)
 		{
 			// Treat input as a translation vector (w = 1)
 			Result = Em.Add(CollectionParameters[3], Result);
diff --git a/Engine/Source/Runtime/Engine/Private/Materials/MaterialIR.cpp b/Engine/Source/Runtime/Engine/Private/Materials/MaterialIR.cpp
--- a/Engine/Source/Runtime/Engine/Private/Materials/MaterialIR.cpp
+++ b/Engine/Source/Runtime/Engine/Private/Materials/MaterialIR.cpp
@@ -448,14 +448,14 @@ FType GetExternalInputType(EExternalInput Id)
 {
 	if (IsExternalInputTexCoordOrPartialDerivative(Id))
 	{
-		return FType::MakeFloat(2);
+		return FType::MakeFloatVector(2);
 	}
 	
 	switch (Id)
 	{
 		case EExternalInput::ViewMaterialTextureMipBias:
 		case EExternalInput::ViewMaterialTextureDerivativeMultiply:
-			 return FType::MakeFloat();
+			 return FType::MakeFloatScalar();
 
 		case EExternalInput::ActorPosition_Absolute:
 		case EExternalInput::ObjectPosition_Absolute:
@@ -465,7 +465,7 @@ FType GetExternalInputType(EExternalInput Id)
 		case EExternalInput::WorldPosition_AbsoluteNoOffsets_Ddx:		// analytic derivative expressions.  So we cast them to LWC on load.  The shader compiler should
 		case EExternalInput::WorldPosition_Absolute_Ddy:				// be able to optimize the math (the tile offsets will all be zero, cancelling things out).
 		case EExternalInput::WorldPosition_AbsoluteNoOffsets_Ddy:
-			return FType::MakeDouble(3);
+			return FType::MakeDoubleVector(3);
 
 		case EExternalInput::ActorPosition_CameraRelative:
 		case EExternalInput::ObjectPosition_CameraRelative:
@@ -475,7 +475,7 @@ FType GetExternalInputType(EExternalInput Id)
 		case EExternalInput::WorldPosition_CameraRelativeNoOffsets_Ddx:
 		case EExternalInput::WorldPosition_CameraRelative_Ddy:
 		case EExternalInput::WorldPosition_CameraRelativeNoOffsets_Ddy:
-			return FType::MakeFloat(3);
+			return FType::MakeFloatVector(3);
 
 		default: UE_MIR_UNREACHABLE();
 	}
diff --git a/Engine/Source/Runtime/Engine/Private/Materials/MaterialIRDebug.cpp b/Engine/Source/Runtime/Engine/Private/Materials/MaterialIRDebug.cpp
--- a/Engine/Source/Runtime/Engine/Private/Materials/MaterialIRDebug.cpp
+++ b/Engine/Source/Runtime/Engine/Private/Materials/MaterialIRDebug.cpp
@@ -368,7 +368,7 @@ struct FDebugDumpIRState
 			{
 				Out.Appendf(TEXT("%.5ff"), Constant->Float);
 			}
-			else if (Constant->Type.IsLWC())
+			else if (Constant->Type.IsDouble())
 			{
 				Out.Appendf(TEXT("%.8f"), Constant->Double);
 			}
diff --git a/Engine/Source/Runtime/Engine/Private/Materials/MaterialIREmitter.cpp b/Engine/Source/Runtime/Engine/Private/Materials/MaterialIREmitter.cpp
--- a/Engine/Source/Runtime/Engine/Private/Materials/MaterialIREmitter.cpp
+++ b/Engine/Source/Runtime/Engine/Private/Materials/MaterialIREmitter.cpp
@@ -482,7 +482,7 @@ bool FEmitter::ToConstantBool(FValueRef Value)
         Errorf(Value, TEXT("Expected a constant bool value, got a non-constant value instead."));
         return false;
     }
-	if (Constant->Type != FType::MakeBool())
+	if (Constant->Type != FType::MakeBoolScalar())
     {
         Errorf(Value, TEXT("Expected a constant bool value, got a '%s' instead."), *Constant->Type.GetSpelling());
         return false;
@@ -703,7 +703,7 @@ FValueRef FEmitter::ConstantFloat4(UE::Math::TVector4<TFloat> InValue)
 
 FValueRef FEmitter::ConstantDouble(TDouble InX)
 {
-	FConstant Scalar = MakePrototype<FConstant>(FType::MakeDouble());
+	FConstant Scalar = MakePrototype<FConstant>(FType::MakeDoubleScalar());
 	Scalar.Double = InX;
 	return EmitPrototype(*this, Scalar);
 }
@@ -711,8 +711,7 @@ FValueRef FEmitter::ConstantDouble(TDouble InX)
 /*--------------------- Other non-instruction values emission ---------------------*/
 
 FValueRef FEmitter::Poison()
-{
- 	if (CVarMaterialIRDebugBreakOnPoison.GetValueOnGameThread())
+{ 	if (CVarMaterialIRDebugBreakOnPoison.GetValueOnGameThread())
  	{
  		UE_DEBUG_BREAK();
  	}
@@ -787,8 +786,8 @@ FValueRef FEmitter::Parameter(FName Name, FMaterialParameterMetadata& Metadata,
 			}
 
 			FType Type = Metadata.Value.Type == EMaterialParameterType::Enumeration
-				? FType::MakeInt()
-				: FType::MakeFloat();
+				? FType::MakeIntScalar()
+				: FType::MakeFloatScalar();
 
 			FUniformParameter Prototype = MakePrototype<FUniformParameter>(Type);
 			Prototype.ParameterIdInModule = RegisterParameter(Module, Info, Metadata);
@@ -808,7 +807,7 @@ FValueRef FEmitter::Parameter(FName Name, FMaterialParameterMetadata& Metadata,
 			}
 
 			EScalarKind ScalarKind = Metadata.Value.Type == EMaterialParameterType::Vector ? EScalarKind::Float : EScalarKind::Double;
-			FUniformParameter Prototype = MakePrototype<FUniformParameter>(FType::MakePrimitive(ScalarKind, 4));
+			FUniformParameter Prototype = MakePrototype<FUniformParameter>(FType::MakeVector(ScalarKind, 4));
 			Prototype.ParameterIdInModule = RegisterParameter(Module, Info, Metadata);
 			return EmitPrototype(*this, Prototype);
 		}
@@ -881,7 +880,7 @@ FValueRef FEmitter::CustomPrimitiveData(uint32 PrimitiveDataIndex)
 
 FValueRef FEmitter::SceneTexture(ESceneTextureId SceneTextureId)
 {
-	FScreenTexture Proto = MakePrototype<FScreenTexture>(MIR::FType::MakeInt(1));
+	FScreenTexture Proto = MakePrototype<FScreenTexture>(MIR::FType::MakeIntScalar());
 	Proto.TextureKind = EScreenTexture::SceneTexture;
 	Proto.Id = SceneTextureId;
 	return EmitPrototype(*this, Proto);
@@ -889,7 +888,7 @@ FValueRef FEmitter::SceneTexture(ESceneTextureId SceneTextureId)
 
 FValueRef FEmitter::UserSceneTexture(FName UserSceneTexture)
 {
-	FScreenTexture Proto = MakePrototype<FScreenTexture>(MIR::FType::MakeInt(1));
+	FScreenTexture Proto = MakePrototype<FScreenTexture>(MIR::FType::MakeIntScalar());
 	Proto.TextureKind = EScreenTexture::UserSceneTexture;
 	Proto.Id = PPI_UserSceneTexture0;				// Actual ID determined later during MIR lowering
 	Proto.UserSceneTexture = UserSceneTexture;
@@ -898,7 +897,7 @@ FValueRef FEmitter::UserSceneTexture(FName UserSceneTexture)
 
 FValueRef FEmitter::DBufferTexture(EDBufferTextureId DBufferId)
 {
-	FScreenTexture Proto = MakePrototype<FScreenTexture>(MIR::FType::MakeInt(1));
+	FScreenTexture Proto = MakePrototype<FScreenTexture>(MIR::FType::MakeIntScalar());
 	Proto.TextureKind = EScreenTexture::DBufferTexture;
 	Proto.DBufferId = DBufferId;
 	return EmitPrototype(*this, Proto);
@@ -907,7 +906,7 @@ FValueRef FEmitter::DBufferTexture(EDBufferTextureId DBufferId)
 FValueRef FEmitter::ScreenTexture(EScreenTexture TextureKind)
 {
 	check(TextureKind != EScreenTexture::SceneTexture && TextureKind != EScreenTexture::UserSceneTexture && TextureKind != EScreenTexture::DBufferTexture);
-	FScreenTexture Proto = MakePrototype<FScreenTexture>(MIR::FType::MakeInt(1));
+	FScreenTexture Proto = MakePrototype<FScreenTexture>(MIR::FType::MakeIntScalar());
 	Proto.TextureKind = TextureKind;
 	Proto.Id = PPI_SceneColor;					// Ignored
 	return EmitPrototype(*this, Proto);
@@ -1748,17 +1747,17 @@ static FType ValidateOperatorAndGetResultType(FEmitter& Emitter, EOperator Op, F
 	// Given input float3, generates LWC3 type with the given tile value and zero offset.  Caller must pass in float3.
 	if (Op == UO_LWCTile)
 	{
-		if (!A->Type.GetPrimitive().IsFloat() || !A->Type.GetPrimitive().IsVector() || A->Type.GetPrimitive().NumComponents() != 3)
+		if (!A->Type.GetPrimitive().IsFloat() || !A->Type.GetPrimitive().IsRowVector() || A->Type.GetPrimitive().NumComponents() != 3)
 		{
 			Emitter.Errorf(A, TEXT("Argument of LWCTile operator expected to be a 3D float vector."));
 			return FType::MakePoison();
 		}
-		return FType::MakeDouble(3);
+		return FType::MakeDoubleVector(3);
 	}
 
 	// For Clamp, we do a special case and demote the first LWC argument if the second and third arguments (min / max) are non-LWC.
 	// We want to do this before fetching FirstArgumentPrimitiveType.
-	if (Op == TO_Clamp && A->Type.GetPrimitive().IsLWC() && !B->Type.IsLWC() && !C->Type.IsLWC())
+	if (Op == TO_Clamp && A->Type.GetPrimitive().IsDouble() && !B->Type.IsDouble() && !C->Type.IsDouble())
 	{
 		A = Emitter.CastToFloatKind(A);
 	}
@@ -1768,7 +1767,7 @@ static FType ValidateOperatorAndGetResultType(FEmitter& Emitter, EOperator Op, F
 
 	// Handle automatic cast to float for operators that don't support LWC inputs
 	FType FirstArgumentType = A->Type;
-	if (FirstArgumentType.GetPrimitive().IsLWC() && !(Signature->ParameterTypes[0] & OPTF_LWCInput))
+	if (FirstArgumentType.GetPrimitive().IsDouble() && !(Signature->ParameterTypes[0] & OPTF_LWCInput))
 	{
 		FirstArgumentType = A->Type.GetPrimitive().ToScalarKind(EScalarKind::Float);
 	}
@@ -1803,7 +1802,7 @@ static FType ValidateOperatorAndGetResultType(FEmitter& Emitter, EOperator Op, F
 			Arguments[i] = Emitter.Cast(Arguments[i], FirstArgumentType);
 			bValid &= Arguments[i].IsValid();
 		}
-		else if (ArgType.GetPrimitive().IsLWC() && !(Filter & OPTF_LWCInput))
+		else if (ArgType.GetPrimitive().IsDouble() && !(Filter & OPTF_LWCInput))
 		{
 			// Handle automatic cast to float for operators that don't support LWC inputs
 			ArgType = ArgType.GetPrimitive().ToScalarKind(EScalarKind::Float);
@@ -1841,7 +1840,7 @@ static FType ValidateOperatorAndGetResultType(FEmitter& Emitter, EOperator Op, F
 
 		if (Filter & OPTF_CheckIsVector3)
 		{
-			if (!ArgType.GetPrimitive().IsVector() || ArgType.GetPrimitive().NumComponents() != 3)
+			if (!ArgType.GetPrimitive().IsRowVector() || ArgType.GetPrimitive().NumComponents() != 3)
 			{
 				Emitter.Errorf(Arguments[i], TEXT("Expected a 3D vector."));
 				bValid = false;
@@ -1938,7 +1937,7 @@ static FType ValidateOperatorAndGetResultType(FEmitter& Emitter, EOperator Op, F
 				check(ArgumentsCommonType && ArgumentsCommonType.AsPrimitive());
 
 				// Lerp doesn't accept LWC for its third input, so we need to check if LWC is allowed per input when casting to the common type.
-				if (ArgumentsCommonType.GetPrimitive().IsLWC() && !(Filter & OPTF_LWCInput))
+				if (ArgumentsCommonType.GetPrimitive().IsDouble() && !(Filter & OPTF_LWCInput))
 				{
 					Arguments[i] = Emitter.Cast(Arguments[i], ArgumentsCommonType.GetPrimitive().ToScalarKind(EScalarKind::Float));
 				}
@@ -1950,7 +1949,7 @@ static FType ValidateOperatorAndGetResultType(FEmitter& Emitter, EOperator Op, F
 				bValid &= !Arguments[i]->IsPoison();
 			}
 			// Fmod doesn't accept LWC for its second input, so we need to potentially cast that to non-LWC.
-			else if (Arguments[i]->Type.GetPrimitive().IsLWC() && !(Filter & OPTF_LWCInput))
+			else if (Arguments[i]->Type.GetPrimitive().IsDouble() && !(Filter & OPTF_LWCInput))
 			{
 				Arguments[i] = Emitter.CastToFloatKind(Arguments[i]);
 			}
@@ -1980,7 +1979,7 @@ static FType ValidateOperatorAndGetResultType(FEmitter& Emitter, EOperator Op, F
 
 		case ORT_FirstArgumentType:
 			// Some operators accept LWC as input, but always output non-LWC (examples: modulo, sign, sine, cosine, saturate, frac). Anything that may output LWC should use ORT_FirstArgumentTypeLWC.
-			return FirstArgumentPrimitiveType.IsLWC() ? FirstArgumentPrimitiveType.ToScalarKind(EScalarKind::Float) : FirstArgumentType;
+			return FirstArgumentPrimitiveType.IsDouble() ? FirstArgumentPrimitiveType.ToScalarKind(EScalarKind::Float) : FirstArgumentType;
 
 		case ORT_FirstArgumentTypeLWC:
 			return FirstArgumentType;
@@ -2174,7 +2173,7 @@ static FValue* TryFoldOperator(FEmitter& Emitter, EOperator Op, FValue* A, FValu
 				float Result = FMath::Sqrt(ConstantDotFloat(AComponents, AComponents, ANumComponents));
 				return Emitter.ConstantFloat(Result);
 			}
-			else if (ResultType.GetPrimitive().IsLWC())
+			else if (ResultType.GetPrimitive().IsDouble())
 			{
 				double Result = FMath::Sqrt(ConstantDotDouble(AComponents, AComponents, ANumComponents));
 				return Emitter.ConstantDouble(Result);
@@ -2196,7 +2195,7 @@ static FValue* TryFoldOperator(FEmitter& Emitter, EOperator Op, FValue* A, FValu
 					float Result = ConstantDotFloat(AComponents, BComponents, ANumComponents);
 					return Emitter.ConstantFloat(Result);
 				}
-				else if (ResultType.GetPrimitive().IsLWC())
+				else if (ResultType.GetPrimitive().IsDouble())
 				{
 					double Result = ConstantDotDouble(AComponents, BComponents, ANumComponents);
 					return Emitter.ConstantDouble(Result);
@@ -2216,7 +2215,7 @@ static FValue* TryFoldOperator(FEmitter& Emitter, EOperator Op, FValue* A, FValu
 					FVector3f Result = AVector.Cross(BVector);
 					return Emitter.ConstantFloat3(Result);
 				}
-				else if (ResultType.GetPrimitive().IsLWC())
+				else if (ResultType.GetPrimitive().IsDouble())
 				{
 					UE_MIR_TODO();
 				}
@@ -2308,7 +2307,7 @@ FValueRef FEmitter::Branch(FValueRef Condition, FValueRef True, FValueRef False)
 	}
 
 	// Condition must be of type bool
-	Condition = Cast(Condition, FType::MakeBool());
+	Condition = Cast(Condition, FType::MakeBoolScalar());
 	if (!Condition)
 	{
 		return Poison();
@@ -2415,11 +2414,12 @@ FValueRef FEmitter::Swizzle(FValueRef Value, FSwizzleMask Mask)
 	
 	// For brevity
 	FPrimitive PrimitiveType = Value->Type.GetPrimitive();
+	int NumComponents = PrimitiveType.NumComponents();
 
 	// Make sure each component in the mask fits the number of components in Value.
 	for (EVectorComponent Component : Mask)
 	{
-		if ((int32)Component >= PrimitiveType.NumRows)
+		if ((int32)Component >= NumComponents)
 		{
 			Errorf(Value, TEXT("Value of type '%s' has no component '%s'."), *Value->Type.GetSpelling(), VectorComponentToString(Component));
 			return Value.ToPoison();
@@ -2428,7 +2428,7 @@ FValueRef FEmitter::Swizzle(FValueRef Value, FSwizzleMask Mask)
 
 	// If the requested number of components is the same as Value and the order in which the components
 	// are specified in the mask is sequential (e.g. x, y, z) then this is a no op, simply return Value as is.
-	if (Mask.NumComponents == PrimitiveType.NumComponents())
+	if (Mask.NumComponents == NumComponents)
 	{
 		bool InOrder = true;
 		for (int32 i = 0; i < Mask.NumComponents; ++i)
@@ -2439,7 +2439,6 @@ FValueRef FEmitter::Swizzle(FValueRef Value, FSwizzleMask Mask)
 				break;
 			}
 		}
-
 		if (InOrder)
 		{
 			return Value;
@@ -2454,12 +2453,10 @@ FValueRef FEmitter::Swizzle(FValueRef Value, FSwizzleMask Mask)
 
 	// Make the result vector type.
 	FType ResultType = FType::MakeVector(PrimitiveType.ScalarKind, Mask.NumComponents);
-	FComposite* Result = MakeCompositePrototype(*this, ResultType, ResultType.GetPrimitive().NumComponents());
-	TArrayView<FValue*> ResultComponents = Result->GetMutableComponents();
-
+	FComposite* Result = MakeCompositePrototype(*this, ResultType, Mask.NumComponents);
 	for (int32 i = 0; i < Mask.NumComponents; ++i)
 	{
-		ResultComponents[i] = Subscript(Value, (int32)Mask.Components[i]);
+		Result->GetMutableComponents()[i] = Subscript(Value, (int32)Mask.Components[i]);
 	}
 
 	return Value.To(EmitPrototype(*this, *Result));
@@ -2524,12 +2521,12 @@ static FValue* CastConstant(FEmitter& Emitter, FConstant* Constant, EScalarKind
 	UE_MIR_UNREACHABLE();
 }
 
-static FValue* CastToPrimitive(FEmitter& Emitter, FValueRef Value, FType TargetType)
+static FValueRef CastToPrimitive(FEmitter& Emitter, FValueRef Value, FType TargetType)
 {
 	if (!Value->Type.AsPrimitive())
 	{
 		Emitter.Errorf(Value, TEXT("Cannot construct a '%s' from non primitive type '%s'."), *Value->Type.GetSpelling(), *TargetType.GetSpelling());
-		return FPoison::Get();
+		return Value.ToPoison();
 	}
 
 	FPrimitive ValuePrimitiveType = Value->Type.GetPrimitive();
@@ -2588,7 +2585,7 @@ static FValue* CastToPrimitive(FEmitter& Emitter, FValueRef Value, FType TargetT
 	// Construct a vector from another vector. If constructed vector is larger, initialize
 	// remaining components to zero. If it's smaller, truncate initializer vector and only use
 	// the necessary components.
-	if (TargetPrimitiveType.IsVector() && ValuePrimitiveType.IsVector())
+	if (TargetPrimitiveType.IsRowVector() && ValuePrimitiveType.IsRowVector())
 	{
 		// #todo-massimo.tristano Use swizzle when scalartypes are the same, and target num components is less than initializer's.
 
@@ -2622,7 +2619,7 @@ static FValue* CastToPrimitive(FEmitter& Emitter, FValueRef Value, FType TargetT
 	}
 	
 	// The two primitive types are identical matrices that differ only by their scalar type.
-	if (TargetPrimitiveType.NumRows    == ValuePrimitiveType.NumRows &&
+	if (TargetPrimitiveType.NumRows == ValuePrimitiveType.NumRows &&
 		TargetPrimitiveType.NumColumns == ValuePrimitiveType.NumColumns)
 	{
 		check(TargetPrimitiveType.IsMatrix());
@@ -2644,7 +2641,7 @@ static FValue* CastToPrimitive(FEmitter& Emitter, FValueRef Value, FType TargetT
 	}
 
 	// Initializer value cannot be used to construct this primitive type.
-	return Emitter.Poison();
+	return Value.ToPoison();
 }
 
 FValueRef FEmitter::Cast(FValueRef Value, FType TargetType)
@@ -2661,12 +2658,13 @@ FValueRef FEmitter::Cast(FValueRef Value, FType TargetType)
 		return Value;
 	}
 	
-	FValue* Result;
+	FValueRef Result = FPoison::Get();
 	if (TargetType.AsPrimitive())
 	{
 		Result = CastToPrimitive(*this, Value, TargetType);
 	}
-	else
+
+	if (Result->IsPoison())
 	{
 		// No other legal conversions applicable. Report error if we haven't converted the value.
 		Errorf(Value, TEXT("Cannot construct a '%s' from a '%s'."), *TargetType.GetSpelling(), *Value->Type.GetSpelling());
@@ -2684,6 +2682,14 @@ FValueRef FEmitter::CastToScalar(FValueRef Value)
 		: Value;
 }
 
+FValueRef FEmitter::CastToVector(FValueRef Value, int32 NumColumns)
+{
+	Value = CheckIsPrimitive(Value);
+	return Value.IsValid()
+		? Cast(Value, FType::MakeVector(Value->Type.GetPrimitive().ScalarKind, NumColumns))
+		: Value;
+}
+
 FValueRef FEmitter::CastToScalarKind(FValueRef Value, EScalarKind ToScalarKind)
 {
 	Value = CheckIsPrimitive(Value);
@@ -2707,19 +2713,19 @@ FValueRef FEmitter::CastToFloatKind(FValueRef Value)
 	return CastToScalarKind(Value, EScalarKind::Float);
 }
 
-FValueRef FEmitter::CastToBool(FValueRef Value, int NumRows)
+FValueRef FEmitter::CastToBool(FValueRef Value, int NumColumns)
 {
-	return Cast(Value, FType::MakeVector(EScalarKind::Bool, NumRows));
+	return Cast(Value, FType::MakeVector(EScalarKind::Bool, NumColumns));
 }
 
-FValueRef FEmitter::CastToInt(FValueRef Value, int NumRows)
+FValueRef FEmitter::CastToInt(FValueRef Value, int NumColumns)
 {
-	return Cast(Value, FType::MakeVector(EScalarKind::Int, NumRows));
+	return Cast(Value, FType::MakeVector(EScalarKind::Int, NumColumns));
 }
 
-FValueRef FEmitter::CastToFloat(FValueRef Value, int NumRows)
+FValueRef FEmitter::CastToFloat(FValueRef Value, int NumColumns)
 {
-	return Cast(Value, FType::MakeVector(EScalarKind::Float, NumRows));
+	return Cast(Value, FType::MakeVector(EScalarKind::Float, NumColumns));
 }
 
 FValueRef FEmitter::StageSwitch(FType Type, TConstArrayView<FValueRef> ValuePerStage)
@@ -2755,7 +2761,7 @@ FValueRef FEmitter::TextureGather(FValueRef Texture, FValueRef TexCoord, ETextur
 		return Poison();
 	}
 
-	FTextureRead Prototype = MakePrototype<FTextureRead>(FType::MakeFloat(4));
+	FTextureRead Prototype = MakePrototype<FTextureRead>(FType::MakeFloatVector(4));
 	Prototype.TextureObject = Texture;
 	Prototype.TexCoord = TexCoord;
 	Prototype.Mode = GatherMode;
@@ -2777,7 +2783,7 @@ FValueRef FEmitter::TextureSample(FValueRef Texture, FValueRef TexCoord, bool bA
 		return Poison();
 	}
 
-	FTextureRead PrototypeHw = MakePrototype<FTextureRead>(FType::MakeFloat(4));
+	FTextureRead PrototypeHw = MakePrototype<FTextureRead>(FType::MakeFloatVector(4));
 	PrototypeHw.TextureObject = Texture;
 	PrototypeHw.TexCoord = TexCoord;
 	PrototypeHw.Mode = ETextureReadMode::MipAuto;
@@ -2819,7 +2825,7 @@ FValueRef FEmitter::TextureSampleLevel(FValueRef Texture, FValueRef TexCoord, FV
 		return Poison();
 	}
 
-	FTextureRead Prototype = MakePrototype<FTextureRead>(FType::MakeFloat(4));
+	FTextureRead Prototype = MakePrototype<FTextureRead>(FType::MakeFloatVector(4));
 	Prototype.TextureObject = Texture;
 	Prototype.TexCoord = TexCoord;
 	Prototype.MipValue = MipLevel;
@@ -2852,7 +2858,7 @@ FValueRef FEmitter::TextureSampleBias(FValueRef Texture, FValueRef TexCoord, FVa
 		MipBias = Operator(BO_Add, MipBias, ExternalInput(EExternalInput::ViewMaterialTextureMipBias));
 	}
 
-	FTextureRead PrototypeHw = MakePrototype<FTextureRead>(FType::MakeFloat(4));
+	FTextureRead PrototypeHw = MakePrototype<FTextureRead>(FType::MakeFloatVector(4));
 	PrototypeHw.TextureObject = Texture;
 	PrototypeHw.TexCoord = TexCoord;
 	PrototypeHw.MipValue = MipBias;
@@ -2886,7 +2892,7 @@ FValueRef FEmitter::TextureSampleGrad(FValueRef Texture, FValueRef TexCoord, FVa
 		return Poison();
 	}
 
-	FTextureRead Prototype = MakePrototype<FTextureRead>(FType::MakeFloat(4));
+	FTextureRead Prototype = MakePrototype<FTextureRead>(FType::MakeFloatVector(4));
 	Prototype.TextureObject = Texture;
 	Prototype.TexCoord = TexCoord;
 	Prototype.TexCoordDdx = TexCoordDdx;
@@ -3377,23 +3383,23 @@ FType FEmitter::TryGetCommonType(FType A, FType B)
 		return FType::MakePoison();
 	}
 	
-	// If both A and B are matrices, their dimensions must match (equality check above didn't trigger).
-	if (PrimitiveA->IsMatrix() || PrimitiveB->IsMatrix())
+	// Can't cast a vector to a matrix.
+	if ((PrimitiveA->IsRowVector() && PrimitiveB->IsMatrix()) || (PrimitiveA->IsMatrix() && PrimitiveB->IsRowVector()))
 	{
 		return FType::MakePoison();
 	}
-
-	// Neither A nor B are matrices, but single scalar or vector. Return the largest.
-	check(PrimitiveA->NumColumns == 1 && PrimitiveB->NumColumns == 1);
+	
+	// Return the primitive type with the maximum number of rows and columns between the two types.
 	EScalarKind ScalarKind = FMath::Max(PrimitiveA->ScalarKind, PrimitiveB->ScalarKind);
 	int32 NumRows = FMath::Max(PrimitiveA->NumRows, PrimitiveB->NumRows);
-	return FType::MakePrimitive(ScalarKind, NumRows, 1);
+	int32 NumColumns = FMath::Max(PrimitiveA->NumColumns, PrimitiveB->NumColumns);
+	return FType::MakePrimitive(ScalarKind, NumRows, NumColumns);
 }
 
 void FEmitter::Initialize()
 {
 	// Create and reference the true/false constants.
-	FConstant Temp = MakePrototype<FConstant>(FType::MakeBool());
+	FConstant Temp = MakePrototype<FConstant>(FType::MakeBoolScalar());
 
 	Temp.Boolean = true;
 	TrueConstant = EmitPrototype(*this, Temp);
diff --git a/Engine/Source/Runtime/Engine/Private/Materials/MaterialIRToHLSLTranslator.cpp b/Engine/Source/Runtime/Engine/Private/Materials/MaterialIRToHLSLTranslator.cpp
--- a/Engine/Source/Runtime/Engine/Private/Materials/MaterialIRToHLSLTranslator.cpp
+++ b/Engine/Source/Runtime/Engine/Private/Materials/MaterialIRToHLSLTranslator.cpp
@@ -754,31 +754,31 @@ struct FPrivate : FMaterialIRToHLSLTranslation
 	{
 		static const FSubstrateLegacyArgument LegacyArguments[] =
 		{
-			{ MIR::FType::MakeFloat(3), TEXTVIEW("BaseColor"), MP_BaseColor, {} },
-			{ MIR::FType::MakeFloat(),  TEXTVIEW("Specular"), MP_SpecularColor, TEXTVIEW("0.5f") },
-			{ MIR::FType::MakeFloat(),  TEXTVIEW("Metallic"), MP_Metallic, {} },
-			{ MIR::FType::MakeFloat(),  TEXTVIEW("Roughness"), MP_Roughness, {} },
-			{ MIR::FType::MakeFloat(),  TEXTVIEW("Anisotropy"), MP_Anisotropy, {} },
-			{ MIR::FType::MakeFloat(3), TEXTVIEW("SubSurfaceColor"), MP_SubsurfaceColor, TEXTVIEW("(float3)1") },
-			{ MIR::FType::MakeFloat(),  TEXTVIEW("SubSurfaceProfileId"), MP_MAX, TEXTVIEW("0") },
-			{ MIR::FType::MakeFloat(),  TEXTVIEW("ClearCoat"), MP_MAX, TEXTVIEW("1.0f") },
-			{ MIR::FType::MakeFloat(),  TEXTVIEW("ClearCoatRoughness"), MP_MAX, TEXTVIEW("0.1f") },
-			{ MIR::FType::MakeFloat(3), TEXTVIEW("Emissive"), MP_EmissiveColor, {} },
-			{ MIR::FType::MakeFloat(),  TEXTVIEW("Opacity"), MP_Opacity, {} },
-			{ MIR::FType::MakeFloat(3), TEXTVIEW("ThinTranslucentTransmittanceColor"), MP_MAX, TEXTVIEW("(float3)0.5f") },
-			{ MIR::FType::MakeFloat(),  TEXTVIEW("ThinTranslucentSurfaceCoverage"), MP_MAX, TEXTVIEW("1.0f") },
-			{ MIR::FType::MakeFloat(3), TEXTVIEW("WaterScatteringCoefficients"), MP_MAX, {} },
-			{ MIR::FType::MakeFloat(3), TEXTVIEW("WaterAbsorptionCoefficients"), MP_MAX, {} },
-			{ MIR::FType::MakeFloat(),  TEXTVIEW("WaterPhaseG"), MP_MAX, {} },
-			{ MIR::FType::MakeFloat(3), TEXTVIEW("ColorScaleBehindWater"), MP_MAX, TEXTVIEW("(float3)1.0f") },
-			{ MIR::FType::MakeInt(),    TEXTVIEW("ShadingModel"), MP_MAX, TEXTVIEW("1") },
-			{ MIR::FType::MakeFloat(3), TEXTVIEW("RawNormal"), MP_Normal, TEXTVIEW("Parameters.TangentToWorld[2]") },
-			{ MIR::FType::MakeFloat(3), TEXTVIEW("RawTangent"), MP_Tangent, TEXTVIEW("Parameters.TangentToWorld[0]") },
-			{ MIR::FType::MakeFloat(3), TEXTVIEW("RawClearCoatNormal"), MP_Normal, TEXTVIEW("Parameters.TangentToWorld[2]") },
-			{ MIR::FType::MakeFloat(3), TEXTVIEW("RawCustomTangent"), MP_Tangent, TEXTVIEW("Parameters.TangentToWorld[0]") },
-			{ MIR::FType::MakeInt(),    TEXTVIEW("SharedLocalBasisIndex"), MP_MAX, {} },
-			{ MIR::FType::MakeInt(),    TEXTVIEW("ClearCoatBottomNormal_SharedLocalBasisIndex"), MP_MAX, {} },
-			{ MIR::FType::MakeInt(),    TEXTVIEW("SharedLocalBasisTypes"), MP_MAX, TEXTVIEW("Parameters.SharedLocalBases.Types") },
+			{ MIR::FType::MakeFloatVector(3), TEXTVIEW("BaseColor"), MP_BaseColor, {} },
+			{ MIR::FType::MakeFloatScalar(),  TEXTVIEW("Specular"), MP_SpecularColor, TEXTVIEW("0.5f") },
+			{ MIR::FType::MakeFloatScalar(),  TEXTVIEW("Metallic"), MP_Metallic, {} },
+			{ MIR::FType::MakeFloatScalar(),  TEXTVIEW("Roughness"), MP_Roughness, {} },
+			{ MIR::FType::MakeFloatScalar(),  TEXTVIEW("Anisotropy"), MP_Anisotropy, {} },
+			{ MIR::FType::MakeFloatVector(3), TEXTVIEW("SubSurfaceColor"), MP_SubsurfaceColor, TEXTVIEW("(float3)1") },
+			{ MIR::FType::MakeFloatScalar(),  TEXTVIEW("SubSurfaceProfileId"), MP_MAX, TEXTVIEW("0") },
+			{ MIR::FType::MakeFloatScalar(),  TEXTVIEW("ClearCoat"), MP_MAX, TEXTVIEW("1.0f") },
+			{ MIR::FType::MakeFloatScalar(),  TEXTVIEW("ClearCoatRoughness"), MP_MAX, TEXTVIEW("0.1f") },
+			{ MIR::FType::MakeFloatVector(3), TEXTVIEW("Emissive"), MP_EmissiveColor, {} },
+			{ MIR::FType::MakeFloatScalar(),  TEXTVIEW("Opacity"), MP_Opacity, {} },
+			{ MIR::FType::MakeFloatVector(3), TEXTVIEW("ThinTranslucentTransmittanceColor"), MP_MAX, TEXTVIEW("(float3)0.5f") },
+			{ MIR::FType::MakeFloatScalar(),  TEXTVIEW("ThinTranslucentSurfaceCoverage"), MP_MAX, TEXTVIEW("1.0f") },
+			{ MIR::FType::MakeFloatVector(3), TEXTVIEW("WaterScatteringCoefficients"), MP_MAX, {} },
+			{ MIR::FType::MakeFloatVector(3), TEXTVIEW("WaterAbsorptionCoefficients"), MP_MAX, {} },
+			{ MIR::FType::MakeFloatScalar(),  TEXTVIEW("WaterPhaseG"), MP_MAX, {} },
+			{ MIR::FType::MakeFloatVector(3), TEXTVIEW("ColorScaleBehindWater"), MP_MAX, TEXTVIEW("(float3)1.0f") },
+			{ MIR::FType::MakeIntScalar(),    TEXTVIEW("ShadingModel"), MP_MAX, TEXTVIEW("1") },
+			{ MIR::FType::MakeFloatVector(3), TEXTVIEW("RawNormal"), MP_Normal, TEXTVIEW("Parameters.TangentToWorld[2]") },
+			{ MIR::FType::MakeFloatVector(3), TEXTVIEW("RawTangent"), MP_Tangent, TEXTVIEW("Parameters.TangentToWorld[0]") },
+			{ MIR::FType::MakeFloatVector(3), TEXTVIEW("RawClearCoatNormal"), MP_Normal, TEXTVIEW("Parameters.TangentToWorld[2]") },
+			{ MIR::FType::MakeFloatVector(3), TEXTVIEW("RawCustomTangent"), MP_Tangent, TEXTVIEW("Parameters.TangentToWorld[0]") },
+			{ MIR::FType::MakeIntScalar(),    TEXTVIEW("SharedLocalBasisIndex"), MP_MAX, {} },
+			{ MIR::FType::MakeIntScalar(),    TEXTVIEW("ClearCoatBottomNormal_SharedLocalBasisIndex"), MP_MAX, {} },
+			{ MIR::FType::MakeIntScalar(),    TEXTVIEW("SharedLocalBasisTypes"), MP_MAX, TEXTVIEW("Parameters.SharedLocalBases.Types") },
 		};
 		return LegacyArguments;
 	}
@@ -787,9 +787,9 @@ struct FPrivate : FMaterialIRToHLSLTranslation
 	{
 		static const FSubstrateLegacyArgument LegacyArguments[] =
 		{
-			{ MIR::FType::MakeFloat(3), TEXTVIEW("Emissive"), MP_EmissiveColor, {} },
-			{ MIR::FType::MakeFloat(), TEXTVIEW("TransmittanceColor"), MP_MAX, TEXTVIEW("(float3)0.0f") },
-			{ MIR::FType::MakeFloat(3), TEXTVIEW("RawNormal"), MP_Normal, TEXTVIEW("Parameters.TangentToWorld[2]") },
+			{ MIR::FType::MakeFloatVector(3), TEXTVIEW("Emissive"), MP_EmissiveColor, {} },
+			{ MIR::FType::MakeFloatScalar(), TEXTVIEW("TransmittanceColor"), MP_MAX, TEXTVIEW("(float3)0.0f") },
+			{ MIR::FType::MakeFloatVector(3), TEXTVIEW("RawNormal"), MP_Normal, TEXTVIEW("Parameters.TangentToWorld[2]") },
 		};
 		return LegacyArguments;
 	}
@@ -832,7 +832,7 @@ struct FPrivate : FMaterialIRToHLSLTranslation
 						{
 							LowerType(Argument.ParameterType);
 							Printer << BeginArgs;
-							for (int32 ComponentIndex = 0; ComponentIndex < PrimitiveType->NumRows; ++ComponentIndex)
+							for (int32 ComponentIndex = 0; ComponentIndex < PrimitiveType->NumComponents(); ++ComponentIndex)
 							{
 								Printer << ListSeparator << TEXTVIEW("0.0f");
 							}
@@ -857,10 +857,10 @@ struct FPrivate : FMaterialIRToHLSLTranslation
 	ENoOp LowerPrimitiveUniformParameter(const MIR::FUniformParameter* UniformParameter)
 	{
 		MIR::FPrimitive PrimitiveType = UniformParameter->Type.GetPrimitive();
-		check(PrimitiveType.IsScalar() || PrimitiveType.IsVector()); // no matrices yet
+		check(PrimitiveType.IsScalar() || PrimitiveType.IsRowVector()); // no matrices yet
 
 		// LWC parameters are handled differently, they have their own dedicated function.
-		if (PrimitiveType.IsLWC())
+		if (PrimitiveType.IsDouble())
 		{
 			return LowerPrimitiveUniformParameterLWC(UniformParameter, PrimitiveType.NumComponents());
 		}
@@ -907,7 +907,7 @@ struct FPrivate : FMaterialIRToHLSLTranslation
 	ENoOp LowerPrimitiveUniformParameterLWC(const MIR::FUniformParameter* UniformParameter, int32 NumComponents)
 	{
 		check(NumComponents >= 2 && NumComponents <= 4);
-		check(UniformParameter->Type.IsLWC());
+		check(UniformParameter->Type.IsDouble());
 
 		const FUniformExpressionSet& UniformExpressionSet = Module->GetCompilationOutput().UniformExpressionSet;
 
@@ -937,16 +937,16 @@ struct FPrivate : FMaterialIRToHLSLTranslation
 				MIR::FPrimitive ArgPrimitiveType = FirstComponentAsSubscript->Arg->Type.GetPrimitive();
 
 				// Exact match?
-				if (ArgPrimitiveType.NumRows == NumComponents)
+				if (ArgPrimitiveType.NumColumns == NumComponents)
 				{
 					return true;
 				}
 
 				// Subset of components? See if we can generate a swizzled vector.
-				if (ArgPrimitiveType.NumRows >= NumComponents)
+				if (ArgPrimitiveType.NumColumns >= NumComponents)
 				{
 					// Allow swizzle for non-LWC values or uniform parameters (the latter having a special case for LWC)
-					return !ArgPrimitiveType.IsLWC() || FirstComponentAsSubscript->Arg->As<MIR::FUniformParameter>();
+					return !ArgPrimitiveType.IsDouble() || FirstComponentAsSubscript->Arg->As<MIR::FUniformParameter>();
 				}
 			}
 		}
@@ -970,7 +970,7 @@ struct FPrivate : FMaterialIRToHLSLTranslation
 			MIR::FPrimitive PrimitiveType = FirstComponent->Type.GetPrimitive();
 			MIR::FPrimitive ArgPrimitiveType = Arg->Type.GetPrimitive();
 
-			if (ArgPrimitiveType.IsLWC())
+			if (ArgPrimitiveType.IsDouble())
 			{
 				// Cast from LWC
 				if (PrimitiveType.IsBoolean())
@@ -990,7 +990,7 @@ struct FPrivate : FMaterialIRToHLSLTranslation
 					Printer << TEXTVIEW("WSDemote(") << LowerVectorCastChain(Arg, NumComponents) << TEXTVIEW(")");
 				}
 			}
-			else if (PrimitiveType.IsLWC())
+			else if (PrimitiveType.IsDouble())
 			{
 				// Cast to LWC
 				if (ArgPrimitiveType.IsBoolean())
@@ -1035,7 +1035,7 @@ struct FPrivate : FMaterialIRToHLSLTranslation
 		{
 			// Finally reached the inner subscript, print its vector argument
 			const MIR::FValue* Arg = FirstComponentAsSubscript->Arg;
-			if (Arg->Type.IsVector() && Arg->Type.IsLWC() && Arg->As<MIR::FUniformParameter>())
+			if (Arg->Type.IsVector() && Arg->Type.IsDouble() && Arg->As<MIR::FUniformParameter>())
 			{
 				// LWC types support special case swizzling logic for uniform parameters on initial fetch. See HasMatchingScalarComponentCastChain above.
 				LowerPrimitiveUniformParameterLWC(static_cast<const MIR::FUniformParameter*>(Arg), NumComponents);
@@ -1043,7 +1043,7 @@ struct FPrivate : FMaterialIRToHLSLTranslation
 			else
 			{
 				Printer << LowerValue(Arg);
-				if (Arg->Type.GetPrimitive().NumRows > NumComponents)
+				if (Arg->Type.GetPrimitive().NumColumns > NumComponents)
 				{
 					Printer << GVector4SwizzleSubset[NumComponents - 1];
 				}
@@ -1062,12 +1062,12 @@ struct FPrivate : FMaterialIRToHLSLTranslation
 		if (const MIR::FScalar* AsScalar = FirstComponent->As<MIR::FScalar>())
 		{
 			// Check if this is a cast from an arithmetic non-LWC scalar type. If so, we can cast the scalar to a float vector, and then to LWC.
-			return !AsScalar->Arg->Type.IsBoolean() && !AsScalar->Arg->Type.IsLWC();
+			return !AsScalar->Arg->Type.IsBoolean() && !AsScalar->Arg->Type.IsDouble();
 		}
 		else if (const MIR::FConstant* AsConstant = FirstComponent->As<MIR::FConstant>())
 		{
 			// Check if this is a cast from a constant representable exactly as a non-LWC float
-			return AsConstant->Type.IsLWC() && FLargeWorldRenderScalar(AsConstant->Double).GetTile() == 0.0f;
+			return AsConstant->Type.IsDouble() && FLargeWorldRenderScalar(AsConstant->Double).GetTile() == 0.0f;
 		}
 		return false;
 	}
@@ -1085,14 +1085,14 @@ struct FPrivate : FMaterialIRToHLSLTranslation
 		bool bSameComponents = true;
 
 		// Track if all components are part of a constant vector
-		bool bConstantVector = Components[0]->As<MIR::FConstant>() && PrimitiveType.IsVector();
+		bool bConstantVector = Components[0]->As<MIR::FConstant>() && PrimitiveType.IsRowVector();
 		
 		// We can also generate tidier HLSL for cases where casts are done for whole vectors.  This is a frequent case for LWC, where casts happen
 		// in both directions (LWC to float and back), due to operations automatically casting their inputs or outputs.  For example, BO_Sub
 		// automatically downcasts LWC to float, but if you feed that into a BO_Add that has LWC for its other input, it will immediately cast it
 		// back to LWC -- two consecutive casts.  These then get expanded into individual scalar casts, which we would like to do as whole vector
 		// casts for readability.  We only need to consider this special case if the first component is a Scalar, and the type is a Vector.
-		bool bWholeVectorCast = PrimitiveType.IsVector() && (Components[0]->As<MIR::FScalar>() || Components[0]->As<MIR::FSubscript>());
+		bool bWholeVectorCast = PrimitiveType.IsRowVector() && (Components[0]->As<MIR::FScalar>() || Components[0]->As<MIR::FSubscript>());
 
 		for (int32 i = 1; i < Components.Num(); ++i)
 		{
@@ -1101,14 +1101,14 @@ struct FPrivate : FMaterialIRToHLSLTranslation
 			bWholeVectorCast = bWholeVectorCast && HasMatchingScalarComponentCastChain(Components[0], Components[i], i, Components.Num());
 		}
 		
-		if (bSameComponents && !PrimitiveType.IsLWC())
+		if (bSameComponents && !PrimitiveType.IsDouble())
 		{
 			Printer << TEXTVIEW("(") << LowerType(Composite->Type) << TEXTVIEW(")") << LowerValue(Components[0]);
 		}
-		else if (bSameComponents && PrimitiveType.IsVector() && IsScalarToLWCVectorCast(Components[0]))
+		else if (bSameComponents && PrimitiveType.IsRowVector() && IsScalarToLWCVectorCast(Components[0]))
 		{
 			// Cast scalar to float vector, then promote, for example:  "WSPromote((float3)1.0f)"
-			Printer << TEXTVIEW("WSPromote") << TEXTVIEW("((float") << PrimitiveType.NumRows << TEXTVIEW(")");
+			Printer << TEXTVIEW("WSPromote") << TEXTVIEW("((float") << PrimitiveType.NumColumns << TEXTVIEW(")");
 			if (const MIR::FScalar* Scalar = Components[0]->As<MIR::FScalar>())
 			{
 				// Print the inner non-LWC scalar value
@@ -1129,7 +1129,7 @@ struct FPrivate : FMaterialIRToHLSLTranslation
 		{
 			LowerVectorCastChain(Components[0], Components.Num());
 		}
-		else if (bConstantVector && PrimitiveType.IsLWC())
+		else if (bConstantVector && PrimitiveType.IsDouble())
 		{
 			// Special case for LWC constant vectors
 			bool bAllTileValuesZero = true;
@@ -1169,9 +1169,9 @@ struct FPrivate : FMaterialIRToHLSLTranslation
 		}
 		else
 		{
-			if (PrimitiveType.IsLWC())
+			if (PrimitiveType.IsDouble())
 			{
-				if (PrimitiveType.IsVector())
+				if (PrimitiveType.IsRowVector())
 				{
 					Printer << TEXTVIEW("MakeWSVector") << BeginArgs;
 				}
@@ -1251,7 +1251,7 @@ struct FPrivate : FMaterialIRToHLSLTranslation
 			return;
 		}
 
-		bool bIsLWC = Operator->AArg->Type.IsLWC() || (Operator->Op == MIR::TO_Select && Operator->BArg->Type.IsLWC());
+		bool bIsLWC = Operator->AArg->Type.IsDouble() || (Operator->Op == MIR::TO_Select && Operator->BArg->Type.IsDouble());
 
 		if (IsOperatorInfix(Operator->Op) && !bIsLWC)
 		{
@@ -1414,7 +1414,7 @@ struct FPrivate : FMaterialIRToHLSLTranslation
 	{
 		if (IsFoldable(Branch, CurrentStage))
 		{
-			if (Branch->TrueArg->Type.IsLWC())
+			if (Branch->TrueArg->Type.IsDouble())
 			{
 				Printer << TEXTVIEW("WSSelect(") << LowerValue(Branch->ConditionArg)
 					<< TEXTVIEW(", ") << LowerValue(Branch->TrueArg)
@@ -1445,7 +1445,7 @@ struct FPrivate : FMaterialIRToHLSLTranslation
 	void LowerSubscript(const MIR::FSubscript* Subscript)
 	{
 		TOptional<MIR::FPrimitive> ArgVectorType = Subscript->Arg->Type.AsVector();
-		if (ArgVectorType && ArgVectorType->IsLWC())
+		if (ArgVectorType && ArgVectorType->IsDouble())
 		{
 			FStringView LwcComponentsStr[] = { TEXTVIEW("WSGetX("), TEXTVIEW("WSGetY("), TEXTVIEW("WSGetZ("), TEXTVIEW("WSGetW(") };
 			check(Subscript->Index <= ArgVectorType->NumComponents());
@@ -1471,7 +1471,7 @@ struct FPrivate : FMaterialIRToHLSLTranslation
 		MIR::FPrimitive PrimitiveType = Scalar->Type.GetPrimitive();
 		MIR::FPrimitive ArgPrimitiveType = Scalar->Arg->Type.GetPrimitive();
 
-		if (ArgPrimitiveType.IsLWC())
+		if (ArgPrimitiveType.IsDouble())
 		{
 			// Cast from LWC
 			if (PrimitiveType.IsBoolean())
@@ -1490,7 +1490,7 @@ struct FPrivate : FMaterialIRToHLSLTranslation
 				Printer << TEXTVIEW("WSDemote(") << LowerValue(Scalar->Arg) << TEXTVIEW(")");
 			}
 		}
-		else if (PrimitiveType.IsLWC())
+		else if (PrimitiveType.IsDouble())
 		{
 			// Cast to LWC
 			if (ArgPrimitiveType.IsBoolean())
@@ -1782,7 +1782,7 @@ struct FPrivate : FMaterialIRToHLSLTranslation
 
 	void LowerHardwarePartialDerivative(const MIR::FHardwarePartialDerivative* StageSwitch)
 	{
-		if (StageSwitch->Arg->Type.IsLWC())
+		if (StageSwitch->Arg->Type.IsDouble())
 		{
 			// Expression emitter assumes these to be LWC for analytic derivative evaluation, so we promote them on load
 			Printer << (StageSwitch->Axis == MIR::EDerivativeAxis::X ? TEXTVIEW("WSPromote(WSDdxDemote(") : TEXTVIEW("WSPromote(WSDdyDemote(")) << LowerValue(StageSwitch->Arg) << TEXTVIEW("))");
@@ -1890,15 +1890,15 @@ struct FPrivate : FMaterialIRToHLSLTranslation
 	{
 		if (TOptional<MIR::FPrimitive> PrimitiveType = Type.AsPrimitive())
 		{
-			if (PrimitiveType->IsLWC())
+			if (PrimitiveType->IsDouble())
 			{
-				if (PrimitiveType->NumRows == 1)
+				if (PrimitiveType->IsScalar())
 				{
 					Printer << TEXTVIEW("FWSScalar");
 				}
-				else if (PrimitiveType->NumColumns == 1)
+				else if (PrimitiveType->IsRowVector())
 				{
-					Printer << TEXTVIEW("FWSVector") << PrimitiveType->NumRows;
+					Printer << TEXTVIEW("FWSVector") << PrimitiveType->NumComponents();
 				}
 				else if (PrimitiveType->bIsLWCInverseMatrix)
 				{
@@ -1918,14 +1918,13 @@ struct FPrivate : FMaterialIRToHLSLTranslation
 					case MIR::EScalarKind::Float: Printer << TEXTVIEW("MaterialFloat"); break;
 				}
 
-				if (PrimitiveType->NumRows > 1)
+				if (PrimitiveType->NumRows == 1 && PrimitiveType->NumColumns > 1)
 				{
-					Printer << PrimitiveType->NumRows;
+					Printer << PrimitiveType->NumColumns;
 				}
-			
-				if (PrimitiveType->NumColumns > 1)
+				else if (PrimitiveType->IsMatrix())
 				{
-					Printer << TEXTVIEW("x") << PrimitiveType->NumColumns;
+					Printer << PrimitiveType->NumRows << TEXTVIEW("x") << PrimitiveType->NumColumns;
 				}
 			}
 		}
diff --git a/Engine/Source/Runtime/Engine/Private/Materials/MaterialIRTypes.cpp b/Engine/Source/Runtime/Engine/Private/Materials/MaterialIRTypes.cpp
--- a/Engine/Source/Runtime/Engine/Private/Materials/MaterialIRTypes.cpp
+++ b/Engine/Source/Runtime/Engine/Private/Materials/MaterialIRTypes.cpp
@@ -60,25 +60,25 @@ FType FType::FromShaderType(const UE::Shader::FType& InShaderType)
 		case UE::Shader::EValueType::Float2:
 		case UE::Shader::EValueType::Float3:
 		case UE::Shader::EValueType::Float4:
-			return FType::MakeVector(EScalarKind::Float, (int)InShaderType.ValueType - (int)UE::Shader::EValueType::Float1 + 1);
+			return FType::MakeFloatVector((int)InShaderType.ValueType - (int)UE::Shader::EValueType::Float1 + 1);
 
 		case UE::Shader::EValueType::Int1:
 		case UE::Shader::EValueType::Int2:
 		case UE::Shader::EValueType::Int3:
 		case UE::Shader::EValueType::Int4:
-			return FType::MakeVector(EScalarKind::Int, (int)InShaderType.ValueType - (int)UE::Shader::EValueType::Int1 + 1);
+			return FType::MakeIntVector((int)InShaderType.ValueType - (int)UE::Shader::EValueType::Int1 + 1);
 
 		case UE::Shader::EValueType::Bool1:
 		case UE::Shader::EValueType::Bool2:
 		case UE::Shader::EValueType::Bool3:
 		case UE::Shader::EValueType::Bool4:
-			return FType::MakeVector(EScalarKind::Bool, (int)InShaderType.ValueType - (int)UE::Shader::EValueType::Bool1 + 1);
+			return FType::MakeBoolVector((int)InShaderType.ValueType - (int)UE::Shader::EValueType::Bool1 + 1);
 
 		case UE::Shader::EValueType::Double1:
 		case UE::Shader::EValueType::Double2:
 		case UE::Shader::EValueType::Double3:
 		case UE::Shader::EValueType::Double4:
-			return FType::MakeVector(EScalarKind::Double, (int)InShaderType.ValueType - (int)UE::Shader::EValueType::Double1 + 1);
+			return FType::MakeDoubleVector((int)InShaderType.ValueType - (int)UE::Shader::EValueType::Double1 + 1);
 
 		default:
 			UE_MIR_UNREACHABLE();
@@ -89,50 +89,24 @@ FType FType::FromMaterialValueType(EMaterialValueType Type)
 {
 	switch (Type)
 	{
-		case MCT_VoidStatement:
-			return FType::MakeVoid();
-
-		case MCT_Float1:
-			return FType::MakeVector(EScalarKind::Float, 1);
-		case MCT_Float2:
-			return FType::MakeVector(EScalarKind::Float, 2);
-		case MCT_Float3:
-			return FType::MakeVector(EScalarKind::Float, 3);
-		case MCT_Float4:
-			return FType::MakeVector(EScalarKind::Float, 4);
-
-		case MCT_Float:
-			return FType::MakeVector(EScalarKind::Float, 4);
-
-		case MCT_Texture2D:
-			return FType::MakeTexture2D();
-
-		case MCT_StaticBool:
-			return FType::MakeVector(EScalarKind::Bool, 1);
-
-		case MCT_UInt1:
-			return FType::MakeVector(EScalarKind::Int, 1);
-		case MCT_UInt2:
-			return FType::MakeVector(EScalarKind::Int, 2);
-		case MCT_UInt3:
-			return FType::MakeVector(EScalarKind::Int, 3);
-		case MCT_UInt4:
-			return FType::MakeVector(EScalarKind::Int, 4);
-
-		case MCT_Bool:
-			return FType::MakeVector(EScalarKind::Bool, 1);
-
-		case MCT_LWCScalar:
-			return FType::MakeVector(EScalarKind::Double, 1);
-		case MCT_LWCVector2:
-			return FType::MakeVector(EScalarKind::Double, 2);
-		case MCT_LWCVector3:
-			return FType::MakeVector(EScalarKind::Double, 3);
-		case MCT_LWCVector4:
-			return FType::MakeVector(EScalarKind::Double, 4);
-
-		default:
-			UE_MIR_UNREACHABLE();
+		case MCT_VoidStatement: return FType::MakeVoid();
+		case MCT_Float1: return FType::MakeFloatScalar();
+		case MCT_Float2: return FType::MakeFloatVector(2);
+		case MCT_Float3: return FType::MakeFloatVector(3);
+		case MCT_Float4: return FType::MakeFloatVector(4);
+		case MCT_Float: return FType::MakeFloatVector(4);
+		case MCT_Texture2D: return FType::MakeTexture2D();
+		case MCT_StaticBool: return FType::MakeBoolScalar();
+		case MCT_UInt1: return FType::MakeIntScalar();
+		case MCT_UInt2: return FType::MakeIntVector(2);
+		case MCT_UInt3: return FType::MakeIntVector(3);
+		case MCT_UInt4: return FType::MakeIntVector(4);
+		case MCT_Bool: return FType::MakeBoolScalar();
+		case MCT_LWCScalar: return FType::MakeDoubleScalar();
+		case MCT_LWCVector2: return FType::MakeDoubleVector(2);
+		case MCT_LWCVector3: return FType::MakeDoubleVector(3);
+		case MCT_LWCVector4: return FType::MakeDoubleVector(4);
+		default: UE_MIR_UNREACHABLE();
 	}
 }
 
@@ -140,16 +114,16 @@ FType FType::FromMaterialParameterType(EMaterialParameterType Type)
 {
 	switch (Type)
 	{
-		case EMaterialParameterType::Scalar: return FType::MakeFloat();
-		case EMaterialParameterType::Vector: return FType::MakeFloat(4);
-		case EMaterialParameterType::DoubleVector: UE_MIR_TODO();
+		case EMaterialParameterType::Scalar: return FType::MakeFloatScalar();
+		case EMaterialParameterType::Vector: return FType::MakeFloatVector(4);
+		case EMaterialParameterType::DoubleVector: return FType::MakeDoubleVector(4);
 		case EMaterialParameterType::Texture: return FType::MakeTexture2D();
 		case EMaterialParameterType::TextureCollection: UE_MIR_TODO();
 		case EMaterialParameterType::ParameterCollection: UE_MIR_TODO();
 		case EMaterialParameterType::Font: UE_MIR_TODO();
 		case EMaterialParameterType::RuntimeVirtualTexture: UE_MIR_TODO();
 		case EMaterialParameterType::SparseVolumeTexture: UE_MIR_TODO();
-		case EMaterialParameterType::StaticSwitch: return FType::MakeBool();
+		case EMaterialParameterType::StaticSwitch: return FType::MakeBoolScalar();
 		default: UE_MIR_UNREACHABLE();
 	}
 }
@@ -166,50 +140,19 @@ FType FType::MakeVoid()
 	return Type;
 }
 
-FType FType::MakeBool(int NumRows, int NumColumns)
-{
-	return MakePrimitive(EScalarKind::Bool, NumRows, NumColumns);
-}
-
-FType FType::MakeInt(int NumRows, int NumColumns)
-{
-	return MakePrimitive(EScalarKind::Int, NumRows, NumColumns);
-}
-
-FType FType::MakeFloat(int NumRows, int NumColumns)
-{
-	return MakePrimitive(EScalarKind::Float, NumRows, NumColumns);
-}
-
-FType FType::MakeDouble(int NumRows, int NumColumns)
-{
-	return MakePrimitive(EScalarKind::Double, NumRows, NumColumns);
-}
-
-FType FType::MakeScalar(EScalarKind InScalarKind)
-{
-	return MakePrimitive(InScalarKind, 1, 1);
-}
-
-FType FType::MakeVector(EScalarKind InScalarKind, int NumRows)
-{
-	return MakePrimitive(InScalarKind, NumRows, 1);
-}
-
-FType FType::MakePrimitive(EScalarKind InScalarKind, int NumRows, int NumColumns, bool bIsLWCInverseMatrix)
+FType FType::MakePrimitive(EScalarKind InScalarKind, int NumRows, int NumColumns, bool bIsDoubleInverseMatrix)
 {
 	check(NumRows >= 1 && NumRows <= 4);
 	check(NumColumns >= 1 && NumColumns <= 4);
-
-	// LWC inverse matrix flag only supported for EScalarKind::LWC and 4x4 dimensions.
-	check(!bIsLWCInverseMatrix || (InScalarKind == EScalarKind::Double && NumRows == 4 && NumColumns == 4));
+	check(!bIsDoubleInverseMatrix || (InScalarKind == EScalarKind::Double && NumRows == 4 && NumColumns == 4));
 	
 	FType Type{};
 	Type.Kind = ETypeKind::Primitive;
 	Type.Primitive.ScalarKind = InScalarKind;
 	Type.Primitive.NumRows = NumRows;
 	Type.Primitive.NumColumns = NumColumns;
-	Type.Primitive.bIsLWCInverseMatrix = bIsLWCInverseMatrix;
+	Type.Primitive.bIsLWCInverseMatrix = bIsDoubleInverseMatrix;
+
 	return Type;
 }
 
@@ -263,9 +206,9 @@ FString FType::GetSpelling() const
 				{
 					Spelling.Append(TEXT("FWSScalar"));
 				}
-				else if (Primitive.IsVector())
+				else if (Primitive.IsRowVector())
 				{
-					Spelling.Appendf(TEXT("FWSVector%d"), Primitive.NumRows);
+					Spelling.Appendf(TEXT("FWSVector%d"), Primitive.NumColumns);
 				}
 				else if (Primitive.bIsLWCInverseMatrix)
 				{
@@ -282,9 +225,9 @@ FString FType::GetSpelling() const
 		{
 			// append nothing
 		}
-		else if (Primitive.IsVector())
+		else if (Primitive.IsRowVector())
 		{
-			Spelling.Appendf(TEXT("%d"), Primitive.NumRows);
+			Spelling.Appendf(TEXT("%d"), Primitive.NumColumns);
 		}
 		else
 		{
@@ -310,9 +253,9 @@ FType FPrimitive::ToScalar() const
 	return FType::MakeScalar(ScalarKind);
 }
 
-FType FPrimitive::ToVector(int InNumRows) const
+FType FPrimitive::ToVector(int InNumColumns) const
 {
-	return FType::MakeVector(ScalarKind, InNumRows);
+	return FType::MakeVector(ScalarKind, InNumColumns);
 }
 
 UE::Shader::EValueType FType::ToValueType() const
@@ -342,14 +285,14 @@ UE::Shader::EValueType FType::ToValueType() const
 			return EValueType::Any;
 		}
 
-		check(Primitive.NumColumns == 1 && Primitive.NumRows <= 4);
+		check(Primitive.NumRows < 4 && Primitive.NumColumns <= 4);
 
 		switch (Primitive.ScalarKind)
 		{
-			case EScalarKind::Bool: return (EValueType)((int)EValueType::Bool1 + Primitive.NumRows - 1);
-			case EScalarKind::Int: 	return (EValueType)((int)EValueType::Int1 + Primitive.NumRows - 1);
-			case EScalarKind::Float: return (EValueType)((int)EValueType::Float1 + Primitive.NumRows - 1);
-			case EScalarKind::Double:	return (EValueType)((int)EValueType::Double1 + Primitive.NumRows - 1);
+			case EScalarKind::Bool: return (EValueType)((int)EValueType::Bool1 + Primitive.NumColumns - 1);
+			case EScalarKind::Int: 	return (EValueType)((int)EValueType::Int1 + Primitive.NumColumns - 1);
+			case EScalarKind::Float: return (EValueType)((int)EValueType::Float1 + Primitive.NumColumns - 1);
+			case EScalarKind::Double:	return (EValueType)((int)EValueType::Double1 + Primitive.NumColumns - 1);
 			default: UE_MIR_UNREACHABLE();
 		}
 	}
diff --git a/Engine/Source/Runtime/Engine/Private/Materials/MaterialIRValueAnalyzer.cpp b/Engine/Source/Runtime/Engine/Private/Materials/MaterialIRValueAnalyzer.cpp
--- a/Engine/Source/Runtime/Engine/Private/Materials/MaterialIRValueAnalyzer.cpp
+++ b/Engine/Source/Runtime/Engine/Private/Materials/MaterialIRValueAnalyzer.cpp
@@ -542,7 +542,7 @@ static void AnalyzePrimitiveUniformParameter(FMaterialIRValueAnalyzer& Analyzer,
 	if (UsedNumComponents >= 4)
 	{
 		check(UsedNumComponents == 4); 
-		const uint32 SizeInBytes = PrimitiveType.IsLWC() ? 32 : 16;
+		const uint32 SizeInBytes = PrimitiveType.IsDouble() ? 32 : 16;
 		GlobalComponentOffset = UniformExpressionSet.AllocateFromUniformBuffer(SizeInBytes) / 4;
 	}
 
diff --git a/Engine/Source/Runtime/Engine/Public/Materials/MaterialIREmitter.h b/Engine/Source/Runtime/Engine/Public/Materials/MaterialIREmitter.h
--- a/Engine/Source/Runtime/Engine/Public/Materials/MaterialIREmitter.h
+++ b/Engine/Source/Runtime/Engine/Public/Materials/MaterialIREmitter.h
@@ -404,6 +404,9 @@ public:
 	// Casts the given primitive value to a scalar. If this is a vector, it will convert to its first component.
 	FValueRef CastToScalar(FValueRef Value);
 
+	// Casts the given primitive value to a row vector of specified columns.
+	FValueRef CastToVector(FValueRef Value, int32 NumColumns);
+
 	// Casts the given primitive value to a given scalar kind, maintaining the same
 	// number of rows and columns.
 	FValueRef CastToScalarKind(FValueRef Value, EScalarKind ToScalarKind);
@@ -413,9 +416,9 @@ public:
 
 	// These functions cast the given value to the a primitive scalar or column vector type of given rows.
 	// If the value is missing ore could not be cast to the relative primitive type, they repors an error and return poison.
-	FValueRef CastToBool(FValueRef Value, int NumRows);
-	FValueRef CastToInt(FValueRef Value, int NumRows);
-	FValueRef CastToFloat(FValueRef Value, int NumRows);
+	FValueRef CastToBool(FValueRef Value, int NumColumns);
+	FValueRef CastToInt(FValueRef Value, int NumColumns);
+	FValueRef CastToFloat(FValueRef Value, int NumColumns);
 	
 	// Extracts the component with given index from the given argument value.
 	// The argument value must be a scalar or a vector primitive type, otherwise it reports an error and returns poison.
diff --git a/Engine/Source/Runtime/Engine/Public/Materials/MaterialIRTypes.h b/Engine/Source/Runtime/Engine/Public/Materials/MaterialIRTypes.h
--- a/Engine/Source/Runtime/Engine/Public/Materials/MaterialIRTypes.h
+++ b/Engine/Source/Runtime/Engine/Public/Materials/MaterialIRTypes.h
@@ -5,6 +5,7 @@
 #include "Shader/ShaderTypes.h"
 #include "MaterialValueType.h"
 #include "Materials/MaterialParameters.h"
+#include "Materials/MaterialExternalCodeRegistry.h"
 
 #if WITH_EDITOR
 
@@ -71,7 +72,7 @@ struct FPrimitive
 	bool IsScalar() const { return NumComponents() == 1; }
 	
 	// Returns whether this primitive type is a column vector (rows > 1, columns = 1).
-	bool IsVector() const { return NumRows > 1 && NumColumns == 1; }
+	bool IsRowVector() const { return NumRows == 1 && NumColumns > 1; }
 	
 	// Returns whether this primitive type is a matrix (rows > 1, columns > 1).
 	bool IsMatrix() const { return NumRows > 1 && NumColumns > 1; }
@@ -89,7 +90,7 @@ struct FPrimitive
 	bool IsFloat() const { return ScalarKind == EScalarKind::Float; }
 
 	// Returns whether this is an LWC primitive type (of any dimension).
-	bool IsLWC() const { return ScalarKind == EScalarKind::Double; }
+	bool IsDouble() const { return ScalarKind == EScalarKind::Double; }
 	
 	// Returns whether this is any floating point primitive type (e.g. float, LWC).
 	bool IsAnyFloat() const { return ScalarKindIsAnyFloat(ScalarKind); }
@@ -104,7 +105,7 @@ struct FPrimitive
 	FType ToScalar() const;
 	
 	// Returns a vector FType with the given number of rows and this primitive's scalar kind.
-	FType ToVector(int NumRows) const;
+	FType ToVector(int NumColumns) const;
 
 	bool operator==(const FPrimitive& Other) const = default;
 	bool operator!=(const FPrimitive& Other) const = default;
@@ -130,28 +131,40 @@ struct FType
 	// Returns the void type, which represents the absence of a value.
 	static FType MakeVoid();
 
-	// Returns a boolean type with the specified dimensions (scalar, vector, or matrix).
-	static FType MakeBool(int NumRows = 1, int NumColumns = 1);
-	
-	// Returns an integer type with the specified dimensions (scalar, vector, or matrix).
-	static FType MakeInt(int NumRows = 1, int NumColumns = 1);
-	
-	// Returns a float type with the specified dimensions (scalar, vector, or matrix).
-	static FType MakeFloat(int NumRows = 1, int NumColumns = 1);
+	// Returns a scalar type of the given scalar kind.
+	static FType MakeScalar(EScalarKind InScalarKind) { return MakePrimitive(InScalarKind, 1, 1); }
 
-	// Returns an double precision type with the specified dimensions (scalar, vector, or matrix).
-	static FType MakeDouble(int NumRows = 1, int NumColumns = 1);
+	// Returns a scalar type representing a boolean value.
+	static FType MakeBoolScalar() { return MakeScalar(EScalarKind::Bool); }
 
-	// Returns a scalar type of the given scalar kind.
-	static FType MakeScalar(EScalarKind InScalarKind);
+	// Returns a scalar type representing a integer value.
+	static FType MakeIntScalar() { return MakeScalar(EScalarKind::Int); }
+
+	// Returns a scalar type representing a single precision floating point value.
+	static FType MakeFloatScalar() { return MakeScalar(EScalarKind::Float); }
+
+	// Returns a scalar type representing a double precision floating point value.
+	static FType MakeDoubleScalar() { return MakeScalar(EScalarKind::Double); }
+
+	// Returns a row-major vector type of the given scalar kind and number of columns.
+	static FType MakeVector(EScalarKind InScalarKind, int NumColumns) { return MakePrimitive(InScalarKind, 1, NumColumns); }
+
+	// Returns a row-major vector type of boolean values with the specified number of columns.
+	static FType MakeBoolVector(int NumColumns) { return MakeVector(EScalarKind::Bool, NumColumns); }
 
-	// Returns a vector type of the given scalar kind and number of rows.
-	static FType MakeVector(EScalarKind InScalarKind, int NumRows);
+	// Returns a row-major vector type of integer values with the specified number of columns.
+	static FType MakeIntVector(int NumColumns) { return MakeVector(EScalarKind::Int, NumColumns); }
+
+	// Returns a row-major vector type of single precision floating point values with the specified number of columns.
+	static FType MakeFloatVector(int NumColumns) { return MakeVector(EScalarKind::Float, NumColumns); }
+
+	// Returns a row-major vector type of double precision floating point values with the specified number of columns.
+	static FType MakeDoubleVector(int NumColumns) { return MakeVector(EScalarKind::Double, NumColumns); }
 
 	// Returns a primitive type with the given kind and dimensions.
-	// LWC inverse matrix flag is only supported for EScalarKind::LWC and 4x4 dimensions.
-	static FType MakePrimitive(EScalarKind InScalarKind, int NumRows = 1, int NumColumns = 1, bool bIsLWCInverseMatrix = false);
-	
+	// Double inverse matrix flag is only supported for EScalarKind::Double and 4x4 dimensions.
+	static FType MakePrimitive(EScalarKind InScalarKind, int NumRows, int NumColumns, bool bIsDoubleInverseMatrix = false);
+
 	// Returns the aggregate type associated with the given material aggregate definition.
 	static FType MakeAggregate(const UMaterialAggregate* Aggregate);
 
@@ -195,7 +208,7 @@ struct FType
 	bool IsFloat() const { return Is(ETypeKind::Primitive) && Primitive.IsFloat(); }
 
 	// Returns true if the given type is a primitive with an LWC scalar kind.
-	bool IsLWC() const { return IsPrimitive() && Primitive.IsLWC(); }
+	bool IsDouble() const { return IsPrimitive() && Primitive.IsDouble(); }
 
 	// Returns true if the given type is a primitive with any float-based scalar kind (e.g. float, LWC).
 	bool IsAnyFloat() const { return Is(ETypeKind::Primitive) && Primitive.IsAnyFloat(); }
@@ -207,11 +220,14 @@ struct FType
 	bool IsBoolScalar() const { return Is(ETypeKind::Primitive) && Primitive.IsBoolScalar(); }
 
 	// Returns true if the given type is a vector primitive.
-	bool IsVector() const { return Is(ETypeKind::Primitive) && Primitive.IsVector(); }
+	bool IsVector() const { return Is(ETypeKind::Primitive) && Primitive.IsRowVector(); }
 
 	// Returns whether this type is a Material Parameter Collection object.
 	bool IsParameterCollection() const { return Kind == ETypeKind::ParameterCollection; }
 
+	// Returns true if the given type is a matrix primitive.
+	bool IsMatrix() const { return Is(ETypeKind::Primitive) && Primitive.IsMatrix(); }
+
 	// Returns whether this type is a Texture2D object.
 	bool IsTexture2D() const { return Kind == ETypeKind::Texture2D; }
 
@@ -225,22 +241,13 @@ struct FType
 	FPrimitive GetPrimitive() const { check(Is(ETypeKind::Primitive)); return Primitive; }
 	
 	// If this is a primitive type, returns its primitive info. Otherwise, returns None.
-	TOptional<FPrimitive> AsPrimitive() const
-	{
-		return IsPrimitive() ? Primitive : TOptional<FPrimitive>{};
-	}
+	TOptional<FPrimitive> AsPrimitive() const { return IsPrimitive() ? Primitive : TOptional<FPrimitive>{}; }
 
 	// If this is a scalar type, returns its primitive info. Otherwise, returns None.
-	TOptional<FPrimitive> AsScalar() const
-	{
-		return IsScalar() ? Primitive : TOptional<FPrimitive>{};
-	}
+	TOptional<FPrimitive> AsScalar() const { return IsScalar() ? Primitive : TOptional<FPrimitive>{}; }
 
 	// If this is a vector type, returns its primitive info. Otherwise, returns None.
-	TOptional<FPrimitive> AsVector() const
-	{
-		return IsVector() ? Primitive : TOptional<FPrimitive>{};
-	}
+	TOptional<FPrimitive> AsVector() const { return IsVector() ? Primitive : TOptional<FPrimitive>{}; }
 
 	// If this is an aggregate type, returns a pointer to the UMaterialAggregate definition. Otherwise, returns nullptr.
 	const UMaterialAggregate* AsAggregate() const { return Is(ETypeKind::Aggregate) ? Aggregate : nullptr; }
