Build fix for 43765033.

#rnx

diff --git a/Engine/Source/Runtime/Engine/Public/Materials/MaterialIRTypes.h b/Engine/Source/Runtime/Engine/Public/Materials/MaterialIRTypes.h
--- a/Engine/Source/Runtime/Engine/Public/Materials/MaterialIRTypes.h
+++ b/Engine/Source/Runtime/Engine/Public/Materials/MaterialIRTypes.h
@@ -5,7 +5,6 @@
 #include "Shader/ShaderTypes.h"
 #include "MaterialValueType.h"
 #include "Materials/MaterialParameters.h"
-#include "Materials/MaterialExternalCodeRegistry.h"
 
 #if WITH_EDITOR
 
