Refactor login in GameMode to pass a static analysis warning. This warning was occurring because bUseUniqueIdCheck was always going to be true on some platforms

#rnx

diff --git a/Engine/Source/Runtime/Engine/Private/GameMode.cpp b/Engine/Source/Runtime/Engine/Private/GameMode.cpp
--- a/Engine/Source/Runtime/Engine/Private/GameMode.cpp
+++ b/Engine/Source/Runtime/Engine/Private/GameMode.cpp
@@ -635,13 +635,29 @@ void AGameMode::AddInactivePlayer(APlayerState* PlayerState, APlayerController*
 					InactivePlayerArray.RemoveAt(Idx, 1);
 					Idx--;
 				}
-				else if ( (!bUseUniqueIdCheck && bHasValidNetworkAddress && (CurrentPlayerState->SavedNetworkAddress == NewPlayerState->SavedNetworkAddress))
-						|| (bUseUniqueIdCheck && (CurrentPlayerState->GetUniqueId() == NewPlayerState->GetUniqueId())) )
+				else 
 				{
-					// destroy the playerstate, then remove it from the tracking
-					CurrentPlayerState->Destroy();
-					InactivePlayerArray.RemoveAt(Idx, 1);
-					Idx--;
+					bool bDestroyPlayerState = false;
+					if (!bUseUniqueIdCheck)
+					{
+						if (bHasValidNetworkAddress && CurrentPlayerState->SavedNetworkAddress == NewPlayerState->SavedNetworkAddress)
+						{
+							bDestroyPlayerState = true;
+						}
+					}
+					else if (CurrentPlayerState->GetUniqueId() == NewPlayerState->GetUniqueId())
+					{
+						bDestroyPlayerState = true;
+					}
+
+					if (bDestroyPlayerState)
+					{
+						// destroy the playerstate, then remove it from the tracking
+						CurrentPlayerState->Destroy();
+						InactivePlayerArray.RemoveAt(Idx, 1);
+						Idx--;
+					}
+					
 				}
 			}
 			InactivePlayerArray.Add(NewPlayerState);
