[SkeletonCompatibility] When using compatible skeletons we could end in a situation where the source skeleton did not have a bone of the target skeleton, which could lead to a crash.

diff --git a/Engine/Source/Runtime/Engine/Private/Animation/AnimationRuntime.cpp b/Engine/Source/Runtime/Engine/Private/Animation/AnimationRuntime.cpp
--- a/Engine/Source/Runtime/Engine/Private/Animation/AnimationRuntime.cpp
+++ b/Engine/Source/Runtime/Engine/Private/Animation/AnimationRuntime.cpp
@@ -2902,7 +2902,7 @@ void FAnimationRuntime::RetargetBoneTransform(const USkeleton* SourceSkeleton, c
 	if (SourceSkeleton)
 	{
 		// Retrieve skeleton, even if it is unreachable (but not GC-ed yet)
-		const bool bEvenIfUnreachable = true; 
+		constexpr bool bEvenIfUnreachable = true; 
 		const USkeleton* TargetSkeleton = RequiredBones.GetSkeletonAsset(bEvenIfUnreachable);
 		const FSkeletonRemapping& SkeletonRemapping = UE::Anim::FSkeletonRemappingRegistry::Get().GetRemapping(SourceSkeleton, TargetSkeleton);
 
@@ -2912,8 +2912,12 @@ void FAnimationRuntime::RetargetBoneTransform(const USkeleton* SourceSkeleton, c
 		// Apply compatible skeleton remapping if required
 		if (SkeletonRemapping.IsValid() && SkeletonRemapping.RequiresReferencePoseRetarget())
 		{
-			SourceSkeletonBoneIndex = SkeletonRemapping.GetSourceSkeletonBoneIndex(TargetSkeletonBoneIndex);
-			BoneTransform = SkeletonRemapping.RetargetBoneTransformToTargetSkeleton(TargetSkeletonBoneIndex, BoneTransform);
+			const int32 SourceIndex = SkeletonRemapping.GetSourceSkeletonBoneIndex(TargetSkeletonBoneIndex);
+			if (SourceIndex != INDEX_NONE)
+			{
+				SourceSkeletonBoneIndex = SourceIndex;
+				BoneTransform = SkeletonRemapping.RetargetBoneTransformToTargetSkeleton(TargetSkeletonBoneIndex, BoneTransform);
+			}
 		}
 
 		// Check what retarget mode to use for the translational retargeting for this specific bone.
