Iris
- Disabled code that ran minimal update on replicationsystem when we do not have any connections as there was some filtering issues

diff --git a/Engine/Source/Runtime/Engine/Private/NetDriver.cpp b/Engine/Source/Runtime/Engine/Private/NetDriver.cpp
--- a/Engine/Source/Runtime/Engine/Private/NetDriver.cpp
+++ b/Engine/Source/Runtime/Engine/Private/NetDriver.cpp
@@ -1105,7 +1105,7 @@ void UNetDriver::TickFlush(float DeltaSeconds)
 	}
 	FSimpleScopeSecondsCounter ScopedTimer(GTickFlushGameDriverTimeSeconds, bEnableTimer);
 
-	if (IsServer() && !bSkipServerReplicateActors)
+	if (IsServer() && ClientConnections.Num() > 0 && !bSkipServerReplicateActors)
 	{
 		// Update all clients.
 #if WITH_SERVER_CODE
