[shaders] fix conflict in retry mechanism between retries for dumping debug info and interactive retries for global/default shaders. the former should always run first, so the latter can provide the user with the debug info path in the case r.DumpShaderDebugInfo is set to 2 (or 3). includes some minor cleanup of the error messages/dialogs encountered when global/default shaders fail to compile for readability.

diff --git a/Engine/Source/Runtime/Engine/Private/ShaderCompiler/ShaderCompiler.cpp b/Engine/Source/Runtime/Engine/Private/ShaderCompiler/ShaderCompiler.cpp
--- a/Engine/Source/Runtime/Engine/Private/ShaderCompiler/ShaderCompiler.cpp
+++ b/Engine/Source/Runtime/Engine/Private/ShaderCompiler/ShaderCompiler.cpp
@@ -557,7 +557,7 @@ static bool CheckSingleJob(const FShaderCompileJob& SingleJob, TArray<FString>&
 {
 	if (SingleJob.bSucceeded)
 	{
-		checkf(SingleJob.Output.ShaderCode.GetShaderCodeSize() > 0, TEXT("Abnormal shader code size for a succeded job: %d bytes"), SingleJob.Output.ShaderCode.GetShaderCodeSize());
+		checkf(SingleJob.Output.ShaderCode.GetShaderCodeSize() > 0, TEXT("Abnormal shader code size for a successful job: %d bytes"), SingleJob.Output.ShaderCode.GetShaderCodeSize());
 	}
 
 	bool bSucceeded = SingleJob.bSucceeded;
@@ -2162,30 +2162,26 @@ void FShaderCompilingManager::ProcessCompiledShaderMaps(
 
 					if (Material->IsDefaultMaterial())
 					{
-						FString ErrorString;
+						FString ErrorString = FString::Printf(TEXT("Failed to compile default material %s!\n"), *Material->GetBaseMaterialPathName());
 
 						// Log the errors unsuppressed before the fatal error, so it's always obvious from the log what the compile error was
 						for (const FString& UniqueError : ErrorInfo.UniqueErrors)
 						{
-							ErrorString += FString::Printf(TEXT("  %s\n"), *UniqueError);
+							ErrorString += FString::Printf(TEXT("%s\n\n"), *UniqueError);
 						}
 
-						ErrorString += FString::Printf(TEXT("Failed to compile default material %s!"), *Material->GetBaseMaterialPathName());
-
 						if (AreShaderErrorsFatal())
 						{
 							// Assert if a default material could not be compiled, since there will be nothing for other failed materials to fall back on.
-							UE_LOG(LogShaderCompilers, Fatal, TEXT("%s"), *ErrorString);
+							UE_LOG(LogShaderCompilers, Fatal, TEXT("\n%s"), *ErrorString);
 						}
 						else
 						{
-							UE_LOG(LogShaderCompilers, Error, TEXT("%s"), *ErrorString);
+							UE_LOG(LogShaderCompilers, Error, TEXT("\n%s"), *ErrorString);
 						}
 					}
 					
-					FString ErrorString;
-
-					ErrorString += FString::Printf(TEXT("Failed to compile Material %s for platform %s, Default Material will be used in game.\n"),
+					FString ErrorString = FString::Printf(TEXT("Failed to compile Material %s for platform %s, Default Material will be used in game.\n"),
 						*Material->GetDebugName(), *LegacyShaderPlatformToShaderFormat(CompilingShaderMap->GetShaderPlatform()).ToString());
 
 					for (const FString& UniqueError : ErrorInfo.UniqueErrors)
@@ -2193,7 +2189,7 @@ void FShaderCompilingManager::ProcessCompiledShaderMaps(
 						FString ErrorMessage = UniqueError;
 						// Work around build machine string matching heuristics that will cause a cook to fail
 						ErrorMessage.ReplaceInline(TEXT("error "), TEXT("err0r "), ESearchCase::CaseSensitive);
-						ErrorString += FString::Printf(TEXT("  %s\n"), *ErrorMessage);
+						ErrorString += FString::Printf(TEXT("%s\n"), *ErrorMessage);
 					}
 
 					UE_LOG(LogShaderCompilers, Warning, TEXT("%s"), *ErrorString);
@@ -2375,7 +2371,7 @@ static void BuildErrorStringAndReport(const FShaderDiagnosticInfo& DiagInfo, FSt
 
 	for (int32 ErrorIndex = 0; ErrorIndex < DiagInfo.UniqueErrors.Num(); ErrorIndex++)
 	{
-		FString UniqueErrorString = DiagInfo.UniqueErrors[ErrorIndex] + TEXT("\n");
+		FString UniqueErrorString = DiagInfo.UniqueErrors[ErrorIndex] + TEXT("\n\n");
 
 		if (FPlatformMisc::IsDebuggerPresent())
 		{
@@ -2395,17 +2391,39 @@ static void BuildErrorStringAndReport(const FShaderDiagnosticInfo& DiagInfo, FSt
 bool FShaderCompilingManager::HandlePotentialRetry(TMap<int32, FShaderMapFinalizeResults>& CompletedShaderMaps)
 {
 #if WITH_EDITORONLY_DATA
+	TArray<FShaderCommonCompileJobPtr> ReissueForDebugJobs;
 	TArray<FShaderCommonCompileJobPtr> ReissueJobs;
 	TArray<int32> ReissueMapIds;
 
 	for (TMap<int32, FShaderMapFinalizeResults>::TIterator It(CompletedShaderMaps); It; ++It)
 	{
+		bool bAnyToReissue = false;
+
 		FShaderMapFinalizeResults& Results = It.Value();
+		// Handle retries based on r.DumpShaderDebugInfo=2/3 first; if any jobs require this we skip the interactive retry prompt until after such jobs run
+		// (so the retry prompt will contain the debug info path)
+		for (int32 FinishedJobIndex = 0; FinishedJobIndex < Results.FinishedJobs.Num(); ++FinishedJobIndex)
+		{
+			FShaderCommonCompileJobPtr Job = Results.FinishedJobs[FinishedJobIndex];
+			bool bEnableDebug = false;
+			Job->ForEachSingleShaderJob([&bEnableDebug](FShaderCompileJob& SingleJob)
+				{
+					// reissue the job if either we want to dump debug info for it, or a retry was requested above for any failed jobs
+					bEnableDebug = GShaderCompilingManager->ShouldRecompileToDumpShaderDebugInfo(SingleJob);
+				});
 
-		bool bRetryCompile = false;
+			if (bEnableDebug)
+			{
+				bAnyToReissue = true;
+				ReissueForDebugJobs.Add(Job);
+				Results.FinishedJobs.RemoveAt(FinishedJobIndex--, EAllowShrinking::No);
+			}
+		}
+
+		bool bInteractiveRetry = false;
 
-		// interactive retries - prompt for global/default shaders that have errors
-		if (!FApp::IsUnattended() && !Results.bAllJobsSucceeded)
+		// interactive retries - prompt for global/default shaders that have errors, if we're not reissuing any jobs to dump debug info first
+		if (!FApp::IsUnattended() && !Results.bAllJobsSucceeded && ReissueForDebugJobs.IsEmpty())
 		{
 			bool bSpecialEngineMaterial = false;
 
@@ -2433,7 +2451,7 @@ bool FShaderCompilingManager::HandlePotentialRetry(TMap<int32, FShaderMapFinaliz
 				const TCHAR* MaterialName = ShaderMap ? ShaderMap->GetFriendlyName() : TEXT("global shaders");
 				FString ErrorString = FString::Printf(TEXT("%i Shader compiler errors compiling %s for platform %s:"), ShaderDiagInfo.UniqueErrors.Num(), MaterialName, *ShaderDiagInfo.TargetShaderPlatformString);
 				UE_LOG(LogShaderCompilers, Warning, TEXT("%s"), *ErrorString);
-				ErrorString += TEXT("\n");
+				ErrorString += TEXT("\n\n");
 
 				bool bAnyErrorLikelyToBeCodeError = false;
 				for (const FShaderCommonCompileJob* Job : ShaderDiagInfo.ErrorJobs)
@@ -2454,60 +2472,57 @@ bool FShaderCompilingManager::HandlePotentialRetry(TMap<int32, FShaderMapFinaliz
 						UE_DEBUG_BREAK();
 						// Set GRetryShaderCompilation to true in the debugger to enable retries in debug
 						// NOTE: MaterialTemplate.usf will not be reloaded when retrying!
-						bRetryCompile = GRetryShaderCompilation;
+						bInteractiveRetry = GRetryShaderCompilation;
 					}
 					else
 					{
-						if (FPlatformMisc::MessageBoxExt(EAppMsgType::YesNo, *FText::Format(NSLOCTEXT("UnrealEd", "Error_RetryShaderCompilation", "{0}\r\n\r\nRetry compilation?"),
+						if (FPlatformMisc::MessageBoxExt(EAppMsgType::YesNo, *FText::Format(NSLOCTEXT("UnrealEd", "Error_RetryShaderCompilation", "{0}\nRetry compilation?"),
 							FText::FromString(ErrorString)).ToString(), TEXT("Error")) == EAppReturnType::Type::Yes)
 						{
-							bRetryCompile = true;
+							bInteractiveRetry = true;
 						}
 					}
 				}
-			}
-		}
-
-		int32 OriginalFinishedJobCount = Results.FinishedJobs.Num();
-		bool bAnyReissued = false;
 
-		// retries based on r.DumpShaderDebugInfo=2/3
-		for (int32 FinishedJobIndex = 0; FinishedJobIndex < Results.FinishedJobs.Num(); ++FinishedJobIndex)
-		{
-			FShaderCommonCompileJobPtr Job = Results.FinishedJobs[FinishedJobIndex];
-			bool bReissue = false, bEnableDebug = false;
-			Job->ForEachSingleShaderJob([&bReissue, &bEnableDebug, bRetryCompile](FShaderCompileJob& SingleJob)
-				{
-					// reissue the job if either we want to dump debug info for it, or a retry was requested above for any failed jobs
-					bEnableDebug = GShaderCompilingManager->ShouldRecompileToDumpShaderDebugInfo(SingleJob);
-					bReissue = bEnableDebug || (bRetryCompile && !SingleJob.bSucceeded);
-				});
-			if (bReissue)
-			{
-				if (bEnableDebug)
+				if (bInteractiveRetry)
 				{
-					// if any stage of a pipeline job is reissued to dump debug info, we need to enable debug info on all stages of the job
-					// since the debug information relies on debug source files existing for every stage
-					Job->ForEachSingleShaderJob([](FShaderCompileJob& SingleJob)
+					bAnyToReissue = true;
+					for (int32 FinishedJobIndex = 0; FinishedJobIndex < Results.FinishedJobs.Num(); ++FinishedJobIndex)
+					{
+						FShaderCommonCompileJobPtr Job = Results.FinishedJobs[FinishedJobIndex];
+						if (!Job->bSucceeded)
 						{
-							SingleJob.Input.DumpDebugInfoPath = GShaderCompilingManager->CreateShaderDebugInfoPath(SingleJob.Input);
-							// Any reissued jobs due to this condition will dump debug information, so increment the dump count here
-							GShaderCompilingManager->IncrementNumDumpedShaderSources();
-						});
+							ReissueJobs.Add(Job);
+							Results.FinishedJobs.RemoveAt(FinishedJobIndex--, EAllowShrinking::No);
+						}
+					}
 				}
-				bAnyReissued = true;
-				ReissueJobs.Add(Job);
-				// remove the job we will reissue from the finished jobs list; no point in allowing shrinking as it will grow again below when the reissued jobs finish
-				Results.FinishedJobs.RemoveAt(FinishedJobIndex--, EAllowShrinking::No);
 			}
 		}
-
-		if (bAnyReissued)
+		if (bAnyToReissue)
 		{
 			ReissueMapIds.Add(It.Key());
 		}
 	}
 
+	if (ReissueForDebugJobs.Num())
+	{
+		check(ReissueJobs.IsEmpty());
+		for (const FShaderCommonCompileJobPtr& DebugJob : ReissueForDebugJobs)
+		{
+			// if any stage of a pipeline job is reissued to dump debug info, we need to enable debug info on all stages of the job
+			// since the debug information relies on debug source files existing for every stage
+			DebugJob->ForEachSingleShaderJob([](FShaderCompileJob& SingleJob)
+				{
+					SingleJob.Input.DumpDebugInfoPath = GShaderCompilingManager->CreateShaderDebugInfoPath(SingleJob.Input);
+					// Any reissued jobs due to this condition will dump debug information, so increment the dump count here
+					GShaderCompilingManager->IncrementNumDumpedShaderSources();
+				});
+		}
+
+		ReissueJobs = MoveTemp(ReissueForDebugJobs);
+	}
+
 	if (ReissueJobs.Num())
 	{
 		// Flush the shader file cache so that any changes will be propagated.
diff --git a/Engine/Source/Runtime/RenderCore/Internal/ShaderDiagnostics.h b/Engine/Source/Runtime/RenderCore/Internal/ShaderDiagnostics.h
--- a/Engine/Source/Runtime/RenderCore/Internal/ShaderDiagnostics.h
+++ b/Engine/Source/Runtime/RenderCore/Internal/ShaderDiagnostics.h
@@ -8,7 +8,7 @@
 
 struct FShaderDiagnosticInfo
 {
-	TArray<FShaderCommonCompileJob*> ErrorJobs;
+	TArray<FShaderCompileJob*> ErrorJobs;
 	TArray<FString> UniqueErrors;
 	TArray<FString> UniqueWarnings;
 	TArray<EShaderPlatform> ErrorPlatforms;
diff --git a/Engine/Source/Runtime/RenderCore/Private/ShaderDiagnostics.cpp b/Engine/Source/Runtime/RenderCore/Private/ShaderDiagnostics.cpp
--- a/Engine/Source/Runtime/RenderCore/Private/ShaderDiagnostics.cpp
+++ b/Engine/Source/Runtime/RenderCore/Private/ShaderDiagnostics.cpp
@@ -93,7 +93,7 @@ int32 FShaderDiagnosticInfo::AddAndProcessErrorsForFailedJobFiltered(FShaderComp
 				if (!GIsBuildMachine && !bReportedDebugInfo && CurrentJob.Input.DumpDebugInfoPath.Len() > 0)
 				{
 					const FString DebugInfoPath = ConvertToNativePlatformAbsolutePath(*CurrentJob.Input.DumpDebugInfoPath);
-					UniqueErrorPrefix += FString::Printf(TEXT("Shader debug info dumped to: \"%s\"\n"), *DebugInfoPath);
+					UniqueErrorPrefix += FString::Printf(TEXT("Shader debug info dumped to: \"%s\"\n\n"), *DebugInfoPath);
 					bReportedDebugInfo = true;
 				}
 
