Iris
- Disabled code that ran minimal update on replicationsystem when we do not have any connections as there was some filtering issues.
- Fixed ensure

diff --git a/Engine/Source/Runtime/Engine/Private/NetDriver.cpp b/Engine/Source/Runtime/Engine/Private/NetDriver.cpp
--- a/Engine/Source/Runtime/Engine/Private/NetDriver.cpp
+++ b/Engine/Source/Runtime/Engine/Private/NetDriver.cpp
@@ -1588,7 +1588,7 @@ void UNetDriver::PostTickFlush()
 	{
 		TGuardValue<bool> GuardInNetTick(bInTick, true);
 #if UE_WITH_IRIS
-		if (ReplicationSystem && (ServerConnection != nullptr || IsServer()))
+		if (ReplicationSystem && (ServerConnection != nullptr || ClientConnections.Num() > 0))
 		{
 			ReplicationSystem->PostSendUpdate();
 		}
