Updating struct tracing setup used by anim next instances

diff --git a/Engine/Plugins/Experimental/AnimNext/Source/AnimNext/Private/RewindDebugger/AnimNextTrace.cpp b/Engine/Plugins/Experimental/AnimNext/Source/AnimNext/Private/RewindDebugger/AnimNextTrace.cpp
--- a/Engine/Plugins/Experimental/AnimNext/Source/AnimNext/Private/RewindDebugger/AnimNextTrace.cpp
+++ b/Engine/Plugins/Experimental/AnimNext/Source/AnimNext/Private/RewindDebugger/AnimNextTrace.cpp
@@ -113,11 +113,11 @@ void FAnimNextTrace::OutputAnimNextInstance(const FUAFAssetInstance* AnimNextIns
 		}
 
 		const FUAFAssetInstance* HostInstance = AnimNextInstance->GetHost();
-		uint64 HostInstanceId = HostInstance ? HostInstance->GetUniqueId() : FObjectTrace::InvalidObjectId;
+		uint64 HostInstanceId = HostInstance ? HostInstance->GetUniqueId() : FObjectTrace::GetObjectId(OuterObject);
 
 		TRACE_OBJECT(OuterObject);
 		TRACE_OBJECT(AnimNextInstance->GetAsset<UAnimNextRigVMAsset>());
-		TRACE_CHILD_ELEMENT_WITH_OUTER(OuterObject, AnimNextInstance->GetUniqueId(), OuterObject, HostInstanceId, FUAFAssetInstance::StaticStruct(), AnimNextInstance->GetAssetName().ToString());
+		TRACE_INSTANCE(OuterObject, AnimNextInstance->GetUniqueId(), HostInstanceId, FUAFAssetInstance::StaticStruct(), AnimNextInstance->GetAssetName().ToString());
 
 		UE_TRACE_LOG(UAF, Instance, UAFChannel)
 			<< Instance.InstanceId(InstanceId)
diff --git a/Engine/Plugins/Experimental/AnimNext/Source/AnimNext/Private/UAFAssetInstance.cpp b/Engine/Plugins/Experimental/AnimNext/Source/AnimNext/Private/UAFAssetInstance.cpp
--- a/Engine/Plugins/Experimental/AnimNext/Source/AnimNext/Private/UAFAssetInstance.cpp
+++ b/Engine/Plugins/Experimental/AnimNext/Source/AnimNext/Private/UAFAssetInstance.cpp
@@ -2,16 +2,13 @@
 
 #include "UAFAssetInstance.h"
 
+#include "ObjectTrace.h"
 #include "Logging/StructuredLog.h"
 #include "Variables/UAFInstanceVariableContainer.h"
 #include "Variables/VariableOverrides.h"
 #include "Variables/VariableOverridesCollection.h"
 #include "UAFAssetInstanceComponent.h"
 
-#if ANIMNEXT_TRACE_ENABLED
-volatile int64 FUAFAssetInstance::NextUniqueId = 1;
-#endif
-
 void FUAFAssetInstance::InitializeVariables(const TSharedPtr<const UE::AnimNext::FVariableOverridesCollection>& InOverrides)
 {
 	Variables.Initialize(*this, InOverrides);
@@ -34,7 +31,7 @@ void FUAFAssetInstance::ReleaseComponents()
 FUAFAssetInstance::FUAFAssetInstance()
 {
 #if ANIMNEXT_TRACE_ENABLED
-	UniqueId = FPlatformAtomics::InterlockedIncrement(&NextUniqueId);
+	UniqueId = FObjectTrace::AllocateInstanceId();
 #endif
 }
 
diff --git a/Engine/Plugins/Experimental/AnimNext/Source/AnimNextEditor/Private/RewindDebugger/AnimNextModuleTrack.cpp b/Engine/Plugins/Experimental/AnimNext/Source/AnimNextEditor/Private/RewindDebugger/AnimNextModuleTrack.cpp
--- a/Engine/Plugins/Experimental/AnimNext/Source/AnimNextEditor/Private/RewindDebugger/AnimNextModuleTrack.cpp
+++ b/Engine/Plugins/Experimental/AnimNext/Source/AnimNextEditor/Private/RewindDebugger/AnimNextModuleTrack.cpp
@@ -60,11 +60,10 @@ void FAnimNextModuleTrackCreator::GetTrackTypesInternal(TArray<RewindDebugger::F
 
 TSharedPtr<RewindDebugger::FRewindDebuggerTrack> FAnimNextModuleTrackCreator::CreateTrackInternal(const RewindDebugger::FObjectId& ObjectId) const
 {
-	return MakeShared<FAnimNextModuleTrack>(ObjectId.GetMainId(), ObjectId.GetChildElementId());
+	return MakeShared<FAnimNextModuleTrack>(ObjectId.GetMainId());
 }
 
-FAnimNextModuleTrack::FAnimNextModuleTrack(uint64 InObjectId, uint64 InInstanceId) :
-   	ObjectId(InObjectId),
+FAnimNextModuleTrack::FAnimNextModuleTrack(uint64 InInstanceId) :
 	InstanceId(InInstanceId)
 {
    	Initialize();
@@ -130,11 +129,6 @@ bool FAnimNextModuleTrack::UpdateInternal()
 			Initialize();
 		}
 
-		if (InstanceId == 0)
-		{
-			AnimNextProvider->GetModuleId(ObjectId, InstanceId);
-		}
-		
 		if(InstanceId != 0)
 		{
 			if (const FDataInterfaceData* Data = AnimNextProvider->GetDataInterfaceData(InstanceId))
@@ -237,20 +231,7 @@ bool FAnimNextModuleTrack::HandleDoubleClickInternal()
 	
 bool FAnimNextModuleTrackCreator::HasDebugInfoInternal(const RewindDebugger::FObjectId& ObjectId) const
 {
-	TRACE_CPUPROFILER_EVENT_SCOPE(FAnimNextModuleTrack::HasDebugInfoInternal);
-
-	const IRewindDebugger* RewindDebugger = IRewindDebugger::Instance();
-	if (const TraceServices::IAnalysisSession* AnalysisSession = RewindDebugger->GetAnalysisSession())
-	{
-		TraceServices::FAnalysisSessionReadScope SessionReadScope(*AnalysisSession);
-    
-		const FAnimNextProvider* AnimNextProvider = IRewindDebugger::Instance()->GetAnalysisSession()->ReadProvider<FAnimNextProvider>("AnimNextProvider");
-    
-		uint64 ModuleId;
-		return AnimNextProvider->GetModuleId(ObjectId.GetMainId(), ModuleId) != 0;
-	}
-
-	return false;
+	return true;
 }
 
 
diff --git a/Engine/Plugins/Experimental/AnimNext/Source/AnimNextEditor/Private/RewindDebugger/AnimNextModuleTrack.h b/Engine/Plugins/Experimental/AnimNext/Source/AnimNextEditor/Private/RewindDebugger/AnimNextModuleTrack.h
--- a/Engine/Plugins/Experimental/AnimNext/Source/AnimNextEditor/Private/RewindDebugger/AnimNextModuleTrack.h
+++ b/Engine/Plugins/Experimental/AnimNext/Source/AnimNextEditor/Private/RewindDebugger/AnimNextModuleTrack.h
@@ -34,7 +34,7 @@ namespace UE::AnimNextEditor
 class FAnimNextModuleTrack : public RewindDebugger::FRewindDebuggerTrack
 {
 public:
-	FAnimNextModuleTrack(uint64 InObjectId, uint64 InstanceId);
+	FAnimNextModuleTrack(uint64 InstanceId);
 	virtual ~FAnimNextModuleTrack();
 	
 	TSharedPtr<SEventTimelineView::FTimelineEventData> GetExistenceRange() const { return ExistenceRange; }
@@ -59,7 +59,7 @@ private:
 	
 	virtual uint64 GetObjectIdInternal() const override
 	{
-		return ObjectId;
+		return InstanceId;
 	}
 	
 	virtual FText GetDisplayNameInternal() const override;
@@ -70,8 +70,7 @@ private:
 
 	TSharedPtr<IDetailsView> DetailsView;
 	FSlateIcon Icon;
-	uint64 ObjectId;
-	uint64 InstanceId = 0; 
+	uint64 InstanceId; 
 	double PreviousScrubTime = -1.0;
 	TWeakObjectPtr<UPropertyBagDetailsObject> DetailsObjectWeakPtr;
 	TSharedPtr<SEventTimelineView::FTimelineEventData> ExistenceRange;
@@ -95,11 +94,6 @@ private:
 		return "AnimNextModule";
 	}
 	
-	virtual bool IsChildElementSupportedInternal() const override
-	{
-		return true;
-	}
-	
 	virtual bool IsCreatingPrimaryChildTrackInternal() const override
 	{
 		return true;
diff --git a/Engine/Plugins/Experimental/AnimNextAnimGraph/Source/AnimNextAnimGraph/Private/Graph/AnimNextGraphInstance.cpp b/Engine/Plugins/Experimental/AnimNextAnimGraph/Source/AnimNextAnimGraph/Private/Graph/AnimNextGraphInstance.cpp
--- a/Engine/Plugins/Experimental/AnimNextAnimGraph/Source/AnimNextAnimGraph/Private/Graph/AnimNextGraphInstance.cpp
+++ b/Engine/Plugins/Experimental/AnimNextAnimGraph/Source/AnimNextAnimGraph/Private/Graph/AnimNextGraphInstance.cpp
@@ -24,7 +24,7 @@ FAnimNextGraphInstance::FAnimNextGraphInstance()
 
 FAnimNextGraphInstance::~FAnimNextGraphInstance()
 {
-	TRACE_CHILD_ELEMENT_LIFETIME_END(GetModuleInstance() ? GetModuleInstance()->GetObject() : nullptr, GetUniqueId());
+	TRACE_INSTANCE_LIFETIME_END(ModuleInstance ? ModuleInstance->GetObject() : nullptr, GetUniqueId());
 	Release();
 }
 
diff --git a/Engine/Source/Runtime/Engine/Private/ObjectTrace.cpp b/Engine/Source/Runtime/Engine/Private/ObjectTrace.cpp
--- a/Engine/Source/Runtime/Engine/Private/ObjectTrace.cpp
+++ b/Engine/Source/Runtime/Engine/Private/ObjectTrace.cpp
@@ -271,6 +271,13 @@ void FObjectTrace::Reset()
 	};
 }
 
+uint64 FObjectTrace::AllocateInstanceId()
+{
+	static int64 volatile CurrentId = 1;
+	return FPlatformAtomics::InterlockedIncrement(&CurrentId);
+}
+
+
 uint64 FObjectTrace::GetObjectId(const UObject* InObject)
 {
 	// An object ID uses a combination of its own and its outer's index
@@ -279,12 +286,10 @@ uint64 FObjectTrace::GetObjectId(const UObject* InObject)
 
 	auto GetObjectIdInner = [](const UObject* InObjectInner)
 	{
-		static uint64 CurrentId = 1;
-
 		FObjectIdAnnotation Annotation = GObjectIdAnnotations.GetAnnotation(InObjectInner);
 		if (Annotation.Id == 0)
 		{
-			Annotation.Id = CurrentId++;
+			Annotation.Id = AllocateInstanceId();
 			GObjectIdAnnotations.AddAnnotation(InObjectInner, MoveTemp(Annotation));
 		}
 
@@ -453,39 +458,37 @@ void FObjectTrace::OutputView(const UObject* InPlayer, const FSceneView* InView)
 		<< View.AspectRatio(AspectRatio);
 }
 
-void FObjectTrace::OutputObject(const UObject* InObject)
+void FObjectTrace::OutputInstance(const UObject* InOuterObject, uint64 InInstanceId, uint64 InOuterId, UStruct* InType, const FString& InName, const FString& InPathName)
 {
-	if (UE::ObjectTrace::IsTracingDisabledForObject(InObject))
+	if (UE::ObjectTrace::IsTracingDisabledForObject(InOuterObject))
 	{
 		return;
 	}
-
-	const uint64 ObjectId = GetObjectId(InObject);
-	if (!AddObjectToTracedSet(UE::ObjectTrace::FObjectId{ObjectId}))
+	
+	if (!AddObjectToTracedSet(UE::ObjectTrace::FObjectId{InInstanceId}))
 	{
 		return;
 	}
 
-	OutputObject(InObject->GetOuter());
-
 	// Trace the object's class first
-	TRACE_CLASS(InObject->GetClass());
-
-	TCHAR ObjectName[FName::StringBufferSize];
-	const uint32 ObjectNameLength = InObject->GetFName().ToString(ObjectName);
-	const FString ObjectPathName = InObject->GetPathName();
+	TRACE_CLASS(InType);
 
 	UE_TRACE_LOG(Object, Object, ObjectChannel)
-		<< Object.Id(ObjectId)
-		<< Object.ClassId(GetObjectId(InObject->GetClass()))
-		<< Object.OuterId(GetObjectId(InObject->GetOuter()))
-		<< Object.Name(ObjectName, ObjectNameLength)
-		<< Object.Path(*ObjectPathName, ObjectPathName.Len());
-
+		<< Object.Id(InInstanceId)
+		<< Object.ClassId(GetObjectId(InType))
+		<< Object.OuterId(InOuterId)
+		<< Object.Name(*InName, InName.Len())
+		<< Object.Path(*InPathName, InPathName.Len());
+		
 	UE_TRACE_LOG(Object, ObjectLifetimeBegin2, ObjectChannel)
 		<< ObjectLifetimeBegin2.Cycle(FPlatformTime::Cycles64())
-		<< ObjectLifetimeBegin2.RecordingTime(FObjectTrace::GetWorldElapsedTime(InObject->GetWorld()))
-		<< ObjectLifetimeBegin2.Id(ObjectId);
+		<< ObjectLifetimeBegin2.RecordingTime(FObjectTrace::GetWorldElapsedTime(InOuterObject->GetWorld()))
+		<< ObjectLifetimeBegin2.Id(InInstanceId);
+}
+
+void FObjectTrace::OutputObject(const UObject* InObject)
+{
+	OutputInstance(InObject, GetObjectId(InObject), GetObjectId(InObject->GetOuter()), InObject->GetClass(), InObject->GetFName().ToString(), InObject->GetPathName());
 }
 
 void FObjectTrace::OutputObjectTransform(const UObject* InObject, const FTransform& InTransform, ETeleportType InTeleportType)
@@ -525,19 +528,23 @@ void FObjectTrace::OutputObjectLifetimeBegin(const UObject* InObject)
 	TRACE_OBJECT(InObject);
 }
 
-void FObjectTrace::OutputObjectLifetimeEnd(const UObject* InObject)
+void FObjectTrace::OutputInstanceLifetimeEnd(const UObject* InOuterObject, uint64 InInstanceId)
 {
-	if (UE::ObjectTrace::IsTracingDisabledForObject(InObject))
+	if (UE::ObjectTrace::IsTracingDisabledForObject(InOuterObject))
 	{
 		return;
 	}
 
-	TRACE_OBJECT(InObject);
-
 	UE_TRACE_LOG(Object, ObjectLifetimeEnd2, ObjectChannel)
 		<< ObjectLifetimeEnd2.Cycle(FPlatformTime::Cycles64())
-		<< ObjectLifetimeEnd2.RecordingTime(FObjectTrace::GetWorldElapsedTime(InObject->GetWorld()))
-		<< ObjectLifetimeEnd2.Id(GetObjectId(InObject));
+		<< ObjectLifetimeEnd2.RecordingTime(FObjectTrace::GetWorldElapsedTime(InOuterObject->GetWorld()))
+		<< ObjectLifetimeEnd2.Id(InInstanceId);
+}
+
+void FObjectTrace::OutputObjectLifetimeEnd(const UObject* InObject)
+{
+	TRACE_OBJECT(InObject);
+	OutputInstanceLifetimeEnd(InObject, GetObjectId(InObject));
 }
 
 void FObjectTrace::OutputPawnPossess(const UObject* InController, const UObject* InPawn)
diff --git a/Engine/Source/Runtime/Engine/Public/ObjectTrace.h b/Engine/Source/Runtime/Engine/Public/ObjectTrace.h
--- a/Engine/Source/Runtime/Engine/Public/ObjectTrace.h
+++ b/Engine/Source/Runtime/Engine/Public/ObjectTrace.h
@@ -61,12 +61,18 @@ struct FObjectTrace
 
 	/** Helper function to output an object */
 	ENGINE_API static void OutputObject(const UObject* InObject);
-
+	
+	/** Helper function to output an "object instance" which is not an actual UObject */
+	ENGINE_API static void OutputInstance(const UObject* InOuterObject, uint64 InInstanceId, uint64 InOuterId, UStruct* InType, const FString& InName, const FString& InPathName = FString());
+	
 	/** Helper function to output an object transform */
 	ENGINE_API static void OutputObjectTransform(const UObject* InObject, const FTransform& InTransform, ETeleportType InTeleportType);
 
 	/** Helper function to output object creation event */
 	ENGINE_API static void OutputObjectLifetimeBegin(const UObject* InObject);
+	
+	/** Helper function to output object destruction event */
+	ENGINE_API static void OutputInstanceLifetimeEnd(const UObject* InOuterObject, uint64 InInstanceId);
 
 	/** Helper function to output object destruction event */
 	ENGINE_API static void OutputObjectLifetimeEnd(const UObject* InObject);
@@ -80,6 +86,9 @@ struct FObjectTrace
 	/** Helper function to output controller attach event */
 	ENGINE_API static void OutputPawnPossess(const UObject* InController, const UObject* InPawn);
 
+	/** Helper function to allocate a new object Id */
+	ENGINE_API static uint64 AllocateInstanceId();
+	
 	/** Helper function to get an object ID from a UObject */
 	ENGINE_API static uint64 GetObjectId(const UObject* InObject);
 
@@ -127,6 +136,9 @@ struct FObjectTrace
 #define TRACE_OBJECT(Object) \
 	FObjectTrace::OutputObject(Object);
 
+#define TRACE_INSTANCE(OuterObject, InstanceId, OuterId, Type, Name) \
+	FObjectTrace::OutputInstance(OuterObject, InstanceId, OuterId, Type, Name);
+
 #define TRACE_VIEW(Player, View) \
 	FObjectTrace::OutputView(Player, View);
 
@@ -144,6 +156,9 @@ struct FObjectTrace
 #define TRACE_OBJECT_LIFETIME_END(Object) \
 	if (CAN_TRACE_OBJECT(Object)) { FObjectTrace::OutputObjectLifetimeEnd(Object); }
 
+#define TRACE_INSTANCE_LIFETIME_END(OuterObject, InstanceId) \
+	if (CAN_TRACE_OBJECT(OuterObject)) { FObjectTrace::OutputInstanceLifetimeEnd(OuterObject, InstanceId); }
+
 #define TRACE_PAWN_POSSESS(Controller, Pawn)\
 	if (CAN_TRACE_OBJECT(Controller) && (Pawn==nullptr || (CAN_TRACE_OBJECT(Pawn)))) { FObjectTrace::OutputPawnPossess(Controller, Pawn); }
 
@@ -224,6 +239,7 @@ struct FObjectTrace
 
 #define TRACE_CLASS(...)
 #define TRACE_OBJECT(...)
+#define TRACE_INSTANCE(...)
 #define TRACE_OBJECT_TRANSFORM(...)
 #define TRACE_OBJECT_EVENT(...)
 #define TRACE_WORLD(...)
@@ -231,6 +247,7 @@ struct FObjectTrace
 #define TRACE_VIEW(...)
 #define TRACE_OBJECT_LIFETIME_BEGIN(...)
 #define TRACE_OBJECT_LIFETIME_END(...)
+#define TRACE_INSTANCE_LIFETIME_END(...)
 #define TRACE_CHILD_ELEMENT(...)
 #define TRACE_CHILD_ELEMENT_WITH_OUTER(...)
 #define TRACE_CHILD_ELEMENT_LIFETIME_BEGIN(...)
