[DSTM] Fix crash when an actor is spawned inside a transaction.

A recent change made a few hours ago (43775980) in the object trace system makes a operation atomic, which is not supported in transactions.

As the atomic operation is done just to increment a counter and we don't need to roll it back on transaction aborts, this CL make the atomic operation in the open.

diff --git a/Engine/Source/Runtime/Engine/Private/ObjectTrace.cpp b/Engine/Source/Runtime/Engine/Private/ObjectTrace.cpp
--- a/Engine/Source/Runtime/Engine/Private/ObjectTrace.cpp
+++ b/Engine/Source/Runtime/Engine/Private/ObjectTrace.cpp
@@ -274,9 +274,15 @@ void FObjectTrace::Reset()
 uint64 FObjectTrace::AllocateInstanceId()
 {
 	static int64 volatile CurrentId = 1;
-	return FPlatformAtomics::InterlockedIncrement(&CurrentId);
-}
 
+	int64 NewID = 0;
+	UE_AUTORTFM_OPEN
+	{
+		NewID = FPlatformAtomics::InterlockedIncrement(&CurrentId);
+	}
+
+	return NewID;
+}
 
 uint64 FObjectTrace::GetObjectId(const UObject* InObject)
 {
