Fix crash when calling DataTable RemoveRow from editor scripting. This fixes a regression introduced in 5.6 from CL 40129784.

diff --git a/Engine/Source/Runtime/Engine/Private/DataTable.cpp b/Engine/Source/Runtime/Engine/Private/DataTable.cpp
--- a/Engine/Source/Runtime/Engine/Private/DataTable.cpp
+++ b/Engine/Source/Runtime/Engine/Private/DataTable.cpp
@@ -243,7 +243,7 @@ void UDataTable::HandleDataTableChanged(FName ChangedRowName)
 			//Avoid iterating on the full table when a specific row was modified. 
 			if (ChangedRowName != NAME_None)
 			{
-				FTableRowBase* CurRow = reinterpret_cast<FTableRowBase*>(RowMap[ChangedRowName]);
+				FTableRowBase* CurRow = reinterpret_cast<FTableRowBase*>(FindRowUnchecked(ChangedRowName));
 				
 				if (CurRow)
 				{
