[DSTM] Fix CIS in FObjectTrace::AllocateInstanceId casued by my previous CL.

diff --git a/Engine/Source/Runtime/Engine/Private/ObjectTrace.cpp b/Engine/Source/Runtime/Engine/Private/ObjectTrace.cpp
--- a/Engine/Source/Runtime/Engine/Private/ObjectTrace.cpp
+++ b/Engine/Source/Runtime/Engine/Private/ObjectTrace.cpp
@@ -279,7 +279,7 @@ uint64 FObjectTrace::AllocateInstanceId()
 	UE_AUTORTFM_OPEN
 	{
 		NewID = FPlatformAtomics::InterlockedIncrement(&CurrentId);
-	}
+	};
 
 	return NewID;
 }
