Skeletal Mesh: Utility function to get all skin weight profile names

#rnx

diff --git a/Engine/Source/Runtime/Engine/Classes/Engine/SkeletalMesh.h b/Engine/Source/Runtime/Engine/Classes/Engine/SkeletalMesh.h
--- a/Engine/Source/Runtime/Engine/Classes/Engine/SkeletalMesh.h
+++ b/Engine/Source/Runtime/Engine/Classes/Engine/SkeletalMesh.h
@@ -2984,6 +2984,13 @@ public:
 	/** Releases all allocated Skin Weight Profile resources, assumes none are currently in use */
 	ENGINE_API void ReleaseSkinWeightProfileResources();
 
+	/**
+	 *	Returns the list of the names of all the skin weight profile on this skeletal mesh
+	 *  @return	The list of skin weight profiles.
+	 */
+	UFUNCTION(BlueprintPure, Category = Mesh, meta = (DisplayName = "Get All Skin Weight Profile Names", ScriptName = "GetAllSkinWeightProfileNames", Keywords = "skin weight profile"))
+	ENGINE_API TArray<FString> K2_GetAllSkinWeightProfileNames() const;
+
 protected:
 	/** Set of skin weight profiles associated with this mesh */
 	UE_DEPRECATED(5.0, "This must be protected for async build, always use the accessors even internally.")
@@ -3002,7 +3009,7 @@ protected:
 	/** Default translucent material to blend on top of this mesh. Mesh will be rendered twice - once with a base material and once with overlay material */
 	UE_DEPRECATED(5.3, "This must be protected for async build, always use the accessors even internally.")
 	UPROPERTY(EditAnywhere, BlueprintReadOnly, Category=Rendering)
-	TObjectPtr<class UMaterialInterface> OverlayMaterial;
+	TObjectPtr<UMaterialInterface> OverlayMaterial;
 
 	/** Default max draw distance for overlay material. A distance of 0 indicates that overlay will be culled using primitive max distance. */
 	UE_DEPRECATED(5.3, "This must be protected for async build, always use the accessors even internally.")
diff --git a/Engine/Source/Runtime/Engine/Private/SkeletalMesh.cpp b/Engine/Source/Runtime/Engine/Private/SkeletalMesh.cpp
--- a/Engine/Source/Runtime/Engine/Private/SkeletalMesh.cpp
+++ b/Engine/Source/Runtime/Engine/Private/SkeletalMesh.cpp
@@ -6044,6 +6044,16 @@ void USkeletalMesh::ReleaseSkinWeightProfileResources()
 	}
 }
 
+TArray<FString> USkeletalMesh::K2_GetAllSkinWeightProfileNames() const
+{
+	TArray<FString> Names;
+	for (const FSkinWeightProfileInfo& SkinWeightProfile : GetSkinWeightProfiles())
+	{
+		Names.Add(SkinWeightProfile.Name.ToString());
+	}
+	return Names;
+}
+
 FSkeletalMeshLODInfo& USkeletalMesh::AddLODInfo()
 {
 	const int32 NewIndex = AddLODInfo({});
