Temporarily forcing all BPs to load iuncooked in hybrid cooked editor (won't affect any other mode)

diff --git a/Engine/Source/Runtime/Engine/Private/Blueprint.cpp b/Engine/Source/Runtime/Engine/Private/Blueprint.cpp
--- a/Engine/Source/Runtime/Engine/Private/Blueprint.cpp
+++ b/Engine/Source/Runtime/Engine/Private/Blueprint.cpp
@@ -500,7 +500,8 @@ void UBlueprint::Serialize(FArchive& Ar)
 	}
 
 	// cooked BP macros don't work well in cooked editor mixing with uncooked BPs, so we load pacakges with macros as uncooked
-	if (!HasAnyFlags(RF_ClassDefaultObject) && Ar.IsPersistent() && BlueprintType == EBlueprintType::BPTYPE_MacroLibrary)
+	// @todo: temporarily marking all BPs to be uncooked, must check in to get a build overnight
+	if (!HasAnyFlags(RF_ClassDefaultObject) && Ar.IsPersistent())// && BlueprintType == EBlueprintType::BPTYPE_MacroLibrary)
 	{
 		GetOutermost()->ThisShouldLoadUncooked(Ar);
 	}
