Fix Subsurface profile missing with Substrate when using conversion-less material.

diff --git a/Engine/Source/Runtime/Engine/Private/Materials/MaterialInterface.cpp b/Engine/Source/Runtime/Engine/Private/Materials/MaterialInterface.cpp
--- a/Engine/Source/Runtime/Engine/Private/Materials/MaterialInterface.cpp
+++ b/Engine/Source/Runtime/Engine/Private/Materials/MaterialInterface.cpp
@@ -1788,6 +1788,17 @@ void UMaterialInterface::UpdateMaterialRenderProxy(FMaterialRenderProxy& Proxy)
 					Entry.Settings = Entry.Profile->Settings;
 				}
 			}
+			// For legacy materials not converted to Substrate, the SSS profiles array is not populated. 
+			// SSS is then retrieve from the legacy path
+			if (Entries.IsEmpty() && GetSubsurfaceProfile_Internal() != nullptr)
+			{
+				FEntry& Entry = Entries.AddDefaulted_GetRef();
+				Entry.Profile = GetSubsurfaceProfile_Internal();
+				if (Entry.Profile)
+				{
+					Entry.Settings = Entry.Profile->Settings;
+				}
+			}
 
 			USubsurfaceProfile* SubsurfaceProfileOverride = GetSubsurfaceProfileOverride_Internal();
 			FSubsurfaceProfileStruct SubsurfaceProfileOverrideSettings;
