Iris
- Second attempt of minimal update for ReplicationSystem when we have no connections.
- Can be disabled with net.iris.UpdateReplicationSystemWithNoConnections 0

diff --git a/Engine/Source/Runtime/Engine/Private/NetDriver.cpp b/Engine/Source/Runtime/Engine/Private/NetDriver.cpp
--- a/Engine/Source/Runtime/Engine/Private/NetDriver.cpp
+++ b/Engine/Source/Runtime/Engine/Private/NetDriver.cpp
@@ -336,6 +336,12 @@ namespace UE::Net::Private
 		SendDormantDestructionOnRemoval,
 		TEXT("If true and net.SerializeNewActorOverrideLevel = 0, servers will send destruction info for dormant destroyed actors in invisible levels, regardless of how they were removed. This ensures they destroy on the client, in-case the they (or the levels they are in) were removed unexpectedly."));
 
+	bool bUpdateReplicationSystemWithNoConnections = true;
+	static FAutoConsoleVariableRef CVarNetUpdateReplicationSystemWithNoConnections(
+		TEXT("net.iris.UpdateReplicationSystemWithNoConnections"),
+		bUpdateReplicationSystemWithNoConnections,
+		TEXT("If true ReplicationSystem::NetUpdate will be called even if we have no connections to ensure that we do not accumulate pending work."));
+
 	static void DumpActiveNetActors(UWorld* InWorld)
 	{
 		UNetDriver *NetDriver = InWorld->NetDriver;
@@ -1105,7 +1111,7 @@ void UNetDriver::TickFlush(float DeltaSeconds)
 	}
 	FSimpleScopeSecondsCounter ScopedTimer(GTickFlushGameDriverTimeSeconds, bEnableTimer);
 
-	if (IsServer() && ClientConnections.Num() > 0 && !bSkipServerReplicateActors)
+	if (IsServer() && (ClientConnections.Num() > 0 || UE::Net::Private::bUpdateReplicationSystemWithNoConnections) && !bSkipServerReplicateActors)
 	{
 		// Update all clients.
 #if WITH_SERVER_CODE
@@ -1588,7 +1594,7 @@ void UNetDriver::PostTickFlush()
 	{
 		TGuardValue<bool> GuardInNetTick(bInTick, true);
 #if UE_WITH_IRIS
-		if (ReplicationSystem && (ServerConnection != nullptr || ClientConnections.Num() > 0))
+		if (ReplicationSystem && (ServerConnection != nullptr || (ClientConnections.Num() > 0) || UE::Net::Private::bUpdateReplicationSystemWithNoConnections))
 		{
 			ReplicationSystem->PostSendUpdate();
 		}
