New material translator:
- Implement better support for Matrices and the operator MatrixMultiply in the new translator.
- Some refactor/cleanup.

#rnx

diff --git a/Engine/Source/Runtime/Engine/Private/Materials/MaterialExpressionsIR.cpp b/Engine/Source/Runtime/Engine/Private/Materials/MaterialExpressionsIR.cpp
--- a/Engine/Source/Runtime/Engine/Private/Materials/MaterialExpressionsIR.cpp
+++ b/Engine/Source/Runtime/Engine/Private/Materials/MaterialExpressionsIR.cpp
@@ -1236,7 +1236,7 @@ void UMaterialExpressionTime::Build(MIR::FEmitter& Em)
 	if (!bOverride_Period)
 	{
 		FStringView InlinedCode = bIgnorePause ? TEXTVIEW("View.<PREV>RealTime") : TEXTVIEW("View.<PREV>GameTime");
-		Em.Output(0, Em.InlineHLSL(ScalarFloatType, InlinedCode, {}, MIR::EValueFlags::SubstituteTagsInInlineHLSL));
+		Em.Output(0, Em.InlineHLSL(ScalarFloatType, InlinedCode, {}, MIR::EValueFlags::SubstituteTags));
 	}
 	else if (Period == 0.0f)
 	{
@@ -1248,7 +1248,7 @@ void UMaterialExpressionTime::Build(MIR::FEmitter& Em)
 		// We want full 32 bit float precision until the fmod when using a period.
 		FValueRef PeriodValue = Em.ConstantFloat(Period);
 		FStringView InlinedCode = bIgnorePause ? TEXTVIEW("fmod(View.<PREV>RealTime, $0)") : TEXTVIEW("fmod(View.<PREV>GameTime, $0)");
-		Em.Output(0, Em.InlineHLSL(ScalarFloatType, InlinedCode, { PeriodValue }, MIR::EValueFlags::SubstituteTagsInInlineHLSL));
+		Em.Output(0, Em.InlineHLSL(ScalarFloatType, InlinedCode, { PeriodValue }, MIR::EValueFlags::SubstituteTags));
 	}
 }
 
@@ -1564,7 +1564,7 @@ static FValueRef BuildViewProperty(MIR::FEmitter& Em, EMaterialExposedViewProper
 	FStringView HLSLCode = bInvProperty && bHasCustomInverseCode ? PropertyMeta.InvPropertyCode : PropertyMeta.PropertyCode;
 	MIR::FType HLSLCodeType = MIR::FType::FromMaterialValueType(PropertyMeta.Type);
 
-	FValueRef Result = Em.InlineHLSL(HLSLCodeType, HLSLCode, {}, MIR::EValueFlags::SubstituteTagsInInlineHLSL);
+	FValueRef Result = Em.InlineHLSL(HLSLCodeType, HLSLCode, {}, MIR::EValueFlags::SubstituteTags);
 
 	// CastToNonLWCIfDisabled
 	TOptional<MIR::FPrimitive> PrimitiveType = HLSLCodeType.AsPrimitive();
diff --git a/Engine/Source/Runtime/Engine/Private/Materials/MaterialIR.cpp b/Engine/Source/Runtime/Engine/Private/Materials/MaterialIR.cpp
--- a/Engine/Source/Runtime/Engine/Private/Materials/MaterialIR.cpp
+++ b/Engine/Source/Runtime/Engine/Private/Materials/MaterialIR.cpp
@@ -741,6 +741,7 @@ const TCHAR* LexToString(EOperator Op)
 		case BO_Min: return TEXT("Min");
 		case BO_Modulo: return TEXT("Modulo");
 		case BO_Multiply: return TEXT("Multiply");
+		case BO_MatrixMultiply: return TEXT("MatrixMultiply");
 		case BO_NotEquals: return TEXT("NotEquals");
 		case BO_Or: return TEXT("Or");
 		case BO_Pow: return TEXT("Pow");
diff --git a/Engine/Source/Runtime/Engine/Private/Materials/MaterialIREmitter.cpp b/Engine/Source/Runtime/Engine/Private/Materials/MaterialIREmitter.cpp
--- a/Engine/Source/Runtime/Engine/Private/Materials/MaterialIREmitter.cpp
+++ b/Engine/Source/Runtime/Engine/Private/Materials/MaterialIREmitter.cpp
@@ -179,13 +179,6 @@ static FName SlowFindInputName(UMaterialExpression* Expression, const FExpressio
 
 /*------------------------------------- FValueRef --------------------------------------*/
 
-static FValueRef With(FValueRef ValueRef, FValue* Value)
-{
-	FValueRef Copy = ValueRef;
-	Copy.Value = Value;
-	return Copy;
-}
-
 static inline bool IsAnyNotValid()
 {
 	return false;
@@ -431,6 +424,16 @@ FValueRef FEmitter::CheckIsScalar(FValueRef Value)
 	return Value;
 }
 
+FValueRef FEmitter::CheckIsVector(FValueRef Value)
+{
+	if (Value.IsValid() && !Value->Type.IsVector())
+	{
+		Errorf(Value, TEXT("Expected a vector value, got a '%s' instead."), *Value->Type.GetSpelling());
+		return Value.ToPoison();
+	}
+	return Value;
+}
+
 FValueRef FEmitter::CheckIsScalarOrVector(FValueRef Value)
 {
 	if (Value.IsValid() && (!Value->Type.AsPrimitive() || Value->Type.GetPrimitive().IsMatrix()))
@@ -441,6 +444,16 @@ FValueRef FEmitter::CheckIsScalarOrVector(FValueRef Value)
 	return Value;
 }
 
+FValueRef FEmitter::CheckIsMatrix(FValueRef Value)
+{
+	if (Value.IsValid() && (!Value->Type.AsPrimitive() || !Value->Type.GetPrimitive().IsMatrix()))
+	{
+		Errorf(Value, TEXT("Expected a matrix value, got a '%s' instead."), *Value->Type.GetSpelling());
+		return Value.ToPoison();
+	}
+	return Value;
+}
+
 FValueRef FEmitter::CheckIsTexture(FValueRef Value)
 {
 	if (Value.IsValid() && !Value->Type.IsTexture2D())
@@ -951,10 +964,11 @@ FValueRef FEmitter::Vector2(FValueRef InX, FValueRef InY)
 		return Poison();
 	}
 
+	// TODO: replace these checks with casts instead
 	check(InX->Type.IsScalar());
 	check(InX->Type == InY->Type);
 
-	TComposite<2> Vector = MakePrototype<TComposite<2>>(FType::MakeVector(InX->Type.AsPrimitive()->ScalarKind, 2));
+	TComposite<2> Vector = MakePrototype<TComposite<2>>(FType::MakeVector(InX->Type.GetPrimitive().ScalarKind, 2));
 
 	TArrayView<FValue*> Components = Vector.GetMutableComponents();
 	Components[0] = InX;
@@ -971,18 +985,47 @@ FValueRef FEmitter::Vector2(FValueRef InX, FValueRef InY)
 	return Value;
 }
 
+FValueRef FEmitter::Vector3(FValueRef InXY, FValueRef InZ)
+{
+	if (IsAnyNotValid(InXY, InZ))
+	{
+		return Poison();
+	}
+
+	// TODO: replace these checks with casts instead
+	check(InXY->Type.IsVector() && InXY->Type.GetPrimitive().NumComponents() == 2);
+	check(InZ->Type.IsPrimitive() && InZ->Type.GetPrimitive().ScalarKind == InXY->Type.GetPrimitive().ScalarKind);
+
+	TComposite<3> Vector = MakePrototype<TComposite<3>>(FType::MakeVector(InXY->Type.GetPrimitive().ScalarKind, 3));
+	TArrayView<FValue*> Components = Vector.GetMutableComponents();
+	Components[0] = Subscript(InXY, 0);
+	Components[1] = Subscript(InXY, 1);
+	Components[2] = InZ;
+
+	FValueRef Value = EmitPrototype(*this, Vector);
+
+	// If all components source inputs match, refer to the same input in the vector value too (for debugging purposes).
+	if (InXY.Input == InZ.Input )
+	{
+		Value.Input = InXY.Input;
+	}
+
+	return Value;
+}
+
 FValueRef FEmitter::Vector3(FValueRef InX, FValueRef InY, FValueRef InZ)
 {
 	if (IsAnyNotValid(InX, InY, InZ))
 	{
 		return Poison();
 	}
-
+	
+	// TODO: replace these checks with casts instead
 	check(InX->Type.IsScalar());
 	check(InX->Type == InY->Type);
 	check(InY->Type == InZ->Type);
 
-	TComposite<3> Vector = MakePrototype<TComposite<3>>(FType::MakeVector(InX->Type.AsPrimitive()->ScalarKind, 3));
+	TComposite<3> Vector = MakePrototype<TComposite<3>>(FType::MakeVector(InX->Type.GetPrimitive().ScalarKind, 3));
 	TArrayView<FValue*> Components = Vector.GetMutableComponents();
 	Components[0] = InX;
 	Components[1] = InY;
@@ -999,6 +1042,65 @@ FValueRef FEmitter::Vector3(FValueRef InX, FValueRef InY, FValueRef InZ)
 	return Value;
 }
 
+FValueRef FEmitter::Vector4(FValueRef InXYZ, FValueRef InW)
+{
+	if (IsAnyNotValid(InXYZ, InW))
+	{
+		return Poison();
+	}
+
+	// TODO: replace these checks with casts instead
+	check(InXYZ->Type.IsVector() && InXYZ->Type.GetPrimitive().NumComponents() == 3);
+	check(InW->Type.IsPrimitive() && InW->Type.GetPrimitive().ScalarKind == InXYZ->Type.GetPrimitive().ScalarKind);
+
+	TComposite<4> Vector = MakePrototype<TComposite<4>>(FType::MakeVector(InXYZ->Type.GetPrimitive().ScalarKind, 4));
+	TArrayView<FValue*> Components = Vector.GetMutableComponents();
+	Components[0] = Subscript(InXYZ, 0);
+	Components[1] = Subscript(InXYZ, 1);
+	Components[2] = Subscript(InXYZ, 2);
+	Components[3] = InW;
+
+	FValueRef Value = EmitPrototype(*this, Vector);
+
+	// If all components source inputs match, refer to the same input in the vector value too (for debugging purposes).
+	if (InXYZ.Input == InW.Input)
+	{
+		Value.Input = InXYZ.Input;
+	}
+
+	return Value;
+}
+
+FValueRef FEmitter::Vector4(FValueRef InXY, FValueRef InZ, FValueRef InW)
+{
+	if (IsAnyNotValid(InXY, InZ, InW))
+	{
+		return Poison();
+	}
+
+	// TODO: replace these checks with casts instead
+	check(InXY->Type.IsVector() && InXY->Type.GetPrimitive().NumComponents() == 2);
+	check(InZ->Type.IsPrimitive() && InZ->Type.GetPrimitive().ScalarKind == InXY->Type.GetPrimitive().ScalarKind);
+	check(InW->Type.IsPrimitive() && InW->Type.GetPrimitive().ScalarKind == InXY->Type.GetPrimitive().ScalarKind);
+
+	TComposite<4> Vector = MakePrototype<TComposite<4>>(FType::MakeVector(InXY->Type.GetPrimitive().ScalarKind, 4));
+	TArrayView<FValue*> Components = Vector.GetMutableComponents();
+	Components[0] = Subscript(InXY, 0);
+	Components[1] = Subscript(InXY, 1);
+	Components[2] = InZ;
+	Components[3] = InW;
+
+	FValueRef Value = EmitPrototype(*this, Vector);
+
+	// If all components source inputs match, refer to the same input in the vector value too (for debugging purposes).
+	if (InXY.Input == InZ.Input && InXY.Input == InW.Input)
+	{
+		Value.Input = InXY.Input;
+	}
+
+	return Value;
+}
+
 FValueRef FEmitter::Vector4(FValueRef InX, FValueRef InY, FValueRef InZ, FValueRef InW)
 {
 	if (IsAnyNotValid(InX, InY, InZ, InW))
@@ -1006,12 +1108,13 @@ FValueRef FEmitter::Vector4(FValueRef InX, FValueRef InY, FValueRef InZ, FValueR
 		return Poison();
 	}
 
+	// TODO: replace these checks with casts instead
 	check(InX->Type.IsScalar());
 	check(InX->Type == InY->Type);
 	check(InY->Type == InZ->Type);
 	check(InZ->Type == InW->Type);
 
-	TComposite<4> Vector = MakePrototype<TComposite<4>>(FType::MakeVector(InX->Type.AsPrimitive()->ScalarKind, 4));
+	TComposite<4> Vector = MakePrototype<TComposite<4>>(FType::MakeVector(InX->Type.GetPrimitive().ScalarKind, 4));
 	TArrayView<FValue*> Components = Vector.GetMutableComponents();
 	Components[0] = InX;
 	Components[1] = InY;
@@ -1303,7 +1406,7 @@ static T FoldScalarOperator(FEmitter& Emitter, EOperator Operator, T A, T B, T C
 // If it returns null, the end result could not be inferred, but the operator could have
 // still been changed to some other (with lower complexity). For example "clamp(x, 0, 1)"
 // will change to "saturate(x)".
-static FValueRef TrySimplifyOperator(FEmitter& Emitter, EOperator& Op, FValueRef& A, FValueRef& B, FValueRef& C)
+static FValueRef TrySimplifyOperator(FEmitter& Emitter, EOperator& Op, FValueRef& A, FValueRef& B, FValueRef& C, FType ResultType)
 {
 	switch (Op)
 	{
@@ -1371,6 +1474,16 @@ static FValueRef TrySimplifyOperator(FEmitter& Emitter, EOperator& Op, FValueRef
 				return B;
 			}
 			break;
+		
+		case BO_MatrixMultiply:
+			if (ResultType.IsScalar())
+			{
+				Op = BO_Dot;
+
+				// The dot could be simplified further, from dot to multiply if A and B are scalars.
+				return TrySimplifyOperator(Emitter, Op, A, B, C, ResultType);
+			}
+			break;
 
 		case BO_Divide:
 			if (A->AreAllNearlyZero() || B->AreAllNearlyOne())
@@ -1419,6 +1532,19 @@ static FValueRef TrySimplifyOperator(FEmitter& Emitter, EOperator& Op, FValueRef
 				return A;
 			}
 			break;
+		
+		case BO_Dot:
+			if (A->Type.IsScalar())
+			{
+				check(B->Type.IsScalar());
+				Op = BO_Multiply;
+				return TrySimplifyOperator(Emitter, Op, A, B, C, ResultType);
+			}
+			if (A->AreAllNearlyZero() || B->AreAllNearlyZero())
+			{
+				return Emitter.ConstantZero(ResultType.GetPrimitive().ScalarKind);
+			}
+			break;
 
 		case BO_Pow:
 			if (A->AreAllNearlyZero())
@@ -1490,12 +1616,12 @@ static FValueRef TrySimplifyOperator(FEmitter& Emitter, EOperator& Op, FValueRef
 
 // Tries to fold (statically evaluate) the operator, assuming that the arguments are all scalar.
 // It returns either the result of the operator or null if it could not be folded.
-static FValueRef TryFoldOperatorScalar(FEmitter& Emitter, EOperator Op, FValueRef A, FValueRef B, FValueRef C)
+static FValueRef TryFoldOperatorScalar(FEmitter& Emitter, EOperator Op, FValueRef A, FValueRef B, FValueRef C, FType ResultType)
 {
 	TOptional<FPrimitive> PrimitiveType = A->Type.AsPrimitive();
 
 	// Try to simplify the operator. This could potentially change Op, A, B and C.
-	if (FValueRef Simplified = TrySimplifyOperator(Emitter, Op, A, B, C))
+	if (FValueRef Simplified = TrySimplifyOperator(Emitter, Op, A, B, C, ResultType))
 	{
 		return Simplified;
 	}
@@ -1572,26 +1698,28 @@ static FValueRef TryFoldOperatorScalar(FEmitter& Emitter, EOperator Op, FValueRe
 	}
 }
 
-/// Used to filter what parameter *primitive* types operators can take.
+// Used to filter what parameter *primitive* types operators can take.
 enum EOperatorParameterTypeFilter
 {
-	OPTF_Unknown = 0xff,										///< Unspecified
-	OPTF_Any = 0,												///< Any primitive type
-	OPTF_CastToFirstArgumentType = 1 << 10,						///< Cast the argument to the first argument's type
-	OPTF_CastToAnyFloat = 1 << 9,								///< Cast the argument to the floating point primitive type of any dimension
-	OPTF_CheckIsBoolean = 1 << 0,								///< Check the type is boolean primitive of any dimension
-	OPTF_CheckIsInteger = 1 << 1,								///< Check the type is integer primitive of any dimension
-	OPTF_CheckIsArithmetic = 1 << 2,							///< Check the type is arithmetic primitive of any dimension (i.e. that supports arithmetic operations)
-	OPTF_CheckIsNotMatrix = 1 << 3,								///< Check the type is any primitive type except matrices
-	OPTF_CheckIsVector3 = 1 << 4,								///< Check the type is a 3D vector of any scalar type
-	OPTF_CheckIsNonNegativeFloatConst = 1 << 5,					///< Check that if the argument is a constant float, it is not negative (x >= 0)
-	OPTF_CheckIsNonZeroFloatConst = 1 << 6,						///< Check that if the argument is a constant float, it is not zero	(x != 0)
-	OPTF_CheckIsOneOrGreaterFloatConst = 1 << 7,				///< Check that if the argument is a constant float, it is 1 or greater (xFloat >= 1)
-	OPTF_CheckIsBetweenMinusOneAndPlusOneFloatConst = 1 << 8,	///< Check that if the argument is a constant float, it is between -1 and 1 (-1 < x < 1)
-	OPTF_CastToCommonType = 1 << 11,	 						///< Cast the argument to the common arguments type
-	OPTF_LWCInput = 1 << 12,									///< Operator can accept LWC input for this argument
+	OPTF_Unknown = 0xff,										// Unspecified
+	OPTF_Any = 0,												// Any primitive type
+	OPTF_CheckIsBoolean = 1 << 1,								// Check the type is boolean primitive of any dimension
+	OPTF_CheckIsInteger = 1 << 2,								// Check the type is integer primitive of any dimension
+	OPTF_CheckIsArithmetic = 1 << 3,							// Check the type is arithmetic primitive of any dimension (i.e. that supports arithmetic operations)
+	OPTF_CheckIsMatrix = 1 << 4,								// Check the type is any matrix type
+	OPTF_CheckIsNotMatrix = 1 << 5,								// Check the type is any primitive type except matrices
+	OPTF_CheckIsVector3 = 1 << 6,								// Check the type is a 3D vector of any scalar type
+	OPTF_CheckIsNonNegativeFloatConst = 1 << 7,					// Check that if the argument is a constant float, it is not negative (x >= 0)
+	OPTF_CheckIsNonZeroFloatConst = 1 << 8,						// Check that if the argument is a constant float, it is not zero	(x != 0)
+	OPTF_CheckIsOneOrGreaterFloatConst = 1 << 9,				// Check that if the argument is a constant float, it is 1 or greater (xFloat >= 1)
+	OPTF_CheckIsBetweenMinusOneAndPlusOneFloatConst = 1 << 10,	// Check that if the argument is a constant float, it is between -1 and 1 (-1 < x < 1)
+	OPTF_CastToFirstArgumentType = 1 << 11,						// Cast the argument to the first argument's type
+	OPTF_CastToAnyFloat = 1 << 12,								// Cast the argument to the floating point primitive type of any dimension
+	OPTF_AllowDouble = 1 << 13,									// This argument is allowed to be a double
+	OPTF_CastToCommonScalarKind = 1 << 14,						// Casts the argument to have the scalar kind in common with other arguments
+	OPTF_CastToCommonType = 1 << 15,	 						// Cast the argument to the common arguments type
 	OPTF_CastToCommonArithmeticType = OPTF_CheckIsArithmetic | OPTF_CastToCommonType,
-	OPTF_CastToCommonArithmeticTypeLWC = OPTF_CheckIsArithmetic | OPTF_CastToCommonType | OPTF_LWCInput,
+	OPTF_CastToCommonArithmeticTypeAllowDouble = OPTF_CheckIsArithmetic | OPTF_CastToCommonType | OPTF_AllowDouble,
 	OPTF_CastToCommonFloatType = OPTF_CastToAnyFloat | OPTF_CastToCommonType,
 };
 
@@ -1600,48 +1728,49 @@ inline EOperatorParameterTypeFilter operator|(EOperatorParameterTypeFilter A, EO
 	return EOperatorParameterTypeFilter(uint32(A) | uint32(B));
 }
 
-/// Used to determine the operator result type based on argument types
+// Used to determine the operator result type based on argument types
 enum EOperatorReturnType
 {
-	ORT_Unknown,							///< Unspecified
-	ORT_FirstArgumentType,					///< The same type as the first argument, LWC input produces float result
-	ORT_BooleanWithFirstArgumentDimensions,	///< A boolean primitive type with the same dimensions (rows and columns) as the first argument type
-	ORT_FirstArgumentTypeToScalarLWC,		///< A scalar primitive type with the same kind as the scalar type of the first argument, LWC results allowed
-	ORT_SecondArgumentType,					///< The same type as the second argument
-	ORT_FirstArgumentTypeLWC,				///< The same type as the first argument, LWC results allowed
+	ORT_Unknown,							// Unspecified
+	ORT_FirstArgumentType,					// The same type as the first argument, LWC input produces float result
+	ORT_BooleanWithFirstArgumentDimensions,	// A boolean primitive type with the same dimensions (rows and columns) as the first argument type
+	ORT_FirstArgumentTypeToScalarLWC,		// A scalar primitive type with the same kind as the scalar type of the first argument, LWC results allowed
+	ORT_SecondArgumentType,					// The same type as the second argument
+	ORT_FirstArgumentTypeLWC,				// The same type as the first argument, LWC results allowed
+	ORT_MatrixMultiplyResult,				// The result type of the matrix multiplication of first two arguments
 };
 
-/// The signature of an operator consisting of its parameter and return type information.
+// The signature of an operator consisting of its parameter and return type information.
 struct FOperatorSignature
 {
 	EOperatorParameterTypeFilter ParameterTypes[3] = { OPTF_Unknown, OPTF_Unknown, OPTF_Unknown };
 	EOperatorReturnType ReturnType = ORT_Unknown;
 };
 
-/// Returns the signature of an operator.
+// Returns the signature of an operator.
 static const FOperatorSignature* GetOperatorSignature(EOperator Op)
 {
 	static const FOperatorSignature* Signatures = [] ()
 	{
 		const FOperatorSignature UnaryFloat = { { OPTF_CheckIsArithmetic | OPTF_CastToAnyFloat }, ORT_FirstArgumentType };
-		const FOperatorSignature UnaryFloatLWC = { { OPTF_CheckIsArithmetic | OPTF_CastToAnyFloat | OPTF_LWCInput }, ORT_FirstArgumentTypeLWC };
-		const FOperatorSignature UnaryFloatLWCDemote = { { OPTF_CheckIsArithmetic | OPTF_CastToAnyFloat | OPTF_LWCInput }, ORT_FirstArgumentType };
+		const FOperatorSignature UnaryFloatLWC = { { OPTF_CheckIsArithmetic | OPTF_CastToAnyFloat | OPTF_AllowDouble }, ORT_FirstArgumentTypeLWC };
+		const FOperatorSignature UnaryFloatLWCDemote = { { OPTF_CheckIsArithmetic | OPTF_CastToAnyFloat | OPTF_AllowDouble }, ORT_FirstArgumentType };
 		const FOperatorSignature UnaryFloatToBoolean = { { OPTF_CheckIsArithmetic | OPTF_CastToAnyFloat }, ORT_BooleanWithFirstArgumentDimensions };
 		const FOperatorSignature BinaryArithmetic = { { OPTF_CastToCommonArithmeticType, OPTF_CastToCommonArithmeticType }, ORT_FirstArgumentType };
-		const FOperatorSignature BinaryArithmeticLWC = { { OPTF_CastToCommonArithmeticTypeLWC, OPTF_CastToCommonArithmeticTypeLWC }, ORT_FirstArgumentTypeLWC };
-		const FOperatorSignature BinaryArithmeticLWCDemote = { { OPTF_CastToCommonArithmeticTypeLWC, OPTF_CastToCommonArithmeticTypeLWC }, ORT_FirstArgumentType };
+		const FOperatorSignature BinaryArithmeticLWC = { { OPTF_CastToCommonArithmeticTypeAllowDouble, OPTF_CastToCommonArithmeticTypeAllowDouble }, ORT_FirstArgumentTypeLWC };
+		const FOperatorSignature BinaryArithmeticLWCDemote = { { OPTF_CastToCommonArithmeticTypeAllowDouble, OPTF_CastToCommonArithmeticTypeAllowDouble }, ORT_FirstArgumentType };
 		const FOperatorSignature BinaryInteger = { { OPTF_CheckIsInteger | OPTF_CastToCommonArithmeticType, OPTF_CheckIsInteger | OPTF_CastToCommonArithmeticType }, ORT_FirstArgumentType };
 		const FOperatorSignature BinaryFloat = { { OPTF_CastToCommonFloatType, OPTF_CastToCommonFloatType }, ORT_FirstArgumentType };
-		const FOperatorSignature BinaryArithmeticComparisonLWC = { { OPTF_CastToCommonArithmeticType | OPTF_LWCInput, OPTF_CastToCommonArithmeticType | OPTF_LWCInput }, ORT_BooleanWithFirstArgumentDimensions };
+		const FOperatorSignature BinaryArithmeticComparisonLWC = { { OPTF_CastToCommonArithmeticType | OPTF_AllowDouble, OPTF_CastToCommonArithmeticType | OPTF_AllowDouble }, ORT_BooleanWithFirstArgumentDimensions };
 		const FOperatorSignature BinaryLogical = { { OPTF_CheckIsBoolean | OPTF_CastToCommonType, OPTF_CheckIsBoolean | OPTF_CastToCommonType }, ORT_FirstArgumentType };
-		const FOperatorSignature TernaryArithmeticLWC = { { OPTF_CastToCommonArithmeticTypeLWC, OPTF_CastToCommonArithmeticTypeLWC, OPTF_CastToCommonArithmeticTypeLWC }, ORT_FirstArgumentTypeLWC };
-		const FOperatorSignature TernaryFloatLWCDemote = { { OPTF_CastToCommonArithmeticTypeLWC | OPTF_CastToAnyFloat, OPTF_CastToCommonArithmeticTypeLWC, OPTF_CastToCommonArithmeticTypeLWC }, ORT_FirstArgumentType };
+		const FOperatorSignature TernaryArithmeticDouble = { { OPTF_CastToCommonArithmeticTypeAllowDouble, OPTF_CastToCommonArithmeticTypeAllowDouble, OPTF_CastToCommonArithmeticTypeAllowDouble }, ORT_FirstArgumentTypeLWC };
+		const FOperatorSignature TernaryFloatLWCDemote = { { OPTF_CastToCommonArithmeticTypeAllowDouble | OPTF_CastToAnyFloat, OPTF_CastToCommonArithmeticTypeAllowDouble, OPTF_CastToCommonArithmeticTypeAllowDouble }, ORT_FirstArgumentType };
 
 		static FOperatorSignature S[OperatorCount];
 
 		/* unary operators */
 		S[UO_BitwiseNot] 				= { { OPTF_CheckIsInteger }, ORT_FirstArgumentType };
-		S[UO_Negate] 					= { { OPTF_CheckIsArithmetic | OPTF_LWCInput }, ORT_FirstArgumentTypeLWC };
+		S[UO_Negate] 					= { { OPTF_CheckIsArithmetic | OPTF_AllowDouble }, ORT_FirstArgumentTypeLWC };
 		S[UO_Not] 						= { { OPTF_CheckIsBoolean }, ORT_FirstArgumentType };
 
 		S[UO_Abs] 						= UnaryFloatLWC;
@@ -1663,7 +1792,7 @@ static const FOperatorSignature* GetOperatorSignature(EOperator Op)
 		S[UO_IsFinite]					= UnaryFloatToBoolean;
 		S[UO_IsInf]						= UnaryFloatToBoolean;
 		S[UO_IsNan]						= UnaryFloatToBoolean;
-		S[UO_Length]					= { { OPTF_CheckIsArithmetic | OPTF_CheckIsNotMatrix | OPTF_CastToAnyFloat | OPTF_LWCInput }, ORT_FirstArgumentTypeToScalarLWC };
+		S[UO_Length]					= { { OPTF_CheckIsArithmetic | OPTF_CheckIsNotMatrix | OPTF_CastToAnyFloat | OPTF_AllowDouble }, ORT_FirstArgumentTypeToScalarLWC };
 		S[UO_Logarithm] 				= { { OPTF_CheckIsArithmetic | OPTF_CheckIsNonZeroFloatConst | OPTF_CheckIsNonNegativeFloatConst | OPTF_CastToAnyFloat }, ORT_FirstArgumentType };
 		S[UO_Logarithm10]				= { { OPTF_CheckIsArithmetic | OPTF_CheckIsNonZeroFloatConst | OPTF_CheckIsNonNegativeFloatConst | OPTF_CastToAnyFloat }, ORT_FirstArgumentType };
 		S[UO_Logarithm2] 				= { { OPTF_CheckIsArithmetic | OPTF_CheckIsNonZeroFloatConst | OPTF_CheckIsNonNegativeFloatConst | OPTF_CastToAnyFloat }, ORT_FirstArgumentType };
@@ -1674,24 +1803,25 @@ static const FOperatorSignature* GetOperatorSignature(EOperator Op)
 		S[UO_Saturate]					= UnaryFloatLWCDemote;
 		S[UO_Sign]						= UnaryFloatLWCDemote;
 		S[UO_Sin] 						= UnaryFloatLWCDemote;
-		S[UO_Sqrt]						= { { OPTF_CheckIsArithmetic | OPTF_CheckIsNonNegativeFloatConst | OPTF_CastToAnyFloat | OPTF_LWCInput }, ORT_FirstArgumentType };
+		S[UO_Sqrt]						= { { OPTF_CheckIsArithmetic | OPTF_CheckIsNonNegativeFloatConst | OPTF_CastToAnyFloat | OPTF_AllowDouble }, ORT_FirstArgumentType };
 		S[UO_Tan] 						= UnaryFloatLWCDemote;
 		S[UO_Tanh] 						= UnaryFloat;
 		S[UO_Truncate]					= UnaryFloatLWC;
 
 		/* binary operators */
-		S[BO_Equals]					= { { OPTF_CastToCommonType | OPTF_LWCInput, OPTF_CastToCommonType | OPTF_LWCInput }, ORT_BooleanWithFirstArgumentDimensions };
+		S[BO_Equals]					= { { OPTF_CastToCommonType | OPTF_AllowDouble, OPTF_CastToCommonType | OPTF_AllowDouble }, ORT_BooleanWithFirstArgumentDimensions };
 		S[BO_GreaterThan]				= BinaryArithmeticComparisonLWC;
 		S[BO_GreaterThanOrEquals]		= BinaryArithmeticComparisonLWC;
 		S[BO_LessThan]					= BinaryArithmeticComparisonLWC;
 		S[BO_LessThanOrEquals]			= BinaryArithmeticComparisonLWC;
-		S[BO_NotEquals]					= { { OPTF_CastToCommonType | OPTF_LWCInput, OPTF_CastToCommonType | OPTF_LWCInput }, ORT_BooleanWithFirstArgumentDimensions };
+		S[BO_NotEquals]					= { { OPTF_CastToCommonType | OPTF_AllowDouble, OPTF_CastToCommonType | OPTF_AllowDouble }, ORT_BooleanWithFirstArgumentDimensions };
 		
 		S[BO_And] 						= BinaryLogical;
 		S[BO_Or] 						= BinaryLogical;
 		S[BO_Add] 						= BinaryArithmeticLWC;
 		S[BO_Subtract] 					= BinaryArithmeticLWC;
 		S[BO_Multiply] 					= BinaryArithmeticLWC;
+		S[BO_MatrixMultiply]			= { { OPTF_CheckIsArithmetic | OPTF_CastToCommonScalarKind, OPTF_CheckIsArithmetic | OPTF_CastToCommonScalarKind }, ORT_MatrixMultiplyResult };
 		S[BO_Divide] 					= BinaryArithmeticLWC;
 		S[BO_Modulo] 					= BinaryArithmetic;
 		S[BO_BitwiseAnd]				= BinaryInteger;
@@ -1699,12 +1829,13 @@ static const FOperatorSignature* GetOperatorSignature(EOperator Op)
 		S[BO_BitShiftLeft]				= BinaryInteger;
 		S[BO_BitShiftRight]				= BinaryInteger;
 
+
 		S[BO_ATan2]						= BinaryFloat;
 		S[BO_ATan2Fast]					= BinaryFloat;
 		S[BO_Cross] 					= { { OPTF_CheckIsArithmetic | OPTF_CheckIsVector3, OPTF_CastToFirstArgumentType }, ORT_FirstArgumentType };
-		S[BO_Distance]					= { { OPTF_CastToCommonFloatType | OPTF_LWCInput, OPTF_CastToCommonFloatType | OPTF_LWCInput }, ORT_FirstArgumentTypeToScalarLWC };
-		S[BO_Dot] 						= { { OPTF_CheckIsArithmetic | OPTF_CheckIsNotMatrix | OPTF_LWCInput, OPTF_CastToFirstArgumentType | OPTF_LWCInput }, ORT_FirstArgumentTypeToScalarLWC };
-		S[BO_Fmod] 						= { { OPTF_CastToCommonFloatType | OPTF_LWCInput, OPTF_CastToCommonFloatType }, ORT_FirstArgumentType };		// First input can be LWC, but second and output are always demoted
+		S[BO_Distance]					= { { OPTF_CastToCommonFloatType | OPTF_AllowDouble, OPTF_CastToCommonFloatType | OPTF_AllowDouble }, ORT_FirstArgumentTypeToScalarLWC };
+		S[BO_Dot] 						= { { OPTF_CheckIsArithmetic | OPTF_CheckIsNotMatrix | OPTF_AllowDouble, OPTF_CastToFirstArgumentType | OPTF_AllowDouble }, ORT_FirstArgumentTypeToScalarLWC };
+		S[BO_Fmod] 						= { { OPTF_CastToCommonFloatType | OPTF_AllowDouble, OPTF_CastToCommonFloatType }, ORT_FirstArgumentType };		// First input can be LWC, but second and output are always demoted
 
 		S[BO_Max] 						= BinaryArithmeticLWC;
 		S[BO_Min] 						= BinaryArithmeticLWC;
@@ -1712,17 +1843,17 @@ static const FOperatorSignature* GetOperatorSignature(EOperator Op)
 		S[BO_Step] 						= BinaryArithmeticLWCDemote;
 
 		/* ternary operators -- note lerp doesn't support LWC for the third argument! */
-		S[TO_Clamp]						= TernaryArithmeticLWC;
-		S[TO_Lerp] 						= { { OPTF_CastToCommonArithmeticTypeLWC | OPTF_CastToAnyFloat, OPTF_CastToCommonArithmeticTypeLWC, OPTF_CastToCommonArithmeticType }, ORT_FirstArgumentTypeLWC };;
-		S[TO_Select]					= { { OPTF_CheckIsBoolean | OPTF_CheckIsNotMatrix, OPTF_CheckIsNotMatrix | OPTF_LWCInput, OPTF_CheckIsNotMatrix | OPTF_LWCInput}, ORT_SecondArgumentType }; // Note: this is a special operator, which is handled manually in the Validate function
+		S[TO_Clamp]						= TernaryArithmeticDouble;
+		S[TO_Lerp] 						= { { OPTF_CastToCommonArithmeticTypeAllowDouble | OPTF_CastToAnyFloat, OPTF_CastToCommonArithmeticTypeAllowDouble, OPTF_CastToCommonArithmeticType }, ORT_FirstArgumentTypeLWC };;
+		S[TO_Select]					= { { OPTF_CheckIsBoolean | OPTF_CheckIsNotMatrix, OPTF_CheckIsNotMatrix | OPTF_AllowDouble, OPTF_CheckIsNotMatrix | OPTF_AllowDouble}, ORT_SecondArgumentType }; // Note: this is a special operator, which is handled manually in the Validate function
 		S[TO_Smoothstep]				= TernaryFloatLWCDemote;
 		return S;
 	} ();
 	return &Signatures[Op];
 }
  
-/// Validates that the types of the arguments are valid for specified operator.
-/// If valid, it returns the type of the result. Otherwise if it is not valid, it returns nullptr.
+// Validates that the types of the arguments are valid for specified operator.
+// If valid, it returns the type of the result. Otherwise if it is not valid, it returns nullptr.
 static FType ValidateOperatorAndGetResultType(FEmitter& Emitter, EOperator Op, FValueRef& A, FValueRef& B, FValueRef& C)
 {
 	// Argument A must have always been provided.
@@ -1737,14 +1868,13 @@ static FType ValidateOperatorAndGetResultType(FEmitter& Emitter, EOperator Op, F
 	// Verify that C argument has been provided if operator is ternary.
 	check(!IsTernaryOperator(Op) || C);
 
-	// Make sure the first argument has primitive type first, since the following operations
-	// assume this.
+	// Make sure the first argument has primitive type first, since the following operations assume this.
 	if (!Emitter.CheckIsPrimitive(A).IsValid())
 	{
 		return FType::MakePoison();
 	}
 
-	// Given input float3, generates LWC3 type with the given tile value and zero offset.  Caller must pass in float3.
+	// SPECIAL CASE: Given input float3, generates double3 type with the given tile value and zero offset. Caller must pass in float3.
 	if (Op == UO_LWCTile)
 	{
 		if (!A->Type.GetPrimitive().IsFloat() || !A->Type.GetPrimitive().IsRowVector() || A->Type.GetPrimitive().NumComponents() != 3)
@@ -1755,7 +1885,7 @@ static FType ValidateOperatorAndGetResultType(FEmitter& Emitter, EOperator Op, F
 		return FType::MakeDoubleVector(3);
 	}
 
-	// For Clamp, we do a special case and demote the first LWC argument if the second and third arguments (min / max) are non-LWC.
+	// SPECIAL CASE: For Clamp, we do a special case and demote the first LWC argument if the second and third arguments (min / max) are non-LWC.
 	// We want to do this before fetching FirstArgumentPrimitiveType.
 	if (Op == TO_Clamp && A->Type.GetPrimitive().IsDouble() && !B->Type.IsDouble() && !C->Type.IsDouble())
 	{
@@ -1767,19 +1897,15 @@ static FType ValidateOperatorAndGetResultType(FEmitter& Emitter, EOperator Op, F
 
 	// Handle automatic cast to float for operators that don't support LWC inputs
 	FType FirstArgumentType = A->Type;
-	if (FirstArgumentType.GetPrimitive().IsDouble() && !(Signature->ParameterTypes[0] & OPTF_LWCInput))
+	if (FirstArgumentType.GetPrimitive().IsDouble() && !(Signature->ParameterTypes[0] & OPTF_AllowDouble))
 	{
 		FirstArgumentType = A->Type.GetPrimitive().ToScalarKind(EScalarKind::Float);
 	}
 
-	// This will hold the common type between arguments types.
-	FType ArgumentsCommonType = FirstArgumentType;
-
 	// Verify that the first argument type is primitive.
 	FValueRef Arguments[] = { A, B, C, nullptr };
 	static const TCHAR* ArgumentsStr[] = { TEXT("first"), TEXT("second"), TEXT("third") };
 
-	bool bValid = true;
 	for (int32 i = 0; Arguments[i]; ++i)
 	{
 		// Check this argument type i primitive.
@@ -1789,61 +1915,56 @@ static FType ValidateOperatorAndGetResultType(FEmitter& Emitter, EOperator Op, F
 			return FType::MakePoison();
 		}
 
-		FType ArgType = Arguments[i]->Type;
 		EOperatorParameterTypeFilter Filter = Signature->ParameterTypes[i];
 		check(Filter != OPTF_Unknown); // No signature specified for this operator.
 
 		if (Filter & OPTF_CastToFirstArgumentType)
 		{
 			check(i > 0); // This check can't apply to the first argument.
-
-			// Cast this argument to the first argument primitive type
-			ArgType = FirstArgumentType;
 			Arguments[i] = Emitter.Cast(Arguments[i], FirstArgumentType);
-			bValid &= Arguments[i].IsValid();
 		}
-		else if (ArgType.GetPrimitive().IsDouble() && !(Filter & OPTF_LWCInput))
+		// Cast argument to float when...
+		else if ( 
+			// ...the argument should be cast to any float (and it's not a float already)
+			((Filter & OPTF_CastToAnyFloat) && !Arguments[i]->Type.GetPrimitive().IsAnyFloat())
+			// ...or the argument is not allowed to be a double and it is.
+			|| (!(Filter & OPTF_AllowDouble) && Arguments[i]->Type.GetPrimitive().IsDouble()))
 		{
-			// Handle automatic cast to float for operators that don't support LWC inputs
-			ArgType = ArgType.GetPrimitive().ToScalarKind(EScalarKind::Float);
-			Arguments[i] = Emitter.Cast(Arguments[i], ArgType);
-		}
-		else if (Filter & OPTF_CastToAnyFloat)
-		{
-			if (!ArgType.GetPrimitive().IsAnyFloat())
-			{
-				ArgType = ArgType.GetPrimitive().ToScalarKind(EScalarKind::Float);
-				Arguments[i] = Emitter.Cast(Arguments[i], ArgType);
-				bValid &= Arguments[i].IsValid();
-			}
+			Arguments[i] = Emitter.CastToFloatKind(Arguments[i]);
+			check(!Arguments[i].IsPoison());
 		}
 		
 		if (Filter & OPTF_CheckIsBoolean)
 		{
-			bValid &= Emitter.CheckIsBoolean(Arguments[i]).IsValid(); 
+			Arguments[i] = Emitter.CheckIsBoolean(Arguments[i]); 
 		}
-		
+			
+		if (Filter & OPTF_CheckIsArithmetic)
+		{
+			Arguments[i] = Emitter.CheckIsArithmetic(Arguments[i]); 
+		}
+
 		if (Filter & OPTF_CheckIsInteger)
 		{
-			bValid &= Emitter.CheckIsInteger(Arguments[i]).IsValid(); 
+			Arguments[i] = Emitter.CheckIsInteger(Arguments[i]); 
 		}
 		
-		if (Filter & OPTF_CheckIsArithmetic)
+		if (Filter & OPTF_CheckIsMatrix)
 		{
-			bValid &= Emitter.CheckIsArithmetic(Arguments[i]).IsValid(); 
+			Arguments[i] = Emitter.CheckIsMatrix(Arguments[i]); 
 		}
 
 		if (Filter & OPTF_CheckIsNotMatrix)
 		{
-			bValid &= Emitter.CheckIsScalarOrVector(Arguments[i]).IsValid(); 
+			Arguments[i] = Emitter.CheckIsScalarOrVector(Arguments[i]); 
 		}
 
 		if (Filter & OPTF_CheckIsVector3)
 		{
-			if (!ArgType.GetPrimitive().IsRowVector() || ArgType.GetPrimitive().NumComponents() != 3)
+			if (!Arguments[i]->Type.GetPrimitive().IsRowVector() || Arguments[i]->Type.GetPrimitive().NumComponents() != 3)
 			{
 				Emitter.Errorf(Arguments[i], TEXT("Expected a 3D vector."));
-				bValid = false;
+				Arguments[i] = Arguments[i].ToPoison();
 			}
 		}
 
@@ -1856,7 +1977,7 @@ static FType ValidateOperatorAndGetResultType(FEmitter& Emitter, EOperator Op, F
 				if (Constant->Float == 0)
 				{
 					Emitter.Errorf(Arguments[i], TEXT("Expected non-zero value."));
-					bValid = false;
+					Arguments[i] = Arguments[i].ToPoison();
 				}
 			}
 
@@ -1866,7 +1987,7 @@ static FType ValidateOperatorAndGetResultType(FEmitter& Emitter, EOperator Op, F
 				if (Constant->Float < 0)
 				{
 					Emitter.Errorf(Arguments[i], TEXT("Expected non-negative value."));
-					bValid = false;
+					Arguments[i] = Arguments[i].ToPoison();
 				}
 			}
 
@@ -1876,7 +1997,7 @@ static FType ValidateOperatorAndGetResultType(FEmitter& Emitter, EOperator Op, F
 				if (Constant->Float < 1)
 				{
 					Emitter.Errorf(Arguments[i], TEXT("Expected a value equal or greater than 1."));
-					bValid = false;
+					Arguments[i] = Arguments[i].ToPoison();
 				}
 			}
 
@@ -1886,19 +2007,22 @@ static FType ValidateOperatorAndGetResultType(FEmitter& Emitter, EOperator Op, F
 				if (Constant->Float < -1 || Constant->Float > 1)
 				{
 					Emitter.Errorf(Arguments[i], TEXT("Expected a value greater than -1 and lower than 1."));
-					bValid = false;
+					Arguments[i] = Arguments[i].ToPoison();
 				}
 			}
 		}
-		
-		// Update the common type.
-		if (i >= 1)
-		{
-			ArgumentsCommonType = Emitter.TryGetCommonType(ArgumentsCommonType, ArgType);
-		}
+	}
+
+	if (Arguments[0].IsPoison() || Arguments[1].IsPoison() || Arguments[2].IsPoison())
+	{
+		return FType::MakePoison();
 	}
 	
-	// The select operator is special insofar as its first argument is a boolean, while the second and third can be any primitive type.
+	// Whether any argument will require to be cast to common type/scalar kind.
+	bool bRequiresArgumentsCommonType = (Signature->ParameterTypes[0] & OPTF_CastToCommonType) || (Signature->ParameterTypes[1] & OPTF_CastToCommonType) || (Signature->ParameterTypes[2] & OPTF_CastToCommonType);
+	bool bRequiresArgumentsCommonScalarKind = (Signature->ParameterTypes[0] & OPTF_CastToCommonScalarKind) || (Signature->ParameterTypes[1] & OPTF_CastToCommonScalarKind) || (Signature->ParameterTypes[2] & OPTF_CastToCommonScalarKind);
+
+	// SPECIAL CASE:  The select operator is special insofar as its first argument is a boolean, while the second and third can be any primitive type.
 	if (Op == TO_Select)
 	{
 		// Cast the second and third argument types to primitive. This is safe as it was already checked earlier.
@@ -1922,56 +2046,65 @@ static FType ValidateOperatorAndGetResultType(FEmitter& Emitter, EOperator Op, F
 		// Cast second and third arguments to their common type.
 		Arguments[1] = Emitter.Cast(Arguments[1], CommonTypeBetweenSecondAndThirdArguments);
 		Arguments[2] = Emitter.Cast(Arguments[2], CommonTypeBetweenSecondAndThirdArguments);
-
-		// Update the valid flag to the results of the conditions above.
-		bValid &= !Arguments[0]->IsPoison() && !Arguments[1]->IsPoison() && !Arguments[2]->IsPoison();
 	}
-	else
+	else if (bRequiresArgumentsCommonType || bRequiresArgumentsCommonScalarKind)
 	{
-		// Cast every argument with the `CastToCommon` to the common type, if necessary.
+		FType ArgumentsCommonType = Arguments[0]->Type;
+		EScalarKind ArgumentsCommonScalarKind = EScalarKind::Bool;
+
+		// Determine the common type and scalar kind (if needeed)
+		// Note: these two cannot be unified, because we can always determine the common (biggest) scalar kind between two
+		// primitive types, but not always can be determined a common type (e.g. a float3 with a float4x4)
 		for (int32 i = 0; Arguments[i]; ++i)
 		{
-			EOperatorParameterTypeFilter Filter = Signature->ParameterTypes[i];
-			if (Filter & OPTF_CastToCommonType)
+			if (bRequiresArgumentsCommonType)
 			{
-				check(ArgumentsCommonType && ArgumentsCommonType.AsPrimitive());
-
-				// Lerp doesn't accept LWC for its third input, so we need to check if LWC is allowed per input when casting to the common type.
-				if (ArgumentsCommonType.GetPrimitive().IsDouble() && !(Filter & OPTF_LWCInput))
+				ArgumentsCommonType = Emitter.GetCommonType(ArgumentsCommonType, Arguments[i]->Type);
+				if (ArgumentsCommonType.IsPoison())
 				{
-					Arguments[i] = Emitter.Cast(Arguments[i], ArgumentsCommonType.GetPrimitive().ToScalarKind(EScalarKind::Float));
-				}
-				else
-				{
-					Arguments[i] = Emitter.Cast(Arguments[i], ArgumentsCommonType);
+					return FType::MakePoison();
 				}
+			}
+			if (bRequiresArgumentsCommonScalarKind)
+			{
+				ArgumentsCommonScalarKind = (EScalarKind)FMath::Max((int)ArgumentsCommonScalarKind, (int)Arguments[i]->Type.GetPrimitive().ScalarKind);
+			}
+		}
 
-				bValid &= !Arguments[i]->IsPoison();
+		// Cast every argument with the `CastToCommon` to the common type, if necessary.
+		for (int32 i = 0; Arguments[i]; ++i)
+		{
+			EOperatorParameterTypeFilter Filter = Signature->ParameterTypes[i];
+			if (Filter & OPTF_CastToCommonScalarKind)
+			{
+				Arguments[i] = Emitter.CastToScalarKind(Arguments[i], ArgumentsCommonScalarKind);
 			}
-			// Fmod doesn't accept LWC for its second input, so we need to potentially cast that to non-LWC.
-			else if (Arguments[i]->Type.GetPrimitive().IsDouble() && !(Filter & OPTF_LWCInput))
+			else if (Filter & OPTF_CastToCommonType)
 			{
-				Arguments[i] = Emitter.CastToFloatKind(Arguments[i]);
+				check(bRequiresArgumentsCommonType);
+
+				// Lerp doesn't accept LWC for its third input, so we need to check if LWC is allowed per input when casting to the common type.
+				FType ToType = (!(Filter & OPTF_AllowDouble) && ArgumentsCommonType.IsDouble()) 
+					? ArgumentsCommonType.GetPrimitive().ToScalarKind(EScalarKind::Float)
+					: ArgumentsCommonType;
+
+				Arguments[i] = Emitter.Cast(Arguments[i], ToType);
 			}
 		}
 	}
+	
+	if (Arguments[0].IsPoison() || Arguments[1].IsPoison() || Arguments[2].IsPoison())
+	{
+		return FType::MakePoison();
+	}
 
 	// Arguments might have changed, update the references.
 	A = Arguments[0];
 	B = Arguments[1];
 	C = Arguments[2];
 
-	if (!bValid)
-	{
-		return FType::MakePoison();
-	}
-	
-	// Update the first argument type, as it might have changed.
-	FirstArgumentType = Arguments[0]->Type;
-
-	const FPrimitive& FirstArgumentPrimitiveType = FirstArgumentType.GetPrimitive();
-
 	// Finally, determine operator result type.
+	const FPrimitive& FirstArgumentPrimitiveType = Arguments[0]->Type.GetPrimitive();
 	switch (Signature->ReturnType)
 	{
 		case ORT_Unknown:
@@ -1979,10 +2112,10 @@ static FType ValidateOperatorAndGetResultType(FEmitter& Emitter, EOperator Op, F
 
 		case ORT_FirstArgumentType:
 			// Some operators accept LWC as input, but always output non-LWC (examples: modulo, sign, sine, cosine, saturate, frac). Anything that may output LWC should use ORT_FirstArgumentTypeLWC.
-			return FirstArgumentPrimitiveType.IsDouble() ? FirstArgumentPrimitiveType.ToScalarKind(EScalarKind::Float) : FirstArgumentType;
+			return FirstArgumentPrimitiveType.IsDouble() ? FirstArgumentPrimitiveType.ToScalarKind(EScalarKind::Float) : Arguments[0]->Type;
 
 		case ORT_FirstArgumentTypeLWC:
-			return FirstArgumentType;
+			return Arguments[0]->Type;
 
 		case ORT_BooleanWithFirstArgumentDimensions:
 			return FType::MakePrimitive(EScalarKind::Bool, FirstArgumentPrimitiveType.NumRows, FirstArgumentPrimitiveType.NumColumns);
@@ -1992,22 +2125,34 @@ static FType ValidateOperatorAndGetResultType(FEmitter& Emitter, EOperator Op, F
 
 		case ORT_SecondArgumentType:
 			return B->Type;
+
+		case ORT_MatrixMultiplyResult:
+		{
+			FType SecondArgumentType = Arguments[1]->Type;
+			if (FirstArgumentType.GetPrimitive().NumColumns != SecondArgumentType.GetPrimitive().NumRows)
+			{
+				Emitter.Errorf({}, TEXT("Cannot matrix multiply a '%s' with a '%s'."), *FirstArgumentType.GetSpelling(), *SecondArgumentType.GetSpelling());
+				return FType::MakePoison();
+			}
+
+			return FType::MakePrimitive(FirstArgumentType.GetPrimitive().ScalarKind, FirstArgumentType.GetPrimitive().NumRows, SecondArgumentType.GetPrimitive().NumColumns);
+		}
 	}
 
 	UE_MIR_UNREACHABLE();
 }
 
-/// Returns whether the operator supports componentwise application. In other words, if the following is true:
-/// 	op(v, w) == [op(v_0, w_0), ..., op(v_n, w_n)]
+// Returns whether the operator supports componentwise application. In other words, if the following is true:
+// 	op(v, w) == [op(v_0, w_0), ..., op(v_n, w_n)]
 static bool IsComponentwiseOperator(EOperator Op)
 {
 	return Op != BO_Dot && Op != BO_Cross;
 }
 
-/// Tries to fold the operator by applying the operator componentwise on arguments components.
-/// If a value is returned, it will be a composite with some component folded to a constant. If some argument
-/// isn't a composite, or all arguments components are non-constant, the folding will not be carried out.
-/// If no folding is carried out, this function simply returns nullptr.
+// Tries to fold the operator by applying the operator componentwise on arguments components.
+// If a value is returned, it will be a composite with some component folded to a constant. If some argument
+// isn't a composite, or all arguments components are non-constant, the folding will not be carried out.
+// If no folding is carried out, this function simply returns nullptr.
 static FValue* TryFoldComponentwiseOperator(FEmitter& Emitter, EOperator Op, FValue* A, FValue* B, FValue* C, FType ResultType)
 {
 	// Check that at least one component of the resulting composite value would folded.
@@ -2030,7 +2175,7 @@ static FValue* TryFoldComponentwiseOperator(FEmitter& Emitter, EOperator Op, FVa
 		FValue* CComponent = C ? Emitter.Subscript(C, i) : nullptr;
 
 		// Try folding the operation, it may return null
-		FValue* ResultComponent = TryFoldOperatorScalar(Emitter, Op, AComponent, BComponent, CComponent);
+		FValue* ResultComponent = TryFoldOperatorScalar(Emitter, Op, AComponent, BComponent, CComponent, ResultType);
 
 		// Update the flags
 		bSomeResultComponentWasFolded |= (bool)ResultComponent;
@@ -2104,8 +2249,8 @@ static FValue* TryFoldComponentwiseOperator(FEmitter& Emitter, EOperator Op, FVa
 	return {};
 }
 
-/// If V is a composite and all its components are constants, it unpacks the components into OutComponents and returns true.
-/// If this is not possible for any reason, it returns false.
+// If V is a composite and all its components are constants, it unpacks the components into OutComponents and returns true.
+// If this is not possible for any reason, it returns false.
 static bool TryUnpackConstantScalarOrVector(FValue* V, TArrayView<FConstant*> OutComponents, int32& OutNumComponents)
 {
 	// V not specified? Or not a scalar/vector?
@@ -2129,7 +2274,7 @@ static bool TryUnpackConstantScalarOrVector(FValue* V, TArrayView<FConstant*> Ou
 	return true;
 }
 
-/// Computes the dot product on two arrays of constant float components.
+// Computes the dot product on two arrays of constant float components.
 static TFloat ConstantDotFloat(TArrayView<FConstant*> AComponents, TArrayView<FConstant*> BComponents, int32 NumComponents)
 {
 	float Result = 0.0f;
@@ -2150,8 +2295,8 @@ static MIR::TDouble ConstantDotDouble(TArrayView<FConstant*> AComponents, TArray
 	return Result;
 }
 
-/// Tries to fold the operator, that is to evaluate its result now at translation time if its arguments are constant.
-/// If the operator could not be folded in any way, it returns nullptr.
+// Tries to fold the operator, that is to evaluate its result now at translation time if its arguments are constant.
+// If the operator could not be folded in any way, it returns nullptr.
 static FValue* TryFoldOperator(FEmitter& Emitter, EOperator Op, FValue* A, FValue* B, FValue* C, FType ResultType)
 {
 	FConstant* AComponents[4];
@@ -2245,7 +2390,7 @@ FValueRef FEmitter::Operator(EOperator Op, FValueRef A, FValueRef B, FValueRef C
 	}
 
 	// Subtract has a special case option to automatically truncate when subtracting two LWC inputs from each other, assuming this is a
-	// transition from LWC space to relative space, and no longer needs to be LWC.  We need to check that all arguments are LWC before
+	// transition from LWC space to relative space, and no longer needs to be LWC. We need to check that all arguments are LWC before
 	// the call to ValidateOperatorAndGetResultType, as that may cast the inputs, changing them.
 	bool bCastLWCOutputToFloat = false;
 	if (Op == BO_Subtract && B)
@@ -2269,7 +2414,7 @@ FValueRef FEmitter::Operator(EOperator Op, FValueRef A, FValueRef B, FValueRef C
 	FValueRef Result;
 
 	// Try to apply some operator identity to simplify the operator.
-	if (FValueRef Simplified = TrySimplifyOperator(*this, Op, A, B, C))
+	if (FValueRef Simplified = TrySimplifyOperator(*this, Op, A, B, C, ResultType))
 	{
 		Result = Simplified;
 	}
@@ -3133,6 +3278,10 @@ static FValue* DifferentiateOperator(FEmitter& E, FOperator* Op, EDerivativeAxis
 		case BO_Multiply:
 			return E.Add(E.Multiply(dF, G), E.Multiply(F, dG));
 
+		// d/dx matmul(f(x), g(x)) = matmul(f'(x), g(x)) + matmul(f(x), g'(x))
+		case BO_MatrixMultiply:
+			return E.Add(E.MatrixMultiply(dF, G), E.MatrixMultiply(F, dG));
+
 		// d/dx (f(x) / g(x)) = (f'(x) * g(x) - f(x) * g'(x)) / g(x)^2
 		case BO_Divide:
 			return E.Divide(E.Subtract(E.Multiply(dF, G), E.Multiply(F, dG)), E.Multiply(G, G));
@@ -3376,6 +3525,7 @@ FType FEmitter::TryGetCommonType(FType A, FType B)
 		return A;
 	}
 
+	// Primitive types can only be constructed from other primitive types.
 	TOptional<FPrimitive> PrimitiveA = A.AsPrimitive();
 	TOptional<FPrimitive> PrimitiveB = B.AsPrimitive();
 	if (!PrimitiveA || !PrimitiveB)
diff --git a/Engine/Source/Runtime/Engine/Private/Materials/MaterialIRToHLSLTranslator.cpp b/Engine/Source/Runtime/Engine/Private/Materials/MaterialIRToHLSLTranslator.cpp
--- a/Engine/Source/Runtime/Engine/Private/Materials/MaterialIRToHLSLTranslator.cpp
+++ b/Engine/Source/Runtime/Engine/Private/Materials/MaterialIRToHLSLTranslator.cpp
@@ -233,28 +233,8 @@ static bool IsFoldable(const MIR::FInstruction* Instr, MIR::EStage Stage)
 	return true;
 }
 
-/// Returns whether the operator in HLSL is infix between its arguments, e.g. "4 + 4"
-static bool IsOperatorInfix(MIR::EOperator Op)
-{
-	switch (Op)
-	{
-		case MIR::BO_GreaterThan:
-		case MIR::BO_GreaterThanOrEquals:
-		case MIR::BO_LessThan:
-		case MIR::BO_LessThanOrEquals:
-		case MIR::BO_Equals:
-		case MIR::BO_NotEquals:
-		case MIR::BO_Add:
-		case MIR::BO_Subtract:
-		case MIR::BO_Multiply:
-		case MIR::BO_Divide:
-		case MIR::BO_Modulo:
-			return true;
-		default: return false;
-	}
-}
-
 static FStringView GVector4SwizzleSubset[4] = { TEXTVIEW(".x"), TEXTVIEW(".xy"), TEXTVIEW(".xyz"), {} };
+
 struct FPrivate : FMaterialIRToHLSLTranslation
 {
 	// Utility for emitting HLSL code
@@ -654,56 +634,49 @@ struct FPrivate : FMaterialIRToHLSLTranslation
 		FlushCodeSubstring(Code.Len());
 	}
 
+	
 	void LowerInlineHLSL(const MIR::FInlineHLSL* ExternalCode)
 	{
-		auto PrintHLSLCode = [this, ExternalCode](FStringView Code) -> void
-			{
-				if (ExternalCode->NumArguments > 0)
-				{
-					check(ExternalCode->Arguments);
-					LowerInlineHLSLWithArgumentsInternal(Code, TConstArrayView<MIR::FValue*>(ExternalCode->Arguments, ExternalCode->NumArguments));
-				}
-				else
-				{
-					Printer << Code;
-				}
-			};
-
+		FStringView FinalCode;
 		if (ExternalCode->HasFlags(MIR::EValueFlags::HasDynamicHLSLCode))
 		{
-			check(ExternalCode->Code.Ptr);
+			FinalCode = ExternalCode->Code;
+		}
+		else if (ExternalCode->HasFlags(MIR::EValueFlags::DerivativeDDX))
+		{
+			FinalCode = ExternalCode->ExternalCodeDeclaration->DefinitionDDX;
+		}
+		else if (ExternalCode->HasFlags(MIR::EValueFlags::DerivativeDDY))
+		{
+			FinalCode = ExternalCode->ExternalCodeDeclaration->DefinitionDDY;
+		}
+		else
+		{
+			FinalCode = ExternalCode->ExternalCodeDeclaration->Definition;
+		}
 
-			// Substitute placeholder tokens now unless disabled for custom nodes
-			if (ExternalCode->HasFlags(MIR::EValueFlags::SubstituteTagsInInlineHLSL))
-			{
-				FString Code = ExternalCode->Code.ToString();
+		// Substitute placeholder tokens now unless disabled for custom nodes, if necessary.
+		FString SubstitutedCode;
+		if (ExternalCode->HasFlags(MIR::EValueFlags::SubstituteTags))
+		{
+			// @todo-laura.hermanns - Must be replaced as soon as this state is available in MIR
+			constexpr bool bCompilingPreviousFrame = false;
+			SubstitutedCode = FString{ FinalCode };
+			SubstitutedCode.ReplaceInline(TEXT("<PREV>"), bCompilingPreviousFrame ? TEXT("Prev") : TEXT(""));
+			SubstitutedCode.ReplaceInline(TEXT("<PREVIOUS>"), bCompilingPreviousFrame ? TEXT("Previous") : TEXT(""));
 
-				// @todo-laura.hermanns - Must be replaced as soon as this state is available in MIR
-				constexpr bool bCompilingPreviousFrame = false;
-				Code.ReplaceInline(TEXT("<PREV>"), bCompilingPreviousFrame ? TEXT("Prev") : TEXT(""));
-	
-				PrintHLSLCode(Code);
-			}
-			else
-			{
-				PrintHLSLCode(ExternalCode->Code);
-			}
+			FinalCode = SubstitutedCode;
+		}
+
+		// Print the final HLSL code.
+		if (ExternalCode->NumArguments > 0)
+		{
+			check(ExternalCode->Arguments);
+			LowerInlineHLSLWithArgumentsInternal(FinalCode, TConstArrayView<MIR::FValue*>(ExternalCode->Arguments, ExternalCode->NumArguments));
 		}
 		else
 		{
-			check(ExternalCode->ExternalCodeDeclaration != nullptr);
-			if (ExternalCode->HasFlags(MIR::EValueFlags::DerivativeDDX))
-			{
-				PrintHLSLCode(ExternalCode->ExternalCodeDeclaration->DefinitionDDX);
-			}
-			else if (ExternalCode->HasFlags(MIR::EValueFlags::DerivativeDDY))
-			{
-				PrintHLSLCode(ExternalCode->ExternalCodeDeclaration->DefinitionDDY);
-			}
-			else
-			{
-				PrintHLSLCode(ExternalCode->ExternalCodeDeclaration->Definition);
-			}
+			Printer << FinalCode;
 		}
 	}
 
@@ -1251,9 +1224,34 @@ struct FPrivate : FMaterialIRToHLSLTranslation
 			return;
 		}
 
-		bool bIsLWC = Operator->AArg->Type.IsDouble() || (Operator->Op == MIR::TO_Select && Operator->BArg->Type.IsDouble());
+		// Whether any of this operator's arguments has double type which requires special handling in the shader.
+		bool bIsDouble = Operator->AArg->Type.IsDouble() || (Operator->BArg && Operator->BArg->Type.IsDouble()) || (Operator->CArg && Operator->CArg->Type.IsDouble());
+
+		// Whether the operator in HLSL is infix between its arguments, e.g. "4 + 4"
+		bool bOperatorIsInfix = false;
+		if (!bIsDouble)
+		{
+			switch (Operator->Op)
+			{
+				case MIR::BO_GreaterThan:
+				case MIR::BO_GreaterThanOrEquals:
+				case MIR::BO_LessThan:
+				case MIR::BO_LessThanOrEquals:
+				case MIR::BO_Equals:
+				case MIR::BO_NotEquals:
+				case MIR::BO_Add:
+				case MIR::BO_Multiply:
+				case MIR::BO_Subtract:
+				case MIR::BO_Divide:
+				case MIR::BO_Modulo:
+					bOperatorIsInfix = true;
+					break;
+
+				default: break;
+			}
+		}
 
-		if (IsOperatorInfix(Operator->Op) && !bIsLWC)
+		if (bOperatorIsInfix)
 		{
 			FStringView OpString;
 			switch (Operator->Op)
@@ -1277,7 +1275,7 @@ struct FPrivate : FMaterialIRToHLSLTranslation
 		else
 		{
 			FStringView OpString;
-			if (bIsLWC)
+			if (bIsDouble)
 			{
 				// "Demotes LWC" indicates the given operator returns a non-LWC float, even if the input is LWC.  Besides that, comparison
 				// operators all return bool, instead of LWC, and a couple specific operators require specific inputs to always be non-LWC
@@ -1315,6 +1313,7 @@ struct FPrivate : FMaterialIRToHLSLTranslation
 					case MIR::BO_Max: OpString = TEXTVIEW("WSMax"); break;
 					case MIR::BO_Min: OpString = TEXTVIEW("WSMin"); break;
 					case MIR::BO_Multiply: OpString = TEXTVIEW("WSMultiply"); break;
+					case MIR::BO_MatrixMultiply: OpString = TEXTVIEW("WSMultiply"); break;			// TODO: implement WSMultiplyVector, WSMultiplyDemote through tiling
 					case MIR::BO_NotEquals: OpString = TEXTVIEW("WSNotEquals"); break;				// Bool output
 					case MIR::BO_Step: OpString = TEXTVIEW("WSStep"); break;						// Demotes LWC
 					case MIR::BO_Subtract: OpString = TEXTVIEW("WSSubtract"); break;
@@ -1375,6 +1374,7 @@ struct FPrivate : FMaterialIRToHLSLTranslation
 					case MIR::BO_Dot: OpString = TEXTVIEW("dot"); break;
 					case MIR::BO_Fmod: OpString = TEXTVIEW("fmod"); break;
 					case MIR::BO_Max: OpString = TEXTVIEW("max"); break;
+					case MIR::BO_MatrixMultiply: OpString = TEXTVIEW("mul"); break;
 					case MIR::BO_Min: OpString = TEXTVIEW("min"); break;
 					case MIR::BO_Or: OpString = TEXTVIEW("or"); break;
 					case MIR::BO_Pow: OpString = TEXTVIEW("pow"); break;
diff --git a/Engine/Source/Runtime/Engine/Public/Materials/MaterialExpressionOperator.h b/Engine/Source/Runtime/Engine/Public/Materials/MaterialExpressionOperator.h
--- a/Engine/Source/Runtime/Engine/Public/Materials/MaterialExpressionOperator.h
+++ b/Engine/Source/Runtime/Engine/Public/Materials/MaterialExpressionOperator.h
@@ -65,6 +65,7 @@ enum class EMaterialExpressionOperatorKind
 	Add,
 	Subtract,
 	Multiply,
+	MatrixMultiply,
 	Divide,
 	Modulo,
 	BitwiseAnd,
diff --git a/Engine/Source/Runtime/Engine/Public/Materials/MaterialIR.h b/Engine/Source/Runtime/Engine/Public/Materials/MaterialIR.h
--- a/Engine/Source/Runtime/Engine/Public/Materials/MaterialIR.h
+++ b/Engine/Source/Runtime/Engine/Public/Materials/MaterialIR.h
@@ -60,8 +60,8 @@ enum class EValueFlags : uint8
 	// Mask that includes all "analyzed in stage" flags.
 	AnalyzedInAnyStageMask = AnalyzedForStageVertex | AnalyzedForStagePixel | AnalyzedForStageCompute,
 
-	// Whether to skip substution of abstract tokens like "<PREV>". Only true for custom nodes.
-	SubstituteTagsInInlineHLSL = 1 << 3,
+	// Whether to substitute abstract tokens like "<PREV>" in code expressions of InlineHLSL instructions.
+	SubstituteTags = 1 << 3,
 
 	// Whether 'FInlineHLSL::Code' field is used. Otherwise, 'ExternalCodeDeclaration' is used.
 	HasDynamicHLSLCode = 1 << 4, 
@@ -720,6 +720,7 @@ enum EOperator
 	BO_Add,
 	BO_Subtract,
 	BO_Multiply,
+	BO_MatrixMultiply,
 	BO_Divide,
 	BO_Modulo,
 	BO_BitwiseAnd,
diff --git a/Engine/Source/Runtime/Engine/Public/Materials/MaterialIREmitter.h b/Engine/Source/Runtime/Engine/Public/Materials/MaterialIREmitter.h
--- a/Engine/Source/Runtime/Engine/Public/Materials/MaterialIREmitter.h
+++ b/Engine/Source/Runtime/Engine/Public/Materials/MaterialIREmitter.h
@@ -239,9 +239,15 @@ public:
 	// Validates that the value is a scalar. If not, reports error and returns poison.
 	FValueRef CheckIsScalar(FValueRef Value);
 
+	// Validates that the value is a vector. If not, reports error and returns poison.
+	FValueRef CheckIsVector(FValueRef Value);
+
 	// Validates that the value is either a scalar or a vector. Reports error and returns poison otherwise.
 	FValueRef CheckIsScalarOrVector(FValueRef Value);
 	
+	// Validates that the value is a matrix. If not, reports error and returns poison.
+	FValueRef CheckIsMatrix(FValueRef Value);
+
 	// Retrieves the texture value flowing into given input.
 	// If the value is missing or cannot be cast to a texture, it reports an error and returns poison.
 	FValueRef CheckIsTexture(FValueRef Value);
@@ -361,15 +367,24 @@ public:
 	// Note: this API is used by the builder and you should not use it manually.
 	FSetMaterialOutput* SetMaterialOutput(EMaterialProperty InProperty, FValue* Arg);
 
-	// Emits a 2D column vector dimensional value and returns it.
+	// Emits a 2D column vector from X and Y values.
 	FValueRef Vector2(FValueRef InX, FValueRef InY);
 
-	// Emits a 3D column vector dimensional value and returns it.
+	// Emits a 3D vector from a 2D vector (XY) and a Z value.
+	FValueRef Vector3(FValueRef InXY, FValueRef InZ);
+
+	// Emits a 3D vector from separate X, Y, and Z values.
 	FValueRef Vector3(FValueRef InX, FValueRef InY, FValueRef InZ);
 
-	// Emits a 4D column vector dimensional value and returns it.
+	// Emits a 4D vector from a 3D vector (XYZ) and a W value.
+	FValueRef Vector4(FValueRef InXYZ, FValueRef InW);
+
+	// Emits a 4D vector from a 2D vector (XY), a Z value, and a W value.
+	FValueRef Vector4(FValueRef InXY, FValueRef InZ, FValueRef InW);
+
+	// Emits a 4D vector from separate X, Y, Z, and W values.
 	FValueRef Vector4(FValueRef InX, FValueRef InY, FValueRef InZ, FValueRef InW);
-	
+
 	// Constructs an aggregate value using the provided aggregate definition.
 	// The resulting value has default-initialized attributes.
 	FValueRef Aggregate(const UMaterialAggregate* InAggregate);
@@ -577,6 +592,7 @@ public:
 	FValueRef Add(FValueRef A, FValueRef B) { return Operator(BO_Add, A, B); }
 	FValueRef Subtract(FValueRef A, FValueRef B) { return Operator(BO_Subtract, A, B); }
 	FValueRef Multiply(FValueRef A, FValueRef B) { return Operator(BO_Multiply, A, B); }
+	FValueRef MatrixMultiply(FValueRef A, FValueRef B) { return Operator(BO_MatrixMultiply, A, B); }
 	FValueRef Divide(FValueRef A, FValueRef B) { return Operator(BO_Divide, A, B); }
 	FValueRef BitwiseAnd(FValueRef A, FValueRef B) { return Operator(BO_BitwiseAnd, A, B); }
 	FValueRef BitwiseOr(FValueRef A, FValueRef B) { return Operator(BO_BitwiseOr, A, B); }
