Fixed AActor IncrementalPreRegisterComponents and IncrementalPreUnregisterComponents, could iterate on null components.

diff --git a/Engine/Source/Runtime/Engine/Private/Actor.cpp b/Engine/Source/Runtime/Engine/Private/Actor.cpp
--- a/Engine/Source/Runtime/Engine/Private/Actor.cpp
+++ b/Engine/Source/Runtime/Engine/Private/Actor.cpp
@@ -6071,7 +6071,7 @@ void AActor::IncrementalPreUnregisterComponents()
 	TInlineComponentArray<UActorComponent*> Components;
 	for (UActorComponent* Component : GetComponents())
 	{
-		if (Component->ShouldIncrementalPreUnregister())
+		if (Component && Component->ShouldIncrementalPreUnregister())
 		{
 			Components.Add(Component);
 		}
@@ -6115,14 +6115,17 @@ void AActor::IncrementalPreRegisterComponents()
 	TInlineComponentArray<UActorComponent*> Components;
 	for (UActorComponent* Component : GetComponents())
 	{
+		if (Component)
+		{
 #if WITH_EDITOR
-		// Ignore registered components (can happen if construction script adds component(s), which happens before this)
-		if (!Component->IsRegistered())
+			// Ignore registered components (can happen if construction script adds component(s), which happens before this)
+			if (!Component->IsRegistered())
 #endif
-		{
-			if (Component->ShouldIncrementalPreRegister(World))
 			{
-				Components.Add(Component);
+				if (Component->ShouldIncrementalPreRegister(World))
+				{
+					Components.Add(Component);
+				}
 			}
 		}
 	}
