Clean up Type tracing code to remove all the "Class" references, now that it supports tracing Struct types as well

diff --git a/Engine/Plugins/Animation/GameplayInsights/Source/GameplayInsights/Private/GameplayAnalyzer.cpp b/Engine/Plugins/Animation/GameplayInsights/Source/GameplayInsights/Private/GameplayAnalyzer.cpp
--- a/Engine/Plugins/Animation/GameplayInsights/Source/GameplayInsights/Private/GameplayAnalyzer.cpp
+++ b/Engine/Plugins/Animation/GameplayInsights/Source/GameplayInsights/Private/GameplayAnalyzer.cpp
@@ -18,7 +18,8 @@ void FGameplayAnalyzer::OnAnalysisBegin(const FOnAnalysisContext& Context)
 	auto& Builder = Context.InterfaceBuilder;
 
 	Builder.RouteEvent(RouteId_RecordingInfo, "Object", "RecordingInfo");
-	Builder.RouteEvent(RouteId_Class, "Object", "Class");
+	Builder.RouteEvent(RouteId_Type, "Object", "Class");
+	Builder.RouteEvent(RouteId_Type, "Object", "Type");
 	Builder.RouteEvent(RouteId_Object, "Object", "Object");
 	Builder.RouteEvent(RouteId_ObjectTransform, "Object", "ObjectTransform");
 	Builder.RouteEvent(RouteId_ObjectEvent, "Object", "ObjectEvent");
@@ -70,24 +71,24 @@ bool FGameplayAnalyzer::OnEvent(uint16 RouteId, EStyle Style, const FOnEventCont
 		GameplayProvider.AppendRecordingInfo(WorldId, Context.EventTime.AsSeconds(Cycle), RecordingIndex, FrameIndex, ElapsedTime);
 		break;
 	}
-	case RouteId_Class:
+	case RouteId_Type:
 	{
-		FString ClassName, ClassPathName;
-		if (EventData.GetString("Name", ClassName))
+		FString TypeName, TypePathName;
+		if (EventData.GetString("Name", TypeName))
 		{
-			EventData.GetString("Path", ClassPathName);
+			EventData.GetString("Path", TypePathName);
 		}
 		else
 		{
-			const TCHAR* ClassNameAndPathName = reinterpret_cast<const TCHAR*>(EventData.GetAttachment());
-			ClassName = ClassNameAndPathName;
-			int32 ClassNameStringLength = EventData.GetValue<int32>("ClassNameStringLength");
-			ClassPathName = ClassNameAndPathName + ClassNameStringLength;
+			const TCHAR* TypeNameAndPathName = reinterpret_cast<const TCHAR*>(EventData.GetAttachment());
+			TypeName = TypeNameAndPathName;
+			int32 TypeNameStringLength = EventData.GetValue<int32>("TypeNameStringLength");
+			TypePathName = TypeNameAndPathName + TypeNameStringLength;
 		}
 
 		uint64 Id = EventData.GetValue<uint64>("Id");
 		uint64 SuperId = EventData.GetValue<uint64>("SuperId");
-		GameplayProvider.AppendClass(Id, SuperId, *ClassName, *ClassPathName);
+		GameplayProvider.AppendClass(Id, SuperId, *TypeName, *TypePathName);
 		break;
 	}
 	case RouteId_Object:
diff --git a/Engine/Plugins/Animation/GameplayInsights/Source/GameplayInsights/Private/GameplayAnalyzer.h b/Engine/Plugins/Animation/GameplayInsights/Source/GameplayInsights/Private/GameplayAnalyzer.h
--- a/Engine/Plugins/Animation/GameplayInsights/Source/GameplayInsights/Private/GameplayAnalyzer.h
+++ b/Engine/Plugins/Animation/GameplayInsights/Source/GameplayInsights/Private/GameplayAnalyzer.h
@@ -20,7 +20,7 @@ private:
 	enum : uint16
 	{
 		RouteId_World,
-		RouteId_Class,
+		RouteId_Type,
 		RouteId_Object,
 		RouteId_ObjectTransform,
 		RouteId_ObjectEvent,
diff --git a/Engine/Plugins/Developer/TraceSourceFiltering/Source/SourceFilteringTrace/Private/SourceFilterTrace.cpp b/Engine/Plugins/Developer/TraceSourceFiltering/Source/SourceFilteringTrace/Private/SourceFilterTrace.cpp
--- a/Engine/Plugins/Developer/TraceSourceFiltering/Source/SourceFilteringTrace/Private/SourceFilterTrace.cpp
+++ b/Engine/Plugins/Developer/TraceSourceFiltering/Source/SourceFilteringTrace/Private/SourceFilterTrace.cpp
@@ -86,7 +86,7 @@ void FSourceFilterTrace::OutputClass(const TSubclassOf<UDataSourceFilter> InClas
 	}
 
 	UClass* Class = InClass.Get();
-	TRACE_CLASS(Class);
+	TRACE_TYPE(Class);
 
 	DataSourceFilterClasses.Add(Class->GetName(), Class);
 
diff --git a/Engine/Source/Runtime/Engine/Private/Animation/AnimTrace.cpp b/Engine/Source/Runtime/Engine/Private/Animation/AnimTrace.cpp
--- a/Engine/Source/Runtime/Engine/Private/Animation/AnimTrace.cpp
+++ b/Engine/Source/Runtime/Engine/Private/Animation/AnimTrace.cpp
@@ -1087,7 +1087,7 @@ void FAnimTrace::OutputAnimNodeValue(const FAnimationBaseContext& InContext, uin
 
 	UObject* AnimInstance = InContext.AnimInstanceProxy->GetAnimInstanceObject();
 	TRACE_OBJECT(AnimInstance);
-	TRACE_CLASS(InValue);
+	TRACE_TYPE(InValue);
 
 	UE_TRACE_LOG(Animation, AnimNodeValueClass, AnimationChannel)
 		<< AnimNodeValueClass.Cycle(FPlatformTime::Cycles64())
diff --git a/Engine/Source/Runtime/Engine/Private/ObjectAsTraceIdProxyArchive.cpp b/Engine/Source/Runtime/Engine/Private/ObjectAsTraceIdProxyArchive.cpp
--- a/Engine/Source/Runtime/Engine/Private/ObjectAsTraceIdProxyArchive.cpp
+++ b/Engine/Source/Runtime/Engine/Private/ObjectAsTraceIdProxyArchive.cpp
@@ -33,7 +33,7 @@ void FObjectAsTraceIdProxyArchive::Write(const UObject* Obj)
 		// first ensure the object has been traced 
 		if (const UClass* Class = Cast<UClass>(Obj))
 		{
-			TRACE_CLASS(Class);
+			TRACE_TYPE(Class);
 		}
 		TRACE_OBJECT(Obj);
 		// then serialize the Object as a Trace Object id;
diff --git a/Engine/Source/Runtime/Engine/Private/ObjectTrace.cpp b/Engine/Source/Runtime/Engine/Private/ObjectTrace.cpp
--- a/Engine/Source/Runtime/Engine/Private/ObjectTrace.cpp
+++ b/Engine/Source/Runtime/Engine/Private/ObjectTrace.cpp
@@ -16,7 +16,7 @@
 
 UE_TRACE_CHANNEL(ObjectChannel)
 
-UE_TRACE_EVENT_BEGIN(Object, Class)
+UE_TRACE_EVENT_BEGIN(Object, Type)
 	UE_TRACE_EVENT_FIELD(uint64, Id)
 	UE_TRACE_EVENT_FIELD(uint64, SuperId)
 	UE_TRACE_EVENT_FIELD(UE::Trace::WideString, Name)
@@ -412,30 +412,30 @@ uint16 FObjectTrace::GetObjectWorldTickCounter(const UObject* InObject)
 	return 0;
 }
 
-void FObjectTrace::OutputClass(const UStruct* InClass)
+void FObjectTrace::OutputType(const UStruct* InType)
 {
 	const bool bChannelEnabled = UE_TRACE_CHANNELEXPR_IS_ENABLED(ObjectChannel);
-	if (!bChannelEnabled || InClass == nullptr)
+	if (!bChannelEnabled || InType == nullptr)
 	{
 		return;
 	}
-	const uint64 ObjectId = GetObjectId(InClass);
+	const uint64 ObjectId = GetObjectId(InType);
 	if (!AddObjectToTracedSet(UE::ObjectTrace::FObjectId{ObjectId}))
 	{
 		return;
 	}
 
-	OutputClass(InClass->GetSuperStruct());
+	OutputType(InType->GetSuperStruct());
 
-	const FString ClassPathName = InClass->GetPathName();
-	TCHAR ClassName[FName::StringBufferSize];
-	const uint32 ClassNameLength = InClass->GetFName().ToString(ClassName);
+	const FString TypePathName = InType->GetPathName();
+	TCHAR TypeName[FName::StringBufferSize];
+	const uint32 TypeNameLength = InType->GetFName().ToString(TypeName);
 
-	UE_TRACE_LOG(Object, Class, ObjectChannel)
-		<< Class.Id(ObjectId)
-		<< Class.SuperId(GetObjectId(InClass->GetSuperStruct()))
-		<< Class.Name(ClassName, ClassNameLength)
-		<< Class.Path(*ClassPathName, ClassPathName.Len());
+	UE_TRACE_LOG(Object, Type, ObjectChannel)
+		<< Type.Id(ObjectId)
+		<< Type.SuperId(GetObjectId(InType->GetSuperStruct()))
+		<< Type.Name(TypeName, TypeNameLength)
+		<< Type.Path(*TypePathName, TypePathName.Len());
 }
 
 void FObjectTrace::OutputView(const UObject* InPlayer, const FSceneView* InView)
@@ -477,7 +477,7 @@ void FObjectTrace::OutputInstance(const UObject* InOuterObject, uint64 InInstanc
 	}
 
 	// Trace the object's class first
-	TRACE_CLASS(InType);
+	TRACE_TYPE(InType);
 
 	UE_TRACE_LOG(Object, Object, ObjectChannel)
 		<< Object.Id(InInstanceId)
@@ -628,7 +628,7 @@ void FObjectTrace::OutputChildElement(const UObject* InObject, const uint64 InCh
 		Name += FString::Printf(TEXT(" : [0x%llX]"), InChildElementId);
 	}
 	
-	TRACE_CLASS(InStructType);
+	TRACE_TYPE(InStructType);
 
 	UE_TRACE_LOG(Object, ChildElement, ObjectChannel)
 		<< ChildElement.Id(ObjectId)
diff --git a/Engine/Source/Runtime/Engine/Public/ObjectTrace.h b/Engine/Source/Runtime/Engine/Public/ObjectTrace.h
--- a/Engine/Source/Runtime/Engine/Public/ObjectTrace.h
+++ b/Engine/Source/Runtime/Engine/Public/ObjectTrace.h
@@ -56,8 +56,15 @@ struct FObjectTrace
 	/** Reset Caches so a new trace can be started*/
 	ENGINE_API static void Reset();
 
-	/** Helper function to output an object */
-	ENGINE_API static void OutputClass(const UStruct* InClass);
+	/** Helper function to output an type (class or struct) */
+	ENGINE_API static void OutputType(const UStruct* InType);
+	
+	/** Helper function to output a class */
+	UE_DEPRECATED(5.7, "OutputClass is now deprecated, use OutputType")
+	ENGINE_API static void OutputClass(const UClass* InClass)
+	{
+		OutputType(InClass);
+	};
 
 	/** Helper function to output an object */
 	ENGINE_API static void OutputObject(const UObject* InObject);
@@ -130,8 +137,11 @@ struct FObjectTrace
 	ENGINE_API static void OutputChildElementEvent(const UObject* InObject, uint64 InChildElementId, const TCHAR* InEvent);
 };
 
+#define TRACE_TYPE(Type) \
+	FObjectTrace::OutputType(Type);
+
 #define TRACE_CLASS(Class) \
-	FObjectTrace::OutputClass(Class);
+	FObjectTrace::OutputType(Class);
 
 #define TRACE_OBJECT(Object) \
 	FObjectTrace::OutputObject(Object);
@@ -238,6 +248,7 @@ struct FObjectTrace
 };
 
 #define TRACE_CLASS(...)
+#define TRACE_TYPE(...)
 #define TRACE_OBJECT(...)
 #define TRACE_INSTANCE(...)
 #define TRACE_OBJECT_TRANSFORM(...)
