Vector Parameter Expression shows output pin display names to avoid confusion on what the first output pin is.  Also fixed "Convert to Constant" for vector parameters.

By default the outputs are now:

1. RGB
2. R
3. G
4. B
5. RGBA

- Users were expecting the first output pin to be RGBA, but it was RGB.
- Show output name on pins by default for vector parameter.
- Add RGBA output as final one to match what TextureSample does.
- Update this behavior in the new material translator (it should really use the outputs array which is already defined).
- If no channel names are specified show R, G, B by default.
- No longer truncate and lose user data when using "Convert to Constant" for Vector Parameter -> Constant.
- Give constant 4 expression display names so the conversion code can match outputs by name when converting nodes.

[RN] major | Rendering - Materials and Shaders

diff --git a/Engine/Source/Editor/MaterialEditor/Private/MaterialEditor.cpp b/Engine/Source/Editor/MaterialEditor/Private/MaterialEditor.cpp
--- a/Engine/Source/Editor/MaterialEditor/Private/MaterialEditor.cpp
+++ b/Engine/Source/Editor/MaterialEditor/Private/MaterialEditor.cpp
@@ -4240,10 +4240,8 @@ void FMaterialEditor::OnConvertObjects()
 				}
 				else if (VectorParameterExpression)
 				{
-					// Technically should be a constant 4 but UMaterialExpressionVectorParameter has an rgb pin, so using Constant3 to avoid a compile error
-					ClassToCreate = UMaterialExpressionConstant3Vector::StaticClass();
+					ClassToCreate = UMaterialExpressionConstant4Vector::StaticClass();
 				}
-			
 
 				if (ClassToCreate)
 				{
@@ -4365,7 +4363,7 @@ void FMaterialEditor::OnConvertObjects()
 						else if (VectorParameterExpression)
 						{
 							bNeedsRefresh = true;
-							CastChecked<UMaterialExpressionConstant3Vector>(NewExpression)->Constant = VectorParameterExpression->DefaultValue;
+							CastChecked<UMaterialExpressionConstant4Vector>(NewExpression)->Constant = VectorParameterExpression->DefaultValue;
 						}
 
 						if (bNeedsRefresh)
diff --git a/Engine/Source/Runtime/Engine/Private/Materials/MaterialExpressions.cpp b/Engine/Source/Runtime/Engine/Private/Materials/MaterialExpressions.cpp
--- a/Engine/Source/Runtime/Engine/Private/Materials/MaterialExpressions.cpp
+++ b/Engine/Source/Runtime/Engine/Private/Materials/MaterialExpressions.cpp
@@ -2445,9 +2445,9 @@ FText UMaterialExpressionTextureBase::GetPreviewOverlayText() const
 
 void UMaterialExpressionTextureBase::AutoSetSampleType()
 {
-	if ( Texture )
+	if (Texture)
 	{
-		SamplerType = GetSamplerTypeForTexture( Texture );
+		SamplerType = GetSamplerTypeForTexture(Texture);
 	}
 }
 
@@ -2590,6 +2590,8 @@ void UMaterialExpressionTextureSample::PostLoad()
 	{
 		TextureObject.Expression = nullptr;
 	}
+
+	AutoSetSampleType();
 }
 
 void UMaterialExpressionTextureSample::ApplyChannelNames()
@@ -5068,11 +5070,13 @@ UMaterialExpressionConstant4Vector::UMaterialExpressionConstant4Vector(const FOb
 	bCollapsed = false;
 
 	Outputs.Reset();
-	Outputs.Add(FExpressionOutput(TEXT(""), 1, 1, 1, 1, 1));
-	Outputs.Add(FExpressionOutput(TEXT(""), 1, 1, 0, 0, 0));
-	Outputs.Add(FExpressionOutput(TEXT(""), 1, 0, 1, 0, 0));
-	Outputs.Add(FExpressionOutput(TEXT(""), 1, 0, 0, 1, 0));
-	Outputs.Add(FExpressionOutput(TEXT(""), 1, 0, 0, 0, 1));
+	Outputs.Add(FExpressionOutput(TEXT("RGBA"), 1, 1, 1, 1, 1));
+	Outputs.Add(FExpressionOutput(TEXT("R"), 1, 1, 0, 0, 0));
+	Outputs.Add(FExpressionOutput(TEXT("G"), 1, 0, 1, 0, 0));
+	Outputs.Add(FExpressionOutput(TEXT("B"), 1, 0, 0, 1, 0));
+	Outputs.Add(FExpressionOutput(TEXT("A"), 1, 0, 0, 0, 1));
+	Outputs.Add(FExpressionOutput(TEXT("RGB"), 1, 1, 1, 1, 0));
+	bShowOutputNameOnPin = true;
 #endif
 }
 
@@ -8295,12 +8299,15 @@ UMaterialExpressionVectorParameter::UMaterialExpressionVectorParameter(const FOb
 {
 #if WITH_EDITORONLY_DATA
 	Outputs.Reset();
-	Outputs.Add(FExpressionOutput(TEXT(""), 1, 1, 1, 1, 0));
-	Outputs.Add(FExpressionOutput(TEXT(""), 1, 1, 0, 0, 0));
-	Outputs.Add(FExpressionOutput(TEXT(""), 1, 0, 1, 0, 0));
-	Outputs.Add(FExpressionOutput(TEXT(""), 1, 0, 0, 1, 0));
-	Outputs.Add(FExpressionOutput(TEXT(""), 1, 0, 0, 0, 1));
+	Outputs.Add(FExpressionOutput(TEXT("RGB"), 1, 1, 1, 1, 0));
+	Outputs.Add(FExpressionOutput(TEXT("R"), 1, 1, 0, 0, 0));
+	Outputs.Add(FExpressionOutput(TEXT("G"), 1, 0, 1, 0, 0));
+	Outputs.Add(FExpressionOutput(TEXT("B"), 1, 0, 0, 1, 0));
+	Outputs.Add(FExpressionOutput(TEXT("A"), 1, 0, 0, 0, 1));
+	Outputs.Add(FExpressionOutput(TEXT("RGBA"), 1, 1, 1, 1, 1));
 	ApplyChannelNames();
+
+	bShowOutputNameOnPin = true;
 #endif // WITH_EDITORONLY_DATA
 }
 
@@ -8403,11 +8410,10 @@ void UMaterialExpressionVectorParameter::PostEditChangeProperty(FPropertyChanged
 
 void UMaterialExpressionVectorParameter::ApplyChannelNames()
 {
-	Outputs[1].OutputName = FName(*ChannelNames.R.ToString());
-	Outputs[2].OutputName = FName(*ChannelNames.G.ToString());
-	Outputs[3].OutputName = FName(*ChannelNames.B.ToString());
-	Outputs[4].OutputName = FName(*ChannelNames.A.ToString());
-	bShowOutputNameOnPin = !ChannelNames.R.IsEmpty() || !ChannelNames.G.IsEmpty() || !ChannelNames.B.IsEmpty() || !ChannelNames.A.IsEmpty();
+	Outputs[1].OutputName = !ChannelNames.R.IsEmpty() ? FName(*ChannelNames.R.ToString()) : TEXT("R");
+	Outputs[2].OutputName = !ChannelNames.G.IsEmpty() ? FName(*ChannelNames.G.ToString()) : TEXT("G");
+	Outputs[3].OutputName = !ChannelNames.B.IsEmpty() ? FName(*ChannelNames.B.ToString()) : TEXT("B");
+	Outputs[4].OutputName = !ChannelNames.A.IsEmpty() ? FName(*ChannelNames.A.ToString()) : TEXT("A");
 }
 
 
diff --git a/Engine/Source/Runtime/Engine/Private/Materials/MaterialExpressionsIR.cpp b/Engine/Source/Runtime/Engine/Private/Materials/MaterialExpressionsIR.cpp
--- a/Engine/Source/Runtime/Engine/Private/Materials/MaterialExpressionsIR.cpp
+++ b/Engine/Source/Runtime/Engine/Private/Materials/MaterialExpressionsIR.cpp
@@ -386,11 +386,12 @@ void UMaterialExpressionParameter::Build(MIR::FEmitter& Em)
 void UMaterialExpressionVectorParameter::Build(MIR::FEmitter& Em)
 {
 	FValueRef Value = BuildMaterialExpressionParameter(Em, this);
-	Em.Output(0, Value);
+	Em.Output(0, Em.Swizzle(Value, MIR::FSwizzleMask::XYZ()));
 	Em.Output(1, Em.Subscript(Value, 0));
 	Em.Output(2, Em.Subscript(Value, 1));
 	Em.Output(3, Em.Subscript(Value, 2));
 	Em.Output(4, Em.Subscript(Value, 3));
+	Em.Output(5, Value);
 }
 
 void UMaterialExpressionDoubleVectorParameter::Build(MIR::FEmitter& Em)
