Change anim notify queue hard world ref to a soft ref.

diff --git a/Engine/Source/Runtime/Engine/Private/Animation/AnimNotifyQueue.cpp b/Engine/Source/Runtime/Engine/Private/Animation/AnimNotifyQueue.cpp
--- a/Engine/Source/Runtime/Engine/Private/Animation/AnimNotifyQueue.cpp
+++ b/Engine/Source/Runtime/Engine/Private/Animation/AnimNotifyQueue.cpp
@@ -52,7 +52,7 @@ bool FAnimNotifyQueue::PassesChanceOfTriggering(const FAnimNotifyEvent* Event) c
 
 void FAnimNotifyQueue::AddAnimNotifiesToDest(bool bSrcIsLeader, const TArray<FAnimNotifyEventReference>& NewNotifies, TArray<FAnimNotifyEventReference>& DestArray, const float InstanceWeight)
 {
-	bool bIsDedicatedServer = World ? (World->GetNetMode() == NM_DedicatedServer) : false;
+	bool bIsDedicatedServer = World.IsValid() ? (World->GetNetMode() == NM_DedicatedServer) : false;
 	// for now there is no filter whatsoever, it just adds everything requested
 	for (const FAnimNotifyEventReference& NotifyRef : NewNotifies)
 	{
diff --git a/Engine/Source/Runtime/Engine/Public/Animation/AnimNotifyQueue.h b/Engine/Source/Runtime/Engine/Public/Animation/AnimNotifyQueue.h
--- a/Engine/Source/Runtime/Engine/Public/Animation/AnimNotifyQueue.h
+++ b/Engine/Source/Runtime/Engine/Public/Animation/AnimNotifyQueue.h
@@ -218,5 +218,5 @@ private:
 	void AddAnimNotifiesToDestNoFiltering(const TArray<FAnimNotifyEventReference>& NewNotifies, TArray<FAnimNotifyEventReference>& DestArray) const;
 
 	UPROPERTY(Transient)
-	TObjectPtr<UWorld> World;
+	TWeakObjectPtr<UWorld> World;
 };
