Undo temp code that wasn't meant to be checked in.

diff --git a/Engine/Source/Runtime/Engine/Private/Materials/MaterialExpressions.cpp b/Engine/Source/Runtime/Engine/Private/Materials/MaterialExpressions.cpp
--- a/Engine/Source/Runtime/Engine/Private/Materials/MaterialExpressions.cpp
+++ b/Engine/Source/Runtime/Engine/Private/Materials/MaterialExpressions.cpp
@@ -2590,8 +2590,6 @@ void UMaterialExpressionTextureSample::PostLoad()
 	{
 		TextureObject.Expression = nullptr;
 	}
-
-	AutoSetSampleType();
 }
 
 void UMaterialExpressionTextureSample::ApplyChannelNames()
