Handle case of EDLs with no ExternalObjects in FSimulateCookedSession

diff --git a/Engine/Source/Runtime/Engine/Private/WorldPartition/DataLayer/ExternalDataLayerManager.cpp b/Engine/Source/Runtime/Engine/Private/WorldPartition/DataLayer/ExternalDataLayerManager.cpp
--- a/Engine/Source/Runtime/Engine/Private/WorldPartition/DataLayer/ExternalDataLayerManager.cpp
+++ b/Engine/Source/Runtime/Engine/Private/WorldPartition/DataLayer/ExternalDataLayerManager.cpp
@@ -826,7 +826,7 @@ bool UExternalDataLayerManager::UnregisterExternalDataLayerInstance(UExternalDat
 
 URuntimeHashExternalStreamingObjectBase* UExternalDataLayerManager::GetExternalStreamingObjectForDataLayerAsset(const UExternalDataLayerAsset* InAsset) const
 {
-	return ExternalStreamingObjects.FindChecked(InAsset);
+	return ExternalStreamingObjects.FindRef(InAsset);
 }
 
 void UExternalDataLayerManager::ForEachExternalStreamingObjects(TFunctionRef<bool(URuntimeHashExternalStreamingObjectBase*)> Func) const
diff --git a/Engine/Source/Runtime/Engine/Private/WorldPartition/WorldPartitionUtils.cpp b/Engine/Source/Runtime/Engine/Private/WorldPartition/WorldPartitionUtils.cpp
--- a/Engine/Source/Runtime/Engine/Private/WorldPartition/WorldPartitionUtils.cpp
+++ b/Engine/Source/Runtime/Engine/Private/WorldPartition/WorldPartitionUtils.cpp
@@ -76,9 +76,11 @@ bool FWorldPartitionUtils::FSimulateCookedSession::SimulateCook(const FSimulateC
 	{	
 		if (const UExternalDataLayerInstance* ExternalInstance = Cast<const UExternalDataLayerInstance>(Instance))
 		{
-			URuntimeHashExternalStreamingObjectBase* EDLStreamingObject = EDLManager->GetExternalStreamingObjectForDataLayerAsset(CastChecked<const UExternalDataLayerAsset>(ExternalInstance->GetAsset()));
-			WorldPartition->RuntimeHash->InjectExternalStreamingObject(EDLStreamingObject);
-			InjectedStreamingObjects.Add(EDLStreamingObject);
+			if (URuntimeHashExternalStreamingObjectBase* EDLStreamingObject = EDLManager->GetExternalStreamingObjectForDataLayerAsset(CastChecked<const UExternalDataLayerAsset>(ExternalInstance->GetAsset())))
+			{
+				WorldPartition->RuntimeHash->InjectExternalStreamingObject(EDLStreamingObject);
+				InjectedStreamingObjects.Add(EDLStreamingObject);
+			}
 		}
 	}
 
