Add null check to FObjectTrace::OutputObject

diff --git a/Engine/Source/Runtime/Engine/Private/ObjectTrace.cpp b/Engine/Source/Runtime/Engine/Private/ObjectTrace.cpp
--- a/Engine/Source/Runtime/Engine/Private/ObjectTrace.cpp
+++ b/Engine/Source/Runtime/Engine/Private/ObjectTrace.cpp
@@ -494,7 +494,10 @@ void FObjectTrace::OutputInstance(const UObject* InOuterObject, uint64 InInstanc
 
 void FObjectTrace::OutputObject(const UObject* InObject)
 {
-	OutputInstance(InObject, GetObjectId(InObject), GetObjectId(InObject->GetOuter()), InObject->GetClass(), InObject->GetFName().ToString(), InObject->GetPathName());
+	if (InObject)
+	{
+		OutputInstance(InObject, GetObjectId(InObject), GetObjectId(InObject->GetOuter()), InObject->GetClass(), InObject->GetFName().ToString(), InObject->GetPathName());
+	}
 }
 
 void FObjectTrace::OutputObjectTransform(const UObject* InObject, const FTransform& InTransform, ETeleportType InTeleportType)
