Backed out 37226074 and 42355427 and made a fix for both issues: brought back call to AWorldDataLayers::ResetDataLayerRuntimeState but don't reset replicated data.

diff --git a/Engine/Source/Runtime/Engine/Private/WorldPartition/DataLayer/DataLayerManager.cpp b/Engine/Source/Runtime/Engine/Private/WorldPartition/DataLayer/DataLayerManager.cpp
--- a/Engine/Source/Runtime/Engine/Private/WorldPartition/DataLayer/DataLayerManager.cpp
+++ b/Engine/Source/Runtime/Engine/Private/WorldPartition/DataLayer/DataLayerManager.cpp
@@ -237,6 +237,11 @@ void UDataLayerManager::Initialize()
 
 void UDataLayerManager::DeInitialize()
 {
+	if (AWorldDataLayers* WorldDataLayers = GetWorldDataLayers())
+	{
+		WorldDataLayers->OnDataLayerManagerDeinitialized();
+	}
+
 #if WITH_EDITOR
 	UActorDescContainerInstance::OnActorDescContainerInstanceInitialized.RemoveAll(this);
 
diff --git a/Engine/Source/Runtime/Engine/Private/WorldPartition/DataLayer/WorldDataLayers.cpp b/Engine/Source/Runtime/Engine/Private/WorldPartition/DataLayer/WorldDataLayers.cpp
--- a/Engine/Source/Runtime/Engine/Private/WorldPartition/DataLayer/WorldDataLayers.cpp
+++ b/Engine/Source/Runtime/Engine/Private/WorldPartition/DataLayer/WorldDataLayers.cpp
@@ -104,13 +104,6 @@ void AWorldDataLayers::PostRegisterAllComponents()
 	}
 }
 
-void AWorldDataLayers::PostUnregisterAllComponents()
-{
-	Super::PostUnregisterAllComponents();
-
-	ResetDataLayerRuntimeStates();
-}
-
 #define AWORLDDATALAYERS_UPDATE_REPLICATED_DATALAYERS(ReplicatedArray, SourceArray) \
 	MARK_PROPERTY_DIRTY_FROM_NAME(AWorldDataLayers, ReplicatedArray, this); \
 	ReplicatedArray = SourceArray;
@@ -219,7 +212,7 @@ void AWorldDataLayers::InitializeDataLayerRuntimeStates()
 	}
 }
 
-void AWorldDataLayers::ResetDataLayerRuntimeStates()
+void AWorldDataLayers::ResetDataLayerRuntimeStates(bool bFlushNet)
 {
 	ActiveDataLayerNames.Reset();
 	LoadedDataLayerNames.Reset();
@@ -228,12 +221,15 @@ void AWorldDataLayers::ResetDataLayerRuntimeStates()
 
 	EffectiveStates.Reset();
 	
-	static const TArray<FName> Empty;
-	FlushNetDormancy();
-	AWORLDDATALAYERS_UPDATE_REPLICATED_DATALAYERS(RepActiveDataLayerNames, Empty);
-	AWORLDDATALAYERS_UPDATE_REPLICATED_DATALAYERS(RepLoadedDataLayerNames, Empty);
-	AWORLDDATALAYERS_UPDATE_REPLICATED_DATALAYERS(RepEffectiveActiveDataLayerNames, Empty);
-	AWORLDDATALAYERS_UPDATE_REPLICATED_DATALAYERS(RepEffectiveLoadedDataLayerNames, Empty);
+	if (bFlushNet)
+	{
+		FlushNetDormancy();
+		static const TArray<FName> Empty;		
+		AWORLDDATALAYERS_UPDATE_REPLICATED_DATALAYERS(RepActiveDataLayerNames, Empty);
+		AWORLDDATALAYERS_UPDATE_REPLICATED_DATALAYERS(RepLoadedDataLayerNames, Empty);
+		AWORLDDATALAYERS_UPDATE_REPLICATED_DATALAYERS(RepEffectiveActiveDataLayerNames, Empty);
+		AWORLDDATALAYERS_UPDATE_REPLICATED_DATALAYERS(RepEffectiveLoadedDataLayerNames, Empty);
+	}
 }
 
 bool AWorldDataLayers::SetDataLayerRuntimeState(const UDataLayerInstance* InDataLayerInstance, EDataLayerRuntimeState InState, bool bInIsRecursive)
@@ -1256,6 +1252,11 @@ void AWorldDataLayers::OnDataLayerManagerInitialized()
 	InitializeDataLayerRuntimeStates();
 }
 
+void AWorldDataLayers::OnDataLayerManagerDeinitialized()
+{
+    ResetDataLayerRuntimeStates(false);
+}
+
 void AWorldDataLayers::Serialize(FArchive& Ar)
 {
 	Super::Serialize(Ar);
diff --git a/Engine/Source/Runtime/Engine/Public/WorldPartition/DataLayer/WorldDataLayers.h b/Engine/Source/Runtime/Engine/Public/WorldPartition/DataLayer/WorldDataLayers.h
--- a/Engine/Source/Runtime/Engine/Public/WorldPartition/DataLayer/WorldDataLayers.h
+++ b/Engine/Source/Runtime/Engine/Public/WorldPartition/DataLayer/WorldDataLayers.h
@@ -94,7 +94,6 @@ public:
 	//~ Begin AActor Interface
 	ENGINE_API virtual void RewindForReplay() override;
 	ENGINE_API virtual void PostRegisterAllComponents() override;
-	ENGINE_API virtual void PostUnregisterAllComponents() override;
 #if WITH_EDITOR
 	ENGINE_API virtual void PreEditUndo() override;
 	ENGINE_API virtual void PostEditUndo() override;
@@ -231,7 +230,7 @@ public:
 
 protected:
 	ENGINE_API void InitializeDataLayerRuntimeStates();
-	ENGINE_API void ResetDataLayerRuntimeStates();
+	ENGINE_API void ResetDataLayerRuntimeStates(bool bFlushNet = true);
 
 	UFUNCTION()
 	ENGINE_API void OnRep_ActiveDataLayerNames();
@@ -267,6 +266,7 @@ private:
 #endif
 
 	ENGINE_API void OnDataLayerManagerInitialized();
+	ENGINE_API void OnDataLayerManagerDeinitialized();
 	ENGINE_API void ResolveEffectiveRuntimeState(const UDataLayerInstance* InDataLayer, bool bInNotifyChange = true);
 	ENGINE_API void DumpDataLayerRecursively(const UDataLayerInstance* DataLayer, FString Prefix, FOutputDevice& OutputDevice) const;
 	bool CanChangeDataLayerRuntimeState(const UDataLayerInstance* InDataLayerInstance, ESetDataLayerRuntimeStateError* OutReason = nullptr) const;
