[ObjectTrace] restored recursive tracing for outer objects (introduced in 43775980)

diff --git a/Engine/Source/Runtime/Engine/Private/ObjectTrace.cpp b/Engine/Source/Runtime/Engine/Private/ObjectTrace.cpp
--- a/Engine/Source/Runtime/Engine/Private/ObjectTrace.cpp
+++ b/Engine/Source/Runtime/Engine/Private/ObjectTrace.cpp
@@ -496,6 +496,7 @@ void FObjectTrace::OutputObject(const UObject* InObject)
 {
 	if (InObject)
 	{
+		OutputObject(InObject->GetOuter());
 		OutputInstance(InObject, GetObjectId(InObject), GetObjectId(InObject->GetOuter()), InObject->GetClass(), InObject->GetFName().ToString(), InObject->GetPathName());
 	}
 }
