Adding a new function UpdateVelocityBeforeMovement that is called right before we start physics. Gives one final chance to update the velocity right before doing the actual position change.

#rnx

diff --git a/Engine/Source/Runtime/Engine/Classes/GameFramework/CharacterMovementComponent.h b/Engine/Source/Runtime/Engine/Classes/GameFramework/CharacterMovementComponent.h
--- a/Engine/Source/Runtime/Engine/Classes/GameFramework/CharacterMovementComponent.h
+++ b/Engine/Source/Runtime/Engine/Classes/GameFramework/CharacterMovementComponent.h
@@ -1594,6 +1594,9 @@ public:
 	/** Update the character state in PerformMovement after the position change. Some rotation updates happen after this. */
 	ENGINE_API virtual void UpdateCharacterStateAfterMovement(float DeltaSeconds);
 
+	/** Give one final chance to update the velocity right before doing the actual position change. */
+	ENGINE_API virtual void UpdateVelocityBeforeMovement(float DeltaSeconds) {}
+
 	/** 
 	 * Handle start swimming functionality
 	 * @param OldLocation - Location on last tick
diff --git a/Engine/Source/Runtime/Engine/Private/Components/CharacterMovementComponent.cpp b/Engine/Source/Runtime/Engine/Private/Components/CharacterMovementComponent.cpp
--- a/Engine/Source/Runtime/Engine/Private/Components/CharacterMovementComponent.cpp
+++ b/Engine/Source/Runtime/Engine/Private/Components/CharacterMovementComponent.cpp
@@ -2954,6 +2954,9 @@ void UCharacterMovementComponent::PerformMovement(float DeltaSeconds)
 		CharacterOwner->ClearJumpInput(DeltaSeconds);
 		NumJumpApexAttempts = 0;
 
+		// Give one final chance to update the velocity right before doing the actual position change.
+		UpdateVelocityBeforeMovement(DeltaSeconds);
+		
 		// change position
 		StartNewPhysics(DeltaSeconds, 0);
 
