PR #13417: bugfix: GetResourceSizeEx miscalculated for non-isotropic UTextureRenderTargetVolume

diff --git a/Engine/Source/Runtime/Engine/Private/TextureRenderTargetVolume.cpp b/Engine/Source/Runtime/Engine/Private/TextureRenderTargetVolume.cpp
--- a/Engine/Source/Runtime/Engine/Private/TextureRenderTargetVolume.cpp
+++ b/Engine/Source/Runtime/Engine/Private/TextureRenderTargetVolume.cpp
@@ -116,7 +116,7 @@ void UTextureRenderTargetVolume::GetResourceSizeEx(FResourceSizeEx& CumulativeRe
 	const int32 BlockBytes	= GPixelFormats[Format].BlockBytes;
 	const int32 NumBlocksX	= (SizeX + BlockSizeX - 1) / BlockSizeX;
 	const int32 NumBlocksY	= (SizeY + BlockSizeY - 1) / BlockSizeY;
-	const int32 NumBlocksZ	= (SizeZ + BlockSizeY - 1) / BlockSizeZ;
+	const int32 NumBlocksZ	= (SizeZ + BlockSizeZ - 1) / BlockSizeZ;
 	const int32 NumBytes	= NumBlocksX * NumBlocksY * NumBlocksZ * BlockBytes;
 
 	CumulativeResourceSize.AddUnknownMemoryBytes(NumBytes);
@@ -166,7 +166,7 @@ void UTextureRenderTargetVolume::PostLoad()
 
 FString UTextureRenderTargetVolume::GetDesc()
 {
-	return FString::Printf( TEXT("Render to Texture Volume %dx%d[%s]"), SizeX, SizeX, GPixelFormats[GetFormat()].Name);
+	return FString::Printf( TEXT("Render to Texture Volume %dx%dx%d[%s]"), SizeX, SizeY, SizeZ, GPixelFormats[GetFormat()].Name);
 }
 
 TSubclassOf<UTexture> UTextureRenderTargetVolume::GetTextureUClass() const
