[Vertex/Index Buffer] The goal here is to make sure enabling deformer graph support for a project alone does not incur additional memory cost, given that no deformer is used. This is done by decoupling Buffer SRV automatic creation from deformer platform support. Previously when deformer is enabled,  it also enables GPUSkinPassthrough, which triggers various systems to be less conservative with memory and create buffer SRVs automatically. Changed to using GPUSkinCache instead of GPU passthrough as the better hint for platform capability.

if skin cache is enabled, most SRVs are usually created by default.
if skin cache = off but deformer = on, buffers are marked as ShaderResource, but SRVs are not automatically created to save memory
if both skin cache and deformers are off, buffers won't even be marked as ShaderResource

diff --git a/Engine/Plugins/Animation/DeformerGraph/Source/OptimusCore/Private/DataInterfaces/OptimusDataInterfaceSkinnedMesh.cpp b/Engine/Plugins/Animation/DeformerGraph/Source/OptimusCore/Private/DataInterfaces/OptimusDataInterfaceSkinnedMesh.cpp
--- a/Engine/Plugins/Animation/DeformerGraph/Source/OptimusCore/Private/DataInterfaces/OptimusDataInterfaceSkinnedMesh.cpp
+++ b/Engine/Plugins/Animation/DeformerGraph/Source/OptimusCore/Private/DataInterfaces/OptimusDataInterfaceSkinnedMesh.cpp
@@ -3,6 +3,7 @@
 #include "OptimusDataInterfaceSkinnedMesh.h"
 
 #include "OptimusDataDomain.h"
+#include "RenderGraphBuilder.h"
 #include "ComponentSources/OptimusSkinnedMeshComponentSource.h"
 
 #include "ComputeFramework/ShaderParamTypeDefinition.h"
@@ -165,6 +166,22 @@ bool FOptimusSkinnedMeshDataProviderProxy::IsValid(FValidationData const& InVali
 	return true;
 }
 
+void FOptimusSkinnedMeshDataProviderProxy::AllocateResources(FRDGBuilder& GraphBuilder, FAllocationData const& InAllocationData)
+{
+	const int32 LodIndex = SkeletalMeshObject->GetLOD();
+	FSkeletalMeshRenderData& SkeletalMeshRenderData = SkeletalMeshObject->GetSkeletalMeshRenderData();
+	FSkeletalMeshLODRenderData* LodRenderData = &SkeletalMeshRenderData.LODRenderData[LodIndex];
+
+	FRawStaticIndexBuffer16or32Interface* IndexBufferInterface = LodRenderData->MultiSizeIndexContainer.GetIndexBuffer();
+	
+	IndexBufferSRV = IndexBufferInterface->GetOrCreateSRV(GraphBuilder.RHICmdList);
+	TangentsSRV = LodRenderData->StaticVertexBuffers.StaticMeshVertexBuffer.GetOrCreateTangentsSRV(GraphBuilder.RHICmdList);	
+	TexCoordsSRV = LodRenderData->StaticVertexBuffers.StaticMeshVertexBuffer.GetOrCreateTexCoordsSRV(GraphBuilder.RHICmdList);
+	
+	// Color SRV is either created by default or not available at all
+	ColorsSRV = LodRenderData->StaticVertexBuffers.ColorVertexBuffer.GetColorComponentsSRV();
+}
+
 void FOptimusSkinnedMeshDataProviderProxy::GatherDispatchData(FDispatchData const& InDispatchData)
 {
 	const int32 LodIndex = SkeletalMeshObject->GetLOD();
@@ -176,17 +193,17 @@ void FOptimusSkinnedMeshDataProviderProxy::GatherDispatchData(FDispatchData cons
 	const TStridedView<FParameters> ParameterArray = MakeStridedParameterView<FParameters>(InDispatchData);
 	for (int32 InvocationIndex = 0; InvocationIndex < ParameterArray.Num(); ++InvocationIndex)
 	{
-		FRHIShaderResourceView* IndexBufferSRV = LodRenderData->MultiSizeIndexContainer.GetIndexBuffer()->GetSRV();
+		FRHIShaderResourceView* MeshIndexBufferSRV = IndexBufferSRV;
 		FRHIShaderResourceView* MeshVertexBufferSRV = LodRenderData->StaticVertexBuffers.PositionVertexBuffer.GetSRV();
-		FRHIShaderResourceView* MeshTangentBufferSRV = LodRenderData->StaticVertexBuffers.StaticMeshVertexBuffer.GetTangentsSRV();
-		FRHIShaderResourceView* MeshUVBufferSRV = LodRenderData->StaticVertexBuffers.StaticMeshVertexBuffer.GetTexCoordsSRV();
-		FRHIShaderResourceView* MeshColorBufferSRV = LodRenderData->StaticVertexBuffers.ColorVertexBuffer.GetColorComponentsSRV();
+		FRHIShaderResourceView* MeshTangentBufferSRV = TangentsSRV;
+		FRHIShaderResourceView* MeshUVBufferSRV = TexCoordsSRV;
+		FRHIShaderResourceView* MeshColorBufferSRV = ColorsSRV;
 
 		FParameters& Parameters = ParameterArray[InvocationIndex];
 		Parameters.NumVertices = LodRenderData->GetNumVertices();
 		Parameters.NumTriangles =  LodRenderData->GetTotalFaces();
 		Parameters.NumUVChannels = LodRenderData->StaticVertexBuffers.StaticMeshVertexBuffer.GetNumTexCoords();
-		Parameters.IndexBuffer = IndexBufferSRV != nullptr ? IndexBufferSRV : NullSRVBinding;
+		Parameters.IndexBuffer = MeshIndexBufferSRV != nullptr ? MeshIndexBufferSRV : NullSRVBinding;
 		Parameters.PositionInputBuffer = MeshVertexBufferSRV != nullptr ? MeshVertexBufferSRV : NullSRVBinding;
 		Parameters.TangentInputBuffer = MeshTangentBufferSRV != nullptr ? MeshTangentBufferSRV : NullSRVBinding;
 		Parameters.UVInputBuffer = MeshUVBufferSRV != nullptr ? MeshUVBufferSRV : NullSRVBinding;
diff --git a/Engine/Plugins/Animation/DeformerGraph/Source/OptimusCore/Private/DataInterfaces/OptimusDataInterfaceSkinnedMesh.h b/Engine/Plugins/Animation/DeformerGraph/Source/OptimusCore/Private/DataInterfaces/OptimusDataInterfaceSkinnedMesh.h
--- a/Engine/Plugins/Animation/DeformerGraph/Source/OptimusCore/Private/DataInterfaces/OptimusDataInterfaceSkinnedMesh.h
+++ b/Engine/Plugins/Animation/DeformerGraph/Source/OptimusCore/Private/DataInterfaces/OptimusDataInterfaceSkinnedMesh.h
@@ -4,6 +4,7 @@
 
 #include "OptimusComputeDataInterface.h"
 #include "ComputeFramework/ComputeDataProvider.h"
+#include "RHIFwd.h"
 
 #include "OptimusDataInterfaceSkinnedMesh.generated.h"
 
@@ -62,14 +63,20 @@ public:
 	FOptimusSkinnedMeshDataProviderProxy(USkinnedMeshComponent* SkinnedMeshComponent);
 
 	//~ Begin FComputeDataProviderRenderProxy Interface
-	bool IsValid(FValidationData const& InValidationData) const override;
-	void GatherDispatchData(FDispatchData const& InDispatchData) override;
+	virtual bool IsValid(FValidationData const& InValidationData) const override;
+	virtual void AllocateResources(FRDGBuilder& GraphBuilder, FAllocationData const& InAllocationData) override;
+	virtual void GatherDispatchData(FDispatchData const& InDispatchData) override;
 	//~ End FComputeDataProviderRenderProxy Interface
 
 private:
 	using FParameters = FSkinnedMeshDataInterfaceParameters;
 
 	FSkeletalMeshObject* SkeletalMeshObject;
+	
+	FShaderResourceViewRHIRef IndexBufferSRV;
+	FShaderResourceViewRHIRef TangentsSRV;
+	FShaderResourceViewRHIRef TexCoordsSRV;
+	FShaderResourceViewRHIRef ColorsSRV;
 };
 
 #undef UE_API
diff --git a/Engine/Plugins/Animation/DeformerGraph/Source/OptimusCore/Private/DataInterfaces/OptimusDataInterfaceSkinnedMeshRead.cpp b/Engine/Plugins/Animation/DeformerGraph/Source/OptimusCore/Private/DataInterfaces/OptimusDataInterfaceSkinnedMeshRead.cpp
--- a/Engine/Plugins/Animation/DeformerGraph/Source/OptimusCore/Private/DataInterfaces/OptimusDataInterfaceSkinnedMeshRead.cpp
+++ b/Engine/Plugins/Animation/DeformerGraph/Source/OptimusCore/Private/DataInterfaces/OptimusDataInterfaceSkinnedMeshRead.cpp
@@ -220,7 +220,14 @@ void FOptimusSkinnedMeshReadDataProviderProxy::AllocateResources(FRDGBuilder& Gr
 		bInvalidatePreviousPosition = true;
 		*LastLodIndexPtr = LodIndex;
 	}
+
+	FSkeletalMeshRenderData const& SkeletalMeshRenderData = SkeletalMeshObject->GetSkeletalMeshRenderData();
+	FSkeletalMeshLODRenderData const* LodRenderData = &SkeletalMeshRenderData.LODRenderData[LodIndex];
+	
+	TangentsSRV = LodRenderData->StaticVertexBuffers.StaticMeshVertexBuffer.GetOrCreateTangentsSRV(GraphBuilder.RHICmdList);
 	
+	// Color SRV is either created by default or not available at all
+	ColorsSRV = LodRenderData->StaticVertexBuffers.ColorVertexBuffer.GetColorComponentsSRV();
 
 	if (InputMask & static_cast<uint64>(ESkinnedMeshReadDataInterfaceInputSelectorMask::Position) )
 	{
@@ -309,8 +316,8 @@ void FOptimusSkinnedMeshReadDataProviderProxy::GatherDispatchData(FDispatchData
 		Parameters.ColorBufferSRV = ColorBufferSRV;
 		
 		FRHIShaderResourceView* MeshVertexBufferSRV = LodRenderData->StaticVertexBuffers.PositionVertexBuffer.GetSRV();
-		FRHIShaderResourceView* MeshTangentBufferSRV = LodRenderData->StaticVertexBuffers.StaticMeshVertexBuffer.GetTangentsSRV();
-		FRHIShaderResourceView* MeshColorBufferSRV = LodRenderData->StaticVertexBuffers.ColorVertexBuffer.GetColorComponentsSRV();
+		FRHIShaderResourceView* MeshTangentBufferSRV = TangentsSRV;
+		FRHIShaderResourceView* MeshColorBufferSRV = ColorsSRV;
 
 		Parameters.PositionStaticBuffer = MeshVertexBufferSRV != nullptr ? MeshVertexBufferSRV : NullSRVBinding;
 		Parameters.TangentStaticBuffer = MeshTangentBufferSRV != nullptr ? MeshTangentBufferSRV : NullSRVBinding;
diff --git a/Engine/Plugins/Animation/DeformerGraph/Source/OptimusCore/Private/DataInterfaces/OptimusDataInterfaceSkinnedMeshRead.h b/Engine/Plugins/Animation/DeformerGraph/Source/OptimusCore/Private/DataInterfaces/OptimusDataInterfaceSkinnedMeshRead.h
--- a/Engine/Plugins/Animation/DeformerGraph/Source/OptimusCore/Private/DataInterfaces/OptimusDataInterfaceSkinnedMeshRead.h
+++ b/Engine/Plugins/Animation/DeformerGraph/Source/OptimusCore/Private/DataInterfaces/OptimusDataInterfaceSkinnedMeshRead.h
@@ -5,6 +5,7 @@
 #include "IOptimusDeformerInstanceAccessor.h"
 #include "OptimusComputeDataInterface.h"
 #include "ComputeFramework/ComputeDataProvider.h"
+#include "RHIFwd.h"
 
 #include "OptimusDataInterfaceSkinnedMeshRead.generated.h"
 
@@ -100,6 +101,9 @@ private:
 	int32* LastLodIndexPtr = nullptr;
 	EMeshDeformerOutputBuffer OutputBuffersFromPreviousInstances;
 
+	FShaderResourceViewRHIRef TangentsSRV;
+	FShaderResourceViewRHIRef ColorsSRV;
+	
 	FRDGBufferSRV* PositionBufferSRV = nullptr;
 	FRDGBufferSRV* TangentBufferSRV = nullptr;
 	FRDGBufferSRV* ColorBufferSRV = nullptr;
diff --git a/Engine/Plugins/FX/Niagara/Source/Niagara/Private/DataInterface/NiagaraDataInterfaceStaticMesh.cpp b/Engine/Plugins/FX/Niagara/Source/Niagara/Private/DataInterface/NiagaraDataInterfaceStaticMesh.cpp
--- a/Engine/Plugins/FX/Niagara/Source/Niagara/Private/DataInterface/NiagaraDataInterfaceStaticMesh.cpp
+++ b/Engine/Plugins/FX/Niagara/Source/Niagara/Private/DataInterface/NiagaraDataInterfaceStaticMesh.cpp
@@ -442,7 +442,14 @@ namespace NDIStaticMeshLocal
 			if (GpuInitializeData.LODResource)
 			{
 				FBufferRHIRef IndexBufferRHIRef = GpuInitializeData.LODResource->IndexBuffer.IndexBufferRHI;
-				const bool bCanCreateIndexSRV = IndexBufferRHIRef.IsValid() && ((IndexBufferRHIRef->GetUsage() & EBufferUsageFlags::ShaderResource) == EBufferUsageFlags::ShaderResource);
+				const bool bAllowCPUAccess = GpuInitializeData.LODResource->IndexBuffer.GetAllowCPUAccess();
+				const bool bCanCreateIndexSRV =
+					IndexBufferRHIRef.IsValid() &&
+					((IndexBufferRHIRef->GetUsage() & EBufferUsageFlags::ShaderResource) == EBufferUsageFlags::ShaderResource) &&
+					// Added just to keep old behavior, where whether the SRV should be created has been inferred from whether GPU Skin cache is enabled or not
+					// Typically if GPU Skin Cache is enabled, the platform likely does not care about memory cost from extra SRVs
+					// Individual systems can further decide what to do on platforms that don't support GPU Skin Cache
+					(AreBufferSRVsAlwaysCreatedByDefault(GMaxRHIShaderPlatform) || bAllowCPUAccess); 
 
 				if (bCanCreateIndexSRV)
 				{
diff --git a/Engine/Plugins/FX/Niagara/Source/Niagara/Private/NiagaraCommon.cpp b/Engine/Plugins/FX/Niagara/Source/Niagara/Private/NiagaraCommon.cpp
--- a/Engine/Plugins/FX/Niagara/Source/Niagara/Private/NiagaraCommon.cpp
+++ b/Engine/Plugins/FX/Niagara/Source/Niagara/Private/NiagaraCommon.cpp
@@ -1042,7 +1042,7 @@ bool FNiagaraUtilities::AllowGPUCulling()
 
 bool FNiagaraUtilities::AreBufferSRVsAlwaysCreated(EShaderPlatform ShaderPlatform)
 {
-	return RHISupportsManualVertexFetch(ShaderPlatform) || IsGPUSkinPassThroughSupported(ShaderPlatform);
+	return AreBufferSRVsAlwaysCreatedByDefault(ShaderPlatform);
 }
 
 ENiagaraCompileUsageStaticSwitch FNiagaraUtilities::ConvertScriptUsageToStaticSwitchUsage(ENiagaraScriptUsage ScriptUsage)
diff --git a/Engine/Plugins/PCG/Source/PCG/Private/Compute/DataInterfaces/PCGStaticMeshDataInterface.cpp b/Engine/Plugins/PCG/Source/PCG/Private/Compute/DataInterfaces/PCGStaticMeshDataInterface.cpp
--- a/Engine/Plugins/PCG/Source/PCG/Private/Compute/DataInterfaces/PCGStaticMeshDataInterface.cpp
+++ b/Engine/Plugins/PCG/Source/PCG/Private/Compute/DataInterfaces/PCGStaticMeshDataInterface.cpp
@@ -270,7 +270,14 @@ void FPCGStaticMeshDataProviderProxy::GatherDispatchData(FDispatchData const& In
 void FPCGStaticMeshDataProviderProxy::AllocateResources(FRDGBuilder& GraphBuilder, FAllocationData const& InAllocationData)
 {
 	FBufferRHIRef IndexBufferRHIRef = LODResources ? LODResources->IndexBuffer.IndexBufferRHI : nullptr;
-	const bool bCanCreateIndexSRV = IndexBufferRHIRef.IsValid() && ((IndexBufferRHIRef->GetUsage() & EBufferUsageFlags::ShaderResource) == EBufferUsageFlags::ShaderResource);
+	const bool bAllowCPUAccess = LODResources? LODResources->IndexBuffer.GetAllowCPUAccess() : false;
+	const bool bCanCreateIndexSRV =
+		IndexBufferRHIRef.IsValid() &&
+		((IndexBufferRHIRef->GetUsage() & EBufferUsageFlags::ShaderResource) == EBufferUsageFlags::ShaderResource) &&
+		// Added just to keep old behavior, where whether the SRV should be created has been inferred from whether GPU Skin cache is enabled or not
+		// Typically if GPU Skin Cache is enabled, the platform likely does not care about memory cost from extra SRVs
+		// Individual systems can further decide what to do on platforms that don't support GPU Skin Cache
+		(AreBufferSRVsAlwaysCreatedByDefault(GMaxRHIShaderPlatform) || bAllowCPUAccess); 
 
 	if (bCanCreateIndexSRV)
 	{
diff --git a/Engine/Source/Runtime/Engine/Private/RawIndexBuffer.cpp b/Engine/Source/Runtime/Engine/Private/RawIndexBuffer.cpp
--- a/Engine/Source/Runtime/Engine/Private/RawIndexBuffer.cpp
+++ b/Engine/Source/Runtime/Engine/Private/RawIndexBuffer.cpp
@@ -25,6 +25,14 @@ void CacheOptimizeIndexBuffer(TArray<IndexDataType,Allocator>& Indices)
 }
 #endif // #if WITH_EDITOR
 
+/**
+ * Whether buffers should have ShaderResource usage flag, so that systems can optionally opt in to create SRVs for them 
+ */
+static bool AreIndexBuffersShaderResources(EShaderPlatform ShaderPlatform)
+{
+	return AreBufferSRVsAlwaysCreatedByDefault(ShaderPlatform) || IsGPUSkinPassThroughSupported(ShaderPlatform);
+} 
+
 /*-----------------------------------------------------------------------------
 FRawIndexBuffer
 -----------------------------------------------------------------------------*/
@@ -335,15 +343,15 @@ FBufferRHIRef FRawStaticIndexBuffer::CreateRHIBuffer(FRHICommandListBase& RHICmd
 
 	if (GetNumIndices() > 0)
 	{
-		// Systems that generate data for GPUSkinPassThrough use index buffer as SRV.
-		bool bSRV = RHISupportsManualVertexFetch(GMaxRHIShaderPlatform) || IsGPUSkinPassThroughSupported(GMaxRHIShaderPlatform);
+		// Systems that generate data for GPUSkinPassThrough use index buffer as SRV, depending on platform, SRV may be created by default or by individual systems
+		bool bIsShaderResource = AreIndexBuffersShaderResources(GMaxRHIShaderPlatform);
 
 		// When bAllowCPUAccess is true, the meshes is likely going to be used for Niagara to spawn particles on mesh surface.
 		// And it can be the case for CPU *and* GPU access: no differenciation today. That is why we create a SRV in this case.
 		// This also avoid setting lots of states on all the members of all the different buffers used by meshes. Follow up: https://jira.it.epicgames.net/browse/UE-69376.
-		bSRV |= IndexStorage.GetAllowCPUAccess();
+		bIsShaderResource |= IndexStorage.GetAllowCPUAccess();
 
-		const EBufferUsageFlags BufferFlags = EBufferUsageFlags::Static | (bSRV ? EBufferUsageFlags::ShaderResource : EBufferUsageFlags::None);
+		const EBufferUsageFlags BufferFlags = EBufferUsageFlags::Static | (bIsShaderResource ? EBufferUsageFlags::ShaderResource : EBufferUsageFlags::None);
 
 		const static FLazyName ClassName32(TEXT("FRawStaticIndexBuffer32"));
 		const static FLazyName ClassName16(TEXT("FRawStaticIndexBuffer16"));
@@ -444,6 +452,30 @@ void FRawStaticIndexBuffer::Discard()
     IndexStorage.Discard();
 }
 
+bool FRawStaticIndexBuffer16or32Interface::IsSRVCreatedByDefault(bool bAllowCPUAccess) const
+{
+	// Use platform specific features like ManualVertexFetch and GPUSkinCache to infer platform capability
+	bool bSRV = AreBufferSRVsAlwaysCreatedByDefault(GMaxRHIShaderPlatform);
+	// When bAllowCPUAccess is true, the meshes is likely going to be used for Niagara to spawn particles on mesh surface.
+	// And it can be the case for CPU *and* GPU access: no differenciation today. That is why we create a SRV in this case.
+	// This also avoid setting lots of states on all the members of all the different buffers used by meshes. Follow up: https://jira.it.epicgames.net/browse/UE-69376.
+	bSRV |= bAllowCPUAccess;
+	return bSRV;
+}
+
+bool FRawStaticIndexBuffer16or32Interface::IsShaderResource(bool bAllowCPUAccess) const
+{
+	// In the case that AreBufferSRVsAlwaysCreatedByDefault is false, there might be other systems (eg. MeshDeformer)
+	// that still want to create SRVs on demand only for specific meshes in the project.
+	bool bCanBeShaderResource = AreIndexBuffersShaderResources(GMaxRHIShaderPlatform);
+	// When bAllowCPUAccess is true, the meshes is likely going to be used for Niagara to spawn particles on mesh surface.
+	// And it can be the case for CPU *and* GPU access: no differenciation today. That is why we create a SRV in this case.
+	// This also avoid setting lots of states on all the members of all the different buffers used by meshes. Follow up: https://jira.it.epicgames.net/browse/UE-69376.
+	bCanBeShaderResource |= bAllowCPUAccess;
+	return bCanBeShaderResource;	
+}
+
+// Deprecated
 bool FRawStaticIndexBuffer16or32Interface::IsSRVNeeded(bool bAllowCPUAccess) const
 {
 	// Systems that generate data for GPUSkinPassThrough use index buffer as SRV.
@@ -478,7 +510,7 @@ FBufferRHIRef FRawStaticIndexBuffer16or32Interface::CreateRHIIndexBufferInternal
 	int32 IndexCount,
 	size_t IndexSize,
 	FResourceArrayInterface* ResourceArray,
-	bool bNeedSRV
+	bool bIsShaderResource
 )
 {
 	const uint32 Size = IndexCount * IndexSize;
@@ -493,7 +525,7 @@ FBufferRHIRef FRawStaticIndexBuffer16or32Interface::CreateRHIIndexBufferInternal
 
 	if (!CreateDesc.IsNull())
 	{
-		if (bNeedSRV)
+		if (bIsShaderResource)
 		{
 			// EBufferUsageFlags::ShaderResource is needed for SkinCache RecomputeSkinTangents
 			CreateDesc.AddUsage(EBufferUsageFlags::ShaderResource);
diff --git a/Engine/Source/Runtime/Engine/Private/Rendering/StaticMeshVertexBuffer.cpp b/Engine/Source/Runtime/Engine/Private/Rendering/StaticMeshVertexBuffer.cpp
--- a/Engine/Source/Runtime/Engine/Private/Rendering/StaticMeshVertexBuffer.cpp
+++ b/Engine/Source/Runtime/Engine/Private/Rendering/StaticMeshVertexBuffer.cpp
@@ -322,29 +322,19 @@ void FStaticMeshVertexBuffer::InitRHI(FRHICommandListBase& RHICmdList)
 	const bool bHadTangentsData = TangentsData != nullptr;
 	const bool bCreateTangentsSRV = bHadTangentsData && TangentsData->GetAllowCPUAccess();
 	TangentsVertexBuffer.VertexBufferRHI = CreateTangentsRHIBuffer(RHICmdList);
-	if (TangentsVertexBuffer.VertexBufferRHI && (bCreateTangentsSRV || RHISupportsManualVertexFetch(GMaxRHIShaderPlatform) || IsGPUSkinPassThroughSupported(GMaxRHIShaderPlatform)))
+
+	if (bCreateTangentsSRV || AreBufferSRVsAlwaysCreatedByDefault(GMaxRHIShaderPlatform))
 	{
-		uint32       Stride = GetUseHighPrecisionTangentBasis() ? 8 : 4;
-		EPixelFormat Format = GetUseHighPrecisionTangentBasis() ? PF_R16G16B16A16_SNORM : PF_R8G8B8A8_SNORM;
-		TangentsSRV = RHICmdList.CreateShaderResourceView(
-			TangentsVertexBuffer.VertexBufferRHI,
-			FRHIViewDesc::CreateBufferSRV()
-				.SetType(FRHIViewDesc::EBufferType::Typed)
-				.SetFormat(Format));
+		TangentsSRV = GetOrCreateTangentsSRV(RHICmdList);
 	}
 
 	const bool bHadTexCoordData = TexcoordData != nullptr;
 	const bool bCreateTexCoordSRV = bHadTexCoordData && TexcoordData->GetAllowCPUAccess();
 	TexCoordVertexBuffer.VertexBufferRHI = CreateTexCoordRHIBuffer(RHICmdList);
-	if (TexCoordVertexBuffer.VertexBufferRHI && (bCreateTexCoordSRV || RHISupportsManualVertexFetch(GMaxRHIShaderPlatform) || IsGPUSkinPassThroughSupported(GMaxRHIShaderPlatform)))
+	
+	if (bCreateTexCoordSRV || AreBufferSRVsAlwaysCreatedByDefault(GMaxRHIShaderPlatform))
 	{
-		uint32       Stride = GetUseFullPrecisionUVs() ? 8 : 4;
-		EPixelFormat Format = GetUseFullPrecisionUVs() ? PF_G32R32F : PF_G16R16F;
-		TextureCoordinatesSRV = RHICmdList.CreateShaderResourceView(
-			TexCoordVertexBuffer.VertexBufferRHI,
-			FRHIViewDesc::CreateBufferSRV()
-				.SetType(FRHIViewDesc::EBufferType::Typed)
-				.SetFormat(Format));
+		TextureCoordinatesSRV = GetOrCreateTexCoordsSRV(RHICmdList);	
 	}
 }
 
@@ -435,6 +425,55 @@ int FStaticMeshVertexBuffer::GetTexCoordSize() const
 	}
 }
 
+FShaderResourceViewRHIRef FStaticMeshVertexBuffer::GetOrCreateTangentsSRV(FRHICommandListBase& RHICmdList) const
+{
+	// Get the existing SRV or create a new externally owned SRV
+	if (TangentsSRV.IsValid())
+	{
+		return TangentsSRV;
+	}
+	
+	if (!TangentsVertexBuffer.VertexBufferRHI)
+	{
+		return {};
+	}
+	
+	uint32       Stride = GetUseHighPrecisionTangentBasis() ? 8 : 4;
+	EPixelFormat Format = GetUseHighPrecisionTangentBasis() ? PF_R16G16B16A16_SNORM : PF_R8G8B8A8_SNORM;
+	
+	FShaderResourceViewRHIRef ExternalTangentsSRV = RHICmdList.CreateShaderResourceView(
+		TangentsVertexBuffer.VertexBufferRHI,
+		FRHIViewDesc::CreateBufferSRV()
+		.SetType(FRHIViewDesc::EBufferType::Typed)
+		.SetFormat(Format));
+	
+	return ExternalTangentsSRV;
+}
+
+FShaderResourceViewRHIRef FStaticMeshVertexBuffer::GetOrCreateTexCoordsSRV(FRHICommandListBase& RHICmdList) const
+{
+	// Get the existing SRV or create a new externally owned SRV
+	if (TextureCoordinatesSRV.IsValid())
+	{
+		return TextureCoordinatesSRV;
+	}
+	
+	if (!TexCoordVertexBuffer.VertexBufferRHI)
+	{
+		return {};
+	}
+
+	uint32       Stride = GetUseFullPrecisionUVs() ? 8 : 4;
+	EPixelFormat Format = GetUseFullPrecisionUVs() ? PF_G32R32F : PF_G16R16F;
+	FShaderResourceViewRHIRef ExternalTextureCoordinatesSRV = RHICmdList.CreateShaderResourceView(
+		TexCoordVertexBuffer.VertexBufferRHI,
+		FRHIViewDesc::CreateBufferSRV()
+			.SetType(FRHIViewDesc::EBufferType::Typed)
+			.SetFormat(Format));
+
+	return ExternalTextureCoordinatesSRV;
+}
+
 void FStaticMeshVertexBuffer::InitTangentAndTexCoordStrides()
 {
 	typedef TStaticMeshVertexTangentDatum<typename TStaticMeshVertexTangentTypeSelector<EStaticMeshVertexTangentBasisType::HighPrecision>::TangentTypeT> HighPrecTangentType;
diff --git a/Engine/Source/Runtime/Engine/Public/RawIndexBuffer.h b/Engine/Source/Runtime/Engine/Public/RawIndexBuffer.h
--- a/Engine/Source/Runtime/Engine/Public/RawIndexBuffer.h
+++ b/Engine/Source/Runtime/Engine/Public/RawIndexBuffer.h
@@ -353,9 +353,17 @@ public:
 		return SRVValue;
 	}
 
+	// Allows for creating SRV on demand
+	// Note: consider cache the SRV locally if multiple compute tasks all need access to one in a single frame
+	virtual FShaderResourceViewRHIRef GetOrCreateSRV(FRHICommandListBase& RHICmdList) const = 0;
+
 protected:
-	ENGINE_API bool IsSRVNeeded(bool bAllowCPUAccess) const;
+	ENGINE_API bool IsSRVCreatedByDefault(bool bAllowCPUAccess) const;
+	ENGINE_API bool IsShaderResource(bool bAllowCPUAccess) const;
 
+	UE_DEPRECATED(5.7, "Use IsSRVCreatedByDefault or IsShaderResource instead, which gives more granular control over platform specific behavior")	
+	ENGINE_API bool IsSRVNeeded(bool bAllowCPUAccess) const;
+	
 	/** Similar to Init/ReleaseRHI but only update existing SRV so references to the SRV stays valid */
 	void InitRHIForStreaming(FRHIBuffer* IntermediateBuffer, size_t IndexSize, FRHIResourceReplaceBatcher& Batcher);
 	void ReleaseRHIForStreaming(FRHIResourceReplaceBatcher& Batcher);
@@ -367,7 +375,7 @@ protected:
 		int32 IndexCount,
 		size_t IndexSize,
 		FResourceArrayInterface* ResourceArray,
-		bool bNeedSRV
+		bool bIsShaderResource
 	);
 
 	// guaranteed only to be valid if the vertex buffer is valid and the buffer was created with the SRV flags
@@ -394,17 +402,11 @@ public:
 	*/
 	virtual void InitRHI(FRHICommandListBase& RHICmdList) override
 	{
-		const bool bHadIndexData = Num() > 0;
 		IndexBufferRHI = CreateRHIBuffer(RHICmdList);
 
-		if (IndexBufferRHI && IsSRVNeeded(Indices.GetAllowCPUAccess()) && bHadIndexData)
+		if (IsSRVCreatedByDefault(Indices.GetAllowCPUAccess()))
 		{
-			// If the index buffer is a placeholder we still need to create a FRHIShaderResourceView.
-			SRVValue = RHICmdList.CreateShaderResourceView(
-				IndexBufferRHI,
-				FRHIViewDesc::CreateBufferSRV()
-					.SetType(FRHIViewDesc::EBufferType::Typed)
-					.SetFormat(sizeof(INDEX_TYPE) == 2 ? PF_R16_UINT : PF_R32_UINT));
+			SRVValue = GetOrCreateSRV(RHICmdList);
 		}
 	}
 	
@@ -415,6 +417,38 @@ public:
 		SRVValue.SafeRelease();
 	}
 
+	virtual FShaderResourceViewRHIRef GetOrCreateSRV(FRHICommandListBase& RHICmdList) const override
+	{
+		if (SRVValue.IsValid())
+		{
+			return SRVValue;
+		}
+		
+		if (!IndexBufferRHI)
+		{
+			return {};
+		}
+		
+		const bool bHadIndexData = Num() > 0;
+		if (!bHadIndexData)
+		{
+			return {};
+		}
+		
+		if ((IndexBufferRHI->GetUsage() & EBufferUsageFlags::ShaderResource) != EBufferUsageFlags::ShaderResource)
+		{
+			return {};
+		}
+		
+		FShaderResourceViewRHIRef ExternalSRV = RHICmdList.CreateShaderResourceView(
+				IndexBufferRHI,
+				FRHIViewDesc::CreateBufferSRV()
+					.SetType(FRHIViewDesc::EBufferType::Typed)
+					.SetFormat(sizeof(INDEX_TYPE) == 2 ? PF_R16_UINT : PF_R32_UINT));
+		
+		return ExternalSRV;
+	}
+	
 	/**
 	* Serializer for this class
 	* @param Ar - archive to serialize to
@@ -517,7 +551,7 @@ public:
 				Indices.Num(),
 				sizeof(INDEX_TYPE),
 				&Indices,
-				IsSRVNeeded(Indices.GetAllowCPUAccess())
+				IsShaderResource(Indices.GetAllowCPUAccess())
 			);
 		}
 		return nullptr;
diff --git a/Engine/Source/Runtime/Engine/Public/Rendering/StaticMeshVertexBuffer.h b/Engine/Source/Runtime/Engine/Public/Rendering/StaticMeshVertexBuffer.h
--- a/Engine/Source/Runtime/Engine/Public/Rendering/StaticMeshVertexBuffer.h
+++ b/Engine/Source/Runtime/Engine/Public/Rendering/StaticMeshVertexBuffer.h
@@ -493,6 +493,12 @@ public:
 
 	FRHIShaderResourceView* GetTangentsSRV() const { return TangentsSRV; }
 	FRHIShaderResourceView* GetTexCoordsSRV() const { return TextureCoordinatesSRV; }
+
+	// Allows for creating SRV on demand
+	// Note: consider cache the SRV locally if multiple compute tasks all need access to one in a single frame
+	ENGINE_API FShaderResourceViewRHIRef GetOrCreateTangentsSRV(FRHICommandListBase& RHICmdList) const;
+	ENGINE_API FShaderResourceViewRHIRef GetOrCreateTexCoordsSRV(FRHICommandListBase& RHICmdList) const;
+
 private:
 
 	/** The vertex data storage type */
diff --git a/Engine/Source/Runtime/RenderCore/Private/RenderUtils.cpp b/Engine/Source/Runtime/RenderCore/Private/RenderUtils.cpp
--- a/Engine/Source/Runtime/RenderCore/Private/RenderUtils.cpp
+++ b/Engine/Source/Runtime/RenderCore/Private/RenderUtils.cpp
@@ -1509,6 +1509,15 @@ RENDERCORE_API bool IsGPUSkinPassThroughSupported(EShaderPlatform Platform)
 	return bMeshDeformersAvailable || IsGPUSkinCacheAvailable(Platform);
 }
 
+RENDERCORE_API bool AreBufferSRVsAlwaysCreatedByDefault(EShaderPlatform ShaderPlatform)
+{
+	// If manual Vertex fetch is supported, vertex data is passed in via an SRV instead of using hardware vertex fetch
+	// GPUSkinCache uses GPUSkinPassThrough + Compute shader and thus also needs SRV
+	// Mesh Deformer is not included here because sometimes mesh deformer are only used on specific meshes so
+	// we want to avoid creating SRVs for meshes that never need them.
+	return RHISupportsManualVertexFetch(ShaderPlatform) || IsGPUSkinCacheAvailable(ShaderPlatform);
+}
+
 RENDERCORE_API bool DoesRuntimeSupportOnePassPointLightShadows(EShaderPlatform Platform)
 {
 	static const auto CVar = IConsoleManager::Get().FindTConsoleVariableDataInt(TEXT("r.Shadow.DetectVertexShaderLayerAtRuntime"));
diff --git a/Engine/Source/Runtime/RenderCore/Public/RenderUtils.h b/Engine/Source/Runtime/RenderCore/Public/RenderUtils.h
--- a/Engine/Source/Runtime/RenderCore/Public/RenderUtils.h
+++ b/Engine/Source/Runtime/RenderCore/Public/RenderUtils.h
@@ -527,6 +527,12 @@ RENDERCORE_API bool IsGPUSkinCacheAvailable(EShaderPlatform Platform);
  */
 RENDERCORE_API bool IsGPUSkinPassThroughSupported(EShaderPlatform Platform);
 
+/**
+ * Returns true if we need to create SRVs for all index/vertex buffers.
+ * This has a memory overhead for the RHI objects so we only return true if the platform or project settings require this.
+ */
+RENDERCORE_API bool AreBufferSRVsAlwaysCreatedByDefault(EShaderPlatform ShaderPlatform);
+
 /*
  * Detect (at runtime) if the runtime supports rendering one-pass point light shadows (i.e., cube maps)
  */
