[MIR] Implement Vector and Position transform material expressions.

- Implement Build() functions for Transform and TransformPosition material expressions.
- Handle matrix swizzles in LowerSubscript().
- Move GetMaterialCommonBasis() functions into UE::MaterialTranslatorUtils namespace to be shared between old and new translator.
- Rename "SubstituteTagsInInlineHLSL" from EValueFlags enum to "SubstituteTags".
- Add several external code declarations to BaseMaterialExpressions.ini file for the various matrices.
- Also support SubstituteTags flag for dynamic inline-HLSL.

#rnx

diff --git a/Engine/Config/BaseMaterialExpressions.ini b/Engine/Config/BaseMaterialExpressions.ini
--- a/Engine/Config/BaseMaterialExpressions.ini
+++ b/Engine/Config/BaseMaterialExpressions.ini
@@ -1057,6 +1057,25 @@ bShaderInputData=true
 +ExternalCodeDeclarations=(ReturnType=Bool,		Name="ReflectionCapturePassSwitch",					Definition="GetReflectionCapturePassSwitchState()")
 +ExternalCodeDeclarations=(ReturnType=Bool,		Name="ShadowReplace",								Definition="GetShadowReplaceState()")
 
+; External code declarations for Transform and TransformPosition nodes
++ExternalCodeDeclarations=(ReturnType=Float4x4,		Name="CameraViewToTranslatedWorld", Definition="ResolvedView.<PREV>CameraViewToTranslatedWorld", Derivative=Zero, bIsInlined=true)
++ExternalCodeDeclarations=(ReturnType=LWCMatrix,	Name="GetInstanceToWorld",			Definition="Get<PREV>InstanceToWorld(Parameters)", Derivative=Zero, bIsInlined=true, EnvironmentDefines=((Name="NEEDS_INSTANCE_LOCAL_TO_WORLD_PS", ShaderFrequency=Pixel)))
++ExternalCodeDeclarations=(ReturnType=LWCMatrix,	Name="GetInstanceToWorldDF",		Definition="Get<PREV>InstanceToWorldDF(Parameters)", Derivative=Zero, bIsInlined=true, EnvironmentDefines=((Name="NEEDS_INSTANCE_LOCAL_TO_WORLD_PS", ShaderFrequency=Pixel)))
++ExternalCodeDeclarations=(ReturnType=LWCMatrix,	Name="GetLocalToWorldDF",			Definition="Get<PREV>LocalToWorldDF(Parameters)", Derivative=Zero, bIsInlined=true)
++ExternalCodeDeclarations=(ReturnType=LWCMatrix,	Name="GetParticleToWorld",			Definition="GetParticleToWorld(Parameters)", Derivative=Zero, bIsInlined=true, EnvironmentDefines=((Name="NEEDS_PARTICLE_LOCAL_TO_WORLD")))
++ExternalCodeDeclarations=(ReturnType=Float3,		Name="GetPreViewTranslation",		Definition="Get<PREV>PreViewTranslation(Parameters)", Derivative=Zero, bIsInlined=true)
++ExternalCodeDeclarations=(ReturnType=LWCMatrix,	Name="GetWorldToInstance",			Definition="GetWorldToInstance(Parameters)", Derivative=Zero, bIsInlined=true, EnvironmentDefines=((Name="NEEDS_INSTANCE_WORLD_TO_LOCAL_PS", ShaderFrequency=Pixel)))
++ExternalCodeDeclarations=(ReturnType=LWCMatrix,	Name="GetWorldToInstanceDF",		Definition="GetWorldToInstanceDF(Parameters)", Derivative=Zero, bIsInlined=true, EnvironmentDefines=((Name="NEEDS_INSTANCE_WORLD_TO_LOCAL_PS", ShaderFrequency=Pixel)))
++ExternalCodeDeclarations=(ReturnType=LWCMatrix,	Name="GetWorldToLocal",				Definition="Get<PREV>WorldToLocal(Parameters)", Derivative=Zero, bIsInlined=true, Domains=(MD_Surface, MD_Volume))
++ExternalCodeDeclarations=(ReturnType=LWCMatrix,	Name="GetWorldToLocalDF",			Definition="Get<PREV>WorldToLocalDF(Parameters)", Derivative=Zero, bIsInlined=true, Domains=(MD_Surface, MD_Volume))
++ExternalCodeDeclarations=(ReturnType=LWCMatrix,	Name="GetWorldToParticle",			Definition="GetWorldToParticle(Parameters)", Derivative=Zero, bIsInlined=true, EnvironmentDefines=((Name="NEEDS_PARTICLE_WORLD_TO_LOCAL")))
++ExternalCodeDeclarations=(ReturnType=LWCMatrix,	Name="ParticleToWorld",				Definition="Parameters.Particle.ParticleToWorld", Derivative=Zero, bIsInlined=true, EnvironmentDefines=((Name="NEEDS_PARTICLE_LOCAL_TO_WORLD")))
++ExternalCodeDeclarations=(ReturnType=Float3x3,		Name="TangentToWorld",				Definition="Parameters.TangentToWorld", Derivative=Zero, bIsInlined=true)
++ExternalCodeDeclarations=(ReturnType=Float4x4,		Name="TranslatedWorldToCameraView", Definition="ResolvedView.<PREV>TranslatedWorldToCameraView", Derivative=Zero, bIsInlined=true)
++ExternalCodeDeclarations=(ReturnType=Float4x4,		Name="TranslatedWorldToView",		Definition="ResolvedView.<PREV>TranslatedWorldToView", Derivative=Zero, bIsInlined=true)
++ExternalCodeDeclarations=(ReturnType=Float4x4,		Name="ViewToTranslatedWorld",		Definition="ResolvedView.<PREV>ViewToTranslatedWorld", Derivative=Zero, bIsInlined=true)
++ExternalCodeDeclarations=(ReturnType=LWCMatrix,	Name="WorldToParticle",				Definition="Parameters.Particle.WorldToParticle", Derivative=Zero, bIsInlined=true, EnvironmentDefines=((Name="NEEDS_PARTICLE_WORLD_TO_LOCAL")))
+
 [/Script/Engine.MaterialExpressionEnumerationParameter]
 +MenuCategories=NSLOCTEXT("MaterialExpression", "Parameters", "Parameters")
 
diff --git a/Engine/Source/Runtime/Engine/Private/Materials/MaterialExpressions.cpp b/Engine/Source/Runtime/Engine/Private/Materials/MaterialExpressions.cpp
--- a/Engine/Source/Runtime/Engine/Private/Materials/MaterialExpressions.cpp
+++ b/Engine/Source/Runtime/Engine/Private/Materials/MaterialExpressions.cpp
@@ -24,6 +24,7 @@
 #include "DataDrivenShaderPlatformInfo.h"
 #include "MaterialDomain.h"
 #include "MaterialShared.h"
+#include "MaterialSharedPrivate.h"
 #include "Materials/HLSLMaterialTranslator.h"
 #include "Materials/MaterialAttributeDefinitionMap.h"
 #include "Materials/MaterialExpressionMaterialAttributeLayers.h"
@@ -11314,37 +11315,11 @@ void UMaterialExpressionAbs::GetCaption(TArray<FString>& OutCaptions) const
 // UMaterialExpressionTransform
 ///////////////////////////////////////////////////////////////////////////////
 
-static EMaterialCommonBasis GetMaterialCommonBasis(EMaterialVectorCoordTransformSource X)
-{
-	static const EMaterialCommonBasis ConversionTable[TRANSFORMSOURCE_MAX] = {
-		MCB_Tangent,					// TRANSFORMSOURCE_Tangent
-		MCB_Local,						// TRANSFORMSOURCE_Local
-		MCB_World,						// TRANSFORMSOURCE_World
-		MCB_View,						// TRANSFORMSOURCE_View
-		MCB_Camera,						// TRANSFORMSOURCE_Camera
-		MCB_MeshParticle,				// TRANSFORMSOURCE_Particle
-		MCB_Instance,					// TRANSFORMSOURCE_Instance
-	};
-	return ConversionTable[X];
-}
-
-static EMaterialCommonBasis GetMaterialCommonBasis(EMaterialVectorCoordTransform X)
-{
-	static const EMaterialCommonBasis ConversionTable[TRANSFORM_MAX] = {
-		MCB_Tangent,					// TRANSFORM_Tangent
-		MCB_Local,						// TRANSFORM_Local
-		MCB_World,						// TRANSFORM_World
-		MCB_View,						// TRANSFORM_View
-		MCB_Camera,						// TRANSFORM_Camera
-		MCB_MeshParticle,				// TRANSFORM_Particle
-		MCB_Instance,					// TRANSFORM_Instance
-	};
-	return ConversionTable[X];
-}
-
 #if WITH_EDITOR
 int32 UMaterialExpressionTransform::Compile(class FMaterialCompiler* Compiler, int32 OutputIndex)
 {
+	using namespace UE::MaterialTranslatorUtils;
+
 	int32 Result = INDEX_NONE;
 
 	if (!Input.GetTracedInput().Expression)
@@ -11386,25 +11361,12 @@ void UMaterialExpressionTransform::GetCaption(TArray<FString>& OutCaptions) cons
 ///////////////////////////////////////////////////////////////////////////////
 // UMaterialExpressionTransformPosition
 ///////////////////////////////////////////////////////////////////////////////
-static EMaterialCommonBasis GetMaterialCommonBasis(EMaterialPositionTransformSource X)
-{
-	static const EMaterialCommonBasis ConversionTable[TRANSFORMPOSSOURCE_MAX] = {
-		MCB_Local,						// TRANSFORMPOSSOURCE_Local
-		MCB_World,						// TRANSFORMPOSSOURCE_World
-		MCB_PeriodicWorld,				// TRANSFORMPOSSOURCE_PeriodicWorld
-		MCB_TranslatedWorld,			// TRANSFORMPOSSOURCE_TranslatedWorld
-		MCB_FirstPerson,				// TRANSFORMPOSSOURCE_FirstPersonTranslatedWorld
-		MCB_View,						// TRANSFORMPOSSOURCE_View
-		MCB_Camera,						// TRANSFORMPOSSOURCE_Camera
-		MCB_MeshParticle,				// TRANSFORMPOSSOURCE_Particle
-		MCB_Instance,					// TRANSFORMPOSSOURCE_Instance
-	};
-	return ConversionTable[X];
-}
 
 #if WITH_EDITOR
 int32 UMaterialExpressionTransformPosition::Compile(class FMaterialCompiler* Compiler, int32 OutputIndex)
 {
+	using namespace UE::MaterialTranslatorUtils;
+
 	if(!Input.GetTracedInput().Expression)
 	{
 		return Compiler->Errorf(TEXT("Missing Transform Position input vector"));
diff --git a/Engine/Source/Runtime/Engine/Private/Materials/MaterialExpressionsIR.cpp b/Engine/Source/Runtime/Engine/Private/Materials/MaterialExpressionsIR.cpp
--- a/Engine/Source/Runtime/Engine/Private/Materials/MaterialExpressionsIR.cpp
+++ b/Engine/Source/Runtime/Engine/Private/Materials/MaterialExpressionsIR.cpp
@@ -1435,6 +1435,369 @@ void UMaterialExpressionTime::Build(MIR::FEmitter& Em)
 	}
 }
 
+// Returns true if the specified value is a constant power of two (scalar or vector).
+static bool IsConstFloatOfPow2Expression(FValueRef TileScaleIndexValue)
+{
+	using namespace UE::MaterialTranslatorUtils;
+	if (MIR::FConstant* ConstIndex = TileScaleIndexValue->As<MIR::FConstant>())
+	{
+		return IsFloatPowerOfTwo(ConstIndex->Float);
+	}
+	else if (MIR::FComposite* Composite = TileScaleIndexValue->As<MIR::FComposite>())
+	{
+		for (MIR::FValue* Component : Composite->GetComponents())
+		{
+			MIR::FConstant* ConstComponent = Component->As<MIR::FConstant>();
+			if (!ConstComponent || !IsFloatPowerOfTwo(ConstComponent->Float))
+			{
+				return false;
+			}
+		}
+		return true;
+	}
+	return false;
+}
+
+static FValueRef EmitPeriodicWorldPositionOrigin(MIR::FEmitter& Em, FValueRef TileScaleIndexValue)
+{
+	FStringView PeriodicWorldOriginFunctionName = IsConstFloatOfPow2Expression(TileScaleIndexValue) ? TEXTVIEW("GetPeriodicWorldOrigin_Pow2($0)") : TEXTVIEW("GetPeriodicWorldOrigin($0)");
+	return Em.InlineHLSL(TileScaleIndexValue->Type, PeriodicWorldOriginFunctionName, { TileScaleIndexValue });
+}
+
+// Emits inline HLSL from an external code declaration that does not take any parameters.
+static FValueRef EmitFixedExternalCode(MIR::FEmitter& Em, FName InExternalCodeIdentifier)
+{
+	return Em.InlineHLSL(MaterialExternalCodeRegistry::Get().FindExternalCode(InExternalCodeIdentifier), {}, MIR::EValueFlags::SubstituteTags);
+}
+
+static FValueRef EmitMatrixMultiply(MIR::FEmitter& Em, FValueRef VectorValue, FValueRef MatrixValue, bool bHasWComponent)
+{
+	return bHasWComponent
+		? Em.Swizzle(Em.Multiply(VectorValue, MatrixValue), MIR::FSwizzleMask::XYZ())  // mul(Float4(V, 1.0), V).xyz
+		: Em.Multiply(VectorValue, Em.Cast(MatrixValue, MIR::FType::MakeFloat(3, 3))); // mul(V, (Float3x3)M)
+}
+
+static FValueRef EmitMultiplyTransposeMatrix(MIR::FEmitter& Em, FValueRef MatrixValue, FValueRef VectorValue, bool bHasWComponent)
+{
+	// TODO: this should be removed when the Transpose operator is added to the translator.
+	return bHasWComponent
+		? Em.Swizzle(Em.Multiply(MatrixValue, VectorValue), MIR::FSwizzleMask::XYZ())  // mul(M, Float4(V, 1.0)).xyz
+		: Em.Multiply(Em.Cast(MatrixValue, MIR::FType::MakeFloat(3, 3)), VectorValue); // mul((Float3x3)M, V)
+}
+
+static FValueRef EmitMultiplyTranslatedMatrix(MIR::FEmitter& Em, FValueRef VectorValue, FValueRef MatrixPreTranslation, bool bHasWComponent)
+{
+	if (bHasWComponent)
+	{
+		// mul(Float4(V, 1.0), DFFastToTranslatedWorld(M, ResolvedView.PreViewTranslation))
+		MatrixPreTranslation = Em.InlineHLSL(MIR::FType::MakeFloat(4, 4), TEXTVIEW("DFFastToTranslatedWorld($0, ResolvedView.<PREV>PreViewTranslation)"), { MatrixPreTranslation }, MIR::EValueFlags::SubstituteTags);
+		return Em.MatrixMultiply(Em.Vector4(VectorValue, Em.ConstantOne(MIR::EScalarKind::Float)), MatrixPreTranslation);
+	}
+	else
+	{
+		// mul(V, DFToFloat3x3(M))
+		MatrixPreTranslation = Em.InlineHLSL(MIR::FType::MakeFloat(3, 3), TEXTVIEW("DFToFloat3x3($0)"), { MatrixPreTranslation });
+		return Em.MatrixMultiply(VectorValue, MatrixPreTranslation);
+	}
+}
+
+static FValueRef EmitMultiplyLWCMatrix(MIR::FEmitter& Em, FValueRef VectorValue, FValueRef MatrixValue, bool bHasWComponent, bool bDemote)
+{
+	return bHasWComponent
+		? Em.InlineHLSL(VectorValue->Type, bDemote ? TEXTVIEW("WSMultiplyDemote($0, $1)") : TEXTVIEW("WSMultiply($0, $1)"), { VectorValue, MatrixValue })
+		: Em.InlineHLSL(VectorValue->Type, TEXTVIEW("WSMultiplyVector($0, $1)"), { VectorValue, MatrixValue });
+}
+
+static FValueRef EmitTransformVectorBase(
+	MIR::FEmitter& Em, FValueRef InputValue, EMaterialCommonBasis TransformSourceBasis, EMaterialCommonBasis TransformDestBasis,
+	bool bIsPositionTransform, MIR::FValueRef PeriodicWorldTileSizeValue, MIR::FValueRef FirstPersonInterpolationAlphaValue)
+{
+	// Construct float3(0,0,x) out of the input if it is a scalar
+	// This way artists can plug in a scalar and it will be treated as height, or a vector displacement
+	if (TransformSourceBasis == MCB_Tangent && InputValue->Type.IsScalar())
+	{
+		FValueRef Zero = Em.ConstantZero(MIR::EScalarKind::Float);
+		InputValue = Em.Vector3(Zero, Zero, InputValue);
+	}
+	else
+	{
+		InputValue = Em.CastToVector(InputValue, 3);
+	}
+
+	if (!InputValue.IsValid())
+	{
+		return InputValue.ToPoison();
+	}
+
+	MIR::FType ResultType = (TransformDestBasis == MCB_World && bIsPositionTransform) ? MIR::FType::MakeDoubleVector(3) : MIR::FType::MakeFloatVector(3);
+
+	EMaterialCommonBasis IntermediaryBasis = MCB_World;
+
+	switch (TransformSourceBasis)
+	{
+		case MCB_Tangent:
+			check(!bIsPositionTransform);
+			if (TransformDestBasis == MCB_World)
+			{
+				return Em.MatrixMultiply(InputValue, EmitFixedExternalCode(Em, TEXT("TangentToWorld")));
+			}
+			// else use MCB_World as intermediary basis
+			break;
+
+		case MCB_Local:
+			switch (TransformDestBasis)
+			{
+				case MCB_World:
+				{
+					FStringView Code = bIsPositionTransform ? TEXTVIEW("TransformLocalPositionTo<PREV>World(Parameters, $0)") : TEXTVIEW("TransformLocalVectorTo<PREV>World(Parameters, $0)");
+					return Em.InlineHLSL(ResultType, Code, { InputValue }, MIR::EValueFlags::SubstituteTags);
+				}
+
+				case MCB_TranslatedWorld:
+					if (bIsPositionTransform)
+					{
+						return EmitMultiplyTranslatedMatrix(Em, InputValue, EmitFixedExternalCode(Em, TEXT("GetLocalToWorldDF")), bIsPositionTransform);
+					}
+					break;
+
+				case MCB_PeriodicWorld:
+				case MCB_FirstPerson:
+					IntermediaryBasis = MCB_TranslatedWorld;
+					break;
+
+				default:
+					// else use MCB_World as intermediary basis
+					break;
+			}
+			break;
+
+		case MCB_TranslatedWorld:
+			switch (TransformDestBasis)
+			{
+				case MCB_World:
+					return bIsPositionTransform
+						? Em.Subscript(InputValue, EmitFixedExternalCode(Em, TEXT("GetPreViewTranslation")))
+						: InputValue;
+
+				case MCB_Camera:
+					return EmitMatrixMultiply(Em, InputValue, EmitFixedExternalCode(Em, TEXT("TranslatedWorldToCameraView")), bIsPositionTransform);
+
+				case MCB_View:
+					return EmitMatrixMultiply(Em, InputValue, EmitFixedExternalCode(Em, TEXT("TranslatedWorldToView")), bIsPositionTransform);
+
+				case MCB_Tangent:
+					return EmitMultiplyTransposeMatrix(Em, EmitFixedExternalCode(Em, TEXT("TangentToWorld")), InputValue, bIsPositionTransform);
+
+				case MCB_Local:
+					return EmitMultiplyTranslatedMatrix(Em, InputValue, EmitFixedExternalCode(Em, TEXT("GetWorldToLocalDF")), bIsPositionTransform);
+
+				case MCB_MeshParticle:
+					return EmitMultiplyTranslatedMatrix(Em, InputValue, EmitFixedExternalCode(Em, TEXT("WorldToParticle")), bIsPositionTransform);
+
+				case MCB_Instance:
+					return EmitMultiplyTranslatedMatrix(Em, InputValue, EmitFixedExternalCode(Em, TEXT("GetWorldToInstanceDF")), bIsPositionTransform);
+
+				case MCB_PeriodicWorld:
+					if (!PeriodicWorldTileSizeValue.IsValid())
+					{
+						Em.Error(TEXT("Missing periodic world tile size"));
+						return Em.Poison();
+					}
+					return Em.Subtract(InputValue, EmitPeriodicWorldPositionOrigin(Em, PeriodicWorldTileSizeValue));
+
+				case MCB_FirstPerson:
+				{
+					if (!FirstPersonInterpolationAlphaValue.IsValid())
+					{
+						Em.Error(TEXT("Missing first person interpolation alpha"));
+						return Em.Poison();
+					}
+					// The first person transform is actually a 3x3 matrix and can therefore be used for derivatives as well.
+					MIR::FValueRef LerpAlphaClampedIndexValue = Em.Saturate(Em.CastToFloat(FirstPersonInterpolationAlphaValue, 1));
+					return Em.InlineHLSL(ResultType, TEXT("TransformTo<PREVIOUS>FirstPerson($0, $1)"), { InputValue, LerpAlphaClampedIndexValue }, MIR::EValueFlags::SubstituteTags);
+				}
+
+				default:
+					break; // else use MCB_World as intermediary basis
+			}
+			break;
+
+		case MCB_World:
+			switch (TransformDestBasis)
+			{
+				case MCB_Tangent:
+					return EmitMultiplyTransposeMatrix(Em, EmitFixedExternalCode(Em, TEXT("TangentToWorld")), InputValue, bIsPositionTransform);
+
+				case MCB_Local:
+					return EmitMultiplyLWCMatrix(Em, InputValue, EmitFixedExternalCode(Em, TEXT("GetWorldToLocal")), bIsPositionTransform, true);
+
+				case MCB_TranslatedWorld:
+					return bIsPositionTransform
+						? Em.Add(InputValue, EmitFixedExternalCode(Em, TEXT("GetPreViewTranslation")))
+						: InputValue;
+
+				case MCB_MeshParticle:
+					return EmitMultiplyLWCMatrix(Em, InputValue, EmitFixedExternalCode(Em, TEXT("GetWorldToParticle")), bIsPositionTransform, true);
+
+				case MCB_Instance:
+					return EmitMultiplyLWCMatrix(Em, InputValue, EmitFixedExternalCode(Em, TEXT("GetWorldToInstance")), bIsPositionTransform, true);
+
+				default:
+					// else use MCB_TranslatedWorld as intermediary basis
+					IntermediaryBasis = MCB_TranslatedWorld;
+					break;
+			}
+			break;
+
+		case MCB_Camera:
+			if (TransformDestBasis == MCB_TranslatedWorld)
+			{
+				return EmitMatrixMultiply(Em, InputValue, EmitFixedExternalCode(Em, TEXT("CameraViewToTranslatedWorld")), bIsPositionTransform);
+			}
+			
+			// else use MCB_TranslatedWorld as intermediary basis
+			IntermediaryBasis = MCB_TranslatedWorld;
+			break;
+
+		case MCB_View:
+			if (TransformDestBasis == MCB_TranslatedWorld)
+			{
+				return EmitMatrixMultiply(Em, InputValue, EmitFixedExternalCode(Em, TEXT("ViewToTranslatedWorld")), bIsPositionTransform);
+			}
+
+			// else use MCB_TranslatedWorld as intermediary basis
+			IntermediaryBasis = MCB_TranslatedWorld;
+			break;
+
+		case MCB_MeshParticle:
+			switch (TransformDestBasis)
+			{
+				case MCB_World:
+					return EmitMultiplyLWCMatrix(Em, InputValue, EmitFixedExternalCode(Em, TEXT("GetParticleToWorld")), bIsPositionTransform, false);
+
+				case MCB_TranslatedWorld:
+					return EmitMultiplyTranslatedMatrix(Em, InputValue, EmitFixedExternalCode(Em, TEXT("ParticleToWorld")), bIsPositionTransform);
+
+				case MCB_PeriodicWorld:
+				case MCB_FirstPerson:
+					IntermediaryBasis = MCB_TranslatedWorld;
+					break;
+
+				default:
+					break; // use World as an intermediary base
+			}
+			break;
+
+		case MCB_Instance:
+			switch(TransformDestBasis)
+			{
+				case MCB_World:
+					return EmitMultiplyLWCMatrix(Em, InputValue, EmitFixedExternalCode(Em, TEXT("GetInstanceToWorld")), bIsPositionTransform, false);
+
+				case MCB_TranslatedWorld:
+					return EmitMultiplyTranslatedMatrix(Em, InputValue, EmitFixedExternalCode(Em, TEXT("GetInstanceToWorldDF")), bIsPositionTransform);
+
+				case MCB_PeriodicWorld:
+				case MCB_FirstPerson:
+					IntermediaryBasis = MCB_TranslatedWorld;
+					break;
+
+				default:
+					break; // use World as an intermediary base
+			}
+			break;
+
+		case MCB_PeriodicWorld:
+			switch (TransformDestBasis)
+			{
+				case MCB_TranslatedWorld:
+					if (!PeriodicWorldTileSizeValue.IsValid())
+					{
+						Em.Error(TEXT("Missing periodic world tile size"));
+						return Em.Poison();
+					}
+					return Em.Add(InputValue, EmitPeriodicWorldPositionOrigin(Em, PeriodicWorldTileSizeValue));
+
+				default:
+					// else use MCB_TranslatedWorld as intermediary basis
+					IntermediaryBasis = MCB_TranslatedWorld;
+					break;
+			}
+			break;
+
+		case MCB_FirstPerson:
+			UE_MIR_UNREACHABLE(); // MCB_FirstPerson is not supported as a source basis. This should've been caught earlier in validation.
+
+		default:
+			UE_MIR_UNREACHABLE();
+	}
+
+	// Check intermediary basis so we don't have infinite recursion
+	check(IntermediaryBasis != TransformSourceBasis);
+	check(IntermediaryBasis != TransformDestBasis);
+
+	// Use intermediary basis
+	FValueRef IntermediaryBasisA = EmitTransformVectorBase(Em, InputValue, TransformSourceBasis, IntermediaryBasis, bIsPositionTransform, PeriodicWorldTileSizeValue, FirstPersonInterpolationAlphaValue);
+	FValueRef IntermediaryBasisB = EmitTransformVectorBase(Em, IntermediaryBasisA, IntermediaryBasis, TransformDestBasis, bIsPositionTransform, PeriodicWorldTileSizeValue, FirstPersonInterpolationAlphaValue);
+
+	return IntermediaryBasisB;
+}
+
+static void BuildTransformVectorBase(
+	MIR::FEmitter& Em, const FExpressionInput* Input, EMaterialCommonBasis TransformSourceBasis, EMaterialCommonBasis TransformDestBasis,
+	bool bIsPositionTransform, MIR::FValueRef PeriodicWorldTileSizeValue, MIR::FValueRef FirstPersonInterpolationAlphaValue)
+{
+	FValueRef InputValue = Em.CheckIsPrimitive(Em.Input(Input));
+	UE_MIR_CHECKPOINT(Em);
+
+	FValueRef OutputValue = EmitTransformVectorBase(Em, InputValue, TransformSourceBasis, TransformDestBasis, bIsPositionTransform, PeriodicWorldTileSizeValue, FirstPersonInterpolationAlphaValue);
+
+	if (TransformSourceBasis == MCB_World && bIsPositionTransform)
+	{
+		if (!OutputValue->Type.IsDouble())
+		{
+			OutputValue = Em.CastToScalarKind(OutputValue, MIR::EScalarKind::Double);
+		}
+	}
+	else
+	{
+		if (OutputValue->Type.IsDouble())
+		{
+			OutputValue = Em.CastToFloatKind(OutputValue);
+		}
+	}
+
+	Em.Output(0, OutputValue);
+}
+
+void UMaterialExpressionTransform::Build(MIR::FEmitter& Em)
+{
+	const EMaterialCommonBasis TransformSourceBasis = UE::MaterialTranslatorUtils::GetMaterialCommonBasis(TransformSourceType);
+	const EMaterialCommonBasis TransformDestBasis = UE::MaterialTranslatorUtils::GetMaterialCommonBasis(TransformType);
+
+	constexpr bool bIsPositionTransform = false;
+	BuildTransformVectorBase(Em, &Input, TransformSourceBasis, TransformDestBasis, bIsPositionTransform, {}, {});
+}
+
+void UMaterialExpressionTransformPosition::Build(MIR::FEmitter& Em)
+{
+	MIR::FValueRef PeriodicWorldTileSizeValue, FirstPersonInterpolationAlphaValue;
+	if (TransformSourceType == TRANSFORMPOSSOURCE_PeriodicWorld || TransformType == TRANSFORMPOSSOURCE_PeriodicWorld)
+	{
+		PeriodicWorldTileSizeValue = Em.InputDefaultFloat(&PeriodicWorldTileSize, ConstPeriodicWorldTileSize);
+	}
+	if (TransformSourceType == TRANSFORMPOSSOURCE_FirstPersonTranslatedWorld || TransformType == TRANSFORMPOSSOURCE_FirstPersonTranslatedWorld)
+	{
+		FirstPersonInterpolationAlphaValue = Em.InputDefaultFloat(&FirstPersonInterpolationAlpha, ConstFirstPersonInterpolationAlpha);
+	}
+
+	const EMaterialCommonBasis TransformSourceBasis = UE::MaterialTranslatorUtils::GetMaterialCommonBasis(TransformSourceType);
+	const EMaterialCommonBasis TransformDestBasis = UE::MaterialTranslatorUtils::GetMaterialCommonBasis(TransformType);
+
+	constexpr bool bIsPositionTransform = true;
+	BuildTransformVectorBase(Em, &Input, TransformSourceBasis, TransformDestBasis, bIsPositionTransform, PeriodicWorldTileSizeValue, FirstPersonInterpolationAlphaValue);
+}
+
 void UMaterialExpressionReroute::Build(MIR::FEmitter& Em)
 {
 	Em.Output(0, Em.TryInput(&Input));
diff --git a/Engine/Source/Runtime/Engine/Private/Materials/MaterialExternalCodeRegistry.h b/Engine/Source/Runtime/Engine/Private/Materials/MaterialExternalCodeRegistry.h
--- a/Engine/Source/Runtime/Engine/Private/Materials/MaterialExternalCodeRegistry.h
+++ b/Engine/Source/Runtime/Engine/Private/Materials/MaterialExternalCodeRegistry.h
@@ -68,6 +68,9 @@ enum class EMaterialValueTypeBridge : uint64
 	UInt                = MCT_UInt UMETA(Hidden),
 	LWCType             = MCT_LWCType UMETA(Hidden),
 	Numeric             = MCT_Numeric UMETA(Hidden),
+	Float3x3  			= MCT_Float3x3,
+	Float4x4  			= MCT_Float4x4,
+	LWCMatrix 			= MCT_LWCMatrix,
 };
 
 /**
diff --git a/Engine/Source/Runtime/Engine/Private/Materials/MaterialIREmitter.cpp b/Engine/Source/Runtime/Engine/Private/Materials/MaterialIREmitter.cpp
--- a/Engine/Source/Runtime/Engine/Private/Materials/MaterialIREmitter.cpp
+++ b/Engine/Source/Runtime/Engine/Private/Materials/MaterialIREmitter.cpp
@@ -2522,6 +2522,12 @@ FValueRef FEmitter::Subscript(FValueRef Value, int32 Index)
 		return Value.ToPoison();
 	}
 
+	if (PrimitiveType->IsMatrix() && PrimitiveType->IsDouble())
+	{
+		Errorf(Value, TEXT("Cannot subscript a double-precision matrix."));
+		return Value.ToPoison();
+	}
+
 	// Avoid subscripting a subscript (e.g. no value.xy.x)
 	if (FSubscript* Subscript = As<FSubscript>(Value))
 	{
diff --git a/Engine/Source/Runtime/Engine/Private/Materials/MaterialIRToHLSLTranslator.cpp b/Engine/Source/Runtime/Engine/Private/Materials/MaterialIRToHLSLTranslator.cpp
--- a/Engine/Source/Runtime/Engine/Private/Materials/MaterialIRToHLSLTranslator.cpp
+++ b/Engine/Source/Runtime/Engine/Private/Materials/MaterialIRToHLSLTranslator.cpp
@@ -1444,26 +1444,40 @@ struct FPrivate : FMaterialIRToHLSLTranslation
 			
 	void LowerSubscript(const MIR::FSubscript* Subscript)
 	{
-		TOptional<MIR::FPrimitive> ArgVectorType = Subscript->Arg->Type.AsVector();
-		if (ArgVectorType && ArgVectorType->IsDouble())
+		if (TOptional<MIR::FPrimitive> ArgVectorType = Subscript->Arg->Type.AsVector())
 		{
-			FStringView LwcComponentsStr[] = { TEXTVIEW("WSGetX("), TEXTVIEW("WSGetY("), TEXTVIEW("WSGetZ("), TEXTVIEW("WSGetW(") };
-			check(Subscript->Index <= ArgVectorType->NumComponents());
-
-			Printer << LwcComponentsStr[Subscript->Index] << LowerValue(Subscript->Arg) << TEXTVIEW(")");
-		}
-		else
-		{
-			LowerValue(Subscript->Arg);
+			if (ArgVectorType->IsDouble())
+			{
+				FStringView LwcComponentsStr[] = { TEXTVIEW("WSGetX("), TEXTVIEW("WSGetY("), TEXTVIEW("WSGetZ("), TEXTVIEW("WSGetW(") };
+				check(Subscript->Index <= ArgVectorType->NumComponents());
 
-			if (ArgVectorType)
+				Printer << LwcComponentsStr[Subscript->Index] << LowerValue(Subscript->Arg) << TEXTVIEW(")");
+			}
+			else
 			{
+				LowerValue(Subscript->Arg);
+				
 				FStringView ComponentsStr[] = { TEXTVIEW(".x"), TEXTVIEW(".y"), TEXTVIEW(".z"), TEXTVIEW(".w") };
 				check(Subscript->Index <= ArgVectorType->NumComponents());
 
 				Printer << ComponentsStr[Subscript->Index];
 			}
 		}
+		else if (TOptional<MIR::FPrimitive> MatrixType = Subscript->Arg->Type.AsMatrix())
+		{
+			check(!MatrixType->IsDouble()); // The emitter should have checked this
+
+			LowerValue(Subscript->Arg);
+
+			// Print matrix component swizzle, e.g. `M._m00`
+			check(Subscript->Index < MatrixType->NumComponents());
+			Printer << TEXTVIEW("._m") << (Subscript->Index % MatrixType->NumRows) << (Subscript->Index / MatrixType->NumRows);
+		}
+		else
+		{
+			// The builder should never emit subscripts of scalar types.
+			UE_MIR_UNREACHABLE();
+		}
 	}
 
 	void LowerScalar(const MIR::FScalar* Scalar)
diff --git a/Engine/Source/Runtime/Engine/Private/Materials/MaterialIRTypes.cpp b/Engine/Source/Runtime/Engine/Private/Materials/MaterialIRTypes.cpp
--- a/Engine/Source/Runtime/Engine/Private/Materials/MaterialIRTypes.cpp
+++ b/Engine/Source/Runtime/Engine/Private/Materials/MaterialIRTypes.cpp
@@ -106,6 +106,9 @@ FType FType::FromMaterialValueType(EMaterialValueType Type)
 		case MCT_LWCVector2: return FType::MakeDoubleVector(2);
 		case MCT_LWCVector3: return FType::MakeDoubleVector(3);
 		case MCT_LWCVector4: return FType::MakeDoubleVector(4);
+		case MCT_Float3x3: return FType::MakeFloat(3, 3);
+		case MCT_Float4x4: return FType::MakeFloat(4, 4);
+		case MCT_LWCMatrix: return FType::MakeDouble(4, 4);
 		default: UE_MIR_UNREACHABLE();
 	}
 }
diff --git a/Engine/Source/Runtime/Engine/Private/Materials/MaterialShader.cpp b/Engine/Source/Runtime/Engine/Private/Materials/MaterialShader.cpp
--- a/Engine/Source/Runtime/Engine/Private/Materials/MaterialShader.cpp
+++ b/Engine/Source/Runtime/Engine/Private/Materials/MaterialShader.cpp
@@ -657,6 +657,64 @@ int32 GetLWCTruncateMode()
 	return LWCTruncateValue;
 }
 
+EMaterialCommonBasis GetMaterialCommonBasis(EMaterialVectorCoordTransformSource X)
+{
+	static const EMaterialCommonBasis ConversionTable[] = {
+		MCB_Tangent,					// TRANSFORMSOURCE_Tangent
+		MCB_Local,						// TRANSFORMSOURCE_Local
+		MCB_World,						// TRANSFORMSOURCE_World
+		MCB_View,						// TRANSFORMSOURCE_View
+		MCB_Camera,						// TRANSFORMSOURCE_Camera
+		MCB_MeshParticle,				// TRANSFORMSOURCE_Particle
+		MCB_Instance,					// TRANSFORMSOURCE_Instance
+	};
+	static_assert(
+		EMaterialVectorCoordTransformSource::TRANSFORMSOURCE_MAX == UE_ARRAY_COUNT(ConversionTable),
+		"Mapping of enum EMaterialVectorCoordTransformSource to EMaterialCommonBasis is out of sync; ConversionTable[] array needs to be updated!");
+	return ConversionTable[X];
+}
+
+EMaterialCommonBasis GetMaterialCommonBasis(EMaterialVectorCoordTransform X)
+{
+	static const EMaterialCommonBasis ConversionTable[] = {
+		MCB_Tangent,					// TRANSFORM_Tangent
+		MCB_Local,						// TRANSFORM_Local
+		MCB_World,						// TRANSFORM_World
+		MCB_View,						// TRANSFORM_View
+		MCB_Camera,						// TRANSFORM_Camera
+		MCB_MeshParticle,				// TRANSFORM_Particle
+		MCB_Instance,					// TRANSFORM_Instance
+	};
+	static_assert(
+		EMaterialVectorCoordTransform::TRANSFORM_MAX == UE_ARRAY_COUNT(ConversionTable),
+		"Mapping of enum EMaterialVectorCoordTransform to EMaterialCommonBasis is out of sync; ConversionTable[] array needs to be updated!");
+	return ConversionTable[X];
+}
+
+EMaterialCommonBasis GetMaterialCommonBasis(EMaterialPositionTransformSource X)
+{
+	static const EMaterialCommonBasis ConversionTable[] = {
+		MCB_Local,						// TRANSFORMPOSSOURCE_Local
+		MCB_World,						// TRANSFORMPOSSOURCE_World
+		MCB_PeriodicWorld,				// TRANSFORMPOSSOURCE_PeriodicWorld
+		MCB_TranslatedWorld,			// TRANSFORMPOSSOURCE_TranslatedWorld
+		MCB_FirstPerson,				// TRANSFORMPOSSOURCE_FirstPersonTranslatedWorld
+		MCB_View,						// TRANSFORMPOSSOURCE_View
+		MCB_Camera,						// TRANSFORMPOSSOURCE_Camera
+		MCB_MeshParticle,				// TRANSFORMPOSSOURCE_Particle
+		MCB_Instance,					// TRANSFORMPOSSOURCE_Instance
+	};
+	static_assert(
+		EMaterialPositionTransformSource::TRANSFORMPOSSOURCE_MAX == UE_ARRAY_COUNT(ConversionTable),
+		"Mapping of enum EMaterialPositionTransformSource to EMaterialCommonBasis is out of sync; ConversionTable[] array needs to be updated!");
+	return ConversionTable[X];
+}
+
+bool IsFloatPowerOfTwo(float InValue)
+{
+	return ((*reinterpret_cast<int*>(&InValue)) & 0x007FFFFF) == 0; // zero mantisse
+}
+
 } // namespace UE::MaterialTranslatorUtils
 
 #endif // WITH_EDITOR
diff --git a/Engine/Source/Runtime/Engine/Private/Materials/MaterialShared.cpp b/Engine/Source/Runtime/Engine/Private/Materials/MaterialShared.cpp
--- a/Engine/Source/Runtime/Engine/Private/Materials/MaterialShared.cpp
+++ b/Engine/Source/Runtime/Engine/Private/Materials/MaterialShared.cpp
@@ -5872,8 +5872,7 @@ int32 FMaterialCompilationOutput::GetNumPostProcessInputsUsed() const
 
 namespace UE::MaterialTranslatorUtils
 {
-// Utility to generate HLSL defines to remap PPI_UserSceneTexture0-6 tokens generated by the HLSL emitter/translator to PPI_PostProcessInput0-6,
-// using unoccupied slots (UsedSceneTextures bits not set) in the order they are encountered.
+
 FString GenerateUserSceneTextureRemapHLSLDefines(const FMaterialCompilationOutput& CompilationOutput)
 {
 	if (CompilationOutput.UserSceneTextureInputs.IsEmpty())
@@ -5931,8 +5930,8 @@ FString SceneTextureIdToHLSLString(ESceneTextureId TexId)
 	{
 		return FString::Printf(TEXT("%d"), TexId);
 	}
-
 }
+
 }  // namespace UE::MaterialTranslatorUtils
 
 #endif // WITH_EDITOR
diff --git a/Engine/Source/Runtime/Engine/Private/Materials/MaterialSharedPrivate.h b/Engine/Source/Runtime/Engine/Private/Materials/MaterialSharedPrivate.h
--- a/Engine/Source/Runtime/Engine/Private/Materials/MaterialSharedPrivate.h
+++ b/Engine/Source/Runtime/Engine/Private/Materials/MaterialSharedPrivate.h
@@ -8,6 +8,8 @@
 #include "Containers/UnrealString.h"
 #include "Engine/TextureDefines.h"
 #include "MaterialShared.h"
+#include "Materials/MaterialExpressionTransform.h"
+#include "Materials/MaterialExpressionTransformPosition.h"
 #include "RHIFeatureLevel.h"
 #include "RHIShaderPlatform.h"
 #include "SceneTypes.h"
@@ -50,19 +52,41 @@ struct FRecordedMaterialResourceForCooking
 
 } // namespace UE::MaterialInterface::Private
 
-/** Utility functions shared between the old and new material translator. */
+// Utility functions shared between the old and new material translator.
 namespace UE::MaterialTranslatorUtils
 {
+	// Converts a TextureAddress enum to its corresponding string representation.
 	const TCHAR* GetVTAddressMode(TextureAddress Address);
+
+	// Determines the texture addressing mode for a texture based on its sampler source.
 	void GetTextureAddressForSamplerSource(const UTexture* InTexture, ESamplerSourceMode InSamplerSource, TextureAddress& OutStaticAddressX, TextureAddress& OutStaticAddressY, TextureAddress& OutStaticAddressZ);
 
+	// Returns true if the LWCs are enabled in the engine.
 	bool IsLWCEnabled();
+
+	// Returns the current LWC truncation mode.
 	int32 GetLWCTruncateMode();
 
+	// Maps a vector transform source to a common basis (e.g., world, local, tangent).
+	EMaterialCommonBasis GetMaterialCommonBasis(EMaterialVectorCoordTransformSource X);
+
+	// Maps a vector transform operation to a common basis.
+	EMaterialCommonBasis GetMaterialCommonBasis(EMaterialVectorCoordTransform X);
+
+	// Maps a position transform source to a common basis.
+	EMaterialCommonBasis GetMaterialCommonBasis(EMaterialPositionTransformSource X);
+
+	// Returns true if the given float is a power of two.
+	bool IsFloatPowerOfTwo(float InValue);
+
+	// Generates HLSL defines to remap PPI_UserSceneTexture0-6 tokens generated by the HLSL emitter/translator to PPI_PostProcessInput0-6,
+	// using unoccupied slots (UsedSceneTextures bits not set) in the order they are encountered.
 	FString GenerateUserSceneTextureRemapHLSLDefines(const FMaterialCompilationOutput& CompilationOutput);
+
+	// Converts a scene texture identifier to its corresponding HLSL string representation.
 	FString SceneTextureIdToHLSLString(ESceneTextureId TexId);
 
-} // /namespace UE::MaterialTranslatorUtils
+} // namespace UE::MaterialTranslatorUtils
 
 ///////////////////////////////////////////////////////
 // CompactBinary interface for FRecordedMaterialResourceForCooking
diff --git a/Engine/Source/Runtime/Engine/Public/MaterialShared.h b/Engine/Source/Runtime/Engine/Public/MaterialShared.h
--- a/Engine/Source/Runtime/Engine/Public/MaterialShared.h
+++ b/Engine/Source/Runtime/Engine/Public/MaterialShared.h
@@ -760,7 +760,10 @@ protected:
 	LAYOUT_FIELD(FRHIUniformBufferLayoutInitializer, UniformBufferLayoutInitializer);
 };
 
-/** Stores outputs from the material compile that need to be saved. */
+/**
+Stores outputs from the material compile that need to be saved.
+Updates to layout fields in this class must be accompanied with a GUID bump to MaterialTranslationDDCVersion entry in FDevSystemGuids.
+*/
 class FMaterialCompilationOutput
 {
 	DECLARE_TYPE_LAYOUT(FMaterialCompilationOutput, NonVirtual);
diff --git a/Engine/Source/Runtime/Engine/Public/MaterialValueType.h b/Engine/Source/Runtime/Engine/Public/MaterialValueType.h
--- a/Engine/Source/Runtime/Engine/Public/MaterialValueType.h
+++ b/Engine/Source/Runtime/Engine/Public/MaterialValueType.h
@@ -68,6 +68,11 @@ enum EMaterialValueType : uint64
 	MCT_TextureMeshPaint      = 1u << 30,
 	MCT_TextureMaterialCache  = 1u << 31,
 
+	/* Matrix types */
+	MCT_Float3x3              = 1ull << 32,
+	MCT_Float4x4              = 1ull << 33,
+	MCT_LWCMatrix             = 1ull << 34,
+
 	/** Reserved range for licensee modifications, any bits in the (inclusive) range is free for modifications */
 	MCT_LicenseeReservedBegin     = 1ull << 48,
 	MCT_LicenseeReservedEnd       = 1ull << 63,
diff --git a/Engine/Source/Runtime/Engine/Public/Materials/MaterialExpressionTransform.h b/Engine/Source/Runtime/Engine/Public/Materials/MaterialExpressionTransform.h
--- a/Engine/Source/Runtime/Engine/Public/Materials/MaterialExpressionTransform.h
+++ b/Engine/Source/Runtime/Engine/Public/Materials/MaterialExpressionTransform.h
@@ -84,6 +84,7 @@ public:
 
 	//~ Begin UMaterialExpression Interface
 #if WITH_EDITOR
+	virtual void Build(MIR::FEmitter& Em) override;
 	virtual int32 Compile(class FMaterialCompiler* Compiler, int32 OutputIndex) override;
 	virtual void GetCaption(TArray<FString>& OutCaptions) const override;
 #endif
diff --git a/Engine/Source/Runtime/Engine/Public/Materials/MaterialExpressionTransformPosition.h b/Engine/Source/Runtime/Engine/Public/Materials/MaterialExpressionTransformPosition.h
--- a/Engine/Source/Runtime/Engine/Public/Materials/MaterialExpressionTransformPosition.h
+++ b/Engine/Source/Runtime/Engine/Public/Materials/MaterialExpressionTransformPosition.h
@@ -94,6 +94,7 @@ public:
 
 	//~ Begin UMaterialExpression Interface
 #if WITH_EDITOR
+	virtual void Build(MIR::FEmitter& Em) override;
 	virtual int32 Compile(class FMaterialCompiler* Compiler, int32 OutputIndex) override;
 	virtual void GetConnectorToolTip(int32 InputIndex, int32 OutputIndex, TArray<FString>& OutToolTip) override;
 	virtual void GetCaption(TArray<FString>& OutCaptions) const override;
diff --git a/Engine/Source/Runtime/Engine/Public/Materials/MaterialIRTypes.h b/Engine/Source/Runtime/Engine/Public/Materials/MaterialIRTypes.h
--- a/Engine/Source/Runtime/Engine/Public/Materials/MaterialIRTypes.h
+++ b/Engine/Source/Runtime/Engine/Public/Materials/MaterialIRTypes.h
@@ -159,6 +159,15 @@ struct FType
 
 	// Returns a row-major vector type of double precision floating point values with the specified number of columns.
 	static FType MakeDoubleVector(int NumColumns) { return MakeVector(EScalarKind::Double, NumColumns); }
+	
+	// Creates an integer primitive type with given dimensions.
+	static FType MakeInt(int NumRows, int NumColumns) { return MakePrimitive(EScalarKind::Int, NumRows, NumColumns); }
+
+	// Creates a single precision float primitive type with given dimensions.
+	static FType MakeFloat(int NumRows, int NumColumns) { return MakePrimitive(EScalarKind::Float, NumRows, NumColumns); }
+
+	// Creates a double precision float primitive type with given dimensions.
+	static FType MakeDouble(int NumRows, int NumColumns) { return MakePrimitive(EScalarKind::Double, NumRows, NumColumns); }
 
 	// Returns a primitive type with the given kind and dimensions.
 	// Double inverse matrix flag is only supported for EScalarKind::Double and 4x4 dimensions.
@@ -248,6 +257,9 @@ struct FType
 	// If this is a vector type, returns its primitive info. Otherwise, returns None.
 	TOptional<FPrimitive> AsVector() const { return IsVector() ? Primitive : TOptional<FPrimitive>{}; }
 
+	// If this is a matrix type, returns its primitive info. Otherwise, returns None.
+	TOptional<FPrimitive> AsMatrix() const { return IsMatrix() ? Primitive : TOptional<FPrimitive>{}; }
+
 	// If this is an aggregate type, returns a pointer to the UMaterialAggregate definition. Otherwise, returns nullptr.
 	const UMaterialAggregate* AsAggregate() const { return Is(ETypeKind::Aggregate) ? Aggregate : nullptr; }
 
