Iris
- Fixed issue where initial state was not applied when beginplay was dispatched for world from AGameStateBase::OnRep_ReplicatedHasBegunPlay()
- Added flag in actor so that we can track if inital state has been applied or not.

diff --git a/Engine/Source/Runtime/Engine/Classes/GameFramework/Actor.h b/Engine/Source/Runtime/Engine/Classes/GameFramework/Actor.h
--- a/Engine/Source/Runtime/Engine/Classes/GameFramework/Actor.h
+++ b/Engine/Source/Runtime/Engine/Classes/GameFramework/Actor.h
@@ -641,6 +641,9 @@ private:
 	/** Set while actor is being constructed. Used to ensure that construction is not re-entrant. */
 	uint8 bActorIsBeingConstructed : 1;
 
+	/** Set for actors created from networking that are not yet fully initialized */
+	uint8 bActorIsPendingPostNetInit : 1;	
+
 	static ENGINE_API uint32 BeginPlayCallDepth;
 
 protected:
@@ -4618,6 +4621,12 @@ public:
 	/** Sets the value of Instigator without causing other side effects to this instance. */
 	ENGINE_API void SetInstigator(APawn* InInstigator);
 
+	/** This should only be used by UNetActorFactory */
+	void SetActorIsPendingPostNetInit(bool bInIsPendingPostNetInit)
+	{
+		bActorIsPendingPostNetInit = bInIsPendingPostNetInit;
+	}
+	
 	//~ End Methods for Replicated Members.
 };
 
diff --git a/Engine/Source/Runtime/Engine/Private/Actor.cpp b/Engine/Source/Runtime/Engine/Private/Actor.cpp
--- a/Engine/Source/Runtime/Engine/Private/Actor.cpp
+++ b/Engine/Source/Runtime/Engine/Private/Actor.cpp
@@ -330,6 +330,8 @@ void AActor::InitializeDefaults()
 	bHasPreRegisteredAllComponents = false;
 	bHasRegisteredAllComponents = false;
 
+	bActorIsPendingPostNetInit = false;
+
 #if WITH_EDITORONLY_DATA
 	LevelInstanceFlags = ELevelInstanceFlags::None;
 	LevelInstanceType = ELevelInstanceType::None;
@@ -4752,6 +4754,17 @@ EActorUpdateOverlapsMethod AActor::GetUpdateOverlapsMethodDuringLevelStreaming()
 
 void AActor::DispatchBeginPlay(bool bFromLevelStreaming)
 {
+#if UE_WITH_IRIS
+	// If we are spawned from networking, the actor is not ready for begin play until the initial state has been applied.
+	if (bActorIsPendingPostNetInit)
+	{
+		if (UE::Net::FReplicationSystemUtil::GetReplicationSystem(this))
+		{
+			return;
+		}
+	}
+#endif
+	
 	UWorld* World = (!HasActorBegunPlay() && IsValidChecked(this) ? GetWorld() : nullptr);
 
 	if (World)
diff --git a/Engine/Source/Runtime/Engine/Private/Net/Iris/ReplicationSystem/NetActorFactory.cpp b/Engine/Source/Runtime/Engine/Private/Net/Iris/ReplicationSystem/NetActorFactory.cpp
--- a/Engine/Source/Runtime/Engine/Private/Net/Iris/ReplicationSystem/NetActorFactory.cpp
+++ b/Engine/Source/Runtime/Engine/Private/Net/Iris/ReplicationSystem/NetActorFactory.cpp
@@ -363,6 +363,9 @@ UNetObjectFactory::FInstantiateResult UNetActorFactory::InstantiateReplicatedObj
 			Result.Flags |= EReplicationBridgeCreateNetRefHandleResultFlags::AllowDestroyInstanceFromRemote;
 		}
 
+		// Mark actor as pending its PostNetInit
+		Actor->SetActorIsPendingPostNetInit(true);
+
 		return Result;
 	}
 
@@ -372,20 +375,23 @@ UNetObjectFactory::FInstantiateResult UNetActorFactory::InstantiateReplicatedObj
 	// If pre-registered, look for existing instance.
 	if (DynamicHeader->bIsPreRegistered)
 	{
-		UObject* FoundObject = Bridge->GetPreRegisteredObject(Context.Handle);
-		if (!FoundObject)
+		AActor* FoundActor = Cast<AActor>(Bridge->GetPreRegisteredObject(Context.Handle));
+		if (!FoundActor)
 		{
 			UE_LOG(LogIris, Error, TEXT("UNetActorFactory::InstantiateReplicatedObjectFromHeader Unable to find pre-registered actor: %s"), *Context.Handle.ToString());
 			return FInstantiateResult();
 		}
 
-		FInstantiateResult Result{ .Instance = FoundObject };
+		FInstantiateResult Result{ .Instance = FoundActor };
 
-		if (NetDriver->ShouldClientDestroyActor(Cast<AActor>(FoundObject)))
+		if (NetDriver->ShouldClientDestroyActor(FoundActor))
 		{
 			Result.Flags |= EReplicationBridgeCreateNetRefHandleResultFlags::AllowDestroyInstanceFromRemote;
 		}
 
+		// Mark actor as pending its PostNetInit
+		FoundActor->SetActorIsPendingPostNetInit(true);
+
 		return Result;
 	}
 
@@ -457,6 +463,9 @@ UNetObjectFactory::FInstantiateResult UNetActorFactory::InstantiateReplicatedObj
 
 	UE_LOG(LogIris, Verbose, TEXT("UNetActorFactory::InstantiateNetObjectFromHeader Spawned Actor: %s using Archetype: %s"), *Actor->GetPathName(), *GetNameSafe(Archetype));
 
+	// Mark actor as pending its PostNetInit
+	Actor->SetActorIsPendingPostNetInit(true);
+
 	return Result;
 }
 
@@ -495,15 +504,21 @@ void UNetActorFactory::PostInstantiation(const FPostInstantiationContext& Contex
 
 void UNetActorFactory::PostInit(const FPostInitContext& Context)
 {
-	// PostNetInit is only called for dynamic actors
-	if (Context.Handle.IsDynamic())
-	{
-		AActor* Actor = CastChecked<AActor>(Context.Instance);
-		LLM_SCOPE_BYNAME(TEXT("UObject/NetworkPostInit"));
-		LLM_SCOPE_DYNAMIC_STAT_OBJECTPATH(Actor->GetPackage(), ELLMTagSet::Assets);
-		LLM_SCOPE_DYNAMIC_STAT_OBJECTPATH(Actor->GetClass(), ELLMTagSet::AssetClasses);
-		UE_TRACE_METADATA_SCOPE_ASSET(Actor, Actor->GetClass());
-		Actor->PostNetInit();
+	AActor* Actor = CastChecked<AActor>(Context.Instance);
+	if (Actor)
+	{
+		// We now know that we have applied initial state
+		Actor->SetActorIsPendingPostNetInit(false);
+	
+		// PostNetInit is only called for dynamic actors
+		if (Context.Handle.IsDynamic())
+		{
+			LLM_SCOPE_BYNAME(TEXT("UObject/NetworkPostInit"));
+			LLM_SCOPE_DYNAMIC_STAT_OBJECTPATH(Actor->GetPackage(), ELLMTagSet::Assets);
+			LLM_SCOPE_DYNAMIC_STAT_OBJECTPATH(Actor->GetClass(), ELLMTagSet::AssetClasses);
+			UE_TRACE_METADATA_SCOPE_ASSET(Actor, Actor->GetClass());
+			Actor->PostNetInit();
+		}
 	}
 }
 
