Improve the Display Debug Manager UI with controls to shift the UI over by columns

diff --git a/Engine/Config/BaseInput.ini b/Engine/Config/BaseInput.ini
--- a/Engine/Config/BaseInput.ini
+++ b/Engine/Config/BaseInput.ini
@@ -48,6 +48,8 @@ ButtonRepeatDelay=0.1
 +DebugExecBindings=(Key=Tab,Command="FocusLastPIEWindow",Control=True,Shift=True)
 +DebugExecBindings=(Key=PageDown,Command="PreviousDebugTarget")
 +DebugExecBindings=(Key=PageUp,Command="NextDebugTarget")
++DebugExecBindings=(Key=RightShift,Command="DebugManager.CycleToNextColumn")
++DebugExecBindings=(Key=LeftShift,Command="DebugManager.CycleToPreviousColumn")
 ; --- PIE only binding
 +DebugExecBindings=(Key=Semicolon,Command="ToggleDebugCamera")
 
diff --git a/Engine/Source/Runtime/Engine/Classes/Engine/Canvas.h b/Engine/Source/Runtime/Engine/Classes/Engine/Canvas.h
--- a/Engine/Source/Runtime/Engine/Classes/Engine/Canvas.h
+++ b/Engine/Source/Runtime/Engine/Classes/Engine/Canvas.h
@@ -59,7 +59,7 @@ private:
 	float MaxCharHeight;
 	FVector2f InitialPos;
 	class UCanvas* Canvas;
-
+	int32 ColumnCount;
 public:
 	FDisplayDebugManager()
 		: DebugTextItem(FCanvasTextItem(FVector2D(0, 0), FText::GetEmpty(), nullptr, FLinearColor::White))
@@ -68,6 +68,7 @@ public:
 		, MaxCharHeight(0.f)
 		, InitialPos(FVector2f::ZeroVector)
 		, Canvas(nullptr)
+		, ColumnCount(0)
 	{
 		DebugTextItem.EnableShadow(FLinearColor::Black);
 	}
@@ -79,6 +80,7 @@ public:
 		InitialPos = (FVector2f)InInitialPosition;
 		CurrentPos = InitialPos;
 		NextColumXPos = 0.f;
+		ColumnCount = 0;
 	}
 
 	void SetFont(const UFont* NewFont)
diff --git a/Engine/Source/Runtime/Engine/Private/HUD.cpp b/Engine/Source/Runtime/Engine/Private/HUD.cpp
--- a/Engine/Source/Runtime/Engine/Private/HUD.cpp
+++ b/Engine/Source/Runtime/Engine/Private/HUD.cpp
@@ -445,6 +445,7 @@ void AHUD::ShowDebugInfo(float& YL, float& YPos)
 			FDisplayDebugManager& DisplayDebugManager = DebugCanvas->DisplayDebugManager;
 			DisplayDebugManager.SetDrawColor(FColor(255, 0, 0));
 			DisplayDebugManager.DrawString(FString::Printf(TEXT("Showing Debug for %s, Press [PageUp] and [PageDown] to cycle between targets."), *GetNameSafe(ShowDebugTargetActor)));
+			DisplayDebugManager.DrawString(TEXT("Press [Right Shift] or [Left Shift] to move the debug menu a column over to see more of the debug menu"));
 		}
 
 		if (IsValid(ShowDebugTargetActor))
diff --git a/Engine/Source/Runtime/Engine/Private/UserInterface/Canvas.cpp b/Engine/Source/Runtime/Engine/Private/UserInterface/Canvas.cpp
--- a/Engine/Source/Runtime/Engine/Private/UserInterface/Canvas.cpp
+++ b/Engine/Source/Runtime/Engine/Private/UserInterface/Canvas.cpp
@@ -3,7 +3,6 @@
 /*=============================================================================
 	Canvas.cpp: Unreal canvas rendering.
 =============================================================================*/
-
 #include "Engine/Canvas.h"
 #include "CanvasBatchedElementRenderItem.h"
 #include "Internationalization/IBreakIterator.h"
@@ -48,6 +47,46 @@ DEFINE_STAT(STAT_Canvas_AddTileRenderTime);
 DEFINE_STAT(STAT_Canvas_AddTriangleRenderTime);
 DEFINE_STAT(STAT_Canvas_NumBatchesCreated);
 
+namespace DisplayDebugManagerCvars
+{
+	static int32 StoredDebugMenuOffset = 0;
+	static int32 LastKnownColumnCount = 0;
+	static int32 MaxColumnsCreated = 0;
+
+	static int32 ColumnOffsetIndex = 0;
+	static FAutoConsoleVariableRef DebugManagerColumnOffsetIndex(
+		TEXT("DebugManager.ColumnOffsetIndex"),
+		ColumnOffsetIndex,
+		TEXT("")
+	);
+
+	static FAutoConsoleCommandWithWorldAndArgs CmdDebugManagerCycleToNextColumn(
+		TEXT("DebugManager.CycleToNextColumn"),
+		TEXT("Will move the text over to the next column on the debugmanager"),
+		FConsoleCommandWithWorldAndArgsDelegate::CreateLambda([](const TArray<FString>& InParams, UWorld* World)
+			{
+				DisplayDebugManagerCvars::ColumnOffsetIndex++;
+				if (DisplayDebugManagerCvars::ColumnOffsetIndex > DisplayDebugManagerCvars::MaxColumnsCreated)
+				{
+					DisplayDebugManagerCvars::ColumnOffsetIndex = DisplayDebugManagerCvars::MaxColumnsCreated;
+				}
+			})
+	);
+
+	static FAutoConsoleCommandWithWorldAndArgs CmdDebugManagerCycleToPreviousColumn(
+		TEXT("DebugManager.CycleToPreviousColumn"),
+		TEXT("Will move the text over to the previous column on the debugmanager"),
+		FConsoleCommandWithWorldAndArgsDelegate::CreateLambda([](const TArray<FString>& InParams, UWorld* World)
+			{
+				DisplayDebugManagerCvars::ColumnOffsetIndex--;
+				if (DisplayDebugManagerCvars::ColumnOffsetIndex < 0)
+				{
+					DisplayDebugManagerCvars::ColumnOffsetIndex = 0;
+				}
+			})
+	);
+}
+
 class FRDGCanvasRenderTarget final : public FRenderTarget
 {
 public:
@@ -2273,9 +2312,11 @@ void FDisplayDebugManager::DrawString(const FString& InDebugString, const float&
 		DebugTextItem.Scale = FVector2D(TextScale, TextScale);
 
 		DebugTextItem.Text = FText::FromString(InDebugString);
-		Canvas->DrawItem(DebugTextItem, FVector2D(CurrentPos.X + OptionalXOffset, CurrentPos.Y));
+		Canvas->DrawItem(DebugTextItem, FVector2D((CurrentPos.X + OptionalXOffset) - DisplayDebugManagerCvars::StoredDebugMenuOffset, CurrentPos.Y));
+
+		const float EndOfStringPosX = CurrentPos.X + OptionalXOffset + (DebugTextItem.DrawnSize.X / Canvas->GetDPIScale());
 
-		NextColumXPos = FMath::Max(NextColumXPos, CurrentPos.X + OptionalXOffset + (DebugTextItem.DrawnSize.X / Canvas->GetDPIScale()));
+		NextColumXPos = FMath::Max(NextColumXPos, EndOfStringPosX);
 		CurrentPos.Y += FMath::Max(GetYStep(), (DebugTextItem.DrawnSize.Y / Canvas->GetDPIScale()));
 		AddColumnIfNeeded();
 	}
@@ -2295,6 +2336,28 @@ void FDisplayDebugManager::AddColumnIfNeeded()
 		{
 			CurrentPos.Y = InitialPos.Y;
 			CurrentPos.X = NextColumXPos + YStep * 2.f;
+
+			// Checking and storing the max amount of columns created
+			if (ColumnCount == 0)
+			{
+				DisplayDebugManagerCvars::MaxColumnsCreated = DisplayDebugManagerCvars::LastKnownColumnCount;
+			}
+			ColumnCount++;
+			DisplayDebugManagerCvars::LastKnownColumnCount = ColumnCount;
+
+
+			// Storing the targeted column offset
+			if (DisplayDebugManagerCvars::ColumnOffsetIndex > 0)
+			{
+				if (ColumnCount == DisplayDebugManagerCvars::ColumnOffsetIndex)
+				{
+					DisplayDebugManagerCvars::StoredDebugMenuOffset = CurrentPos.X;
+				}
+			}
+			else
+			{
+				DisplayDebugManagerCvars::StoredDebugMenuOffset = 0;
+			}
 		}
 	}
-}
+}
\ No newline at end of file
