[Engine] [Input] Update SetBind to include modifier keys.

Had a request to make it so you could add modifier keys (so can use the same key to toggle slomo on with regular press, and off with shift press).
It seemed the modifier aspects of the debug exec binds didn't work.

#tests used the commands in the comment in a multiplayer environment (serv0erexec tells client to server RPC up the command). Turning on/off worked.  Rebinding worked
[RN] SetBind cheat can now use modifiers on the same key.
Preflight was ran in other stream before migrating this change.

diff --git a/Engine/Source/Runtime/Engine/Classes/GameFramework/PlayerInput.h b/Engine/Source/Runtime/Engine/Classes/GameFramework/PlayerInput.h
--- a/Engine/Source/Runtime/Engine/Classes/GameFramework/PlayerInput.h
+++ b/Engine/Source/Runtime/Engine/Classes/GameFramework/PlayerInput.h
@@ -466,7 +466,12 @@ public:
 	UFUNCTION(exec)
 	void SetMouseSensitivity(const float Sensitivity) { SetMouseSensitivity(Sensitivity, Sensitivity); }
 
-	/** Exec function to add a debug exec command */
+	/** Exec function to add a debug exec command.
+	* Use + to addmodifiers; do not include any spaces between modifiers, so that it is not conflated with the command to be bound.
+	* example usage: 
+	* SetBind k serverexec slomo 0.1
+	* SetBind shift+k serverexec slomo 1
+	*/
 	UFUNCTION(exec)
 	ENGINE_API void SetBind(FName BindName, const FString& Command);
 
diff --git a/Engine/Source/Runtime/Engine/Private/UserInterface/PlayerInput.cpp b/Engine/Source/Runtime/Engine/Private/UserInterface/PlayerInput.cpp
--- a/Engine/Source/Runtime/Engine/Private/UserInterface/PlayerInput.cpp
+++ b/Engine/Source/Runtime/Engine/Private/UserInterface/PlayerInput.cpp
@@ -2455,7 +2455,37 @@ FKeyBind UPlayerInput::GetExecBind(FString const& ExecCommand)
 void UPlayerInput::SetBind(FName BindName, const FString& Command)
 {
 #if !UE_BUILD_SHIPPING
-	FKey const BindKey(BindName);
+
+	// Split modfiiers; eg controL+alt+j into [control, alt, j]
+	TArray<FString> Tokens;
+	BindName.ToString().ParseIntoArray(Tokens, TEXT("+"), true);
+
+	FKey BindKey;
+	FKeyBind NewBind;
+	for (const FString& Token : Tokens)
+	{
+		if (Token.Equals(TEXT("ctrl"), ESearchCase::IgnoreCase)
+			|| Token.Equals(TEXT("control"), ESearchCase::IgnoreCase))
+		{
+			NewBind.Control = true;
+		}
+		else if (Token.Equals(TEXT("alt"), ESearchCase::IgnoreCase))
+		{
+			NewBind.Alt = true;
+		}
+		else if (Token.Equals(TEXT("shift"), ESearchCase::IgnoreCase))
+		{
+			NewBind.Shift = true;
+		}
+		else if (Token.Equals(TEXT("cmd"), ESearchCase::IgnoreCase))
+		{
+			NewBind.Cmd = true;
+		}
+		else if (Token.Len() > 0)
+		{
+			BindKey = FKey(FName(Token));
+		}
+	}
 	if (BindKey.IsValid())
 	{
 		FString CommandMod = Command;
@@ -2466,7 +2496,13 @@ void UPlayerInput::SetBind(FName BindName, const FString& Command)
 
 		for(int32 BindIndex = DebugExecBindings.Num()-1;BindIndex >= 0;BindIndex--)
 		{
-			if (DebugExecBindings[BindIndex].Key == BindKey)
+			// if it matches an existing binding (except for command), then overwrite the old one.
+			const FKeyBind& ExistingBinding = DebugExecBindings[BindIndex];
+			if (ExistingBinding.Key == BindKey
+				&& ExistingBinding.Alt == NewBind.Alt
+				&& ExistingBinding.Cmd == NewBind.Cmd
+				&& ExistingBinding.Shift == NewBind.Shift
+				&& ExistingBinding.Control ==  NewBind.Control)
 			{
 				DebugExecBindings[BindIndex].Command = CommandMod;
 				SaveConfig();
@@ -2474,7 +2510,6 @@ void UPlayerInput::SetBind(FName BindName, const FString& Command)
 			}
 		}
 
-		FKeyBind NewBind;
 		NewBind.Key = BindKey;
 		NewBind.Command = CommandMod;
 		DebugExecBindings.Add(NewBind);
