More fixes for CDO/Archetype determinism

#rnx

diff --git a/Engine/Source/Developer/Localization/Private/LocalizationTargetTypes.cpp b/Engine/Source/Developer/Localization/Private/LocalizationTargetTypes.cpp
--- a/Engine/Source/Developer/Localization/Private/LocalizationTargetTypes.cpp
+++ b/Engine/Source/Developer/Localization/Private/LocalizationTargetTypes.cpp
@@ -442,6 +442,14 @@ namespace
 	};
 }
 
+ULocalizationTarget::ULocalizationTarget(const FObjectInitializer& ObjectInitializer)
+{
+	if(IsTemplate())
+	{
+		Settings.Guid = FGuid();
+	}
+}
+
 #if WITH_EDITOR
 void ULocalizationTarget::PostEditChangeProperty(FPropertyChangedEvent& PropertyChangedEvent)
 {
diff --git a/Engine/Source/Developer/Localization/Public/LocalizationTargetTypes.h b/Engine/Source/Developer/Localization/Public/LocalizationTargetTypes.h
--- a/Engine/Source/Developer/Localization/Public/LocalizationTargetTypes.h
+++ b/Engine/Source/Developer/Localization/Public/LocalizationTargetTypes.h
@@ -514,6 +514,8 @@ class ULocalizationTarget : public UObject
 	GENERATED_BODY()
 
 public:
+	UE_API ULocalizationTarget(const FObjectInitializer& ObjectInitializer);
+
 	UPROPERTY(EditAnywhere, Category = "Target", meta=(ShowOnlyInnerProperties))
 	FLocalizationTargetSettings Settings;
 
diff --git a/Engine/Source/Runtime/Engine/Private/Components/ReflectionCaptureComponent.cpp b/Engine/Source/Runtime/Engine/Private/Components/ReflectionCaptureComponent.cpp
--- a/Engine/Source/Runtime/Engine/Private/Components/ReflectionCaptureComponent.cpp
+++ b/Engine/Source/Runtime/Engine/Private/Components/ReflectionCaptureComponent.cpp
@@ -826,7 +826,10 @@ void UReflectionCaptureComponent::PostInitProperties()
 	Super::PostInitProperties();
 
 	// Gets overwritten with saved value (if being loaded from disk)
-	FPlatformMisc::CreateGuid(MapBuildDataId);
+	if(!IsTemplate())
+	{
+		FPlatformMisc::CreateGuid(MapBuildDataId);
+	}
 #if WITH_EDITOR
 	bMapBuildDataIdLoaded = false;
 #endif
diff --git a/Engine/Source/Runtime/Engine/Private/Level.cpp b/Engine/Source/Runtime/Engine/Private/Level.cpp
--- a/Engine/Source/Runtime/Engine/Private/Level.cpp
+++ b/Engine/Source/Runtime/Engine/Private/Level.cpp
@@ -648,7 +648,10 @@ void ULevel::PostInitProperties()
 	Super::PostInitProperties();
 
 	// Initialize LevelBuildDataId to something unique, in case this is a new ULevel
-	LevelBuildDataId = FGuid::NewGuid();
+	if(!IsTemplate())
+	{
+		LevelBuildDataId = FGuid::NewGuid();
+	}
 }
 
 void ULevel::Serialize( FArchive& Ar )
diff --git a/Engine/Source/Runtime/Engine/Private/Materials/MaterialExpressions.cpp b/Engine/Source/Runtime/Engine/Private/Materials/MaterialExpressions.cpp
--- a/Engine/Source/Runtime/Engine/Private/Materials/MaterialExpressions.cpp
+++ b/Engine/Source/Runtime/Engine/Private/Materials/MaterialExpressions.cpp
@@ -19885,7 +19885,7 @@ bool UMaterialExpressionNamedRerouteDeclaration::GetRerouteInput(FExpressionInpu
 void UMaterialExpressionNamedRerouteDeclaration::UpdateVariableGuid(bool bForceGeneration, bool bAllowMarkingPackageDirty)
 {
 	// If we are in the editor, and we don't have a valid GUID yet, generate one.
-	if (GIsEditor && !FApp::IsGame())
+	if (GIsEditor && !FApp::IsGame() && !IsTemplate())
 	{
 		if (bForceGeneration || !VariableGuid.IsValid())
 		{
diff --git a/Engine/Source/Runtime/Experimental/Dataflow/Engine/Private/Dataflow/DataflowSubGraph.cpp b/Engine/Source/Runtime/Experimental/Dataflow/Engine/Private/Dataflow/DataflowSubGraph.cpp
--- a/Engine/Source/Runtime/Experimental/Dataflow/Engine/Private/Dataflow/DataflowSubGraph.cpp
+++ b/Engine/Source/Runtime/Experimental/Dataflow/Engine/Private/Dataflow/DataflowSubGraph.cpp
@@ -38,7 +38,7 @@ FDataflowSubGraphDelegates::FOnDataflowSubGraphLoaded FDataflowSubGraphDelegates
 
 UDataflowSubGraph::UDataflowSubGraph(class FObjectInitializer const& ObjectInitializer)
 	: UEdGraph(ObjectInitializer)
-	, SubGraphGuid(FGuid::NewGuid())
+	, SubGraphGuid(IsTemplate() ? FGuid() : FGuid::NewGuid())
 	, bIsLoaded(false)
 {}
 
