Skip BP graph properties by default for editor-only gathers

This data was originally marked as editor-only to allow potentially localizing the BP graphs themselves, but in practice this is not something that we do, and just adds bloat to gathers that include editor-only data.

If you were relying on the BP graph properties being included in the localization gather, then you can set the CVar "Localization.SkipBlueprintGraphProperties" to false to restore the old behavior.

diff --git a/Engine/Source/Runtime/CoreUObject/Private/Serialization/PropertyLocalizationDataGathering.cpp b/Engine/Source/Runtime/CoreUObject/Private/Serialization/PropertyLocalizationDataGathering.cpp
--- a/Engine/Source/Runtime/CoreUObject/Private/Serialization/PropertyLocalizationDataGathering.cpp
+++ b/Engine/Source/Runtime/CoreUObject/Private/Serialization/PropertyLocalizationDataGathering.cpp
@@ -47,8 +47,11 @@ FPropertyLocalizationDataGatherer::FPropertyLocalizationDataGatherer(TArray<FGat
 	// Iterate any bytecode containing objects
 	for (const FObjectAndGatherFlags& BytecodeToGather : BytecodePendingGather)
 	{
-		const UStruct* Struct = CastChecked<UStruct>(BytecodeToGather.Object);
-		GatherScriptBytecode(Struct->GetPathName(), Struct->Script, !!(BytecodeToGather.GatherTextFlags & EPropertyLocalizationGathererTextFlags::ForceEditorOnlyScriptData));
+		if (!EnumHasAnyFlags(BytecodeToGather.GatherTextFlags, EPropertyLocalizationGathererTextFlags::SkipScriptData))
+		{
+			const UStruct* Struct = CastChecked<UStruct>(BytecodeToGather.Object);
+			GatherScriptBytecode(Struct->GetPathName(), Struct->Script, !!(BytecodeToGather.GatherTextFlags & EPropertyLocalizationGathererTextFlags::ForceEditorOnlyScriptData));
+		}
 	}
 	BytecodePendingGather.Reset();
 }
@@ -412,7 +415,7 @@ void FPropertyLocalizationDataGatherer::GatherLocalizationDataFromChildTextPrope
 			{
 				MarkDefaultTextInstance(*TextElementValueAddress);
 			}
-			else
+			else if (!EnumHasAnyFlags(GatherTextFlags, EPropertyLocalizationGathererTextFlags::SkipProperties))
 			{
 				UPackage* const PropertyPackage = TextProperty->GetOutermost();
 				if (FTextInspector::GetFlags(*TextElementValueAddress) & ETextFlag::ConvertedProperty)
diff --git a/Engine/Source/Runtime/CoreUObject/Public/Serialization/PropertyLocalizationDataGathering.h b/Engine/Source/Runtime/CoreUObject/Public/Serialization/PropertyLocalizationDataGathering.h
--- a/Engine/Source/Runtime/CoreUObject/Public/Serialization/PropertyLocalizationDataGathering.h
+++ b/Engine/Source/Runtime/CoreUObject/Public/Serialization/PropertyLocalizationDataGathering.h
@@ -58,10 +58,21 @@ enum class EPropertyLocalizationGathererTextFlags : uint8
 	 */
 	ForceIsDefaultValue = 1<<3,
 
+	/**
+	 * Don't process text gathered from object properties.
+	 * @note Does not apply to properties gathered from script data (see SkipScriptData).
+	 */
+	SkipProperties = 1<<4,
+
+	/**
+	 * Don't process text gathered from script data.
+	 */
+	SkipScriptData = 1<<5,
+
 	/**
 	 * Don't process any sub-objects (either inner objects or object pointers).
 	 */
-	SkipSubObjects = 1<<4,
+	SkipSubObjects = 1<<6,
 };
 ENUM_CLASS_FLAGS(EPropertyLocalizationGathererTextFlags);
 
diff --git a/Engine/Source/Runtime/Engine/Private/Blueprint.cpp b/Engine/Source/Runtime/Engine/Private/Blueprint.cpp
--- a/Engine/Source/Runtime/Engine/Private/Blueprint.cpp
+++ b/Engine/Source/Runtime/Engine/Private/Blueprint.cpp
@@ -41,6 +41,7 @@
 #include "Engine/InheritableComponentHandler.h"
 
 #if WITH_EDITORONLY_DATA
+#include "HAL/IConsoleManager.h"
 #include "Kismet2/KismetDebugUtilities.h"
 #endif
 
@@ -189,14 +190,17 @@ void FBPVariableDescription::RemoveMetaData(const FName Key)
 // UBlueprintCore
 
 #if WITH_EDITORONLY_DATA
-namespace
+namespace BlueprintLocalization
 {
-	void GatherBlueprintForLocalization(const UObject* const Object, FPropertyLocalizationDataGatherer& PropertyLocalizationDataGatherer, const EPropertyLocalizationGathererTextFlags GatherTextFlags)
+	static bool bSkipGraphProperties = true;
+	static FAutoConsoleVariableRef CVarSkipGraphProperties(TEXT("Localization.SkipBlueprintGraphProperties"), bSkipGraphProperties, TEXT("True to skip BP graph properties from the localization gather, or false to mark them as editor-only"));
+
+	void GatherForLocalization(const UObject* const Object, FPropertyLocalizationDataGatherer& PropertyLocalizationDataGatherer, const EPropertyLocalizationGathererTextFlags GatherTextFlags)
 	{
 		const UBlueprintCore* const BlueprintCore = CastChecked<UBlueprintCore>(Object);
 
 		// Blueprint assets never exist at runtime, so treat all of their properties as editor-only, but allow their script (which is available at runtime) to be gathered by a game
-		EPropertyLocalizationGathererTextFlags BlueprintGatherFlags = GatherTextFlags | EPropertyLocalizationGathererTextFlags::ForceEditorOnlyProperties;
+		EPropertyLocalizationGathererTextFlags BlueprintGatherFlags = GatherTextFlags | (bSkipGraphProperties ? EPropertyLocalizationGathererTextFlags::SkipProperties : EPropertyLocalizationGathererTextFlags::ForceEditorOnlyProperties);
 
 #if WITH_EDITOR
 		if (const UBlueprint* const Blueprint = Cast<UBlueprint>(Object))
@@ -274,7 +278,7 @@ UBlueprintCore::UBlueprintCore(const FObjectInitializer& ObjectInitializer)
 	: Super(ObjectInitializer)
 {
 #if WITH_EDITORONLY_DATA
-	{ static const FAutoRegisterLocalizationDataGatheringCallback AutomaticRegistrationOfLocalizationGatherer(UBlueprintCore::StaticClass(), &GatherBlueprintForLocalization); }
+	{ static const FAutoRegisterLocalizationDataGatheringCallback AutomaticRegistrationOfLocalizationGatherer(UBlueprintCore::StaticClass(), &BlueprintLocalization::GatherForLocalization); }
 #endif
 
 	bLegacyNeedToPurgeSkelRefs = true;
