New material translator:
- Introduced Transpose operator.
- Added a GetCommonType() API that accepts an array of FValuePtr, so the values originating input can be reported if a common type does not exist.
- Some renamings in the Emitter.

#rnx

diff --git a/Engine/Source/Runtime/Engine/Private/Materials/MaterialExpressionsIR.cpp b/Engine/Source/Runtime/Engine/Private/Materials/MaterialExpressionsIR.cpp
--- a/Engine/Source/Runtime/Engine/Private/Materials/MaterialExpressionsIR.cpp
+++ b/Engine/Source/Runtime/Engine/Private/Materials/MaterialExpressionsIR.cpp
@@ -2188,12 +2188,13 @@ uint32 GetMaterialExpressionOperatorArity(EMaterialExpressionOperatorKind Operat
 
 void UMaterialExpressionOperator::Build(MIR::FEmitter& Em)
 {
+	FValueRef AValue = Em.InputDefaultFloat(&DynamicInputs[0].ExpressionInput, DynamicInputs[0].ConstValue);
 	if (bAllowAddPin)
 	{
 		MIR::EOperator OpMIR = MaterialExpressionOperatorToMIR(Operator);
 
 		// Apply operation to iteratively to all input values
-		FValueRef Value = Em.InputDefaultFloat(&DynamicInputs[0].ExpressionInput, DynamicInputs[0].ConstValue);
+		FValueRef Value = AValue;
 		for (int32 i = 1; i < DynamicInputs.Num(); i++)
 		{
 			FValueRef CurValue = Em.InputDefaultFloat(&DynamicInputs[i].ExpressionInput, DynamicInputs[i].ConstValue);
@@ -2207,7 +2208,6 @@ void UMaterialExpressionOperator::Build(MIR::FEmitter& Em)
 		MIR::EOperator OpMIR = MaterialExpressionOperatorToMIR(Operator);
 		int32 OperatorArity = MIR::GetOperatorArity(OpMIR);
 
-		FValueRef AValue = Em.InputDefaultFloat(&DynamicInputs[0].ExpressionInput, DynamicInputs[0].ConstValue);
 		FValueRef BValue = OperatorArity >= 2 ? Em.InputDefaultFloat(&DynamicInputs[1].ExpressionInput, DynamicInputs[1].ConstValue) : FValueRef{};
 		FValueRef CValue = OperatorArity >= 3 ? Em.InputDefaultFloat(&DynamicInputs[2].ExpressionInput, DynamicInputs[2].ConstValue) : FValueRef{};
 
diff --git a/Engine/Source/Runtime/Engine/Private/Materials/MaterialIR.cpp b/Engine/Source/Runtime/Engine/Private/Materials/MaterialIR.cpp
--- a/Engine/Source/Runtime/Engine/Private/Materials/MaterialIR.cpp
+++ b/Engine/Source/Runtime/Engine/Private/Materials/MaterialIR.cpp
@@ -720,6 +720,7 @@ const TCHAR* LexToString(EOperator Op)
 		case UO_Sqrt: return TEXT("Sqrt");
 		case UO_Tan: return TEXT("Tan");
 		case UO_Tanh: return TEXT("Tanh");
+		case UO_Transpose: return TEXT("Transpose");
 		case UO_Truncate: return TEXT("Truncate");
 
 		/* Binary operators */
diff --git a/Engine/Source/Runtime/Engine/Private/Materials/MaterialIREmitter.cpp b/Engine/Source/Runtime/Engine/Private/Materials/MaterialIREmitter.cpp
--- a/Engine/Source/Runtime/Engine/Private/Materials/MaterialIREmitter.cpp
+++ b/Engine/Source/Runtime/Engine/Private/Materials/MaterialIREmitter.cpp
@@ -250,6 +250,41 @@ void FEmitter::Error(FStringView Message)
 
 /*--------------------------------- Type handling ----------------------------------*/
 
+FType FEmitter::TryGetCommonType(FType A, FType B)
+{
+	// Trivial case: types are equal
+	if (A == B)
+	{
+		return A;
+	}
+
+	// Primitive types can only be constructed from other primitive types.
+	TOptional<FPrimitive> PrimitiveA = A.AsPrimitive();
+	TOptional<FPrimitive> PrimitiveB = B.AsPrimitive();
+	if (!PrimitiveA || !PrimitiveB)
+	{
+		return FType::MakePoison();
+	}
+	
+	// No common type between row and column vectors
+	if ((PrimitiveA->IsRowVector() && PrimitiveB->IsColumnVector()) || (PrimitiveA->IsColumnVector() && PrimitiveB->IsRowVector()))
+	{
+		return FType::MakePoison();
+	}
+
+	// Can't cast a vector to a matrix.
+	if ((PrimitiveA->IsRowVector() && PrimitiveB->IsMatrix()) || (PrimitiveA->IsMatrix() && PrimitiveB->IsRowVector()))
+	{
+		return FType::MakePoison();
+	}
+	
+	// Return the primitive type with the maximum number of rows and columns between the two types.
+	EScalarKind ScalarKind = FMath::Max(PrimitiveA->ScalarKind, PrimitiveB->ScalarKind);
+	int32 NumRows = FMath::Max(PrimitiveA->NumRows, PrimitiveB->NumRows);
+	int32 NumColumns = FMath::Max(PrimitiveA->NumColumns, PrimitiveB->NumColumns);
+	return FType::MakePrimitive(ScalarKind, NumRows, NumColumns);
+}
+
 FType FEmitter::GetCommonType(FType A, FType B)
 {
 	if (FType CommonType = TryGetCommonType(A, B))
@@ -258,7 +293,82 @@ FType FEmitter::GetCommonType(FType A, FType B)
 	}
 
 	Errorf(TEXT("No common type between '%s' and '%s'."), *A.GetSpelling(), *B.GetSpelling());
-	return {};
+	return FType::MakePoison();
+}
+
+FType FEmitter::GetCommonType(TConstArrayView<FValueRef> Values)
+{
+	check(!Values.IsEmpty() && Values[0]);
+
+	// Find the common type among non-null values
+	FType CommonType = Values[0]->Type;
+	for (int32 i = 1; i < Values.Num(); ++i)
+	{
+		if (Values[i].IsPoison())
+		{
+			return FType::MakePoison();
+		}
+		if (Values[i])
+		{
+			CommonType = TryGetCommonType(CommonType, Values[i]->Type);
+		}
+	}
+
+	// If common type is valid, return it
+	if (!CommonType.IsPoison())
+	{
+		return CommonType;
+	}
+
+	// ...otherwise generate an error. This error message prints the input the values
+	// come from, if available.
+
+	// Search for the last valid index in the values array, so that we know when to print " and "
+	int32 LastIndex = Values.Num() - 1;
+	for (; LastIndex >= 0; --LastIndex)
+	{
+		if (Values[LastIndex])
+		{
+			break;
+		}
+	}
+
+	// Whether some value has already been reported (used to print the comma ", ")
+	bool bSomeValueAlreadyPrinted = false;
+
+	FString ErrorMsg{ TEXTVIEW("No common type between ") };
+	for (int32 i = 0; i < Values.Num(); ++i)
+	{
+		if (!Values[i])
+		{
+			continue;
+		}
+
+		FValueRef Value = Values[i];
+
+		if (i == LastIndex)
+		{
+			ErrorMsg.Append(TEXTVIEW(" and "));
+		}
+		else if (bSomeValueAlreadyPrinted)
+		{
+			ErrorMsg.Append(TEXTVIEW(", "));
+		}
+
+		ErrorMsg.Appendf(TEXT("'%s'"), *Value->Type.GetSpelling());
+		
+		if (Value.Input)
+		{
+			ErrorMsg.Appendf(TEXT(" (from input '%s')"), *SlowFindInputName(Expression, Value.Input).ToString());
+		}
+
+		bSomeValueAlreadyPrinted = true;
+	}
+
+	ErrorMsg.AppendChar('.');
+	Error(ErrorMsg);
+
+	return FType::MakePoison();
 }
 
 FType FEmitter::GetMaterialAggregateAttributeType(const UMaterialAggregate* Aggregate, int32 AttributeIndex)
@@ -1699,52 +1809,52 @@ static FValueRef TryFoldOperatorScalar(FEmitter& Emitter, EOperator Op, FValueRe
 }
 
 // Used to filter what parameter *primitive* types operators can take.
-enum EOperatorParameterTypeFilter
-{
-	OPTF_Unknown = 0xff,										// Unspecified
-	OPTF_Any = 0,												// Any primitive type
-	OPTF_CheckIsBoolean = 1 << 1,								// Check the type is boolean primitive of any dimension
-	OPTF_CheckIsInteger = 1 << 2,								// Check the type is integer primitive of any dimension
-	OPTF_CheckIsArithmetic = 1 << 3,							// Check the type is arithmetic primitive of any dimension (i.e. that supports arithmetic operations)
-	OPTF_CheckIsMatrix = 1 << 4,								// Check the type is any matrix type
-	OPTF_CheckIsNotMatrix = 1 << 5,								// Check the type is any primitive type except matrices
-	OPTF_CheckIsVector3 = 1 << 6,								// Check the type is a 3D vector of any scalar type
-	OPTF_CheckIsNonNegativeFloatConst = 1 << 7,					// Check that if the argument is a constant float, it is not negative (x >= 0)
-	OPTF_CheckIsNonZeroFloatConst = 1 << 8,						// Check that if the argument is a constant float, it is not zero	(x != 0)
-	OPTF_CheckIsOneOrGreaterFloatConst = 1 << 9,				// Check that if the argument is a constant float, it is 1 or greater (xFloat >= 1)
-	OPTF_CheckIsBetweenMinusOneAndPlusOneFloatConst = 1 << 10,	// Check that if the argument is a constant float, it is between -1 and 1 (-1 < x < 1)
-	OPTF_CastToFirstArgumentType = 1 << 11,						// Cast the argument to the first argument's type
-	OPTF_CastToAnyFloat = 1 << 12,								// Cast the argument to the floating point primitive type of any dimension
-	OPTF_AllowDouble = 1 << 13,									// This argument is allowed to be a double
-	OPTF_CastToCommonScalarKind = 1 << 14,						// Casts the argument to have the scalar kind in common with other arguments
-	OPTF_CastToCommonType = 1 << 15,	 						// Cast the argument to the common arguments type
-	OPTF_CastToCommonArithmeticType = OPTF_CheckIsArithmetic | OPTF_CastToCommonType,
-	OPTF_CastToCommonArithmeticTypeAllowDouble = OPTF_CheckIsArithmetic | OPTF_CastToCommonType | OPTF_AllowDouble,
-	OPTF_CastToCommonFloatType = OPTF_CastToAnyFloat | OPTF_CastToCommonType,
+enum EOperatorParameterFlags
+{
+	OPF_Unknown = 0xff,                                         // Unspecified
+	OPF_Any = 0,                                                // Any primitive type
+	OPF_CheckIsBoolean = 1 << 1,                                // Check the type is boolean primitive of any dimension
+	OPF_CheckIsInteger = 1 << 2,                                // Check the type is integer primitive of any dimension
+	OPF_CheckIsArithmetic = 1 << 3,                             // Check the type is arithmetic primitive of any dimension (i.e. that supports arithmetic operations)
+	OPF_CheckIsMatrix = 1 << 4,                                 // Check the type is any matrix type
+	OPF_CheckIsNotMatrix = 1 << 5,                              // Check the type is any primitive type except matrices
+	OPF_CheckIsVector3 = 1 << 6,                                // Check the type is a 3D vector of any scalar type
+	OPF_CheckIsNonNegativeFloatConst = 1 << 7,                  // Check that if the argument is a constant float, it is not negative (x >= 0)
+	OPF_CheckIsNonZeroFloatConst = 1 << 8,                      // Check that if the argument is a constant float, it is not zero    (x != 0)
+	OPF_CheckIsOneOrGreaterFloatConst = 1 << 9,                 // Check that if the argument is a constant float, it is 1 or greater (xFloat >= 1)
+	OPF_CheckIsBetweenMinusOneAndPlusOneFloatConst = 1 << 10,   // Check that if the argument is a constant float, it is between -1 and 1 (-1 < x < 1)
+	OPF_CastToFirstArgumentType = 1 << 11,                      // Cast the argument to the first argument's type
+	OPF_CastToAnyFloat = 1 << 12,                               // Cast the argument to the floating point primitive type of any dimension
+	OPF_AllowDouble = 1 << 13,                                  // This argument is allowed to be a double
+	OPF_CastToCommonScalarKind = 1 << 14,                       // Casts the argument to have the scalar kind in common with other arguments
+	OPF_CastToCommonType = 1 << 15,                             // Cast the argument to the common arguments type
+	OPF_CastToCommonArithmeticType = OPF_CheckIsArithmetic | OPF_CastToCommonType,
+	OPF_CastToCommonArithmeticTypeAllowDouble = OPF_CheckIsArithmetic | OPF_CastToCommonType | OPF_AllowDouble,
+	OPF_CastToCommonFloatType = OPF_CastToAnyFloat | OPF_CastToCommonType,
 };
 
-inline EOperatorParameterTypeFilter operator|(EOperatorParameterTypeFilter A, EOperatorParameterTypeFilter B)
+inline EOperatorParameterFlags operator|(EOperatorParameterFlags A, EOperatorParameterFlags B)
 {
-	return EOperatorParameterTypeFilter(uint32(A) | uint32(B));
+	return EOperatorParameterFlags(uint32(A) | uint32(B));
 }
 
 // Used to determine the operator result type based on argument types
-enum EOperatorReturnType
-{
-	ORT_Unknown,							// Unspecified
-	ORT_FirstArgumentType,					// The same type as the first argument, LWC input produces float result
-	ORT_BooleanWithFirstArgumentDimensions,	// A boolean primitive type with the same dimensions (rows and columns) as the first argument type
-	ORT_FirstArgumentTypeToScalarLWC,		// A scalar primitive type with the same kind as the scalar type of the first argument, LWC results allowed
-	ORT_SecondArgumentType,					// The same type as the second argument
-	ORT_FirstArgumentTypeLWC,				// The same type as the first argument, LWC results allowed
-	ORT_MatrixMultiplyResult,				// The result type of the matrix multiplication of first two arguments
+enum EOperatorResult
+{
+	OR_Unknown,							    // Unspecified
+	OR_FirstArgumentType,					// The same type as the first argument, LWC input produces float result
+	OR_BooleanWithFirstArgumentDimensions,	// A boolean primitive type with the same dimensions (rows and columns) as the first argument type
+	OR_FirstArgumentTypeToScalarLWC,		// A scalar primitive type with the same kind as the scalar type of the first argument, LWC results allowed
+	OR_SecondArgumentType,					// The same type as the second argument
+	OR_FirstArgumentTypeAllowDouble,		// The same type as the first argument, LWC results allowed
+	OR_MatrixMultiplyResult,				// The result type of the matrix multiplication of first two arguments
 };
 
 // The signature of an operator consisting of its parameter and return type information.
 struct FOperatorSignature
 {
-	EOperatorParameterTypeFilter ParameterTypes[3] = { OPTF_Unknown, OPTF_Unknown, OPTF_Unknown };
-	EOperatorReturnType ReturnType = ORT_Unknown;
+	EOperatorParameterFlags ParameterFlags[3] = { OPF_Unknown, OPF_Unknown, OPF_Unknown };
+	EOperatorResult 		Result = OR_Unknown;
 };
 
 // Returns the signature of an operator.
@@ -1752,77 +1862,77 @@ static const FOperatorSignature* GetOperatorSignature(EOperator Op)
 {
 	static const FOperatorSignature* Signatures = [] ()
 	{
-		const FOperatorSignature UnaryFloat = { { OPTF_CheckIsArithmetic | OPTF_CastToAnyFloat }, ORT_FirstArgumentType };
-		const FOperatorSignature UnaryFloatLWC = { { OPTF_CheckIsArithmetic | OPTF_CastToAnyFloat | OPTF_AllowDouble }, ORT_FirstArgumentTypeLWC };
-		const FOperatorSignature UnaryFloatLWCDemote = { { OPTF_CheckIsArithmetic | OPTF_CastToAnyFloat | OPTF_AllowDouble }, ORT_FirstArgumentType };
-		const FOperatorSignature UnaryFloatToBoolean = { { OPTF_CheckIsArithmetic | OPTF_CastToAnyFloat }, ORT_BooleanWithFirstArgumentDimensions };
-		const FOperatorSignature BinaryArithmetic = { { OPTF_CastToCommonArithmeticType, OPTF_CastToCommonArithmeticType }, ORT_FirstArgumentType };
-		const FOperatorSignature BinaryArithmeticLWC = { { OPTF_CastToCommonArithmeticTypeAllowDouble, OPTF_CastToCommonArithmeticTypeAllowDouble }, ORT_FirstArgumentTypeLWC };
-		const FOperatorSignature BinaryArithmeticLWCDemote = { { OPTF_CastToCommonArithmeticTypeAllowDouble, OPTF_CastToCommonArithmeticTypeAllowDouble }, ORT_FirstArgumentType };
-		const FOperatorSignature BinaryInteger = { { OPTF_CheckIsInteger | OPTF_CastToCommonArithmeticType, OPTF_CheckIsInteger | OPTF_CastToCommonArithmeticType }, ORT_FirstArgumentType };
-		const FOperatorSignature BinaryFloat = { { OPTF_CastToCommonFloatType, OPTF_CastToCommonFloatType }, ORT_FirstArgumentType };
-		const FOperatorSignature BinaryArithmeticComparisonLWC = { { OPTF_CastToCommonArithmeticType | OPTF_AllowDouble, OPTF_CastToCommonArithmeticType | OPTF_AllowDouble }, ORT_BooleanWithFirstArgumentDimensions };
-		const FOperatorSignature BinaryLogical = { { OPTF_CheckIsBoolean | OPTF_CastToCommonType, OPTF_CheckIsBoolean | OPTF_CastToCommonType }, ORT_FirstArgumentType };
-		const FOperatorSignature TernaryArithmeticDouble = { { OPTF_CastToCommonArithmeticTypeAllowDouble, OPTF_CastToCommonArithmeticTypeAllowDouble, OPTF_CastToCommonArithmeticTypeAllowDouble }, ORT_FirstArgumentTypeLWC };
-		const FOperatorSignature TernaryFloatLWCDemote = { { OPTF_CastToCommonArithmeticTypeAllowDouble | OPTF_CastToAnyFloat, OPTF_CastToCommonArithmeticTypeAllowDouble, OPTF_CastToCommonArithmeticTypeAllowDouble }, ORT_FirstArgumentType };
+		const FOperatorSignature UnaryFloat = { { OPF_CheckIsArithmetic | OPF_CastToAnyFloat }, OR_FirstArgumentType };
+		const FOperatorSignature UnaryFloatOrDouble = { { OPF_CheckIsArithmetic | OPF_CastToAnyFloat | OPF_AllowDouble }, OR_FirstArgumentTypeAllowDouble };
+		const FOperatorSignature UnaryFloatLWCDemote = { { OPF_CheckIsArithmetic | OPF_CastToAnyFloat | OPF_AllowDouble }, OR_FirstArgumentType };
+		const FOperatorSignature UnaryFloatToBoolean = { { OPF_CheckIsArithmetic | OPF_CastToAnyFloat }, OR_BooleanWithFirstArgumentDimensions };
+		const FOperatorSignature BinaryArithmetic = { { OPF_CastToCommonArithmeticType, OPF_CastToCommonArithmeticType }, OR_FirstArgumentType };
+		const FOperatorSignature BinaryArithmeticAllowDouble = { { OPF_CastToCommonArithmeticTypeAllowDouble, OPF_CastToCommonArithmeticTypeAllowDouble }, OR_FirstArgumentTypeAllowDouble };
+		const FOperatorSignature BinaryArithmeticLWCDemote = { { OPF_CastToCommonArithmeticTypeAllowDouble, OPF_CastToCommonArithmeticTypeAllowDouble }, OR_FirstArgumentType };
+		const FOperatorSignature BinaryInteger = { { OPF_CheckIsInteger | OPF_CastToCommonArithmeticType, OPF_CheckIsInteger | OPF_CastToCommonArithmeticType }, OR_FirstArgumentType };
+		const FOperatorSignature BinaryFloat = { { OPF_CastToCommonFloatType, OPF_CastToCommonFloatType }, OR_FirstArgumentType };
+		const FOperatorSignature BinaryArithmeticComparisonAllowDouble = { { OPF_CastToCommonArithmeticType | OPF_AllowDouble, OPF_CastToCommonArithmeticType | OPF_AllowDouble }, OR_BooleanWithFirstArgumentDimensions };
+		const FOperatorSignature BinaryLogical = { { OPF_CheckIsBoolean | OPF_CastToCommonType, OPF_CheckIsBoolean | OPF_CastToCommonType }, OR_FirstArgumentType };
+		const FOperatorSignature TernaryArithmeticDouble = { { OPF_CastToCommonArithmeticTypeAllowDouble, OPF_CastToCommonArithmeticTypeAllowDouble, OPF_CastToCommonArithmeticTypeAllowDouble }, OR_FirstArgumentTypeAllowDouble };
+		const FOperatorSignature TernaryFloatDoubleDemote = { { OPF_CastToCommonArithmeticTypeAllowDouble | OPF_CastToAnyFloat, OPF_CastToCommonArithmeticTypeAllowDouble, OPF_CastToCommonArithmeticTypeAllowDouble }, OR_FirstArgumentType };
 
 		static FOperatorSignature S[OperatorCount];
 
 		/* unary operators */
-		S[UO_BitwiseNot] 				= { { OPTF_CheckIsInteger }, ORT_FirstArgumentType };
-		S[UO_Negate] 					= { { OPTF_CheckIsArithmetic | OPTF_AllowDouble }, ORT_FirstArgumentTypeLWC };
-		S[UO_Not] 						= { { OPTF_CheckIsBoolean }, ORT_FirstArgumentType };
+		S[UO_BitwiseNot] 				= { { OPF_CheckIsInteger }, OR_FirstArgumentType };
+		S[UO_Negate] 					= { { OPF_CheckIsArithmetic | OPF_AllowDouble }, OR_FirstArgumentTypeAllowDouble };
+		S[UO_Not] 						= { { OPF_CheckIsBoolean }, OR_FirstArgumentType };
 
-		S[UO_Abs] 						= UnaryFloatLWC;
+		S[UO_Abs] 						= UnaryFloatOrDouble;
 		S[UO_ACos] 						= UnaryFloatLWCDemote;
 		S[UO_ACosFast]					= UnaryFloat;
-		S[UO_ACosh] 					= { { OPTF_CheckIsArithmetic | OPTF_CastToAnyFloat | OPTF_CheckIsOneOrGreaterFloatConst }, ORT_FirstArgumentType };
+		S[UO_ACosh] 					= { { OPF_CheckIsArithmetic | OPF_CastToAnyFloat | OPF_CheckIsOneOrGreaterFloatConst }, OR_FirstArgumentType };
 		S[UO_ASin] 						= UnaryFloatLWCDemote;
 		S[UO_ASinFast]					= UnaryFloat;
 		S[UO_ASinh] 					= UnaryFloat;
 		S[UO_ATan] 						= UnaryFloatLWCDemote;
 		S[UO_ATanFast]					= UnaryFloat;
-		S[UO_ATanh] 					= { { OPTF_CheckIsArithmetic | OPTF_CastToAnyFloat | OPTF_CheckIsBetweenMinusOneAndPlusOneFloatConst }, ORT_FirstArgumentType };
-		S[UO_Ceil] 						= UnaryFloatLWC;
+		S[UO_ATanh] 					= { { OPF_CheckIsArithmetic | OPF_CastToAnyFloat | OPF_CheckIsBetweenMinusOneAndPlusOneFloatConst }, OR_FirstArgumentType };
+		S[UO_Ceil] 						= UnaryFloatOrDouble;
 		S[UO_Cos] 						= UnaryFloatLWCDemote;
 		S[UO_Exponential]				= UnaryFloat;
 		S[UO_Exponential2]				= UnaryFloat;
-		S[UO_Floor] 					= UnaryFloatLWC;
+		S[UO_Floor] 					= UnaryFloatOrDouble;
 		S[UO_Frac]						= UnaryFloatLWCDemote;
 		S[UO_IsFinite]					= UnaryFloatToBoolean;
 		S[UO_IsInf]						= UnaryFloatToBoolean;
 		S[UO_IsNan]						= UnaryFloatToBoolean;
-		S[UO_Length]					= { { OPTF_CheckIsArithmetic | OPTF_CheckIsNotMatrix | OPTF_CastToAnyFloat | OPTF_AllowDouble }, ORT_FirstArgumentTypeToScalarLWC };
-		S[UO_Logarithm] 				= { { OPTF_CheckIsArithmetic | OPTF_CheckIsNonZeroFloatConst | OPTF_CheckIsNonNegativeFloatConst | OPTF_CastToAnyFloat }, ORT_FirstArgumentType };
-		S[UO_Logarithm10]				= { { OPTF_CheckIsArithmetic | OPTF_CheckIsNonZeroFloatConst | OPTF_CheckIsNonNegativeFloatConst | OPTF_CastToAnyFloat }, ORT_FirstArgumentType };
-		S[UO_Logarithm2] 				= { { OPTF_CheckIsArithmetic | OPTF_CheckIsNonZeroFloatConst | OPTF_CheckIsNonNegativeFloatConst | OPTF_CastToAnyFloat }, ORT_FirstArgumentType };
+		S[UO_Length]					= { { OPF_CheckIsArithmetic | OPF_CheckIsNotMatrix | OPF_CastToAnyFloat | OPF_AllowDouble }, OR_FirstArgumentTypeToScalarLWC };
+		S[UO_Logarithm] 				= { { OPF_CheckIsArithmetic | OPF_CheckIsNonZeroFloatConst | OPF_CheckIsNonNegativeFloatConst | OPF_CastToAnyFloat }, OR_FirstArgumentType };
+		S[UO_Logarithm10]				= { { OPF_CheckIsArithmetic | OPF_CheckIsNonZeroFloatConst | OPF_CheckIsNonNegativeFloatConst | OPF_CastToAnyFloat }, OR_FirstArgumentType };
+		S[UO_Logarithm2] 				= { { OPF_CheckIsArithmetic | OPF_CheckIsNonZeroFloatConst | OPF_CheckIsNonNegativeFloatConst | OPF_CastToAnyFloat }, OR_FirstArgumentType };
 		S[UO_LWCTile]					= { };		// UNUSED
-		S[UO_Reciprocal]				= { { OPTF_CheckIsArithmetic | OPTF_CheckIsNonZeroFloatConst | OPTF_CastToAnyFloat }, ORT_FirstArgumentType };
-		S[UO_Round] 					= UnaryFloatLWC;
-		S[UO_Rsqrt]						= { { OPTF_CheckIsArithmetic | OPTF_CheckIsNonZeroFloatConst | OPTF_CastToAnyFloat }, ORT_FirstArgumentType };
+		S[UO_Reciprocal]				= { { OPF_CheckIsArithmetic | OPF_CheckIsNonZeroFloatConst | OPF_CastToAnyFloat }, OR_FirstArgumentType };
+		S[UO_Round] 					= UnaryFloatOrDouble;
+		S[UO_Rsqrt]						= { { OPF_CheckIsArithmetic | OPF_CheckIsNonZeroFloatConst | OPF_CastToAnyFloat }, OR_FirstArgumentType };
 		S[UO_Saturate]					= UnaryFloatLWCDemote;
 		S[UO_Sign]						= UnaryFloatLWCDemote;
 		S[UO_Sin] 						= UnaryFloatLWCDemote;
-		S[UO_Sqrt]						= { { OPTF_CheckIsArithmetic | OPTF_CheckIsNonNegativeFloatConst | OPTF_CastToAnyFloat | OPTF_AllowDouble }, ORT_FirstArgumentType };
+		S[UO_Sqrt]						= { { OPF_CheckIsArithmetic | OPF_CheckIsNonNegativeFloatConst | OPF_CastToAnyFloat | OPF_AllowDouble }, OR_FirstArgumentType };
 		S[UO_Tan] 						= UnaryFloatLWCDemote;
 		S[UO_Tanh] 						= UnaryFloat;
-		S[UO_Truncate]					= UnaryFloatLWC;
+		S[UO_Truncate]					= UnaryFloatOrDouble;
 
 		/* binary operators */
-		S[BO_Equals]					= { { OPTF_CastToCommonType | OPTF_AllowDouble, OPTF_CastToCommonType | OPTF_AllowDouble }, ORT_BooleanWithFirstArgumentDimensions };
-		S[BO_GreaterThan]				= BinaryArithmeticComparisonLWC;
-		S[BO_GreaterThanOrEquals]		= BinaryArithmeticComparisonLWC;
-		S[BO_LessThan]					= BinaryArithmeticComparisonLWC;
-		S[BO_LessThanOrEquals]			= BinaryArithmeticComparisonLWC;
-		S[BO_NotEquals]					= { { OPTF_CastToCommonType | OPTF_AllowDouble, OPTF_CastToCommonType | OPTF_AllowDouble }, ORT_BooleanWithFirstArgumentDimensions };
+		S[BO_Equals]					= { { OPF_CastToCommonType | OPF_AllowDouble, OPF_CastToCommonType | OPF_AllowDouble }, OR_BooleanWithFirstArgumentDimensions };
+		S[BO_GreaterThan]				= BinaryArithmeticComparisonAllowDouble;
+		S[BO_GreaterThanOrEquals]		= BinaryArithmeticComparisonAllowDouble;
+		S[BO_LessThan]					= BinaryArithmeticComparisonAllowDouble;
+		S[BO_LessThanOrEquals]			= BinaryArithmeticComparisonAllowDouble;
+		S[BO_NotEquals]					= { { OPF_CastToCommonType | OPF_AllowDouble, OPF_CastToCommonType | OPF_AllowDouble }, OR_BooleanWithFirstArgumentDimensions };
 		
 		S[BO_And] 						= BinaryLogical;
 		S[BO_Or] 						= BinaryLogical;
-		S[BO_Add] 						= BinaryArithmeticLWC;
-		S[BO_Subtract] 					= BinaryArithmeticLWC;
-		S[BO_Multiply] 					= BinaryArithmeticLWC;
-		S[BO_MatrixMultiply]			= { { OPTF_CheckIsArithmetic | OPTF_CastToCommonScalarKind, OPTF_CheckIsArithmetic | OPTF_CastToCommonScalarKind }, ORT_MatrixMultiplyResult };
-		S[BO_Divide] 					= BinaryArithmeticLWC;
+		S[BO_Add] 						= BinaryArithmeticAllowDouble;
+		S[BO_Subtract] 					= BinaryArithmeticAllowDouble;
+		S[BO_Multiply] 					= BinaryArithmeticAllowDouble;
+		S[BO_MatrixMultiply]			= { { OPF_CheckIsArithmetic | OPF_CastToCommonScalarKind, OPF_CheckIsArithmetic | OPF_CastToCommonScalarKind }, OR_MatrixMultiplyResult };
+		S[BO_Divide] 					= BinaryArithmeticAllowDouble;
 		S[BO_Modulo] 					= BinaryArithmetic;
 		S[BO_BitwiseAnd]				= BinaryInteger;
 		S[BO_BitwiseOr]					= BinaryInteger;
@@ -1832,21 +1942,21 @@ static const FOperatorSignature* GetOperatorSignature(EOperator Op)
 
 		S[BO_ATan2]						= BinaryFloat;
 		S[BO_ATan2Fast]					= BinaryFloat;
-		S[BO_Cross] 					= { { OPTF_CheckIsArithmetic | OPTF_CheckIsVector3, OPTF_CastToFirstArgumentType }, ORT_FirstArgumentType };
-		S[BO_Distance]					= { { OPTF_CastToCommonFloatType | OPTF_AllowDouble, OPTF_CastToCommonFloatType | OPTF_AllowDouble }, ORT_FirstArgumentTypeToScalarLWC };
-		S[BO_Dot] 						= { { OPTF_CheckIsArithmetic | OPTF_CheckIsNotMatrix | OPTF_AllowDouble, OPTF_CastToFirstArgumentType | OPTF_AllowDouble }, ORT_FirstArgumentTypeToScalarLWC };
-		S[BO_Fmod] 						= { { OPTF_CastToCommonFloatType | OPTF_AllowDouble, OPTF_CastToCommonFloatType }, ORT_FirstArgumentType };		// First input can be LWC, but second and output are always demoted
+		S[BO_Cross] 					= { { OPF_CheckIsArithmetic | OPF_CheckIsVector3, OPF_CastToFirstArgumentType }, OR_FirstArgumentType };
+		S[BO_Distance]					= { { OPF_CastToCommonFloatType | OPF_AllowDouble, OPF_CastToCommonFloatType | OPF_AllowDouble }, OR_FirstArgumentTypeToScalarLWC };
+		S[BO_Dot] 						= { { OPF_CheckIsArithmetic | OPF_CheckIsNotMatrix | OPF_AllowDouble, OPF_CastToFirstArgumentType | OPF_AllowDouble }, OR_FirstArgumentTypeToScalarLWC };
+		S[BO_Fmod] 						= { { OPF_CastToCommonFloatType | OPF_AllowDouble, OPF_CastToCommonFloatType }, OR_FirstArgumentType };		// First input can be LWC, but second and output are always demoted
 
-		S[BO_Max] 						= BinaryArithmeticLWC;
-		S[BO_Min] 						= BinaryArithmeticLWC;
+		S[BO_Max] 						= BinaryArithmeticAllowDouble;
+		S[BO_Min] 						= BinaryArithmeticAllowDouble;
 		S[BO_Pow] 						= BinaryFloat;
 		S[BO_Step] 						= BinaryArithmeticLWCDemote;
 
 		/* ternary operators -- note lerp doesn't support LWC for the third argument! */
 		S[TO_Clamp]						= TernaryArithmeticDouble;
-		S[TO_Lerp] 						= { { OPTF_CastToCommonArithmeticTypeAllowDouble | OPTF_CastToAnyFloat, OPTF_CastToCommonArithmeticTypeAllowDouble, OPTF_CastToCommonArithmeticType }, ORT_FirstArgumentTypeLWC };;
-		S[TO_Select]					= { { OPTF_CheckIsBoolean | OPTF_CheckIsNotMatrix, OPTF_CheckIsNotMatrix | OPTF_AllowDouble, OPTF_CheckIsNotMatrix | OPTF_AllowDouble}, ORT_SecondArgumentType }; // Note: this is a special operator, which is handled manually in the Validate function
-		S[TO_Smoothstep]				= TernaryFloatLWCDemote;
+		S[TO_Lerp] 						= { { OPF_CastToCommonArithmeticTypeAllowDouble | OPF_CastToAnyFloat, OPF_CastToCommonArithmeticTypeAllowDouble, OPF_CastToCommonArithmeticType }, OR_FirstArgumentTypeAllowDouble };;
+		S[TO_Select]					= { { OPF_CheckIsBoolean | OPF_CheckIsNotMatrix, OPF_CheckIsNotMatrix | OPF_AllowDouble, OPF_CheckIsNotMatrix | OPF_AllowDouble}, OR_SecondArgumentType }; // Note: this is a special operator, which is handled manually in the Validate function
+		S[TO_Smoothstep]				= TernaryFloatDoubleDemote;
 		return S;
 	} ();
 	return &Signatures[Op];
@@ -1897,7 +2007,7 @@ static FType ValidateOperatorAndGetResultType(FEmitter& Emitter, EOperator Op, F
 
 	// Handle automatic cast to float for operators that don't support LWC inputs
 	FType FirstArgumentType = A->Type;
-	if (FirstArgumentType.GetPrimitive().IsDouble() && !(Signature->ParameterTypes[0] & OPTF_AllowDouble))
+	if (FirstArgumentType.GetPrimitive().IsDouble() && !(Signature->ParameterFlags[0] & OPF_AllowDouble))
 	{
 		FirstArgumentType = A->Type.GetPrimitive().ToScalarKind(EScalarKind::Float);
 	}
@@ -1915,10 +2025,10 @@ static FType ValidateOperatorAndGetResultType(FEmitter& Emitter, EOperator Op, F
 			return FType::MakePoison();
 		}
 
-		EOperatorParameterTypeFilter Filter = Signature->ParameterTypes[i];
-		check(Filter != OPTF_Unknown); // No signature specified for this operator.
+		EOperatorParameterFlags Filter = Signature->ParameterFlags[i];
+		check(Filter != OPF_Unknown); // No signature specified for this operator.
 
-		if (Filter & OPTF_CastToFirstArgumentType)
+		if (Filter & OPF_CastToFirstArgumentType)
 		{
 			check(i > 0); // This check can't apply to the first argument.
 			Arguments[i] = Emitter.Cast(Arguments[i], FirstArgumentType);
@@ -1926,40 +2036,40 @@ static FType ValidateOperatorAndGetResultType(FEmitter& Emitter, EOperator Op, F
 		// Cast argument to float when...
 		else if ( 
 			// ...the argument should be cast to any float (and it's not a float already)
-			((Filter & OPTF_CastToAnyFloat) && !Arguments[i]->Type.GetPrimitive().IsAnyFloat())
+			((Filter & OPF_CastToAnyFloat) && !Arguments[i]->Type.GetPrimitive().IsAnyFloat())
 			// ...or the argument is not allowed to be a double and it is.
-			|| (!(Filter & OPTF_AllowDouble) && Arguments[i]->Type.GetPrimitive().IsDouble()))
+			|| (!(Filter & OPF_AllowDouble) && Arguments[i]->Type.GetPrimitive().IsDouble()))
 		{
 			Arguments[i] = Emitter.CastToFloatKind(Arguments[i]);
 			check(!Arguments[i].IsPoison());
 		}
 		
-		if (Filter & OPTF_CheckIsBoolean)
+		if (Filter & OPF_CheckIsBoolean)
 		{
 			Arguments[i] = Emitter.CheckIsBoolean(Arguments[i]); 
 		}
 			
-		if (Filter & OPTF_CheckIsArithmetic)
+		if (Filter & OPF_CheckIsArithmetic)
 		{
 			Arguments[i] = Emitter.CheckIsArithmetic(Arguments[i]); 
 		}
 
-		if (Filter & OPTF_CheckIsInteger)
+		if (Filter & OPF_CheckIsInteger)
 		{
 			Arguments[i] = Emitter.CheckIsInteger(Arguments[i]); 
 		}
 		
-		if (Filter & OPTF_CheckIsMatrix)
+		if (Filter & OPF_CheckIsMatrix)
 		{
 			Arguments[i] = Emitter.CheckIsMatrix(Arguments[i]); 
 		}
 
-		if (Filter & OPTF_CheckIsNotMatrix)
+		if (Filter & OPF_CheckIsNotMatrix)
 		{
 			Arguments[i] = Emitter.CheckIsScalarOrVector(Arguments[i]); 
 		}
 
-		if (Filter & OPTF_CheckIsVector3)
+		if (Filter & OPF_CheckIsVector3)
 		{
 			if (!Arguments[i]->Type.GetPrimitive().IsRowVector() || Arguments[i]->Type.GetPrimitive().NumComponents() != 3)
 			{
@@ -1971,9 +2081,9 @@ static FType ValidateOperatorAndGetResultType(FEmitter& Emitter, EOperator Op, F
 		// The following checks are only applicable if the argument is constant.
 		if (FConstant* Constant = Arguments[i]->As<FConstant>())
 		{
-			if (Filter & OPTF_CheckIsNonZeroFloatConst)
+			if (Filter & OPF_CheckIsNonZeroFloatConst)
 			{
-				check((Filter & OPTF_CastToAnyFloat) || (Filter & OPTF_CastToCommonFloatType));
+				check((Filter & OPF_CastToAnyFloat) || (Filter & OPF_CastToCommonFloatType));
 				if (Constant->Float == 0)
 				{
 					Emitter.Errorf(Arguments[i], TEXT("Expected non-zero value."));
@@ -1981,9 +2091,9 @@ static FType ValidateOperatorAndGetResultType(FEmitter& Emitter, EOperator Op, F
 				}
 			}
 
-			if (Filter & OPTF_CheckIsNonNegativeFloatConst)
+			if (Filter & OPF_CheckIsNonNegativeFloatConst)
 			{
-				check((Filter & OPTF_CastToAnyFloat) || (Filter & OPTF_CastToCommonFloatType));
+				check((Filter & OPF_CastToAnyFloat) || (Filter & OPF_CastToCommonFloatType));
 				if (Constant->Float < 0)
 				{
 					Emitter.Errorf(Arguments[i], TEXT("Expected non-negative value."));
@@ -1991,9 +2101,9 @@ static FType ValidateOperatorAndGetResultType(FEmitter& Emitter, EOperator Op, F
 				}
 			}
 
-			if (Filter & OPTF_CheckIsOneOrGreaterFloatConst)
+			if (Filter & OPF_CheckIsOneOrGreaterFloatConst)
 			{
-				check((Filter & OPTF_CastToAnyFloat) || (Filter & OPTF_CastToCommonFloatType));
+				check((Filter & OPF_CastToAnyFloat) || (Filter & OPF_CastToCommonFloatType));
 				if (Constant->Float < 1)
 				{
 					Emitter.Errorf(Arguments[i], TEXT("Expected a value equal or greater than 1."));
@@ -2001,9 +2111,9 @@ static FType ValidateOperatorAndGetResultType(FEmitter& Emitter, EOperator Op, F
 				}
 			}
 
-			if (Filter & OPTF_CheckIsBetweenMinusOneAndPlusOneFloatConst)
+			if (Filter & OPF_CheckIsBetweenMinusOneAndPlusOneFloatConst)
 			{
-				check((Filter & OPTF_CastToAnyFloat) || (Filter & OPTF_CastToCommonFloatType));
+				check((Filter & OPF_CastToAnyFloat) || (Filter & OPF_CastToCommonFloatType));
 				if (Constant->Float < -1 || Constant->Float > 1)
 				{
 					Emitter.Errorf(Arguments[i], TEXT("Expected a value greater than -1 and lower than 1."));
@@ -2019,8 +2129,8 @@ static FType ValidateOperatorAndGetResultType(FEmitter& Emitter, EOperator Op, F
 	}
 	
 	// Whether any argument will require to be cast to common type/scalar kind.
-	bool bRequiresArgumentsCommonType = (Signature->ParameterTypes[0] & OPTF_CastToCommonType) || (Signature->ParameterTypes[1] & OPTF_CastToCommonType) || (Signature->ParameterTypes[2] & OPTF_CastToCommonType);
-	bool bRequiresArgumentsCommonScalarKind = (Signature->ParameterTypes[0] & OPTF_CastToCommonScalarKind) || (Signature->ParameterTypes[1] & OPTF_CastToCommonScalarKind) || (Signature->ParameterTypes[2] & OPTF_CastToCommonScalarKind);
+	bool bRequiresArgumentsCommonType = (Signature->ParameterFlags[0] & OPF_CastToCommonType) || (Signature->ParameterFlags[1] & OPF_CastToCommonType) || (Signature->ParameterFlags[2] & OPF_CastToCommonType);
+	bool bRequiresArgumentsCommonScalarKind = (Signature->ParameterFlags[0] & OPF_CastToCommonScalarKind) || (Signature->ParameterFlags[1] & OPF_CastToCommonScalarKind) || (Signature->ParameterFlags[2] & OPF_CastToCommonScalarKind);
 
 	// SPECIAL CASE:  The select operator is special insofar as its first argument is a boolean, while the second and third can be any primitive type.
 	if (Op == TO_Select)
@@ -2049,23 +2159,23 @@ static FType ValidateOperatorAndGetResultType(FEmitter& Emitter, EOperator Op, F
 	}
 	else if (bRequiresArgumentsCommonType || bRequiresArgumentsCommonScalarKind)
 	{
-		FType ArgumentsCommonType = Arguments[0]->Type;
-		EScalarKind ArgumentsCommonScalarKind = EScalarKind::Bool;
-
 		// Determine the common type and scalar kind (if needeed)
 		// Note: these two cannot be unified, because we can always determine the common (biggest) scalar kind between two
 		// primitive types, but not always can be determined a common type (e.g. a float3 with a float4x4)
-		for (int32 i = 0; Arguments[i]; ++i)
+		FType ArgumentsCommonType{};
+		if (bRequiresArgumentsCommonType)
 		{
-			if (bRequiresArgumentsCommonType)
+			ArgumentsCommonType = Emitter.GetCommonType({ Arguments[0], Arguments[1], Arguments[2] });
+			if (ArgumentsCommonType.IsPoison())
 			{
-				ArgumentsCommonType = Emitter.GetCommonType(ArgumentsCommonType, Arguments[i]->Type);
-				if (ArgumentsCommonType.IsPoison())
-				{
-					return FType::MakePoison();
-				}
+				return FType::MakePoison();
 			}
-			if (bRequiresArgumentsCommonScalarKind)
+		}
+
+		EScalarKind ArgumentsCommonScalarKind = EScalarKind::Bool;
+		if (bRequiresArgumentsCommonScalarKind)
+		{
+			for (int32 i = 0; Arguments[i]; ++i)
 			{
 				ArgumentsCommonScalarKind = (EScalarKind)FMath::Max((int)ArgumentsCommonScalarKind, (int)Arguments[i]->Type.GetPrimitive().ScalarKind);
 			}
@@ -2074,17 +2184,17 @@ static FType ValidateOperatorAndGetResultType(FEmitter& Emitter, EOperator Op, F
 		// Cast every argument with the `CastToCommon` to the common type, if necessary.
 		for (int32 i = 0; Arguments[i]; ++i)
 		{
-			EOperatorParameterTypeFilter Filter = Signature->ParameterTypes[i];
-			if (Filter & OPTF_CastToCommonScalarKind)
+			EOperatorParameterFlags Filter = Signature->ParameterFlags[i];
+			if (Filter & OPF_CastToCommonScalarKind)
 			{
 				Arguments[i] = Emitter.CastToScalarKind(Arguments[i], ArgumentsCommonScalarKind);
 			}
-			else if (Filter & OPTF_CastToCommonType)
+			else if (Filter & OPF_CastToCommonType)
 			{
 				check(bRequiresArgumentsCommonType);
 
 				// Lerp doesn't accept LWC for its third input, so we need to check if LWC is allowed per input when casting to the common type.
-				FType ToType = (!(Filter & OPTF_AllowDouble) && ArgumentsCommonType.IsDouble()) 
+				FType ToType = (!(Filter & OPF_AllowDouble) && ArgumentsCommonType.IsDouble()) 
 					? ArgumentsCommonType.GetPrimitive().ToScalarKind(EScalarKind::Float)
 					: ArgumentsCommonType;
 
@@ -2105,28 +2215,28 @@ static FType ValidateOperatorAndGetResultType(FEmitter& Emitter, EOperator Op, F
 
 	// Finally, determine operator result type.
 	const FPrimitive& FirstArgumentPrimitiveType = Arguments[0]->Type.GetPrimitive();
-	switch (Signature->ReturnType)
+	switch (Signature->Result)
 	{
-		case ORT_Unknown:
+		case OR_Unknown:
 			UE_MIR_UNREACHABLE(); // missing operator signature declaration
 
-		case ORT_FirstArgumentType:
-			// Some operators accept LWC as input, but always output non-LWC (examples: modulo, sign, sine, cosine, saturate, frac). Anything that may output LWC should use ORT_FirstArgumentTypeLWC.
+		case OR_FirstArgumentType:
+			// Some operators accept LWC as input, but always output non-LWC (examples: modulo, sign, sine, cosine, saturate, frac). Anything that may output LWC should use OR_FirstArgumentTypeLWC.
 			return FirstArgumentPrimitiveType.IsDouble() ? FirstArgumentPrimitiveType.ToScalarKind(EScalarKind::Float) : Arguments[0]->Type;
 
-		case ORT_FirstArgumentTypeLWC:
+		case OR_FirstArgumentTypeAllowDouble:
 			return Arguments[0]->Type;
 
-		case ORT_BooleanWithFirstArgumentDimensions:
+		case OR_BooleanWithFirstArgumentDimensions:
 			return FType::MakePrimitive(EScalarKind::Bool, FirstArgumentPrimitiveType.NumRows, FirstArgumentPrimitiveType.NumColumns);
 
- 		case ORT_FirstArgumentTypeToScalarLWC:
+ 		case OR_FirstArgumentTypeToScalarLWC:
 			return FirstArgumentPrimitiveType.ToScalar();
 
-		case ORT_SecondArgumentType:
+		case OR_SecondArgumentType:
 			return B->Type;
 
-		case ORT_MatrixMultiplyResult:
+		case OR_MatrixMultiplyResult:
 		{
 			FType SecondArgumentType = Arguments[1]->Type;
 			if (FirstArgumentType.GetPrimitive().NumColumns != SecondArgumentType.GetPrimitive().NumRows)
@@ -2384,24 +2494,15 @@ static FValue* TryFoldOperator(FEmitter& Emitter, EOperator Op, FValue* A, FValu
 
 FValueRef FEmitter::Operator(EOperator Op, FValueRef A, FValueRef B, FValueRef C)
 {
-	if (!A.IsValid() || (B && !B.IsValid()) || (C && !C.IsValid()))
+	// Transpose is a translation-time operation only that never creates a runtime Operator instruction.
+	if (Op == UO_Transpose)
 	{
-		return Poison();
+		return Transpose(A);
 	}
 
-	// Subtract has a special case option to automatically truncate when subtracting two LWC inputs from each other, assuming this is a
-	// transition from LWC space to relative space, and no longer needs to be LWC. We need to check that all arguments are LWC before
-	// the call to ValidateOperatorAndGetResultType, as that may cast the inputs, changing them.
-	bool bCastLWCOutputToFloat = false;
-	if (Op == BO_Subtract && B)
+	if (!A.IsValid() || (B && !B.IsValid()) || (C && !C.IsValid()))
 	{
-		TOptional<FPrimitive> AType = A->Type.AsPrimitive();
-		TOptional<FPrimitive> BType = B->Type.AsPrimitive();
-	
-		if (AType && AType->ScalarKind == EScalarKind::Double && BType && BType->ScalarKind == EScalarKind::Double && UE::MaterialTranslatorUtils::GetLWCTruncateMode() == 2)
-		{
-			bCastLWCOutputToFloat = true;
-		}
+		return Poison();
 	}
 
 	// Validate the operation and retrieve the result type.
@@ -2435,8 +2536,10 @@ FValueRef FEmitter::Operator(EOperator Op, FValueRef A, FValueRef B, FValueRef C
 		Result = EmitPrototype(*this, Proto);
 	}
 
-	// Optionally cast the output to float
-	if (bCastLWCOutputToFloat)
+	// Subtract has a special case option to automatically truncate when subtracting two double-precision inputs from each other, assuming this is a
+	// transition from double-precision space to relative space, and no longer needs to be double-precision.
+	// We need to check that all arguments are double-precision before the call to ValidateOperatorAndGetResultType, as that may cast the inputs, changing them.
+	if (Op == BO_Subtract && UE::MaterialTranslatorUtils::GetLWCTruncateMode() == 2 && A->Type.IsDouble() && B->Type.IsDouble())
 	{
 		Result = CastToFloatKind(Result);
 	}
@@ -2466,7 +2569,7 @@ FValueRef FEmitter::Branch(FValueRef Condition, FValueRef True, FValueRef False)
 
 	// If the condition is not static, make both true and false arguments have the same type,
 	// by casting false argument into the true's type.
-	FType CommonType = GetCommonType(True->Type, False->Type);
+	FType CommonType = GetCommonType({ True, False });
 	if (!CommonType)
 	{
 		return Poison();
@@ -2613,6 +2716,45 @@ FValueRef FEmitter::Swizzle(FValueRef Value, FSwizzleMask Mask)
 	return Value.To(EmitPrototype(*this, *Result));
 }
 
+FValueRef FEmitter::Transpose(FValueRef Value)
+{
+	if (!Value.IsValid())
+	{
+		return Value;
+	}
+
+	TOptional<FPrimitive> PrimitiveType = Value->Type.AsPrimitive();
+	if (!PrimitiveType)
+	{
+		Errorf(Value, TEXT("Cannot transpose a non primitive value of type '%s'."), *Value->Type.GetSpelling());
+		return Value.ToPoison();
+	}
+
+	// A transposed scalar is itself.
+	if (PrimitiveType->IsScalar())
+	{
+		return Value;
+	}
+
+	// Build the result type (swap rows and columns)
+	const int32 OrigRows = PrimitiveType->NumRows;
+    const int32 OrigColumns = PrimitiveType->NumColumns;
+    FType ResultType = FType::MakePrimitive(PrimitiveType->ScalarKind, OrigColumns, OrigRows);
+
+    // Make a composite prototype with space for all components
+    FComposite* Prototype = MakeCompositePrototype(*this, ResultType, OrigRows * OrigColumns);
+
+	// Transpose the components
+	TArrayView<FValue*> Components = Prototype->GetMutableComponents();
+	for (int32 i = 0, NumComponents = OrigRows * OrigColumns; i < NumComponents; ++i) {
+		int32 OrigRow = i % OrigRows;
+		int32 OrigColumn = i / OrigRows;
+		Components[i] = Subscript(Value, OrigRow * OrigColumns + OrigColumn);
+	}
+
+	return EmitPrototype(*this, *Prototype);
+}
+
 static FValue* CastConstant(FEmitter& Emitter, FConstant* Constant, EScalarKind ConstantScalarKind, EScalarKind TargetKind)
 {
 	if (ConstantScalarKind == TargetKind)
@@ -3523,35 +3665,6 @@ FValueRef FEmitter::PromoteSubstrateParameter()
 	return EmitPrototype(*this, Prototype);
 }
 
-FType FEmitter::TryGetCommonType(FType A, FType B)
-{
-	// Trivial case: types are equal
-	if (A == B)
-	{
-		return A;
-	}
-
-	// Primitive types can only be constructed from other primitive types.
-	TOptional<FPrimitive> PrimitiveA = A.AsPrimitive();
-	TOptional<FPrimitive> PrimitiveB = B.AsPrimitive();
-	if (!PrimitiveA || !PrimitiveB)
-	{
-		return FType::MakePoison();
-	}
-	
-	// Can't cast a vector to a matrix.
-	if ((PrimitiveA->IsRowVector() && PrimitiveB->IsMatrix()) || (PrimitiveA->IsMatrix() && PrimitiveB->IsRowVector()))
-	{
-		return FType::MakePoison();
-	}
-	
-	// Return the primitive type with the maximum number of rows and columns between the two types.
-	EScalarKind ScalarKind = FMath::Max(PrimitiveA->ScalarKind, PrimitiveB->ScalarKind);
-	int32 NumRows = FMath::Max(PrimitiveA->NumRows, PrimitiveB->NumRows);
-	int32 NumColumns = FMath::Max(PrimitiveA->NumColumns, PrimitiveB->NumColumns);
-	return FType::MakePrimitive(ScalarKind, NumRows, NumColumns);
-}
-
 void FEmitter::Initialize()
 {
 	// Create and reference the true/false constants.
diff --git a/Engine/Source/Runtime/Engine/Public/Materials/MaterialExpressionOperator.h b/Engine/Source/Runtime/Engine/Public/Materials/MaterialExpressionOperator.h
--- a/Engine/Source/Runtime/Engine/Public/Materials/MaterialExpressionOperator.h
+++ b/Engine/Source/Runtime/Engine/Public/Materials/MaterialExpressionOperator.h
@@ -51,6 +51,7 @@ enum class EMaterialExpressionOperatorKind
 	Sqrt,
 	Tan,
 	Tanh,
+	Transpose,
 	Truncate,
 
 	// Binary
diff --git a/Engine/Source/Runtime/Engine/Public/Materials/MaterialIR.h b/Engine/Source/Runtime/Engine/Public/Materials/MaterialIR.h
--- a/Engine/Source/Runtime/Engine/Public/Materials/MaterialIR.h
+++ b/Engine/Source/Runtime/Engine/Public/Materials/MaterialIR.h
@@ -700,8 +700,9 @@ enum EOperator
 	UO_Sqrt,
 	UO_Tan,
 	UO_Tanh,
+	UO_Transpose,
 	UO_Truncate,
-
+	
 	BO_FirstBinaryOperator,
 
 	// Binary comparisons
diff --git a/Engine/Source/Runtime/Engine/Public/Materials/MaterialIREmitter.h b/Engine/Source/Runtime/Engine/Public/Materials/MaterialIREmitter.h
--- a/Engine/Source/Runtime/Engine/Public/Materials/MaterialIREmitter.h
+++ b/Engine/Source/Runtime/Engine/Public/Materials/MaterialIREmitter.h
@@ -180,10 +180,15 @@ public:
 	// Note; that this may cause a loss of data when casting from int to float.
 	// It returns the common type if it exists, nullptr otherwise.
 	FType TryGetCommonType(FType A, FType B);
-
+	
 	// Gets the common type between A and B and reports an error if no such type exists.
 	// See TryGetCommonType()
 	FType GetCommonType(FType A, FType B);
+		
+	// Gets the common type between the specified values types. Note that values in this array
+	// are allowed to be null/poison.
+	// See GetCommonType(FType A, FType B)
+	FType GetCommonType(TConstArrayView<FValueRef> Values);
 
 	// Returns the type of the i-th attribute in the given material aggregate.
 	FType GetMaterialAggregateAttributeType(const UMaterialAggregate* Aggregate, int32 AttributeIndex);
@@ -445,6 +450,11 @@ public:
 	// given argument and mask it reports an error and returns poison.
 	FValueRef Swizzle(FValueRef Value, FSwizzleMask Mask);
 
+	// Transposes the given MxN primitive value into the NxM result.
+	// Note if Value is a row-vector, returns a column-vector and viceversa.
+	// If Value is not a primitive type, reports an error and returns poison.
+	FValueRef Transpose(FValueRef Value);
+
 	// Emits a stage switch instruction. Use this instruction to select a different value
 	// in each execution stage.
 	// Note: ValuePerStage must have no more than MIR::NumStages elements.
diff --git a/Engine/Source/Runtime/Engine/Public/Materials/MaterialIRTypes.h b/Engine/Source/Runtime/Engine/Public/Materials/MaterialIRTypes.h
--- a/Engine/Source/Runtime/Engine/Public/Materials/MaterialIRTypes.h
+++ b/Engine/Source/Runtime/Engine/Public/Materials/MaterialIRTypes.h
@@ -72,6 +72,9 @@ struct FPrimitive
 	
 	// Returns whether this primitive type is a column vector (rows > 1, columns = 1).
 	bool IsRowVector() const { return NumRows == 1 && NumColumns > 1; }
+
+	// Returns whether this primitive type is a row vector (rows = 1, columns > 1).
+	bool IsColumnVector() const { return NumRows > 1 && NumColumns == 1; }
 	
 	// Returns whether this primitive type is a matrix (rows > 1, columns > 1).
 	bool IsMatrix() const { return NumRows > 1 && NumColumns > 1; }
