Reduce log spam from AssetCompilingManager

#rnx

diff --git a/Engine/Source/Runtime/Engine/Private/AssetCompilingManager.cpp b/Engine/Source/Runtime/Engine/Private/AssetCompilingManager.cpp
--- a/Engine/Source/Runtime/Engine/Private/AssetCompilingManager.cpp
+++ b/Engine/Source/Runtime/Engine/Private/AssetCompilingManager.cpp
@@ -165,13 +165,20 @@ public:
 		}
 
 		const int64 MemoryAvailableAgainstHardLimit = FMath::Max(0, HardMemoryLimit - TotalEstimatedMemory);
-		
-		if (SystemAvailableMemory < MemoryAvailableAgainstHardLimit) {
-			UE_LOG(LogAsyncCompilation, Warning, TEXT(
-				"Not enough system memory to provide asset compilation full hard limit of %.3f MB - SystemAvailableMemory = %.3f MB < MemoryAvailableAgainstHardLimit = %.3f MB"),
-				   HardMemoryLimit / (1024 * 1024.f),
-				   SystemAvailableMemory / (1024 * 1024.f),
-				   MemoryAvailableAgainstHardLimit / (1024 * 1024.f));
+
+		if (SystemAvailableMemory < MemoryAvailableAgainstHardLimit) 
+		{
+			int64 CurrentDefecit = MemoryAvailableAgainstHardLimit - SystemAvailableMemory;
+
+			if (CurrentDefecit > LargestDefecit) 
+			{
+				LargestDefecit = CurrentDefecit;
+				UE_LOG(LogAsyncCompilation, Display, TEXT(
+					"Not enough system memory to provide asset compilation full hard limit of %.3f MB - SystemAvailableMemory = %.3f MB < MemoryAvailableAgainstHardLimit = %.3f MB"),
+					   HardMemoryLimit / (1024 * 1024.f),
+					   SystemAvailableMemory / (1024 * 1024.f),
+					   MemoryAvailableAgainstHardLimit / (1024 * 1024.f));
+			}
 		}
 
 		const int64 MemoryAvailableForAssetCompilation = FMath::Min(
@@ -373,7 +380,7 @@ public:
 		if ( RequiredMemory < 0 )
 		{
 			// log the types that are failing to set required memory
-			UE_LOG(LogAsyncCompilation, Warning, TEXT("AddQueuedWork [%s] did not supply a RequiredMemory estimate. This is a bug, all asset types should provide an estimate. When no estimate is provided asset compilation is slowed as the queue has to operate very conservatively."),DebugName);
+			UE_LOG(LogAsyncCompilation, Display, TEXT("AddQueuedWork [%s] did not supply a RequiredMemory estimate. This is a bug, all asset types should provide an estimate. When no estimate is provided asset compilation is slowed as the queue has to operate very conservatively."),DebugName);
 		}
 		else
 		{
@@ -388,6 +395,7 @@ private:
 	TFunction<int32()> MaxBackgroundConcurrency;
 
 	std::atomic<int64> TotalEstimatedMemory {0};
+	mutable std::atomic<int64> LargestDefecit {0};
 };
 
 #endif // #if WITH_EDITOR
