[SimpleStreamableAssetManager] Fixed InvestigateRenderAsset console command to work with SimpleStreamableAssetManager.

diff --git a/Engine/Source/Runtime/Engine/Private/Streaming/SimpleStreamableAssetManager.cpp b/Engine/Source/Runtime/Engine/Private/Streaming/SimpleStreamableAssetManager.cpp
--- a/Engine/Source/Runtime/Engine/Private/Streaming/SimpleStreamableAssetManager.cpp
+++ b/Engine/Source/Runtime/Engine/Private/Streaming/SimpleStreamableAssetManager.cpp
@@ -646,7 +646,12 @@ void FSimpleStreamableAssetManager::GetRenderAssetScreenSize_Impl(
 					const bool bForcedLODs = AssetBoundElement.bForceLOD;
 					FRenderAssetInstanceAsyncView::ProcessElement(AssetType, BoundsViewInfo, TexelFactor, bForcedLODs, MaxSize, MaxSize_VisibleOnly, MaxNumForcedLODs);
 
-					if (MaxSize_VisibleOnly >= MaxAssetSize || MaxNumForcedLODs >= MaxAllowedMip)
+					if (LogPrefix)
+					{
+						FBoxSphereBounds Bounds = GetBounds(ObjectRegistrationIndex);
+						FRenderAssetInstanceView::OutputToLog(Bounds, ObjectRegistrationIndex, ObjectBounds4, TexelFactor, bForcedLODs, BoundsViewInfo.MaxNormalizedSize, BoundsViewInfo.MaxNormalizedSize_VisibleOnly, LogPrefix);
+					}
+					else if (MaxSize_VisibleOnly >= MaxAssetSize || MaxNumForcedLODs >= MaxAllowedMip)
 					{
 						return;
 					}
diff --git a/Engine/Source/Runtime/Engine/Private/Streaming/TextureInstanceView.cpp b/Engine/Source/Runtime/Engine/Private/Streaming/TextureInstanceView.cpp
--- a/Engine/Source/Runtime/Engine/Private/Streaming/TextureInstanceView.cpp
+++ b/Engine/Source/Runtime/Engine/Private/Streaming/TextureInstanceView.cpp
@@ -125,11 +125,14 @@ FBoxSphereBounds FRenderAssetInstanceView::FRenderAssetLinkConstIterator::GetBou
 void FRenderAssetInstanceView::FRenderAssetLinkConstIterator::OutputToLog(float MaxNormalizedSize, float MaxNormalizedSize_VisibleOnly, const TCHAR* Prefix) const
 {
 #if !UE_BUILD_SHIPPING
-	const UPrimitiveComponent* Component = GetComponent();
 	FBoxSphereBounds Bounds = GetBounds();
-	float TexelFactor = GetTexelFactor();
-	bool bForceLoad = GetForceLoad();
+	FRenderAssetInstanceView::OutputToLog(Bounds, GetBoundsIndex(), State.Bounds4, GetTexelFactor(), GetForceLoad(), MaxNormalizedSize, MaxNormalizedSize_VisibleOnly, Prefix, GetComponent());
+#endif
+}
 
+void FRenderAssetInstanceView::OutputToLog(const FBoxSphereBounds& Bounds, int32 BoundsIndex, const TArray<FBounds4>& Bounds4, float TexelFactor, bool bForceLoad, float MaxNormalizedSize, float MaxNormalizedSize_VisibleOnly, const TCHAR* Prefix, const UPrimitiveComponent* Component)
+{
+#if !UE_BUILD_SHIPPING
 	// Log the component reference.
 	if (Component)
 	{
@@ -149,19 +152,18 @@ void FRenderAssetInstanceView::FRenderAssetLinkConstIterator::OutputToLog(float
 	{
 		if (GIsEditor) // In editor, visibility information is unreliable and we only consider the max.
 		{
-			UE_LOG(LogContentStreaming, Log, TEXT("    Size=%f, BoundIndex=%d"), TexelFactor * FMath::Max(MaxNormalizedSize, MaxNormalizedSize_VisibleOnly), GetBoundsIndex());
+			UE_LOG(LogContentStreaming, Log, TEXT("    Size=%f, BoundIndex=%d"), TexelFactor * FMath::Max(MaxNormalizedSize, MaxNormalizedSize_VisibleOnly), BoundsIndex);
 		}
 		else if (MaxNormalizedSize_VisibleOnly > 0)
 		{
-			UE_LOG(LogContentStreaming, Log, TEXT("    OnScreenSize=%f, BoundIndex=%d"), TexelFactor * MaxNormalizedSize_VisibleOnly, GetBoundsIndex());
+			UE_LOG(LogContentStreaming, Log, TEXT("    OnScreenSize=%f, BoundIndex=%d"), TexelFactor * MaxNormalizedSize_VisibleOnly, BoundsIndex);
 		}
 		else
 		{
-			const int32 BoundIndex = GetBoundsIndex();
-			if (State.Bounds4.IsValidIndex(BoundIndex / 4))
+			if (Bounds4.IsValidIndex(BoundsIndex / 4))
 			{
-				float LastRenderTime = State.Bounds4[BoundIndex / 4].LastRenderTime[BoundIndex % 4];
-				UE_LOG(LogContentStreaming, Log, TEXT("    OffScreenSize=%f, LastRenderTime= %.3f, BoundIndex=%d"), TexelFactor * MaxNormalizedSize, LastRenderTime, BoundIndex);
+				float LastRenderTime = Bounds4[BoundsIndex / 4].LastRenderTime[BoundsIndex % 4];
+				UE_LOG(LogContentStreaming, Log, TEXT("    OffScreenSize=%f, LastRenderTime= %.3f, BoundIndex=%d"), TexelFactor * MaxNormalizedSize, LastRenderTime, BoundsIndex);
 			}
 			else
 			{
diff --git a/Engine/Source/Runtime/Engine/Private/Streaming/TextureInstanceView.h b/Engine/Source/Runtime/Engine/Private/Streaming/TextureInstanceView.h
--- a/Engine/Source/Runtime/Engine/Private/Streaming/TextureInstanceView.h
+++ b/Engine/Source/Runtime/Engine/Private/Streaming/TextureInstanceView.h
@@ -242,6 +242,8 @@ public:
 
 	static void GetDistanceAndRange(const UPrimitiveComponent* Component, const FBoxSphereBounds& RenderAssetInstanceBounds, float& MinDistanceSq, float& MinRangeSq, float& MaxRangeSq);
 
+	static void OutputToLog(const FBoxSphereBounds& Bounds, int32 BoundsIndex, const TArray<FBounds4>& Bounds4, float TexelFactor, bool bForceLoad, float MaxNormalizedSize, float MaxNormalizedSize_VisibleOnly, const TCHAR* Prefix, const UPrimitiveComponent* Component = nullptr);
+
 protected:
 
 	TArray<FBounds4> Bounds4;
