Deprecated TIsConst.

diff --git a/Engine/Source/Runtime/Core/Public/Containers/ArrayView.h b/Engine/Source/Runtime/Core/Public/Containers/ArrayView.h
--- a/Engine/Source/Runtime/Core/Public/Containers/ArrayView.h
+++ b/Engine/Source/Runtime/Core/Public/Containers/ArrayView.h
@@ -941,9 +941,3 @@ inline TArray<InElementType, InAllocatorType>& TArray<InElementType, InAllocator
 	CopyToEmpty(Other.GetData(), Other.Num(), ArrayMax);
 	return *this;
 }
-
-#if UE_ENABLE_INCLUDE_ORDER_DEPRECATED_IN_5_4
-#include "Templates/IsConst.h"
-#include "Templates/IsSigned.h"
-#include "Templates/PointerIsConvertibleFromTo.h"
-#endif
diff --git a/Engine/Source/Runtime/Core/Public/Memory/MemoryView.h b/Engine/Source/Runtime/Core/Public/Memory/MemoryView.h
--- a/Engine/Source/Runtime/Core/Public/Memory/MemoryView.h
+++ b/Engine/Source/Runtime/Core/Public/Memory/MemoryView.h
@@ -10,7 +10,6 @@
 #include "Templates/AndOrNot.h"
 #include "Templates/EnableIf.h"
 #include "Templates/Identity.h"
-#include "Templates/IsConst.h"
 #include "Templates/PointerIsConvertibleFromTo.h"
 #include "Templates/UnrealTemplate.h"
 #include "Traits/IsContiguousContainer.h"
@@ -300,3 +299,7 @@ template <
 	using DataType = std::conditional_t<bIsConst, const void, void>;
 	return TMemoryView<DataType>(GetData(Container), GetNum(Container) * sizeof(ElementType));
 }
+
+#if UE_ENABLE_INCLUDE_ORDER_DEPRECATED_IN_5_7
+#include "Templates/IsConst.h"
+#endif
diff --git a/Engine/Source/Runtime/Core/Public/Templates/IsConst.h b/Engine/Source/Runtime/Core/Public/Templates/IsConst.h
--- a/Engine/Source/Runtime/Core/Public/Templates/IsConst.h
+++ b/Engine/Source/Runtime/Core/Public/Templates/IsConst.h
@@ -2,17 +2,21 @@
 
 #pragma once
 
+// HEADER_UNIT_SKIP - Deprecated
+
+UE_DEPRECATED_HEADER(5.7, "IsConst.h has been deprecated- please include <type_traits> and use std::is_const instead.")
+
 /**
  * Traits class which tests if a type is const.
  */
 template <typename T>
-struct TIsConst
+struct UE_DEPRECATED(5.7, "TIsConst has been deprecated - please use std::is_const instead.") TIsConst
 {
 	static constexpr bool Value = false;
 };
 
 template <typename T>
-struct TIsConst<const T>
+struct UE_DEPRECATED(5.7, "TIsConst has been deprecated - please use std::is_const instead.") TIsConst<const T>
 {
 	static constexpr bool Value = true;
 };
diff --git a/Engine/Source/Runtime/Engine/Private/Components/ApplicationLifecycleComponent.cpp b/Engine/Source/Runtime/Engine/Private/Components/ApplicationLifecycleComponent.cpp
--- a/Engine/Source/Runtime/Engine/Private/Components/ApplicationLifecycleComponent.cpp
+++ b/Engine/Source/Runtime/Engine/Private/Components/ApplicationLifecycleComponent.cpp
@@ -1,8 +1,8 @@
 // Copyright Epic Games, Inc. All Rights Reserved.
-// ApplicationLifecycleComponent.cpp: Component to handle receiving notifications from the OS about application state (activated, suspended, termination, etc)
+// ApplicationLifecycleComponent.cpp: Component 
+// handle receiving notifications from the OS about application state (activated, suspended, termination, etc)
 
 #include "Components/ApplicationLifecycleComponent.h"
-#include "Templates/IsConst.h"
 
 #include UE_INLINE_GENERATED_CPP_BY_NAME(ApplicationLifecycleComponent)
 
@@ -42,3 +42,6 @@ void UApplicationLifecycleComponent::OnUnregister()
 	FCoreDelegates::OnLowPowerMode.RemoveAll(this);
 }
 
+#if UE_ENABLE_INCLUDE_ORDER_DEPRECATED_IN_5_7
+#include "Templates/IsConst.h"
+#endif
diff --git a/Engine/Source/Runtime/SlateCore/Private/FastUpdate/SlateInvalidationWidgetList.cpp b/Engine/Source/Runtime/SlateCore/Private/FastUpdate/SlateInvalidationWidgetList.cpp
--- a/Engine/Source/Runtime/SlateCore/Private/FastUpdate/SlateInvalidationWidgetList.cpp
+++ b/Engine/Source/Runtime/SlateCore/Private/FastUpdate/SlateInvalidationWidgetList.cpp
@@ -10,7 +10,6 @@
 #include "Layout/Children.h"
 #include "Misc/StringBuilder.h"
 #include "Rendering/SlateLayoutTransform.h"
-#include "Templates/IsConst.h"
 #include "Templates/Tuple.h"
 #include "Types/ReflectionMetadata.h"
 #include "Types/SlateAttributeMetaData.h"
