Fix memory leak when trying to trigger the display of a FCanvasBatchedElementRenderItem without primitive to draw, on the game thread.

#tests Locally modified some code to easily trigger the leaking code path. Validated that memory was freed in both normal and problematic cases. Validate that canvas display was still working.

diff --git a/Engine/Source/Runtime/Engine/Private/UserInterface/Canvas.cpp b/Engine/Source/Runtime/Engine/Private/UserInterface/Canvas.cpp
--- a/Engine/Source/Runtime/Engine/Private/UserInterface/Canvas.cpp
+++ b/Engine/Source/Runtime/Engine/Private/UserInterface/Canvas.cpp
@@ -512,7 +512,8 @@ bool FCanvasBatchedElementRenderItem::Render_GameThread(const FCanvas* Canvas, F
 
 		// current render target set for the canvas  (eg. an FSlateTextureRenderTarget2DResource)
 		const FRenderTarget* CanvasRenderTarget = Canvas->GetRenderTarget();
-		float Gamma = 1.0f / CanvasRenderTarget->GetDisplayGamma(); // GetDisplayGamma typically == 2.2
+		const float Gamma = 1.0f / CanvasRenderTarget->GetDisplayGamma(); // GetDisplayGamma typically == 2.2
+		const bool bAllowDeleteOnRender = Canvas->GetAllowedModes() & FCanvas::Allow_DeleteOnRender;
 		
 		// Render the batched elements.
 		struct FBatchedDrawParameters
@@ -522,7 +523,7 @@ bool FCanvasBatchedElementRenderItem::Render_GameThread(const FCanvas* Canvas, F
 			uint32 ViewportSizeX;
 			uint32 ViewportSizeY;
 			float DisplayGamma;
-			uint32 AllowedCanvasModes;
+			bool bAllowDeleteOnRender;
 			ERHIFeatureLevel::Type FeatureLevel;
 			EShaderPlatform ShaderPlatform;
 		};
@@ -534,7 +535,7 @@ bool FCanvasBatchedElementRenderItem::Render_GameThread(const FCanvas* Canvas, F
 			(uint32)CanvasRenderTarget->GetSizeXY().X,
 			(uint32)CanvasRenderTarget->GetSizeXY().Y,
 			Gamma,
-			Canvas->GetAllowedModes(),
+			bAllowDeleteOnRender,
 			Canvas->GetFeatureLevel(),
 			Canvas->GetShaderPlatform()
 		};
@@ -559,15 +560,16 @@ bool FCanvasBatchedElementRenderItem::Render_GameThread(const FCanvas* Canvas, F
 				DrawParameters.bHitTesting,
 				DrawParameters.DisplayGamma);
 
-			if(DrawParameters.AllowedCanvasModes & FCanvas::Allow_DeleteOnRender )
+			if (DrawParameters.bAllowDeleteOnRender)
 			{
 				delete DrawParameters.RenderData;
 			}
 		});
-	}
-	if( Canvas->GetAllowedModes() & FCanvas::Allow_DeleteOnRender )
-	{
-		Data = NULL;
+
+		if (bAllowDeleteOnRender)
+		{
+			Data = nullptr;
+		}
 	}
 	return bDirty;
 }
