New material translator: minor refactoring to FMaterialIREmitter::Operator().

#rnx

diff --git a/Engine/Source/Runtime/Engine/Private/Materials/MaterialIREmitter.cpp b/Engine/Source/Runtime/Engine/Private/Materials/MaterialIREmitter.cpp
--- a/Engine/Source/Runtime/Engine/Private/Materials/MaterialIREmitter.cpp
+++ b/Engine/Source/Runtime/Engine/Private/Materials/MaterialIREmitter.cpp
@@ -2033,8 +2033,8 @@ static FType ValidateOperatorAndGetResultType(FEmitter& Emitter, EOperator Op, F
 			check(i > 0); // This check can't apply to the first argument.
 			Arguments[i] = Emitter.Cast(Arguments[i], FirstArgumentType);
 		}
-		// Cast argument to float when...
 		else if ( 
+			// Cast argument to float when...
 			// ...the argument should be cast to any float (and it's not a float already)
 			((Filter & OPF_CastToAnyFloat) && !Arguments[i]->Type.GetPrimitive().IsAnyFloat())
 			// ...or the argument is not allowed to be a double and it is.
@@ -2193,7 +2193,7 @@ static FType ValidateOperatorAndGetResultType(FEmitter& Emitter, EOperator Op, F
 			{
 				check(bRequiresArgumentsCommonType);
 
-				// Lerp doesn't accept LWC for its third input, so we need to check if LWC is allowed per input when casting to the common type.
+				// Lerp doesn't accept double for its third input, so we need to check if double is allowed per input when casting to the common type.
 				FType ToType = (!(Filter & OPF_AllowDouble) && ArgumentsCommonType.IsDouble()) 
 					? ArgumentsCommonType.GetPrimitive().ToScalarKind(EScalarKind::Float)
 					: ArgumentsCommonType;
@@ -2407,8 +2407,14 @@ static MIR::TDouble ConstantDotDouble(TArrayView<FConstant*> AComponents, TArray
 
 // Tries to fold the operator, that is to evaluate its result now at translation time if its arguments are constant.
 // If the operator could not be folded in any way, it returns nullptr.
-static FValue* TryFoldOperator(FEmitter& Emitter, EOperator Op, FValue* A, FValue* B, FValue* C, FType ResultType)
+static FValueRef TryFoldOperator(FEmitter& Emitter, EOperator Op, FValueRef A, FValueRef B, FValueRef C, FType ResultType)
 {
+	// First, try to apply some operator identity to simplify the operator.
+	if (FValueRef Simplified = TrySimplifyOperator(Emitter, Op, A, B, C, ResultType))
+	{
+		return Simplified;
+	}
+
 	FConstant* AComponents[4];
 	int32      ANumComponents;
 	
@@ -2514,13 +2520,8 @@ FValueRef FEmitter::Operator(EOperator Op, FValueRef A, FValueRef B, FValueRef C
 
 	FValueRef Result;
 
-	// Try to apply some operator identity to simplify the operator.
-	if (FValueRef Simplified = TrySimplifyOperator(*this, Op, A, B, C, ResultType))
-	{
-		Result = Simplified;
-	}
 	// Try folding the operator first.
-	else if (FValue* FoldedValue = TryFoldOperator(*this, Op, A, B, C, ResultType))
+	if (FValueRef FoldedValue = TryFoldOperator(*this, Op, A, B, C, ResultType))
 	{
 		Result = FoldedValue;
 	}
