Iris - Seamless travel support.
If seamless travel support is not needed for a project the cvar net.Iris.AlwaysCreateLevelFilteringGroupsForPersistentLevels can be set to false to avoid unnecessary overhead of filtering objects in the persistent level.

diff --git a/Engine/Source/Runtime/Engine/Classes/Engine/NetDriver.h b/Engine/Source/Runtime/Engine/Classes/Engine/NetDriver.h
--- a/Engine/Source/Runtime/Engine/Classes/Engine/NetDriver.h
+++ b/Engine/Source/Runtime/Engine/Classes/Engine/NetDriver.h
@@ -1987,6 +1987,9 @@ public:
 	inline UReplicationSystem* GetReplicationSystem() const { return ReplicationSystem; }
 
 	void UpdateGroupFilterStatusForLevel(const ULevel* Level, UE::Net::FNetObjectGroupHandle LevelGroupHandle);
+	void DisallowReplicationOfAllLevelGroupsForConnection(UE::Net::FConnectionHandle Handle);
+	void UpdateGroupFilterStatusForAllLevelsForConnection(UE::Net::FConnectionHandle Handle);
+	void UpdateGroupFilterStatusForAllLevelsForAllConnections();
 
 	/** Returns NetTokenStore that is required to create and serialize NetTokens */
 	const UE::Net::FNetTokenStore* GetNetTokenStore() const { return NetTokenStore.Get(); }
diff --git a/Engine/Source/Runtime/Engine/Private/Net/Iris/ReplicationSystem/EngineReplicationBridge.cpp b/Engine/Source/Runtime/Engine/Private/Net/Iris/ReplicationSystem/EngineReplicationBridge.cpp
--- a/Engine/Source/Runtime/Engine/Private/Net/Iris/ReplicationSystem/EngineReplicationBridge.cpp
+++ b/Engine/Source/Runtime/Engine/Private/Net/Iris/ReplicationSystem/EngineReplicationBridge.cpp
@@ -67,6 +67,13 @@ extern bool GDefaultUseSubObjectReplicationList;
 namespace UE::Net::Private
 {
 
+static bool bAlwaysCreateLevelFilteringGroupsForPersistentLevels = true;
+static FAutoConsoleVariableRef CVarAlwaysCreateLevelFilteringGroupsForPersistentLevels(
+	TEXT("net.Iris.AlwaysCreateLevelFilteringGroupsForPersistentLevels"),
+	bAlwaysCreateLevelFilteringGroupsForPersistentLevels,
+	TEXT("When true a world's persistent level will have a filter group created such that replication of objects therein can be enabled and disabled at any point in time. This is required for seamless travel but otherwise an unnecessary overhead.")
+);
+
 static bool bEnableActorLevelChanges = true;
 static FAutoConsoleVariableRef CVarEnableActorLevelChanges(
 	TEXT("net.Iris.EnableActorLevelChanges"),
@@ -189,6 +196,12 @@ void UEngineReplicationBridge::Deinitialize()
 	ObjectReferencePackageMap = nullptr;
 }
 
+void UEngineReplicationBridge::OnPreSeamlessTravelGarbageCollect()
+{
+	Super::OnPreSeamlessTravelGarbageCollect();
+	ensureMsgf(UE::Net::Private::bAlwaysCreateLevelFilteringGroupsForPersistentLevels, TEXT("Seamless travel requires setting cvar net.Iris.AlwaysCreateLevelFilteringGroupsForPersistentLevels to true"));
+}
+
 UE::Net::FNetRefHandle UEngineReplicationBridge::StartReplicatingActor(AActor* Actor, const FActorReplicationParams& ActorParams )
 {
 	using namespace UE::Net;
@@ -850,9 +863,7 @@ void UEngineReplicationBridge::ActorChangedLevel(const AActor* Actor, const ULev
 	}
 
 	// Remove from previous level group
-	const bool bPreviousLevelIsPersistentLevel = PreviousLevel && (PreviousLevel->IsPersistentLevel() || (PreviousLevel == NetDriver->GetWorld()->PersistentLevel));
-	
-	if (PreviousLevel && !bPreviousLevelIsPersistentLevel)
+	if (PreviousLevel)
 	{
 		const UPackage* const PreviousLevelPackage = PreviousLevel->GetOutermost();
 		const FName PreviousLevelPackageName = PreviousLevelPackage->GetFName();
@@ -864,7 +875,7 @@ void UEngineReplicationBridge::ActorChangedLevel(const AActor* Actor, const ULev
 			UE_LOG_BRIDGEID(LogIris, Verbose, TEXT("ActorChangedLevel: removing %s from GroupIndex: %u PreviousLevel: %s"), *ActorRefHandle.ToString(), PreviousLevelGroup.GetGroupIndex(), ToCStr(PreviousLevelPackageName.ToString()));
 			GetReplicationSystem()->RemoveFromGroup(PreviousLevelGroup, ActorRefHandle);
 		}
-		else
+		else if (PreviousLevelGroup.IsValid())
 		{
 			UE_LOG_BRIDGEID(LogIris, Warning, TEXT("ActorChangedLevel: %s not found in GroupIndex: %u PreviousLevel: %s"), *ActorRefHandle.ToString(), PreviousLevelGroup.GetGroupIndex(), ToCStr(PreviousLevelPackageName.ToString()));
 		}
@@ -893,8 +904,8 @@ void UEngineReplicationBridge::AddRootObjectToLevelGroup(const UObject* RootObje
 		return;
 	}
 	
-	// Don't filter out actors in the persistent level
-	if (Level && (!Level->IsPersistentLevel() || (Level != NetDriver->GetWorld()->PersistentLevel)))
+	// For seamless travel we need to maintain a level group for all kinds of levels, including persistent ones.
+	if (Level && (UE::Net::Private::bAlwaysCreateLevelFilteringGroupsForPersistentLevels || !Level->IsPersistentLevel() || (Level != NetDriver->GetWorld()->PersistentLevel)))
 	{
 		const UPackage* const LevelPackage = Level->GetOutermost();
 		const FName PackageName = LevelPackage->GetFName();
diff --git a/Engine/Source/Runtime/Engine/Private/NetConnection.cpp b/Engine/Source/Runtime/Engine/Private/NetConnection.cpp
--- a/Engine/Source/Runtime/Engine/Private/NetConnection.cpp
+++ b/Engine/Source/Runtime/Engine/Private/NetConnection.cpp
@@ -1970,7 +1970,16 @@ void UNetConnection::SetClientWorldPackageName(FName NewClientWorldPackageName)
 {
 	ClientWorldPackageName = NewClientWorldPackageName;
 	
-	UpdateAllCachedLevelVisibility();
+#if UE_WITH_IRIS
+	if (Driver->IsUsingIrisReplication())
+	{
+		Driver->UpdateGroupFilterStatusForAllLevelsForConnection(GetConnectionHandle());
+	}
+	else
+#endif
+	{
+		UpdateAllCachedLevelVisibility();
+	}
 }
 
 void UNetConnection::ValidateSendBuffer()
diff --git a/Engine/Source/Runtime/Engine/Private/NetDriver.cpp b/Engine/Source/Runtime/Engine/Private/NetDriver.cpp
--- a/Engine/Source/Runtime/Engine/Private/NetDriver.cpp
+++ b/Engine/Source/Runtime/Engine/Private/NetDriver.cpp
@@ -964,7 +964,7 @@ FNetworkObjectInfo* UNetDriver::FindNetworkObjectInfo(const AActor* InActor)
 #if UE_WITH_IRIS
 	if (ReplicationSystem)
 	{
-		if (UObjectReplicationBridge* Bridge = ReplicationSystem->GetReplicationBridgeAs<UObjectReplicationBridge>())
+		if (const UObjectReplicationBridge* Bridge = ReplicationSystem->GetReplicationBridgeAs<UObjectReplicationBridge>())
 		{
 			const UE::Net::FNetRefHandle ActorRefHandle = Bridge->GetReplicatedRefHandle(InActor);
 			
@@ -6828,6 +6828,7 @@ void UNetDriver::SetWorld(class UWorld* InWorld)
 			}
 			else
 			{
+				UpdateGroupFilterStatusForAllLevelsForAllConnections();				
 				UE::Net::FReplicationSystemUtil::BeginReplicationForActorsInWorldForNetDriver(InWorld, this);
 			}
 		}
@@ -6873,13 +6874,6 @@ void UNetDriver::ResetGameWorldState()
 	{
 		ReplicationDriver->ResetGameWorldState();
 	}
-
-#if UE_WITH_IRIS
-	if (ReplicationSystem)
-	{
-		ReplicationSystem->ResetGameWorldState();
-	}
-#endif // UE_WITH_IRIS
 }
 
 void UNetDriver::CleanPackageMaps()
@@ -6953,10 +6947,16 @@ void UNetDriver::CleanupWorldForSeamlessTravel()
 
 void UNetDriver::PreSeamlessTravelGarbageCollect()
 {
+	// Destruction of our current World is imminent. Reset anything tied to it.
 	ResetGameWorldState();
 
 	// Clear out the network objects list, the traveling actors are in SeamlessTravelingObjects
 	NetworkObjects->Reset();
+
+	if (ReplicationSystem)
+	{
+		ReplicationSystem->PreSeamlessTravelGarbageCollect();
+	}
 }
 
 void UNetDriver::PostSeamlessTravelGarbageCollect()
@@ -6964,6 +6964,11 @@ void UNetDriver::PostSeamlessTravelGarbageCollect()
 	CleanPackageMaps();
 
 	NetworkObjects->OnPostSeamlessTravel();
+
+	if (ReplicationSystem)
+	{
+		ReplicationSystem->PostSeamlessTravelGarbageCollect();
+	}
 }
 
 void UNetDriver::SetReplicationDriver(UReplicationDriver* NewReplicationDriver)
@@ -7087,9 +7092,12 @@ void UNetDriver::UpdateGroupFilterStatusForLevel(const ULevel* Level, UE::Net::F
 	// Build connection mask
 	FNetBitArray ConnectionMask(ReplicationSystem->GetMaxConnectionCount());
 		
-	for (UNetConnection* Connection : ClientConnections)
+	for (const UNetConnection* Connection : ClientConnections)
 	{
-		const bool bIsVisible = (bIsPersistentLevel && WorldPackageName == Connection->GetClientWorldPackageName()) || Connection->ClientVisibleLevelNames.Contains(LevelPackageName);
+		// For a level to be visible the client world need to match ours. 
+		const bool bIsSameWorld = (WorldPackageName == Connection->GetClientWorldPackageName());
+		// If the world matches then the level is loaded if it's persistent or if the client has streamed it in.
+		const bool bIsVisible = bIsSameWorld && (bIsPersistentLevel || Connection->ClientVisibleLevelNames.Contains(LevelPackageName));
 		ConnectionMask.SetBitValue(Connection->GetConnectionHandle().GetParentConnectionId(), bIsVisible);
 	}
 
@@ -7097,6 +7105,90 @@ void UNetDriver::UpdateGroupFilterStatusForLevel(const ULevel* Level, UE::Net::F
 	ReplicationSystem->SetGroupFilterStatus(LevelGroupHandle, ConnectionMask, ENetFilterStatus::Allow);
 }
 
+void UNetDriver::DisallowReplicationOfAllLevelGroupsForConnection(UE::Net::FConnectionHandle ConnectionHandle)
+{
+	if (!ConnectionHandle.IsValid() || ConnectionHandle.IsChildConnection())
+	{
+		return;
+	}
+
+	const UReplicationBridge* Bridge = ReplicationSystem ? ReplicationSystem->GetReplicationBridge() : nullptr;
+	if (Bridge == nullptr)
+	{
+		return;
+	}
+
+	const uint32 ConnectionId = ConnectionHandle.GetParentConnectionId();
+	const TMap<FObjectKey, UE::Net::FNetObjectGroupHandle>& LevelGroups = Bridge->GetAllLevelGroups();
+	for (const TMap<FObjectKey, UE::Net::FNetObjectGroupHandle>::ElementType& LevelAndGroup : LevelGroups)
+	{
+		ReplicationSystem->SetGroupFilterStatus(LevelAndGroup.Value, ConnectionId, UE::Net::ENetFilterStatus::Disallow);
+	}
+}
+
+void UNetDriver::UpdateGroupFilterStatusForAllLevelsForConnection(UE::Net::FConnectionHandle ConnectionHandle)
+{
+	if (!ConnectionHandle.IsValid() || ConnectionHandle.IsChildConnection())
+	{
+		return;
+	}
+
+	const UReplicationBridge* Bridge = ReplicationSystem ? ReplicationSystem->GetReplicationBridge() : nullptr;
+	if (Bridge == nullptr)
+	{
+		return;
+	}
+
+	const UNetConnection* Connection = GetConnectionByHandle(ConnectionHandle);
+	if (!ensureMsgf(Connection, TEXT("Unknown connection with parent ID %u child ID %u"), ConnectionHandle.GetParentConnectionId(), ConnectionHandle.GetChildConnectionId()))
+	{
+		return;
+	}
+
+	// If client and server are in different worlds we disallow replication of world objects.
+	if (GetWorldPackage()->GetFName() != Connection->GetClientWorldPackageName())
+	{
+		DisallowReplicationOfAllLevelGroupsForConnection(ConnectionHandle);
+		return;
+	}
+
+	const uint32 ConnectionId = ConnectionHandle.GetParentConnectionId();
+	const TMap<FObjectKey, UE::Net::FNetObjectGroupHandle>& LevelGroups = Bridge->GetAllLevelGroups();
+	for (const TMap<FObjectKey, UE::Net::FNetObjectGroupHandle>::ElementType& LevelAndGroup : LevelGroups)
+	{
+		if (const ULevel* Level = Cast<ULevel>(LevelAndGroup.Key.ResolveObjectPtr()))
+		{
+			const bool bIsVisible = (Level->IsPersistentLevel() || Connection->ClientVisibleLevelNames.Contains(Level->GetOutermost()->GetFName()));
+			ReplicationSystem->SetGroupFilterStatus(LevelAndGroup.Value, ConnectionId, bIsVisible ? UE::Net::ENetFilterStatus::Allow : UE::Net::ENetFilterStatus::Disallow);
+		}
+	}
+}
+
+void UNetDriver::UpdateGroupFilterStatusForAllLevelsForAllConnections()
+{
+	using namespace UE::Net;
+
+	if (ClientConnections.IsEmpty())
+	{
+		return;
+	}
+
+	const UReplicationBridge* Bridge = ReplicationSystem ? ReplicationSystem->GetReplicationBridge() : nullptr;
+	if (Bridge == nullptr)
+	{
+		return;
+	}
+
+	const TMap<FObjectKey, UE::Net::FNetObjectGroupHandle>& LevelGroups = Bridge->GetAllLevelGroups();
+	for (const TMap<FObjectKey, UE::Net::FNetObjectGroupHandle>::ElementType& LevelAndGroup : LevelGroups)
+	{
+		if (const ULevel* Level = Cast<ULevel>(LevelAndGroup.Key.ResolveObjectPtr()))
+		{
+			UpdateGroupFilterStatusForLevel(Level, LevelAndGroup.Value);
+		}
+	}
+}
+
 void UNetDriver::SetReplicationSystem(UReplicationSystem* InReplicationSystem)
 {
 	if (ReplicationSystem)
diff --git a/Engine/Source/Runtime/Engine/Public/Net/Iris/ReplicationSystem/EngineReplicationBridge.h b/Engine/Source/Runtime/Engine/Public/Net/Iris/ReplicationSystem/EngineReplicationBridge.h
--- a/Engine/Source/Runtime/Engine/Public/Net/Iris/ReplicationSystem/EngineReplicationBridge.h
+++ b/Engine/Source/Runtime/Engine/Public/Net/Iris/ReplicationSystem/EngineReplicationBridge.h
@@ -130,6 +130,8 @@ public:
 	UE::Net::FNetObjectFactoryId GetSubObjectFactoryId() const { return SubObjectFactoryId; }
 
 protected:
+	// UReplicationBridge
+	virtual void OnPreSeamlessTravelGarbageCollect() override;
 
 	// UObjectReplicationBridge
 	virtual void Initialize(UReplicationSystem* ReplicationSystem) override;
diff --git a/Engine/Source/Runtime/Experimental/Iris/Core/Private/Iris/ReplicationSystem/ReplicationBridge.cpp b/Engine/Source/Runtime/Experimental/Iris/Core/Private/Iris/ReplicationSystem/ReplicationBridge.cpp
--- a/Engine/Source/Runtime/Experimental/Iris/Core/Private/Iris/ReplicationSystem/ReplicationBridge.cpp
+++ b/Engine/Source/Runtime/Experimental/Iris/Core/Private/Iris/ReplicationSystem/ReplicationBridge.cpp
@@ -1199,6 +1199,11 @@ UE::Net::FNetObjectGroupHandle UReplicationBridge::GetLevelGroup(const UObject*
 	return (LevelGroupHandle != nullptr ? *LevelGroupHandle : UE::Net::FNetObjectGroupHandle());
 }
 
+const TMap<FObjectKey, UE::Net::FNetObjectGroupHandle>& UReplicationBridge::GetAllLevelGroups() const
+{
+	return LevelGroups;
+}
+
 void UReplicationBridge::DestroyGlobalNetHandle(UE::Net::Private::FInternalNetRefIndex InternalReplicationIndex)
 {
 	using namespace UE::Net;
@@ -1233,3 +1238,29 @@ FString UReplicationBridge::PrintObjectFromNetRefHandle(FNetRefHandle RefHandle)
 {
 	return NetRefHandleManager->PrintObjectFromNetRefHandle(RefHandle);
 }
+
+void UReplicationBridge::PreSeamlessTravelGarbageCollect()
+{
+	RemoveDestructionInfosForGroup(UE::Net::FNetObjectGroupHandle());
+
+	for (const TPair<FObjectKey, UE::Net::FNetObjectGroupHandle>& LevelAndGroup : LevelGroups)
+	{
+		ReplicationSystem->DestroyGroup(LevelAndGroup.Value);
+	}
+	LevelGroups.Empty();
+
+	OnPreSeamlessTravelGarbageCollect();
+}
+
+void UReplicationBridge::OnPreSeamlessTravelGarbageCollect()
+{
+}
+
+void UReplicationBridge::PostSeamlessTravelGarbageCollect()
+{
+	OnPostSeamlessTravelGarbageCollect();
+}
+
+void UReplicationBridge::OnPostSeamlessTravelGarbageCollect()
+{
+}
diff --git a/Engine/Source/Runtime/Experimental/Iris/Core/Private/Iris/ReplicationSystem/ReplicationSystem.cpp b/Engine/Source/Runtime/Experimental/Iris/Core/Private/Iris/ReplicationSystem/ReplicationSystem.cpp
--- a/Engine/Source/Runtime/Experimental/Iris/Core/Private/Iris/ReplicationSystem/ReplicationSystem.cpp
+++ b/Engine/Source/Runtime/Experimental/Iris/Core/Private/Iris/ReplicationSystem/ReplicationSystem.cpp
@@ -1100,14 +1100,22 @@ void UReplicationSystem::CollectGarbage()
 	bDoCollectGarbage = 0;
 }
 
-void UReplicationSystem::ResetGameWorldState()
+void UReplicationSystem::NotifyStreamingLevelUnload(const UObject* Level)
 {
-	Impl->ReplicationSystemInternal.GetReplicationBridge()->RemoveDestructionInfosForGroup(UE::Net::FNetObjectGroupHandle());
+	Impl->ReplicationSystemInternal.GetReplicationBridge()->NotifyStreamingLevelUnload(Level);
 }
 
-void UReplicationSystem::NotifyStreamingLevelUnload(const UObject* Level)
+void UReplicationSystem::PreSeamlessTravelGarbageCollect()
 {
-	Impl->ReplicationSystemInternal.GetReplicationBridge()->NotifyStreamingLevelUnload(Level);
+	UE_LOG(LogIris, Verbose, TEXT("UReplicationSystem::PreSeamlessTravelGarbageCollect IsServer: %d"), IsServer());
+	Impl->ReplicationSystemInternal.PreSeamlessTravelGarbageCollect();
+}
+
+void UReplicationSystem::PostSeamlessTravelGarbageCollect()
+{
+	UE_LOG(LogIris, Verbose, TEXT("UReplicationSystem::PostSeamlessTravelGarbageCollect IsServer: %d"), IsServer());
+	CollectGarbage();
+	Impl->ReplicationSystemInternal.PostSeamlessTravelGarbageCollect();
 }
 
 void UReplicationSystem::AddConnection(uint32 ConnectionId)
diff --git a/Engine/Source/Runtime/Experimental/Iris/Core/Private/Iris/ReplicationSystem/ReplicationSystemInternal.h b/Engine/Source/Runtime/Experimental/Iris/Core/Private/Iris/ReplicationSystem/ReplicationSystemInternal.h
--- a/Engine/Source/Runtime/Experimental/Iris/Core/Private/Iris/ReplicationSystem/ReplicationSystemInternal.h
+++ b/Engine/Source/Runtime/Experimental/Iris/Core/Private/Iris/ReplicationSystem/ReplicationSystemInternal.h
@@ -162,6 +162,16 @@ public:
 		return bIsInParallelPhase;
 	}
 
+	void PreSeamlessTravelGarbageCollect()
+	{
+		GetReplicationBridge()->PreSeamlessTravelGarbageCollect();
+	}
+
+	void PostSeamlessTravelGarbageCollect()
+	{
+		GetReplicationBridge()->PostSeamlessTravelGarbageCollect();
+	}
+
 private:
 	FReplicationProtocolManager ReplicationProtocolManager;
 	FNetRefHandleManager NetRefHandleManager;
diff --git a/Engine/Source/Runtime/Experimental/Iris/Core/Public/Iris/ReplicationSystem/ReplicationBridge.h b/Engine/Source/Runtime/Experimental/Iris/Core/Public/Iris/ReplicationSystem/ReplicationBridge.h
--- a/Engine/Source/Runtime/Experimental/Iris/Core/Public/Iris/ReplicationSystem/ReplicationBridge.h
+++ b/Engine/Source/Runtime/Experimental/Iris/Core/Public/Iris/ReplicationSystem/ReplicationBridge.h
@@ -132,6 +132,9 @@ public:
 	/** Get the group associated with the level in order to control connection filtering for it. */
 	IRISCORE_API UE::Net::FNetObjectGroupHandle GetLevelGroup(const UObject* Level) const;
 
+	/** Get all level groups. The object key is to a ULevel. The TMap must be processed immediately or copied- DO NOT CACHE. */
+	IRISCORE_API const TMap<FObjectKey, UE::Net::FNetObjectGroupHandle>& GetAllLevelGroups() const;
+
 	/** Returns true when we are in the middle of processing incoming data. */
 	bool IsInReceiveUpdate() const { return bInReceiveUpdate; }
 
@@ -208,6 +211,12 @@ protected:
 	/** Returns if the bridge is allowed to create new destruction info at this moment. */
 	virtual bool CanCreateDestructionInfo() const { return true; }
 
+	/** Called just prior to garbage collecting/destroying the previous world tied to a NetDriver during seamless travel. Subclasses can implement virtual method OnPreSeamlessTravelGarbageCollect. */
+	void PreSeamlessTravelGarbageCollect();
+
+	/** Called just after garbage collecting/destroying the previous world tied to a NetDriver during seamless travel. Subclasses can implement virtual method OnPostSeamlessTravelGarbageCollect. */
+	void PostSeamlessTravelGarbageCollect();
+
 protected:
 
 	// Forward calls to internal operations that we allow replication bridges to access
@@ -258,6 +267,12 @@ protected:
 	virtual void SendErrorWithNetRefHandle(UE::Net::ENetRefHandleError ErrorType, FNetRefHandle RefHandle, uint32 ConnectionId) {}
 	virtual void SendErrorWithNetRefHandle(UE::Net::ENetRefHandleError ErrorType, FNetRefHandle RefHandle, uint32 ConnectionId, const TArray<FNetRefHandle>& ExtraNetRefHandle) {}
 
+	/** Called from PreSeamlessTravelGarbageCollect. */
+	IRISCORE_API virtual void OnPreSeamlessTravelGarbageCollect();
+
+	/** Called from PostSeamlessTravelGarbageCollect. */
+	IRISCORE_API virtual void OnPostSeamlessTravelGarbageCollect();
+
 private:
 
 	// Internal operations invoked by ReplicationSystem/ReplicationWriter
diff --git a/Engine/Source/Runtime/Experimental/Iris/Core/Public/Iris/ReplicationSystem/ReplicationSystem.h b/Engine/Source/Runtime/Experimental/Iris/Core/Public/Iris/ReplicationSystem/ReplicationSystem.h
--- a/Engine/Source/Runtime/Experimental/Iris/Core/Public/Iris/ReplicationSystem/ReplicationSystem.h
+++ b/Engine/Source/Runtime/Experimental/Iris/Core/Public/Iris/ReplicationSystem/ReplicationSystem.h
@@ -763,9 +763,14 @@ private:
 private:
 	friend UNetDriver;
 
-	IRISCORE_API void ResetGameWorldState();
 	IRISCORE_API void NotifyStreamingLevelUnload(const UObject* Level);
 
+	/** Called just prior to garbage collecting/destroying the previous world tied to a NetDriver during seamless travel, but after ResetGameWorldState. */
+	IRISCORE_API void PreSeamlessTravelGarbageCollect();
+
+	/** Called just after garbage collecting/destroying the previous world tied to a NetDriver during seamless travel. */
+	IRISCORE_API void PostSeamlessTravelGarbageCollect();
+
 	void SetPIEInstanceID(int32 InPIEInstanceID) { PIEInstanceID = InPIEInstanceID; }
 
 private:
