Handle case on the RecordTextureStreamingInfo material expression where a texture isn't set.
We do the same thing in UMaterialExpressionTextureProperty.
Without this we can hit an assert later in the FMaterialCompiler implementation.

diff --git a/Engine/Source/Runtime/Engine/Private/Materials/MaterialExpressions.cpp b/Engine/Source/Runtime/Engine/Private/Materials/MaterialExpressions.cpp
--- a/Engine/Source/Runtime/Engine/Private/Materials/MaterialExpressions.cpp
+++ b/Engine/Source/Runtime/Engine/Private/Materials/MaterialExpressions.cpp
@@ -24238,7 +24238,11 @@ int32 UMaterialExpressionRecordTextureStreamingInfo::Compile(class FMaterialComp
 		return Compiler->Errorf(TEXT("Missing input TextureObject"));
 	}
 	const int32 TextureCodeIndex = TextureObject.Compile(Compiler);
-	
+	if (TextureCodeIndex == INDEX_NONE)
+	{
+		return INDEX_NONE;
+	}
+
 	if (!Coordinates.GetTracedInput().Expression)
 	{
 		return Compiler->Errorf(TEXT("Missing input UV Coordinates"));
