IncrementalCook: [Materials]: append source hashes to the material shader layout hash; this is needed since that hash is now used to invalidate incremental cook results (which is currently not happening on a shader code change)
#rnx
#incrementalcook

diff --git a/Engine/Source/Runtime/Engine/Private/Materials/MaterialShader.cpp b/Engine/Source/Runtime/Engine/Private/Materials/MaterialShader.cpp
--- a/Engine/Source/Runtime/Engine/Private/Materials/MaterialShader.cpp
+++ b/Engine/Source/Runtime/Engine/Private/Materials/MaterialShader.cpp
@@ -3166,6 +3166,9 @@ private:
 					const uint64 TypeNameHash = ShaderType->GetHashedName().GetHash();
 					Hasher.Update((uint8*)&TypeNameHash, sizeof(TypeNameHash));
 					Hasher.Update((uint8*)&PermutationId, sizeof(PermutationId));
+
+					const FSHAHash& Hash = ShaderType->GetSourceHash(Platform);
+					Hasher.Update((uint8*)&Hash, sizeof(Hash));
 				}
 			}
 		}
@@ -3181,6 +3184,9 @@ private:
 
 					const uint64 TypeNameHash = ShaderPipelineType->GetHashedName().GetHash();
 					Hasher.Update((uint8*)&TypeNameHash, sizeof(TypeNameHash));
+
+					const FSHAHash& Hash = ShaderPipelineType->GetSourceHash(Platform);
+					Hasher.Update((uint8*)&Hash, sizeof(Hash));
 				}
 			}
 		}
@@ -3197,6 +3203,9 @@ private:
 					continue;
 				}
 
+				const FSHAHash& VertexFactorySourceHash = VertexFactoryType->GetSourceHash(Platform);
+				Hasher.Update((uint8*)&VertexFactorySourceHash, sizeof(VertexFactorySourceHash));
+
 				const int32 PermutationCount = ShaderType->GetPermutationCount();
 				for (int32 PermutationId = 0; PermutationId < PermutationCount; ++PermutationId)
 				{
@@ -3211,6 +3220,9 @@ private:
 						const uint64 TypeNameHash = ShaderType->GetHashedName().GetHash();
 						Hasher.Update((uint8*)&TypeNameHash, sizeof(TypeNameHash));
 						Hasher.Update((uint8*)&PermutationId, sizeof(PermutationId));
+
+						const FSHAHash& Hash = ShaderType->GetSourceHash(Platform);
+						Hasher.Update((uint8*)&Hash, sizeof(Hash));
 					}
 				}
 			}
@@ -3235,6 +3247,9 @@ private:
 
 						const uint64 TypeNameHash = ShaderPipelineType->GetHashedName().GetHash();
 						Hasher.Update((uint8*)&TypeNameHash, sizeof(TypeNameHash));
+
+						const FSHAHash& Hash = ShaderPipelineType->GetSourceHash(Platform);
+						Hasher.Update((uint8*)&Hash, sizeof(Hash));
 					}
 				}
 			}
