PR #10685: Add "Setter=SetShowMouseCursor" to APlayerController::bShowMouseCursor.

diff --git a/Engine/Source/Runtime/Engine/Classes/GameFramework/PlayerController.h b/Engine/Source/Runtime/Engine/Classes/GameFramework/PlayerController.h
--- a/Engine/Source/Runtime/Engine/Classes/GameFramework/PlayerController.h
+++ b/Engine/Source/Runtime/Engine/Classes/GameFramework/PlayerController.h
@@ -523,7 +523,7 @@ public:
 	ENGINE_API void SetDeprecatedInputRollScale(float NewValue);
 
 	/** Whether the mouse cursor should be displayed. */
-	UPROPERTY(EditAnywhere, BlueprintReadWrite, Category=MouseInterface)
+	UPROPERTY(EditAnywhere, BlueprintReadWrite, Setter=SetShowMouseCursor, Category=MouseInterface)
 	uint32 bShowMouseCursor:1;
 
 	/** Whether actor/component click events should be generated. */
