[Backout] - CL43927074
#rnx

Original CL Desc
-----------------------------------------------------------------
IncrementalCook: [Materials]: append source hashes to the material shader layout hash; this is needed since that hash is now used to invalidate incremental cook results (which is currently not happening on a shader code change)
#rnx
#incrementalcook

diff --git a/Engine/Source/Runtime/Engine/Private/Materials/MaterialShader.cpp b/Engine/Source/Runtime/Engine/Private/Materials/MaterialShader.cpp
--- a/Engine/Source/Runtime/Engine/Private/Materials/MaterialShader.cpp
+++ b/Engine/Source/Runtime/Engine/Private/Materials/MaterialShader.cpp
@@ -3166,9 +3166,6 @@ private:
 					const uint64 TypeNameHash = ShaderType->GetHashedName().GetHash();
 					Hasher.Update((uint8*)&TypeNameHash, sizeof(TypeNameHash));
 					Hasher.Update((uint8*)&PermutationId, sizeof(PermutationId));
-
-					const FSHAHash& Hash = ShaderType->GetSourceHash(Platform);
-					Hasher.Update((uint8*)&Hash, sizeof(Hash));
 				}
 			}
 		}
@@ -3184,9 +3181,6 @@ private:
 
 					const uint64 TypeNameHash = ShaderPipelineType->GetHashedName().GetHash();
 					Hasher.Update((uint8*)&TypeNameHash, sizeof(TypeNameHash));
-
-					const FSHAHash& Hash = ShaderPipelineType->GetSourceHash(Platform);
-					Hasher.Update((uint8*)&Hash, sizeof(Hash));
 				}
 			}
 		}
@@ -3203,9 +3197,6 @@ private:
 					continue;
 				}
 
-				const FSHAHash& VertexFactorySourceHash = VertexFactoryType->GetSourceHash(Platform);
-				Hasher.Update((uint8*)&VertexFactorySourceHash, sizeof(VertexFactorySourceHash));
-
 				const int32 PermutationCount = ShaderType->GetPermutationCount();
 				for (int32 PermutationId = 0; PermutationId < PermutationCount; ++PermutationId)
 				{
@@ -3220,9 +3211,6 @@ private:
 						const uint64 TypeNameHash = ShaderType->GetHashedName().GetHash();
 						Hasher.Update((uint8*)&TypeNameHash, sizeof(TypeNameHash));
 						Hasher.Update((uint8*)&PermutationId, sizeof(PermutationId));
-
-						const FSHAHash& Hash = ShaderType->GetSourceHash(Platform);
-						Hasher.Update((uint8*)&Hash, sizeof(Hash));
 					}
 				}
 			}
@@ -3247,9 +3235,6 @@ private:
 
 						const uint64 TypeNameHash = ShaderPipelineType->GetHashedName().GetHash();
 						Hasher.Update((uint8*)&TypeNameHash, sizeof(TypeNameHash));
-
-						const FSHAHash& Hash = ShaderPipelineType->GetSourceHash(Platform);
-						Hasher.Update((uint8*)&Hash, sizeof(Hash));
 					}
 				}
 			}
