World Partition: runtime partitions will now default to ERuntimePartitionCellBoundsMethod::UseContent when creating new partitions.
Keep the CDO to ERuntimePartitionCellBoundsMethod::UseMinContentCellBounds to maintain backward compatibility of existing partitions.

diff --git a/Engine/Source/Runtime/Engine/Private/WorldPartition/RuntimeHashSet/RuntimePartition.cpp b/Engine/Source/Runtime/Engine/Private/WorldPartition/RuntimeHashSet/RuntimePartition.cpp
--- a/Engine/Source/Runtime/Engine/Private/WorldPartition/RuntimeHashSet/RuntimePartition.cpp
+++ b/Engine/Source/Runtime/Engine/Private/WorldPartition/RuntimeHashSet/RuntimePartition.cpp
@@ -12,6 +12,9 @@ URuntimePartition::URuntimePartition(const FObjectInitializer& ObjectInitializer
 {
 #if WITH_EDITOR
 	SetDefaultValues();
+	
+	// Keep the original BoundsMethod by default to maintain backward compatibility
+	BoundsMethod = ERuntimePartitionCellBoundsMethod::UseMinContentCellBounds;
 #endif
 }
 
@@ -21,7 +24,7 @@ void URuntimePartition::SetDefaultValues()
 	bBlockOnSlowStreaming = false;
 	bClientOnlyVisible = false;
 	Priority = 0;
-	BoundsMethod = ERuntimePartitionCellBoundsMethod::UseMinContentCellBounds;
+	BoundsMethod = ERuntimePartitionCellBoundsMethod::UseContent;
 	LoadingRange = 25600;
 	DebugColor = FLinearColor::MakeRandomSeededColor(GetTypeHash(GetName()));
 	HLODIndex = INDEX_NONE;
