UE: Include the external package name in the debug name for body instance warnings

diff --git a/Engine/Source/Runtime/Engine/Private/PhysicsEngine/BodyInstance.cpp b/Engine/Source/Runtime/Engine/Private/PhysicsEngine/BodyInstance.cpp
--- a/Engine/Source/Runtime/Engine/Private/PhysicsEngine/BodyInstance.cpp
+++ b/Engine/Source/Runtime/Engine/Private/PhysicsEngine/BodyInstance.cpp
@@ -1140,6 +1140,12 @@ TSharedPtr<TArray<ANSICHAR>> GetDebugDebugName(const UPrimitiveComponent* Primit
 	if (PrimitiveComp)
 	{
 		DebugName += FString::Printf(TEXT("%s %s "), *AActor::GetDebugName(PrimitiveComp->GetOwner()), *PrimitiveComp->GetPathName());
+		AActor* PrimitiveCompOwner = PrimitiveComp->GetOwner();
+		UPackage* ExternalPackage = PrimitiveCompOwner ? PrimitiveCompOwner->GetExternalPackage() : nullptr;
+		if (ExternalPackage)
+		{
+			DebugName += FString::Printf(TEXT("%s "), *ExternalPackage->GetName());
+		}
 	}
 	else if (SourceObject && SourceObject->GetOuter())
 	{
