[ZenLoader]
- Add test for recursive loads inside completion callback

diff --git a/Engine/Source/Runtime/Engine/Private/Tests/Loading/AsyncLoadingTests_RecursiveLoads.cpp b/Engine/Source/Runtime/Engine/Private/Tests/Loading/AsyncLoadingTests_RecursiveLoads.cpp
--- a/Engine/Source/Runtime/Engine/Private/Tests/Loading/AsyncLoadingTests_RecursiveLoads.cpp
+++ b/Engine/Source/Runtime/Engine/Private/Tests/Loading/AsyncLoadingTests_RecursiveLoads.cpp
@@ -299,6 +299,48 @@ bool FLoadingTests_RecursiveLoads_FullFlushFrom_PostLoad::RunTest(const FString&
 	return true;
 }
 
+IMPLEMENT_CUSTOM_SIMPLE_AUTOMATION_TEST(
+	FLoadingTests_RecursiveLoads_FullFlushFrom_CompletionCallback,
+	FLoadingTests_RecursiveLoads_Base,
+	TEXT("System.Engine.Loading.RecursiveLoads.FullFlushFrom.CompletionCallback"),
+	EAutomationTestFlags::EditorContext | EAutomationTestFlags::EngineFilter
+)
+bool FLoadingTests_RecursiveLoads_FullFlushFrom_CompletionCallback::RunTest(const FString& Parameters)
+{
+	FLoadingTestsScope LoadingTestScope(this);
+
+	bool bIsCallback1Called = false;
+	bool bIsCallback2Called = false;
+	int32 RequestId1 = 0;
+	
+	// Create a request before starting the loading test so we get a request that is not tagged as partial.
+	RequestId1 = LoadPackageAsync(FLoadingTestsScope::PackagePath1,
+		FLoadPackageAsyncDelegate::CreateLambda(
+			[this, &bIsCallback1Called, &bIsCallback2Called](const FName& LoadedPackageName, UPackage* LoadedPackage, EAsyncLoadingResult::Type Result)
+			{
+				bIsCallback1Called = true;
+
+				int32 RequestId2 = 
+					LoadPackageAsync(FLoadingTestsScope::PackagePath2,
+						FLoadPackageAsyncDelegate::CreateLambda(
+							[&bIsCallback2Called](const FName& LoadedPackageName, UPackage* LoadedPackage, EAsyncLoadingResult::Type Result)
+							{
+								bIsCallback2Called = true;
+							}
+						));
+
+				FlushAsyncLoading(RequestId2);
+				TestTrue(TEXT("Callback function 2 should have been called"), bIsCallback2Called);
+			}
+		));
+
+	FlushAsyncLoading(RequestId1);
+	TestTrue(TEXT("Callback function 1 should have been called"), bIsCallback1Called);
+	TestTrue(TEXT("Callback function 2 should have been called"), bIsCallback2Called);
+
+	return true;
+}
+
 IMPLEMENT_CUSTOM_SIMPLE_AUTOMATION_TEST(
 	FLoadingTests_RecursiveLoads_FromBothSerializeAndPostLoad,
 	FLoadingTests_RecursiveLoads_Base,
