[UE]: Fix for a regression that causes the editor to crash when attempting to place a cooked Actor-based Blueprint class via drag/drop.

diff --git a/Engine/Source/Runtime/Engine/Private/BlueprintGeneratedClass.cpp b/Engine/Source/Runtime/Engine/Private/BlueprintGeneratedClass.cpp
--- a/Engine/Source/Runtime/Engine/Private/BlueprintGeneratedClass.cpp
+++ b/Engine/Source/Runtime/Engine/Private/BlueprintGeneratedClass.cpp
@@ -2170,6 +2170,12 @@ UClass* UBlueprintGeneratedClass::RegenerateClass(UClass* ClassToRegenerate, UOb
 			InheritableComponentHandler->PreloadAll();
 		}
 
+		// Make BPGC from a cooked package standalone so it doesn't get GCed
+		if (GEditor && bCooked)
+		{
+			SetFlags(RF_Standalone);
+		}
+
 		// If this is a cooked class, warn about any uncooked classes appearing above it in the inheritance hierarchy.
 		if(bCooked)
 		{
