Force errors in the streaming manager with a cvar

diff --git a/Engine/Source/Runtime/Engine/Private/StreamableManager.cpp b/Engine/Source/Runtime/Engine/Private/StreamableManager.cpp
--- a/Engine/Source/Runtime/Engine/Private/StreamableManager.cpp
+++ b/Engine/Source/Runtime/Engine/Private/StreamableManager.cpp
@@ -58,7 +58,15 @@ static FAutoConsoleVariableRef CVarStreamableStripDebugNameInShipping(
 	ECVF_Default
 );
 
-
+#if !UE_BUILD_SHIPPING
+static FString GStreamableErrorOnRequest = "";
+static FAutoConsoleVariableRef CVarStreamableErrorOnRequest(
+		TEXT("s.errorOnRequest"),
+		GStreamableErrorOnRequest,
+		TEXT("If this is set and the errorOnRequest string partially matches the current debug name of the asset being loaded, then an error will be thrown instead of completion."),
+		ECVF_Default
+	);
+#endif
 FStreamableDownloadCachePin::FStreamableDownloadCachePin() = default;
 FStreamableDownloadCachePin::~FStreamableDownloadCachePin() = default;
 FString FStreamableDownloadCachePin::GetDebugName() const { return LexToString(ContentHandle); }
@@ -1679,6 +1687,38 @@ FStreamable* FStreamableManager::StreamInternal(const FSoftObjectPath& InTargetN
 		FindInMemory(TargetName, Existing);
 	}
 	
+
+#if !UE_BUILD_SHIPPING
+
+	if (!GStreamableErrorOnRequest.IsEmpty())
+	{
+		if (FCString::Strstr(*Handle->DebugName, *GStreamableErrorOnRequest) != nullptr)
+		{
+			TOptional<UE::UnifiedError::FError> Error;
+			switch(FMath::Rand() / (RAND_MAX / 4))
+			{
+			case 0:
+				Error = UE::UnifiedError::StreamableManager::PackageLoadFailed::MakeError();
+				break;
+			case 1:
+				Error = UE::UnifiedError::StreamableManager::PackageLoadCanceled::MakeError();
+				break;
+			case 2:
+				Error = UE::UnifiedError::StreamableManager::DownloadError::MakeError();
+				break;
+			default:
+				Error = UE::UnifiedError::StreamableManager::PackageNameInvalid::MakeError();
+				break;
+			}
+
+			Handle->SetError(Error.GetValue());
+			Existing->bAsyncLoadRequestOutstanding = false;
+			Existing->bLoadFailed = true;
+			return Existing;
+		}
+	}
+#endif
+
 	if (!Existing->Target)
 	{
 		// Disable failed flag as it may have been added at a later point
