String table registration and GC fixes

* Only string table assets outside of the /Temp root are now auto-registered, ignoring assets loaded into scratch areas for things like diffs.
* Reloading a string table asset no longer leaves it in an unregistered state.
* Stale or garbage string table assets are now removed from the keep-alive list.

diff --git a/Engine/Source/Runtime/Core/Private/Internationalization/StringTableRegistry.cpp b/Engine/Source/Runtime/Core/Private/Internationalization/StringTableRegistry.cpp
--- a/Engine/Source/Runtime/Core/Private/Internationalization/StringTableRegistry.cpp
+++ b/Engine/Source/Runtime/Core/Private/Internationalization/StringTableRegistry.cpp
@@ -75,11 +75,13 @@ void FStringTableRegistry::RegisterStringTable(const FName InTableId, FStringTab
 	RegisteredStringTables.Add(InTableId, MoveTemp(InTable));
 }
 
-void FStringTableRegistry::UnregisterStringTable(const FName InTableId)
+bool FStringTableRegistry::UnregisterStringTable(const FName InTableId)
 {
+	bool bDidUnregister = false;
+
 	{
 		UE::TScopeLock RegisteredStringTablesLock(RegisteredStringTablesCS);
-		RegisteredStringTables.Remove(InTableId);
+		bDidUnregister = RegisteredStringTables.Remove(InTableId) > 0;
 	}
 
 #if WITH_EDITOR
@@ -94,6 +96,8 @@ void FStringTableRegistry::UnregisterStringTable(const FName InTableId)
 		}
 	}
 #endif // WITH_EDITOR
+
+	return bDidUnregister;
 }
 
 FStringTablePtr FStringTableRegistry::FindMutableStringTable(const FName InTableId) const
diff --git a/Engine/Source/Runtime/Core/Public/Internationalization/StringTableRegistry.h b/Engine/Source/Runtime/Core/Public/Internationalization/StringTableRegistry.h
--- a/Engine/Source/Runtime/Core/Public/Internationalization/StringTableRegistry.h
+++ b/Engine/Source/Runtime/Core/Public/Internationalization/StringTableRegistry.h
@@ -30,8 +30,8 @@ public:
 	/** Register a string table with the given ID */
 	CORE_API void RegisterStringTable(const FName InTableId, FStringTableRef InTable);
 
-	/** Unregister a string table with the given ID */
-	CORE_API void UnregisterStringTable(const FName InTableId);
+	/** Unregister a string table with the given ID, returning true if there was a registered string table with the given ID */
+	CORE_API bool UnregisterStringTable(const FName InTableId);
 
 	/** Try and find a string table with the given ID */
 	CORE_API FStringTablePtr FindMutableStringTable(const FName InTableId) const;
diff --git a/Engine/Source/Runtime/Engine/Private/Internationalization/StringTable.cpp b/Engine/Source/Runtime/Engine/Private/Internationalization/StringTable.cpp
--- a/Engine/Source/Runtime/Engine/Private/Internationalization/StringTable.cpp
+++ b/Engine/Source/Runtime/Engine/Private/Internationalization/StringTable.cpp
@@ -76,6 +76,14 @@ namespace UE::StringTable
 		}
 	}
 #endif	// WITH_EDITORONLY_DATA
+
+	namespace Private
+	{
+		bool ShouldAutoRegister(const FName InTableId)
+		{
+			return !FPackageName::IsTempPackage(FNameBuilder(InTableId).ToView());
+		}
+	}
 }
 
 class FStringTableEngineBridge : public IStringTableEngineBridge, public FGCObject
@@ -317,6 +325,7 @@ private:
 	virtual void AddReferencedObjects(FReferenceCollector& Collector) override
 	{
 		FScopeLock KeepAliveStringTablesLock(&KeepAliveStringTablesCS);
+		KeepAliveStringTables.RemoveAllSwap([](const TObjectPtr<UStringTable>& InStringTableAsset) { return !InStringTableAsset || InStringTableAsset->HasAnyFlags(RF_MirroredGarbage | RF_NewerVersionExists); });
 		Collector.AddReferencedObjects(KeepAliveStringTables);
 	}
 
@@ -353,7 +362,7 @@ UStringTable::UStringTable()
 
 	StringTable->SetNamespace(GetName());
 
-	if (!HasAnyFlags(RF_ClassDefaultObject))
+	if (IsAsset() && UE::StringTable::Private::ShouldAutoRegister(GetStringTableId()))
 	{
 		FStringTableEngineBridge::Get().RegisterForGC(this);
 		FStringTableRegistry::Get().RegisterStringTable(GetStringTableId(), StringTable.ToSharedRef());
@@ -371,11 +380,12 @@ void UStringTable::InitializeEngineBridge()
 
 void UStringTable::FinishDestroy()
 {
-	if (!HasAnyFlags(RF_ClassDefaultObject))
+	FStringTableEngineBridge::Get().UnregisterForGC(this);
+	if (!FStringTableEngineBridge::Get().IsStringTableAssetBeingReplaced(this))
 	{
-		FStringTableEngineBridge::Get().UnregisterForGC(this);
 		FStringTableRegistry::Get().UnregisterStringTable(GetStringTableId());
 	}
+
 	StringTable.Reset();
 
 	Super::FinishDestroy();
@@ -404,11 +414,17 @@ void UStringTable::PostLoad()
 bool UStringTable::Rename(const TCHAR* NewName, UObject* NewOuter, ERenameFlags Flags)
 {
 	const bool bRenamed = Super::Rename(NewName, NewOuter, Flags);
-	if (bRenamed && !HasAnyFlags(RF_ClassDefaultObject))
+	if (bRenamed)
 	{
-		FStringTableRegistry::Get().UnregisterStringTable(GetStringTableId());
+		const bool bWasRegistered = FStringTableRegistry::Get().UnregisterStringTable(GetStringTableId());
+		const bool bWasManuallyRegistered = bWasRegistered && !UE::StringTable::Private::ShouldAutoRegister(GetStringTableId());
+
 		StringTableId = *GetPathName();
-		FStringTableRegistry::Get().RegisterStringTable(GetStringTableId(), StringTable.ToSharedRef());
+		
+		if (bWasRegistered && (bWasManuallyRegistered || UE::StringTable::Private::ShouldAutoRegister(GetStringTableId())))
+		{
+			FStringTableRegistry::Get().RegisterStringTable(GetStringTableId(), StringTable.ToSharedRef());
+		}
 	}
 	return bRenamed;
 }
