Fix warning V557: Array overrun is possible. The 'PointIndex' index is pointing beyond array bound.
#rnx

diff --git a/Engine/Source/Runtime/Engine/Private/Rendering/SkeletalMeshLODImporterData.cpp b/Engine/Source/Runtime/Engine/Private/Rendering/SkeletalMeshLODImporterData.cpp
--- a/Engine/Source/Runtime/Engine/Private/Rendering/SkeletalMeshLODImporterData.cpp
+++ b/Engine/Source/Runtime/Engine/Private/Rendering/SkeletalMeshLODImporterData.cpp
@@ -1815,7 +1815,7 @@ void FSkeletalMeshImportData::SplitVerticesBySmoothingGroups()
 		PointToRawMap.Empty(PointNum);
 		for (int32 PointIndex = 0; PointIndex < PointNum; ++PointIndex)
 		{
-			PointToRawMap[PointIndex] = PointIndex;
+			PointToRawMap.Add(PointIndex);
 		}
 	}
 }
