[MIR] Don't sample with virtual-texturing if input texture is not a VT and vice-versa.

This also handles the cases of having either FTextureObject or FUniformParameter in MIR at various places as opposed to always expecting FTextureObject.

#rnx

diff --git a/Engine/Source/Runtime/Engine/Private/Materials/MaterialIR.cpp b/Engine/Source/Runtime/Engine/Private/Materials/MaterialIR.cpp
--- a/Engine/Source/Runtime/Engine/Private/Materials/MaterialIR.cpp
+++ b/Engine/Source/Runtime/Engine/Private/Materials/MaterialIR.cpp
@@ -374,6 +374,28 @@ bool FValue::Equals(const FValue* Other) const
 	return FMemory::Memcmp(this, Other, SizeInBytes) == 0;
 }
 
+UTexture* FValue::GetTexture() const
+{
+	if (const FTextureObject* TextureObject = As<FTextureObject>())
+	{
+		return TextureObject->Texture;
+	}
+	return nullptr;
+}
+
+int32 FValue::GetUniformParameterIndex() const
+{
+	if (const MIR::FTextureObject* TextureObject = As<MIR::FTextureObject>())
+	{
+		return (int32)TextureObject->Analysis_UniformParameterIndex;
+	}
+	if (const MIR::FUniformParameter* UniformParameter = As<MIR::FUniformParameter>())
+	{
+		return (int32)UniformParameter->Analysis_UniformParameterIndex;
+	}
+	return INDEX_NONE;
+}
+
 FInstruction* AsInstruction(FValue* Value)
 {
 	return Value && (Value->Kind > VK_InstructionBegin && Value->Kind < VK_InstructionEnd) ? static_cast<FInstruction*>(Value) : nullptr;
diff --git a/Engine/Source/Runtime/Engine/Private/Materials/MaterialIREmitter.cpp b/Engine/Source/Runtime/Engine/Private/Materials/MaterialIREmitter.cpp
--- a/Engine/Source/Runtime/Engine/Private/Materials/MaterialIREmitter.cpp
+++ b/Engine/Source/Runtime/Engine/Private/Materials/MaterialIREmitter.cpp
@@ -3033,15 +3033,63 @@ FValueRef FEmitter::StageSwitch(FType Type, TConstArrayView<FValueRef> ValuePerS
 	return EmitPrototype(*this, Prototype);
 }
 
-static FValueRef VTPageTableLoadFromSamplerSource(FEmitter& Em, FValueRef Texture, ESamplerSourceMode SamplerSourceMode, FValueRef TexCoord, FValueRef TexCoordDdx, FValueRef TexCoordDdy)
+// Maps a sampler type from standard texture (ST) to virtual texture (VT).
+static EMaterialSamplerType PromoteSamplerTypeFromSTtoVT(EMaterialSamplerType InSamplerType)
 {
-	FTextureObject* TextureObject = Texture->As<FTextureObject>();
-	check(TextureObject);
+	switch (InSamplerType)
+	{
+		case SAMPLERTYPE_Color: return SAMPLERTYPE_VirtualColor;
+		case SAMPLERTYPE_Grayscale: return SAMPLERTYPE_VirtualGrayscale;
+		case SAMPLERTYPE_Alpha: return SAMPLERTYPE_VirtualAlpha;
+		case SAMPLERTYPE_Normal: return SAMPLERTYPE_VirtualNormal;
+		case SAMPLERTYPE_Masks: return SAMPLERTYPE_VirtualMasks;
+		case SAMPLERTYPE_LinearColor: return SAMPLERTYPE_VirtualLinearColor;
+		case SAMPLERTYPE_LinearGrayscale: return SAMPLERTYPE_VirtualLinearGrayscale;
+	}
+	return InSamplerType;
+}
+
+// Maps a sampler type from virtual texture (VT) to standard texture (ST).
+static EMaterialSamplerType DemoteSamplerTypeFromVTtoST(EMaterialSamplerType InSamplerType)
+{
+	switch (InSamplerType)
+	{
+		case SAMPLERTYPE_VirtualColor: return SAMPLERTYPE_Color;
+		case SAMPLERTYPE_VirtualGrayscale: return SAMPLERTYPE_Grayscale;
+		case SAMPLERTYPE_VirtualAlpha: return SAMPLERTYPE_Alpha;
+		case SAMPLERTYPE_VirtualNormal: return SAMPLERTYPE_Normal;
+		case SAMPLERTYPE_VirtualMasks: return SAMPLERTYPE_Masks;
+		case SAMPLERTYPE_VirtualLinearColor: return SAMPLERTYPE_LinearColor;
+		case SAMPLERTYPE_VirtualLinearGrayscale: return SAMPLERTYPE_LinearGrayscale;
+	}
+	return InSamplerType;
+}
 
+static bool IsVirtualTexture(FValueRef Texture)
+{
+	if (FUniformParameter* UniformParameter = Texture->As<FUniformParameter>())
+	{
+		return IsVirtualSamplerType(UniformParameter->SamplerType);
+	}
+	else if (FTextureObject* TextureObject = Texture->As<FTextureObject>())
+	{
+		return (TextureObject->Texture->GetMaterialType() & MCT_TextureVirtual) != 0;
+	}
+	return false;
+}
+
+static EMaterialSamplerType MapSamplerTypeForTexture(FValueRef InTexture, EMaterialSamplerType InSamplerType)
+{
+	// Can't sample with virtual texturing if input texture is not a virtual texture
+	return IsVirtualTexture(InTexture) ? PromoteSamplerTypeFromSTtoVT(InSamplerType) : DemoteSamplerTypeFromVTtoST(InSamplerType);
+}
+
+static FValueRef VTPageTableLoadFromSamplerSource(FEmitter& Em, FValueRef Texture, ESamplerSourceMode SamplerSourceMode, FValueRef TexCoord, FValueRef TexCoordDdx, FValueRef TexCoordDdy)
+{
 	TextureAddress StaticAddressX = TA_Wrap;
 	TextureAddress StaticAddressY = TA_Wrap;
 	TextureAddress StaticAddressZ = TA_Wrap;
-	UE::MaterialTranslatorUtils::GetTextureAddressForSamplerSource(TextureObject->Texture, SamplerSourceMode, StaticAddressX, StaticAddressY, StaticAddressZ);
+	UE::MaterialTranslatorUtils::GetTextureAddressForSamplerSource(Texture->GetTexture(), SamplerSourceMode, StaticAddressX, StaticAddressY, StaticAddressZ);
 
 	return Em.VTPageTableLoad(Texture, StaticAddressX, StaticAddressY, TexCoord, TexCoordDdx, TexCoordDdy);
 }
@@ -3055,6 +3103,8 @@ FValueRef FEmitter::TextureGather(FValueRef Texture, FValueRef TexCoord, ETextur
 		return Poison();
 	}
 
+	SamplerType = MapSamplerTypeForTexture(Texture, SamplerType);
+
 	FTextureRead Prototype = MakePrototype<FTextureRead>(FType::MakeFloatVector(4));
 	Prototype.TextureObject = Texture;
 	Prototype.TexCoord = TexCoord;
@@ -3077,6 +3127,8 @@ FValueRef FEmitter::TextureSample(FValueRef Texture, FValueRef TexCoord, bool bA
 		return Poison();
 	}
 
+	SamplerType = MapSamplerTypeForTexture(Texture, SamplerType);
+
 	FTextureRead PrototypeHw = MakePrototype<FTextureRead>(FType::MakeFloatVector(4));
 	PrototypeHw.TextureObject = Texture;
 	PrototypeHw.TexCoord = TexCoord;
@@ -3119,6 +3171,8 @@ FValueRef FEmitter::TextureSampleLevel(FValueRef Texture, FValueRef TexCoord, FV
 		return Poison();
 	}
 
+	SamplerType = MapSamplerTypeForTexture(Texture, SamplerType);
+
 	FTextureRead Prototype = MakePrototype<FTextureRead>(FType::MakeFloatVector(4));
 	Prototype.TextureObject = Texture;
 	Prototype.TexCoord = TexCoord;
@@ -3152,6 +3206,8 @@ FValueRef FEmitter::TextureSampleBias(FValueRef Texture, FValueRef TexCoord, FVa
 		MipBias = Operator(BO_Add, MipBias, ExternalInput(EExternalInput::ViewMaterialTextureMipBias));
 	}
 
+	SamplerType = MapSamplerTypeForTexture(Texture, SamplerType);
+
 	FTextureRead PrototypeHw = MakePrototype<FTextureRead>(FType::MakeFloatVector(4));
 	PrototypeHw.TextureObject = Texture;
 	PrototypeHw.TexCoord = TexCoord;
@@ -3186,6 +3242,8 @@ FValueRef FEmitter::TextureSampleGrad(FValueRef Texture, FValueRef TexCoord, FVa
 		return Poison();
 	}
 
+	SamplerType = MapSamplerTypeForTexture(Texture, SamplerType);
+
 	FTextureRead Prototype = MakePrototype<FTextureRead>(FType::MakeFloatVector(4));
 	Prototype.TextureObject = Texture;
 	Prototype.TexCoord = TexCoord;
diff --git a/Engine/Source/Runtime/Engine/Private/Materials/MaterialIRToHLSLTranslator.cpp b/Engine/Source/Runtime/Engine/Private/Materials/MaterialIRToHLSLTranslator.cpp
--- a/Engine/Source/Runtime/Engine/Private/Materials/MaterialIRToHLSLTranslator.cpp
+++ b/Engine/Source/Runtime/Engine/Private/Materials/MaterialIRToHLSLTranslator.cpp
@@ -1590,10 +1590,11 @@ struct FPrivate : FMaterialIRToHLSLTranslation
 
 		checkf(TextureRead->VTPageTable, TEXT("Missing page table for virtual texture read instruction"));
 		const MIR::FVTPageTableRead* VTPageTableResult = TextureRead->VTPageTable->As<MIR::FVTPageTableRead>();
-		const MIR::FTextureObject* TextureObject = TextureRead->TextureObject->As<MIR::FTextureObject>();
 
-		int32 VirtualTextureIndex = TextureObject->Analysis_UniformParameterIndex;
-		int32 VTLayerIndex = VTPageTableResult->VTLayerIndex;
+		const int32 VirtualTextureIndex = TextureRead->TextureObject->GetUniformParameterIndex();
+		check(VirtualTextureIndex >= 0);
+
+		const int32 VTLayerIndex = VTPageTableResult->VTLayerIndex;
 
 		// Sampling function
 		Printer << TEXTVIEW("TextureVirtualSample");
diff --git a/Engine/Source/Runtime/Engine/Private/Materials/MaterialIRValueAnalyzer.cpp b/Engine/Source/Runtime/Engine/Private/Materials/MaterialIRValueAnalyzer.cpp
--- a/Engine/Source/Runtime/Engine/Private/Materials/MaterialIRValueAnalyzer.cpp
+++ b/Engine/Source/Runtime/Engine/Private/Materials/MaterialIRValueAnalyzer.cpp
@@ -126,10 +126,8 @@ static int32 AcquireVTStackIndex(FMaterialIRValueAnalyzer& Analyzer, MIR::FVTPag
 static void AnalyzeInStageVTPageTableRead(FMaterialIRValueAnalyzer& Analyzer, MIR::FVTPageTableRead* VTPageTableRead, MIR::EStage Stage)
 {
 	check(VTPageTableRead->TextureObject);
-	MIR::FTextureObject* Texture = VTPageTableRead->TextureObject->As<MIR::FTextureObject>();
-	check(Texture);
 
-	const int32 TextureUniformIndex = (int32)Texture->Analysis_UniformParameterIndex;
+	const int32 TextureUniformIndex = VTPageTableRead->TextureObject->GetUniformParameterIndex();
 	check(TextureUniformIndex >= 0);
 
 	// Only support GPU feedback from pixel shader
diff --git a/Engine/Source/Runtime/Engine/Public/Materials/MaterialIR.h b/Engine/Source/Runtime/Engine/Public/Materials/MaterialIR.h
--- a/Engine/Source/Runtime/Engine/Public/Materials/MaterialIR.h
+++ b/Engine/Source/Runtime/Engine/Public/Materials/MaterialIR.h
@@ -234,6 +234,12 @@ struct FValue
 	// Returns whether this value exactly equals Other.
 	bool Equals(const FValue* Other) const;
 
+	// Returns the UTexture object if this value is of type FTextureObject. Returns null otherwise.
+	UTexture* GetTexture() const;
+
+	// Returns the uniform parameter index if this value is of type FTextureObject or FUniformParameter. Retunrs INDEX_NONE otherwise.
+	int32 GetUniformParameterIndex() const;
+
 	// Tries to cast this value to specified type T and returns the casted pointer, if possible (nullptr otherwise).
 	template <typename T>
 	T* As() { return IsA(T::TypeKind) ? static_cast<T*>(this) : nullptr; }
@@ -869,7 +875,7 @@ const TCHAR* LexToString(ETextureReadMode Mode);
 // This instruction performs texture read operaation (sample or gather).
 struct FTextureRead : TInstruction<VK_TextureRead, 6>
 {
-	// The texture object to sample.
+	// The texture object to sample. Can be FTextureObject or FUniformParameter.
 	FValue* TextureObject;
 
 	// The texture coordinate at which to sample.
