- Added "DefaultValue" to material aggregate attributes, which is assigned to the attribute when an aggregate value is created without an assignment to that attribute.
- MaterialAttributes aggregate uses the appropriate default values as indicated by FMaterialAttributeDefinitionMap::GetDefaultValue.

#rnx

diff --git a/Engine/Source/Runtime/Engine/Private/Materials/MaterialAggregate.cpp b/Engine/Source/Runtime/Engine/Private/Materials/MaterialAggregate.cpp
--- a/Engine/Source/Runtime/Engine/Private/Materials/MaterialAggregate.cpp
+++ b/Engine/Source/Runtime/Engine/Private/Materials/MaterialAggregate.cpp
@@ -120,6 +120,7 @@ static FMaterialAggregateAttribute MaterialPropertyToAttribute(EMaterialProperty
 	FMaterialAggregateAttribute Attribute;
 	Attribute.Name = *FMaterialAttributeDefinitionMap::GetAttributeName(Property);
 	Attribute.Type = MaterialValueTypeToMaterialAggregateAttributeType(FMaterialAttributeDefinitionMap::GetValueType(Property));
+	Attribute.DefaultValue = FMaterialAttributeDefinitionMap::GetDefaultValue(Property);
 	return Attribute;
 }
 
diff --git a/Engine/Source/Runtime/Engine/Private/Materials/MaterialIREmitter.cpp b/Engine/Source/Runtime/Engine/Private/Materials/MaterialIREmitter.cpp
--- a/Engine/Source/Runtime/Engine/Private/Materials/MaterialIREmitter.cpp
+++ b/Engine/Source/Runtime/Engine/Private/Materials/MaterialIREmitter.cpp
@@ -1247,6 +1247,28 @@ FValueRef FEmitter::Aggregate(const UMaterialAggregate* InAggregate)
 	return Aggregate(InAggregate, {}, TConstArrayView<FValueRef>{});
 }
 
+static FValueRef EmitAttributeDefaultValue(FEmitter& Emitter, const FMaterialAggregateAttribute& Attribute)
+{
+	switch (Attribute.Type)
+	{
+		case EMaterialAggregateAttributeType::Bool1: return Emitter.ConstantBool((bool)Attribute.DefaultValue.X);
+		case EMaterialAggregateAttributeType::Bool2: return Emitter.ConstantBool2((bool)Attribute.DefaultValue.X, (bool)Attribute.DefaultValue.Y);
+		case EMaterialAggregateAttributeType::Bool3: return Emitter.ConstantBool3((bool)Attribute.DefaultValue.X, (bool)Attribute.DefaultValue.Y, (bool)Attribute.DefaultValue.Z); 
+		case EMaterialAggregateAttributeType::Bool4: return Emitter.ConstantBool4((bool)Attribute.DefaultValue.X, (bool)Attribute.DefaultValue.Y, (bool)Attribute.DefaultValue.Z, (bool)Attribute.DefaultValue.W);
+		case EMaterialAggregateAttributeType::ShadingModel:
+		case EMaterialAggregateAttributeType::UInt1: return Emitter.ConstantInt((TInteger)Attribute.DefaultValue.X);
+		case EMaterialAggregateAttributeType::UInt2: return Emitter.ConstantInt2({ (TInteger)Attribute.DefaultValue.X, (TInteger)Attribute.DefaultValue.Y });
+		case EMaterialAggregateAttributeType::UInt3: return Emitter.ConstantInt3({ (TInteger)Attribute.DefaultValue.X, (TInteger)Attribute.DefaultValue.Y, (TInteger)Attribute.DefaultValue.Z });
+		case EMaterialAggregateAttributeType::UInt4: return Emitter.ConstantInt4({ (TInteger)Attribute.DefaultValue.X, (TInteger)Attribute.DefaultValue.Y, (TInteger)Attribute.DefaultValue.Z, (TInteger)Attribute.DefaultValue.W });
+		case EMaterialAggregateAttributeType::Float1: return Emitter.ConstantFloat((TFloat)Attribute.DefaultValue.X);
+		case EMaterialAggregateAttributeType::Float2: return Emitter.ConstantFloat2({ (TFloat)Attribute.DefaultValue.X, (TFloat)Attribute.DefaultValue.Y });
+		case EMaterialAggregateAttributeType::Float3: return Emitter.ConstantFloat3({ (TFloat)Attribute.DefaultValue.X, (TFloat)Attribute.DefaultValue.Y, (TFloat)Attribute.DefaultValue.Z });
+		case EMaterialAggregateAttributeType::Float4: return Emitter.ConstantFloat4(UE::Math::TVector4<TFloat>{ Attribute.DefaultValue.X, Attribute.DefaultValue.Y, Attribute.DefaultValue.Z, Attribute.DefaultValue.W });
+		case EMaterialAggregateAttributeType::Aggregate: return Emitter.Aggregate(Attribute.Aggregate);
+		default: UE_MIR_UNREACHABLE();
+	}
+}
+
 FValueRef FEmitter::Aggregate(const UMaterialAggregate* InAggregate, FValueRef InPrototype, TConstArrayView<FValueRef> AttributeValues)
 {
 	// Check that the specified prototype (if any) aggregate matches the one provided.
@@ -1280,8 +1302,8 @@ FValueRef FEmitter::Aggregate(const UMaterialAggregate* InAggregate, FValueRef I
 		}
 		else
 		{
-			// ...otherwise construct the default value.
-			Components[i] = ConstantDefault(AttributeType);
+			// ...otherwise construct the default value as indicated in the attribute.
+			Components[i] = EmitAttributeDefaultValue(*this, InAggregate->Attributes[i]);
 		}
 	}
 
@@ -1329,7 +1351,7 @@ FValueRef FEmitter::Aggregate(const UMaterialAggregate* InAggregate, FValueRef I
 		else
 		{
 			// ...otherwise construct the default value.
-			Components[i] = ConstantDefault(AttributeType);
+			Components[i] = EmitAttributeDefaultValue(*this, InAggregate->Attributes[i]);
 		}
 	}
 
diff --git a/Engine/Source/Runtime/Engine/Public/Materials/MaterialAggregate.h b/Engine/Source/Runtime/Engine/Public/Materials/MaterialAggregate.h
--- a/Engine/Source/Runtime/Engine/Public/Materials/MaterialAggregate.h
+++ b/Engine/Source/Runtime/Engine/Public/Materials/MaterialAggregate.h
@@ -59,6 +59,10 @@ struct FMaterialAggregateAttribute
 	UPROPERTY(EditAnywhere, Category=MaterialAggregateAttribute, meta=(EditCondition="Type == EMaterialAggregateAttributeType::Aggregate", EditConditionHides))
 	TObjectPtr<UMaterialAggregate> Aggregate{};
 
+	// This attribute default value, assigned when the parent aggregate is constructed without an assignment to this attribute.
+	UPROPERTY(EditAnywhere, Category=MaterialAggregateAttribute)
+	FVector4f DefaultValue;
+
 #if WITH_EDITOR
 	// Converts this attribute's type to a corresponding material value type.
 	EMaterialValueType ToMaterialValueType() const;
