Fix FirstAvailableLOD to be correct when LODResources only has one entry and that entry has no vertices.
#rnx

diff --git a/Engine/Source/Runtime/Engine/Private/StaticMeshSceneProxy.cpp b/Engine/Source/Runtime/Engine/Private/StaticMeshSceneProxy.cpp
--- a/Engine/Source/Runtime/Engine/Private/StaticMeshSceneProxy.cpp
+++ b/Engine/Source/Runtime/Engine/Private/StaticMeshSceneProxy.cpp
@@ -308,7 +308,7 @@ FStaticMeshSceneProxy::FStaticMeshSceneProxy(const FStaticMeshSceneProxyDesc& In
 
 	// Find the first LOD with any vertices (ie that haven't been stripped)
 	int FirstAvailableLOD = 0;
-	for (; FirstAvailableLOD < RenderData->LODResources.Num(); FirstAvailableLOD++)
+	for (; FirstAvailableLOD < RenderData->LODResources.Num()-1; FirstAvailableLOD++)
 	{
 		if (RenderData->LODResources[FirstAvailableLOD].GetNumVertices() > 0)
 		{
