Simple change to rename Esrb field to ESRB.

diff --git a/Engine/Source/Runtime/Engine/Public/Subtitles/SubtitlesAndClosedCaptionsDelegates.h b/Engine/Source/Runtime/Engine/Public/Subtitles/SubtitlesAndClosedCaptionsDelegates.h
--- a/Engine/Source/Runtime/Engine/Public/Subtitles/SubtitlesAndClosedCaptionsDelegates.h
+++ b/Engine/Source/Runtime/Engine/Public/Subtitles/SubtitlesAndClosedCaptionsDelegates.h
@@ -21,7 +21,7 @@ enum class ESubtitleTiming : uint8
 
 // ESRB rating categories.
 UENUM(BlueprintType)
-enum class ESRB : uint8
+enum class EESRB : uint8
 {
 	Everyone,
 	Everyone10Plus,
@@ -85,8 +85,8 @@ public:
 	float Priority = 1.f;
 
 	// ESRB rating category.
-	UPROPERTY(EditAnywhere, BlueprintReadWrite, Category = Subtitles)
-	ESRB Esrb = ESRB::Everyone;
+	UPROPERTY(EditAnywhere, BlueprintReadWrite, Category = Subtitles, meta = (DisplayName = "ESRB"))
+	EESRB Esrb = EESRB::Everyone;
 
 	// Subtitle type for type-specific rendering.
 	UPROPERTY(EditAnywhere, BlueprintReadWrite, Category = Subtitles)
