Added PreDeinitialize callback to UWorldSubsystem

diff --git a/Engine/Source/Runtime/Engine/Private/World.cpp b/Engine/Source/Runtime/Engine/Private/World.cpp
--- a/Engine/Source/Runtime/Engine/Private/World.cpp
+++ b/Engine/Source/Runtime/Engine/Private/World.cpp
@@ -6292,6 +6292,14 @@ void UWorld::CleanupWorldInternal(bool bSessionEnded, bool bCleanupResources, bo
 
 	PSCPool.Cleanup(this);
 
+	if (bCleanupResources)
+	{
+		SubsystemCollection.ForEachSubsystem([](UWorldSubsystem* WorldSubsystem)
+		{
+			WorldSubsystem->PreDeinitialize();
+		});
+	}
+
 	FWorldDelegates::OnPostWorldCleanup.Broadcast(this, bSessionEnded, bCleanupResources);
 
 	if (bCleanupResources)
diff --git a/Engine/Source/Runtime/Engine/Public/Subsystems/WorldSubsystem.h b/Engine/Source/Runtime/Engine/Public/Subsystems/WorldSubsystem.h
--- a/Engine/Source/Runtime/Engine/Public/Subsystems/WorldSubsystem.h
+++ b/Engine/Source/Runtime/Engine/Public/Subsystems/WorldSubsystem.h
@@ -31,6 +31,9 @@ public:
 
 	/** Called once all UWorldSubsystems have been initialized */
 	virtual void PostInitialize() {}
+
+	/** Called once all UWorldSubsystems are about to be deinitialized */
+	virtual void PreDeinitialize() {}
 	
 	/** Called when world is ready to start gameplay before the game mode transitions to the correct state and call BeginPlay on all actors */
 	virtual void OnWorldBeginPlay(UWorld& InWorld) {}
