Return the playercontroller if found from the owner

diff --git a/Engine/Source/Runtime/Engine/Private/PhysicsEngine/AsyncPhysicsInputComponent.cpp b/Engine/Source/Runtime/Engine/Private/PhysicsEngine/AsyncPhysicsInputComponent.cpp
--- a/Engine/Source/Runtime/Engine/Private/PhysicsEngine/AsyncPhysicsInputComponent.cpp
+++ b/Engine/Source/Runtime/Engine/Private/PhysicsEngine/AsyncPhysicsInputComponent.cpp
@@ -224,7 +224,7 @@ APlayerController* UAsyncPhysicsInputComponent::GetPlayerController()
 
 	if (APawn* Pawn = Cast<APawn>(GetOwner()))
 	{
-		Pawn->GetController<APlayerController>();
+		return Pawn->GetController<APlayerController>();
 	}
 
 	return nullptr;
