Re-added the CDO check to String Table destruction to avoid a crash for programs that don't call RequestEngineExit

#rnx

diff --git a/Engine/Source/Runtime/Engine/Private/Internationalization/StringTable.cpp b/Engine/Source/Runtime/Engine/Private/Internationalization/StringTable.cpp
--- a/Engine/Source/Runtime/Engine/Private/Internationalization/StringTable.cpp
+++ b/Engine/Source/Runtime/Engine/Private/Internationalization/StringTable.cpp
@@ -380,10 +380,13 @@ void UStringTable::InitializeEngineBridge()
 
 void UStringTable::FinishDestroy()
 {
-	FStringTableEngineBridge::Get().UnregisterForGC(this);
-	if (!FStringTableEngineBridge::Get().IsStringTableAssetBeingReplaced(this))
+	if (!HasAnyFlags(RF_ClassDefaultObject)) // Needed to avoid a crash during shutdown for programs that don't call RequestEngineExit
 	{
-		FStringTableRegistry::Get().UnregisterStringTable(GetStringTableId());
+		FStringTableEngineBridge::Get().UnregisterForGC(this);
+		if (!HasAnyFlags(RF_NewerVersionExists))
+		{
+			FStringTableRegistry::Get().UnregisterStringTable(GetStringTableId());
+		}
 	}
 
 	StringTable.Reset();
