[Cell Transformer] Replace references to objects that could be recreated during RerunConstructionScripts triggered from ApplyRuntimeCellsTransformerStack

ApplyRuntimeCellsTransformerStack is executed before components are registered, which means that for actors that have actors attached to them, AttachChildren array won't be populated yet. That information is not serialized (unlike AttachParent property) and is constructed during component registration based on AttachParent information. Because of that, when recreating objects during RerunConstructionScripts triggered from ApplyRuntimeCellsTransformerStack, we could recreate a parent object without updating AttachParent property on its children. This can lead to "Reachable garbage object" GC Warning - AttachParent pointing at garbage objects. Under normal flow that's handled by RerunConstructionScripts, using AttachChildren array. Since that is not available yet in this case, we have to handle that on our side.

diff --git a/Engine/Source/Runtime/Engine/Private/WorldPartition/WorldPartition.cpp b/Engine/Source/Runtime/Engine/Private/WorldPartition/WorldPartition.cpp
--- a/Engine/Source/Runtime/Engine/Private/WorldPartition/WorldPartition.cpp
+++ b/Engine/Source/Runtime/Engine/Private/WorldPartition/WorldPartition.cpp
@@ -41,6 +41,7 @@
 #include "LocationVolume.h"
 #include "Engine/ExternalObjectAndActorDependencyGatherer.h"
 #include "Engine/LevelScriptBlueprint.h"
+#include "Serialization/ArchiveReplaceObjectRef.h"
 #include "Settings/EditorExperimentalSettings.h"
 #include "WorldPartition/DataLayer/WorldDataLayersActorDesc.h"
 #include "WorldPartition/IWorldPartitionEditorModule.h"
@@ -2498,7 +2499,13 @@ void UWorldPartition::ApplyRuntimeCellsTransformerStack(ULevel* InLevel)
 		{
 			return IsValid(InActor) && !InActor->IsChildActor();
 		});
-	
+		
+		TMap<UObject*, UObject*> CombinedReplacementMap;
+		FDelegateHandle OnObjectsReplacedHandle = FCoreUObjectDelegates::OnObjectsReplaced.AddLambda([&CombinedReplacementMap](const TMap<UObject*, UObject*>& ReplacementMap)
+		{
+			CombinedReplacementMap.Append(ReplacementMap);
+		});
+
 		for (TObjectPtr<AActor>& LevelActor : LevelActors)
 		{
 			// Deal with deleted child actors that might have been deleted from previous calls to RerunConstructionScripts
@@ -2507,6 +2514,57 @@ void UWorldPartition::ApplyRuntimeCellsTransformerStack(ULevel* InLevel)
 				LevelActor->RerunConstructionScripts();
 			}
 		}
+		
+		FCoreUObjectDelegates::OnObjectsReplaced.Remove(OnObjectsReplacedHandle);
+		
+		TMap<UObject*, UObject*> FinalReplacementMap;
+		TSet<UObject*> CurrentlyResolving;
+		TFunction<UObject*(UObject*)> ResolveFinalReplacement = [&CombinedReplacementMap, &FinalReplacementMap, &CurrentlyResolving, &ResolveFinalReplacement](UObject* Object) -> UObject*
+		{
+			// Already resolved
+			if (UObject** Final = FinalReplacementMap.Find(Object))
+			{
+				return *Final;
+			}
+			
+			// Track objects being currently resolved to enable cycle detection
+			bool bIsAlreadyInSet = false;
+			CurrentlyResolving.Add(Object, &bIsAlreadyInSet);
+
+			// Cycle detection
+			if (bIsAlreadyInSet)
+			{
+				UE_LOG(LogWorldPartition, Log, TEXT("Detected cycle in object remapping during ApplyRuntimeCellsTransformerStack at %s, ignoring"), *GetNameSafe(Object));
+				return nullptr;
+			}
+			
+			// Resolve final replacement recursively
+			UObject* FinalReplacement = Object;
+			if (UObject** ReplacementObject = CombinedReplacementMap.Find(Object))
+			{
+				FinalReplacement = ResolveFinalReplacement(*ReplacementObject);
+			}
+			
+			// Finished resolving, remove from currently resolving objects
+			CurrentlyResolving.Remove(Object);
+
+			// Add to the final map if found replacement
+			if (FinalReplacement != Object)
+			{
+				FinalReplacementMap.Add(Object, FinalReplacement);
+			}
+
+			return FinalReplacement;
+		};
+		
+		// Find final replacement for each replaced object, combines A -> B and B -> C replacements into A -> C
+		for (TPair<UObject*, UObject*> Replacement : CombinedReplacementMap)
+		{
+			ResolveFinalReplacement(Replacement.Key);
+		}
+		
+		// Apply final replacements
+		FArchiveReplaceObjectRef<UObject> ReplaceAr(InLevel, FinalReplacementMap);
 
 		ApplyTransformPhase([](const FRuntimeCellTransformerInstance& TransformerInstance, ULevel* Level) { TransformerInstance.PreTransform(Level); });
 		ApplyTransformPhase([](const FRuntimeCellTransformerInstance& TransformerInstance, ULevel* Level) { TransformerInstance.Transform(Level); });
