Fix for broken root motion extraction during animation modifier evaluation

Fixes a regression in root motion extraction which prevented animation modifiers from reading raw uncompressed data during the compilation process, causing them to fail

#tests animsandbox

diff --git a/Engine/Source/Runtime/Engine/Private/Animation/AnimSequence.cpp b/Engine/Source/Runtime/Engine/Private/Animation/AnimSequence.cpp
--- a/Engine/Source/Runtime/Engine/Private/Animation/AnimSequence.cpp
+++ b/Engine/Source/Runtime/Engine/Private/Animation/AnimSequence.cpp
@@ -2213,11 +2213,17 @@ FTransform UAnimSequence::ExtractRootTrackTransform_Lockless(const FAnimExtractC
 
 		bool bUseRawDataForPoseExtraction = false;
 #if WITH_EDITOR
+		bValidCompressedData = PlatformCompressedData.IsValid(this);
 		if (RequiredBones)
 		{
 			bUseRawDataForPoseExtraction = RequiredBones->ShouldUseRawData();
 		}
-		bValidCompressedData = PlatformCompressedData.IsValid(this);
+		else
+		{
+			// Allow reading raw data in editor, as long as we don't explcitly request compressed data
+			bUseRawDataForPoseExtraction = !bValidCompressedData && !ExtractionContext.bEnforceCompressedDataSampling; 
+		}
+
 		check(!ExtractionContext.bEnforceCompressedDataSampling || (bValidCompressedData && !bUseRawDataForPoseExtraction));	
 #endif // WITH_EDITOR
 		
