[SkeletalMesh]
- Fix support for parallel serialization

diff --git a/Engine/Source/Runtime/Engine/Private/SkeletalMesh.cpp b/Engine/Source/Runtime/Engine/Private/SkeletalMesh.cpp
--- a/Engine/Source/Runtime/Engine/Private/SkeletalMesh.cpp
+++ b/Engine/Source/Runtime/Engine/Private/SkeletalMesh.cpp
@@ -1936,7 +1936,7 @@ void USkeletalMesh::Serialize( FArchive& Ar )
 	}
 
 	// Don't try to finish compilation while still in the loader as we can't possibly have our mesh being compiled if it has not been loaded yet.
-	if (!IsInAsyncLoadingThread())
+	if (!IsInAsyncLoadingThread() && !IsInParallelLoadingThread())
 	{
 		// Since UPROPERTY are accessed directly by offset during serialization instead of using accessors, 
 		// the protection put in place to automatically finish compilation if a locked property is accessed will not work. 
