IncrementalCook: [Materials]: append source hashes to the material shader layout hash; this is needed since that hash is now used to invalidate incremental cook results (which is currently not happening on a shader code change)
Also change FMaterialShaderMapLayout.ShaderMapHash to be accessible only if AllowShaderCompiling. It is currently only needed in editor, and with the modification to append source hashes, it now relies on GetShaderFileHash, which is already only available if AllowShaderCompiling is true due to failure to call InitializeShaderHashCache in the non-AllowShaderCompiling case.
#rnx
#incrementalcook

diff --git a/Engine/Source/Runtime/Engine/Private/Materials/MaterialShader.cpp b/Engine/Source/Runtime/Engine/Private/Materials/MaterialShader.cpp
--- a/Engine/Source/Runtime/Engine/Private/Materials/MaterialShader.cpp
+++ b/Engine/Source/Runtime/Engine/Private/Materials/MaterialShader.cpp
@@ -522,7 +522,7 @@ void HashMaterialDependenciesForCook(FCbFieldViewIterator Args, UE::Cook::FCookD
 		}
 		const FMaterialShaderMapLayout& Layout = AcquireMaterialShaderMapLayout(Resource.ShaderPlatform,
 			GetShaderPermutationFlags(Resource.ShaderMapId.LayoutParams), Resource.ShaderParameters);
-		Context.Update(&Layout.ShaderMapHash.Hash, sizeof(Layout.ShaderMapHash.Hash));
+		Context.Update(&Layout.ShaderCompileHash.Hash, sizeof(Layout.ShaderCompileHash.Hash));
 
 		FBlake3Hash ShaderMapPlatformHash = GetShaderPlatformIncrementalCookHash(Resource.ShaderPlatform);
 		Context.Update(&ShaderMapPlatformHash.GetBytes(), sizeof(ShaderMapPlatformHash.GetBytes()));
@@ -3151,6 +3151,13 @@ private:
 		const TArray<FShaderPipelineType*>& SortedMeshMaterialPipelineTypes = FShaderPipelineType::GetSortedTypes(FShaderType::EShaderTypeForDynamicCast::MeshMaterial);
 
 		FSHA1 Hasher;
+#if WITH_EDITOR
+		// GetSourceHash is only available if AllowShaderCompiling (it asserts otherwise).
+		// We only need the hash of sources in the editor, to support incremental cooking's need to know
+		// when to recook the Material.
+		check(AllowShaderCompiling());
+		FSHA1 ShaderCompileHasher;
+#endif
 
 		for (FShaderType* BaseShaderType : SortedMaterialShaderTypes)
 		{
@@ -3166,6 +3173,10 @@ private:
 					const uint64 TypeNameHash = ShaderType->GetHashedName().GetHash();
 					Hasher.Update((uint8*)&TypeNameHash, sizeof(TypeNameHash));
 					Hasher.Update((uint8*)&PermutationId, sizeof(PermutationId));
+#if WITH_EDITOR
+					const FSHAHash& Hash = ShaderType->GetSourceHash(Platform);
+					ShaderCompileHasher.Update((uint8*)&Hash, sizeof(Hash));
+#endif
 				}
 			}
 		}
@@ -3181,6 +3192,10 @@ private:
 
 					const uint64 TypeNameHash = ShaderPipelineType->GetHashedName().GetHash();
 					Hasher.Update((uint8*)&TypeNameHash, sizeof(TypeNameHash));
+#if WITH_EDITOR
+					const FSHAHash& Hash = ShaderPipelineType->GetSourceHash(Platform);
+					ShaderCompileHasher.Update((uint8*)&Hash, sizeof(Hash));
+#endif
 				}
 			}
 		}
@@ -3197,6 +3212,11 @@ private:
 					continue;
 				}
 
+#if WITH_EDITOR
+				const FSHAHash& VertexFactorySourceHash = VertexFactoryType->GetSourceHash(Platform);
+				ShaderCompileHasher.Update((uint8*)&VertexFactorySourceHash, sizeof(VertexFactorySourceHash));
+#endif
+
 				const int32 PermutationCount = ShaderType->GetPermutationCount();
 				for (int32 PermutationId = 0; PermutationId < PermutationCount; ++PermutationId)
 				{
@@ -3211,6 +3231,10 @@ private:
 						const uint64 TypeNameHash = ShaderType->GetHashedName().GetHash();
 						Hasher.Update((uint8*)&TypeNameHash, sizeof(TypeNameHash));
 						Hasher.Update((uint8*)&PermutationId, sizeof(PermutationId));
+#if WITH_EDITOR
+						const FSHAHash& Hash = ShaderType->GetSourceHash(Platform);
+						ShaderCompileHasher.Update((uint8*)&Hash, sizeof(Hash));
+#endif
 					}
 				}
 			}
@@ -3235,6 +3259,10 @@ private:
 
 						const uint64 TypeNameHash = ShaderPipelineType->GetHashedName().GetHash();
 						Hasher.Update((uint8*)&TypeNameHash, sizeof(TypeNameHash));
+#if WITH_EDITOR
+						const FSHAHash& Hash = ShaderPipelineType->GetSourceHash(Platform);
+						ShaderCompileHasher.Update((uint8*)&Hash, sizeof(Hash));
+#endif
 					}
 				}
 			}
@@ -3242,6 +3270,11 @@ private:
 
 		Hasher.Final();
 		Hasher.GetHash(Layout.ShaderMapHash.Hash);
+#if WITH_EDITOR
+		ShaderCompileHasher.Update((uint8*)&Layout.ShaderMapHash, sizeof(Layout.ShaderMapHash));
+		ShaderCompileHasher.Final();
+		ShaderCompileHasher.GetHash(Layout.ShaderCompileHash.Hash);
+#endif
 	}
 
 	TArray<TUniquePtr<FMaterialShaderMapLayout>> ShaderMapLayouts;
diff --git a/Engine/Source/Runtime/Engine/Private/Materials/MaterialShaderMapLayout.h b/Engine/Source/Runtime/Engine/Private/Materials/MaterialShaderMapLayout.h
--- a/Engine/Source/Runtime/Engine/Private/Materials/MaterialShaderMapLayout.h
+++ b/Engine/Source/Runtime/Engine/Private/Materials/MaterialShaderMapLayout.h
@@ -44,6 +44,10 @@ class FMaterialShaderMapLayout : public FShaderMapLayout
 public:
 	TArray<FMeshMaterialShaderMapLayout> MeshShaderMaps;
 	FSHAHash ShaderMapHash;
+#if WITH_EDITOR
+	/** Only available if AllowShaderCompiling(). Includes all of ShaderMapHash, plus hash of the shader source files. */
+	FSHAHash ShaderCompileHash;
+#endif
 };
 
 const FMaterialShaderMapLayout& AcquireMaterialShaderMapLayout(EShaderPlatform Platform, EShaderPermutationFlags Flags, const FMaterialShaderParameters& MaterialParameters);
