FSpline: Increment version in UpdateSpline

diff --git a/Engine/Source/Runtime/Engine/Private/Curves/Spline.cpp b/Engine/Source/Runtime/Engine/Private/Curves/Spline.cpp
--- a/Engine/Source/Runtime/Engine/Private/Curves/Spline.cpp
+++ b/Engine/Source/Runtime/Engine/Private/Curves/Spline.cpp
@@ -1556,6 +1556,8 @@ void FSpline::UpdateSpline(const FUpdateSplineParams& InParams)
     case 1: LegacyData->UpdateSpline(InParams); break;
     case 2: NewData->UpdateSpline(InParams); break;
     }
+
+	++Version;
 }
 
 TSharedPtr<TSplineInterface<FVector>> FSpline::GetSplineInterface() const
