- Add lock to GPU Sort Manager AddTask to ensure no concurrent access
- This fixes a race between Niagara & Cacsade adding tasks

#tests Add lots of instance of Niagara & Cascade (GPU) that use sorting

diff --git a/Engine/Source/Runtime/Engine/Private/GPUSortManager.cpp b/Engine/Source/Runtime/Engine/Private/GPUSortManager.cpp
--- a/Engine/Source/Runtime/Engine/Private/GPUSortManager.cpp
+++ b/Engine/Source/Runtime/Engine/Private/GPUSortManager.cpp
@@ -636,6 +636,8 @@ bool FGPUSortManager::AddTask(FRHICommandListBase& RHICmdList, FAllocationInfo&
 	// Invalid if it needs to be sorted before the key are to be generated. Ex : Cascade opaque materials using depth buffer in the simulation.
 	checkSlow(!EnumHasAllFlags(TaskFlags, EGPUSortFlags::KeyGenAfterPostRenderOpaque | EGPUSortFlags::SortAfterPreRender));
 
+	UE::TScopeLock Lock(AddTaskMutex);
+
 	// Search into the current batch for a compatible entries
 	for (FSortBatch& SortBatch : SortBatches)
 	{
@@ -670,11 +672,6 @@ bool FGPUSortManager::AddTask(FRHICommandListBase& RHICmdList, FAllocationInfo&
 	return true;
 }
 
-bool FGPUSortManager::AddTask(FAllocationInfo& OutInfo, int32 ValueCount, EGPUSortFlags TaskFlags)
-{
-	return AddTask(FRHICommandListImmediate::Get(), OutInfo, ValueCount, TaskFlags);
-}
-
 void FGPUSortManager::FinalizeSortBatches()
 {
 	// At this point the batches are complete as no more tasks will be added to them.
diff --git a/Engine/Source/Runtime/Engine/Public/GPUSortManager.h b/Engine/Source/Runtime/Engine/Public/GPUSortManager.h
--- a/Engine/Source/Runtime/Engine/Public/GPUSortManager.h
+++ b/Engine/Source/Runtime/Engine/Public/GPUSortManager.h
@@ -424,9 +424,6 @@ public:
 	 */
 	ENGINE_API bool AddTask(FRHICommandListBase& RHICmdList, FAllocationInfo& OutInfo, int32 ValueCount, EGPUSortFlags TaskFlags);
 
-	UE_DEPRECATED(5.4, "AddTask now requires a command list.")
-	ENGINE_API bool AddTask(FAllocationInfo& OutInfo, int32 ValueCount, EGPUSortFlags TaskFlags);
-
 	/**
 	 * Callback that needs to be called in the rendering loop, after calls to FFXSystemInterface::PreRender() are issued.
 	 * At this point, no other GPU sort tasks are allowed to be created through AddTask() as the sort batches
@@ -499,6 +496,9 @@ private:
 	/** Delete all GPU sort buffers. */
 	void ReleaseSortBuffers();
 
+	/** Used to ensure no concurrent access during AddTask */
+	UE::FMutex AddTaskMutex;
+
 	/**
 	 * The different batches that will be used this frame, each one ending up being sorting work on the GPU.
 	 * In many cases, there will only be one batch but some cases makes it possible to have 2 
