WorldPartition: better Insights tags for UWorldPartition::ApplyRuntimeCellsTransformerStack.

diff --git a/Engine/Source/Runtime/Engine/Private/WorldPartition/WorldPartition.cpp b/Engine/Source/Runtime/Engine/Private/WorldPartition/WorldPartition.cpp
--- a/Engine/Source/Runtime/Engine/Private/WorldPartition/WorldPartition.cpp
+++ b/Engine/Source/Runtime/Engine/Private/WorldPartition/WorldPartition.cpp
@@ -2475,12 +2475,16 @@ void UWorldPartition::ApplyRuntimeCellsTransformerStack(ULevel* InLevel)
 {
 	if (!RuntimeCellsTransformerStack.IsEmpty())
 	{
+		TRACE_CPUPROFILER_EVENT_SCOPE(ApplyRuntimeCellsTransformerStack);
+
 		auto ApplyTransformPhase = [this, InLevel](TFunction<void(const FRuntimeCellTransformerInstance& TransformerInstance, ULevel* Level)> Func)
 		{
 			for (const FRuntimeCellTransformerInstance& TransformerInstance : RuntimeCellsTransformerStack)
 			{
 				if (TransformerInstance.Class)
 				{
+					TRACE_CPUPROFILER_EVENT_SCOPE_TEXT(*TransformerInstance.Class->GetName());
+
 					TPair<double, int32>& TotalStats = RuntimeCellsTransformerStackTimes.FindOrAdd(TransformerInstance.Class);
 					TotalStats.Key -= FPlatformTime::Seconds();
 
@@ -2492,83 +2496,96 @@ void UWorldPartition::ApplyRuntimeCellsTransformerStack(ULevel* InLevel)
 			}
 		};
 
-		// Rerun construction scripts for actors in levels before executing the cell transformers, as the Blueprint parent classes might have changed 
-		// and we want to make sure the constructed actors are up to date.
-		TArray<TObjectPtr<AActor>> LevelActors;
-		Algo::CopyIf(InLevel->Actors, LevelActors, [](AActor* InActor)
 		{
-			return IsValid(InActor) && !InActor->IsChildActor();
-		});
+			TRACE_CPUPROFILER_EVENT_SCOPE(RerunConstructionScripts);
+
+			// Rerun construction scripts for actors in levels before executing the cell transformers, as the Blueprint parent classes might have changed 
+			// and we want to make sure the constructed actors are up to date.
+			TArray<TObjectPtr<AActor>> LevelActors;
+			Algo::CopyIf(InLevel->Actors, LevelActors, [](AActor* InActor)
+			{
+				return IsValid(InActor) && !InActor->IsChildActor();
+			});
 		
-		TMap<UObject*, UObject*> CombinedReplacementMap;
-		FDelegateHandle OnObjectsReplacedHandle = FCoreUObjectDelegates::OnObjectsReplaced.AddLambda([&CombinedReplacementMap](const TMap<UObject*, UObject*>& ReplacementMap)
-		{
-			CombinedReplacementMap.Append(ReplacementMap);
-		});
+			TMap<UObject*, UObject*> CombinedReplacementMap;
+			FDelegateHandle OnObjectsReplacedHandle = FCoreUObjectDelegates::OnObjectsReplaced.AddLambda([&CombinedReplacementMap](const TMap<UObject*, UObject*>& ReplacementMap)
+			{
+				CombinedReplacementMap.Append(ReplacementMap);
+			});
 
-		for (TObjectPtr<AActor>& LevelActor : LevelActors)
-		{
-			// Deal with deleted child actors that might have been deleted from previous calls to RerunConstructionScripts
-			if (IsValid(LevelActor))
+			for (TObjectPtr<AActor>& LevelActor : LevelActors)
 			{
-				LevelActor->RerunConstructionScripts();
+				// Deal with deleted child actors that might have been deleted from previous calls to RerunConstructionScripts
+				if (IsValid(LevelActor))
+				{
+					LevelActor->RerunConstructionScripts();
+				}
 			}
-		}
 		
-		FCoreUObjectDelegates::OnObjectsReplaced.Remove(OnObjectsReplacedHandle);
+			FCoreUObjectDelegates::OnObjectsReplaced.Remove(OnObjectsReplacedHandle);
 		
-		TMap<UObject*, UObject*> FinalReplacementMap;
-		TSet<UObject*> CurrentlyResolving;
-		TFunction<UObject*(UObject*)> ResolveFinalReplacement = [&CombinedReplacementMap, &FinalReplacementMap, &CurrentlyResolving, &ResolveFinalReplacement](UObject* Object) -> UObject*
-		{
-			// Already resolved
-			if (UObject** Final = FinalReplacementMap.Find(Object))
+			TMap<UObject*, UObject*> FinalReplacementMap;
+			TSet<UObject*> CurrentlyResolving;
+			TFunction<UObject*(UObject*)> ResolveFinalReplacement = [&CombinedReplacementMap, &FinalReplacementMap, &CurrentlyResolving, &ResolveFinalReplacement](UObject* Object) -> UObject*
 			{
-				return *Final;
-			}
+				// Already resolved
+				if (UObject** Final = FinalReplacementMap.Find(Object))
+				{
+					return *Final;
+				}
 			
-			// Track objects being currently resolved to enable cycle detection
-			bool bIsAlreadyInSet = false;
-			CurrentlyResolving.Add(Object, &bIsAlreadyInSet);
+				// Track objects being currently resolved to enable cycle detection
+				bool bIsAlreadyInSet = false;
+				CurrentlyResolving.Add(Object, &bIsAlreadyInSet);
 
-			// Cycle detection
-			if (bIsAlreadyInSet)
-			{
-				UE_LOG(LogWorldPartition, Log, TEXT("Detected cycle in object remapping during ApplyRuntimeCellsTransformerStack at %s, ignoring"), *GetNameSafe(Object));
-				return nullptr;
-			}
+				// Cycle detection
+				if (bIsAlreadyInSet)
+				{
+					UE_LOG(LogWorldPartition, Log, TEXT("Detected cycle in object remapping during ApplyRuntimeCellsTransformerStack at %s, ignoring"), *GetNameSafe(Object));
+					return nullptr;
+				}
 			
-			// Resolve final replacement recursively
-			UObject* FinalReplacement = Object;
-			if (UObject** ReplacementObject = CombinedReplacementMap.Find(Object))
-			{
-				FinalReplacement = ResolveFinalReplacement(*ReplacementObject);
-			}
+				// Resolve final replacement recursively
+				UObject* FinalReplacement = Object;
+				if (UObject** ReplacementObject = CombinedReplacementMap.Find(Object))
+				{
+					FinalReplacement = ResolveFinalReplacement(*ReplacementObject);
+				}
 			
-			// Finished resolving, remove from currently resolving objects
-			CurrentlyResolving.Remove(Object);
+				// Finished resolving, remove from currently resolving objects
+				CurrentlyResolving.Remove(Object);
 
-			// Add to the final map if found replacement
-			if (FinalReplacement != Object)
+				// Add to the final map if found replacement
+				if (FinalReplacement != Object)
+				{
+					FinalReplacementMap.Add(Object, FinalReplacement);
+				}
+
+				return FinalReplacement;
+			};
+		
+			// Find final replacement for each replaced object, combines A -> B and B -> C replacements into A -> C
+			for (TPair<UObject*, UObject*> Replacement : CombinedReplacementMap)
 			{
-				FinalReplacementMap.Add(Object, FinalReplacement);
+				ResolveFinalReplacement(Replacement.Key);
 			}
-
-			return FinalReplacement;
-		};
 		
-		// Find final replacement for each replaced object, combines A -> B and B -> C replacements into A -> C
-		for (TPair<UObject*, UObject*> Replacement : CombinedReplacementMap)
-		{
-			ResolveFinalReplacement(Replacement.Key);
+			// Apply final replacements
+			FArchiveReplaceObjectRef<UObject> ReplaceAr(InLevel, FinalReplacementMap);
 		}
-		
-		// Apply final replacements
-		FArchiveReplaceObjectRef<UObject> ReplaceAr(InLevel, FinalReplacementMap);
 
-		ApplyTransformPhase([](const FRuntimeCellTransformerInstance& TransformerInstance, ULevel* Level) { TransformerInstance.PreTransform(Level); });
-		ApplyTransformPhase([](const FRuntimeCellTransformerInstance& TransformerInstance, ULevel* Level) { TransformerInstance.Transform(Level); });
-		ApplyTransformPhase([](const FRuntimeCellTransformerInstance& TransformerInstance, ULevel* Level) { TransformerInstance.PostTransform(Level); });
+		{
+			TRACE_CPUPROFILER_EVENT_SCOPE(PreTransform);
+			ApplyTransformPhase([](const FRuntimeCellTransformerInstance& TransformerInstance, ULevel* Level) { TransformerInstance.PreTransform(Level); });
+		}
+		{
+			TRACE_CPUPROFILER_EVENT_SCOPE(Transform);
+			ApplyTransformPhase([](const FRuntimeCellTransformerInstance& TransformerInstance, ULevel* Level) { TransformerInstance.Transform(Level); });
+		}
+		{
+			TRACE_CPUPROFILER_EVENT_SCOPE(PostTransform);
+			ApplyTransformPhase([](const FRuntimeCellTransformerInstance& TransformerInstance, ULevel* Level) { TransformerInstance.PostTransform(Level); });
+		}
 	}
 }
 
diff --git a/Engine/Source/Runtime/Engine/Private/WorldPartition/WorldPartitionLevelStreamingDynamic.cpp b/Engine/Source/Runtime/Engine/Private/WorldPartition/WorldPartitionLevelStreamingDynamic.cpp
--- a/Engine/Source/Runtime/Engine/Private/WorldPartition/WorldPartitionLevelStreamingDynamic.cpp
+++ b/Engine/Source/Runtime/Engine/Private/WorldPartition/WorldPartitionLevelStreamingDynamic.cpp
@@ -542,7 +542,6 @@ void UWorldPartitionLevelStreamingDynamic::FinalizeRuntimeLevel()
 	// Execute the cells transformers on the cell's level
 	if (StreamingCell.IsValid() && OuterWorld->IsGameWorld())
 	{
-		TRACE_CPUPROFILER_EVENT_SCOPE(ApplyRuntimeCellsTransformerStack);
 		OuterWorldPartition->ApplyRuntimeCellsTransformerStack(RuntimeLevel);
 	}
 
