Remove SyncFromWorld AddColumn from Actor::AddOwnedComponent
* AActor::Modify is called which will cause the actor to get a SyncFromWorld tag already in a deferred manner due to the callback in EditorDataStorageCompatibility

diff --git a/Engine/Source/Runtime/Engine/Private/Actor.cpp b/Engine/Source/Runtime/Engine/Private/Actor.cpp
--- a/Engine/Source/Runtime/Engine/Private/Actor.cpp
+++ b/Engine/Source/Runtime/Engine/Private/Actor.cpp
@@ -3860,16 +3860,6 @@ void AActor::AddOwnedComponent(UActorComponent* Component)
 		{
 			InstanceComponents.Add(Component);
 		}
-#if WITH_EDITOR
-		{
-			using namespace UE::Editor::DataStorage;
-			if (ICoreProvider* Storage = UE::Editor::DataStorage::GetMutableDataStorageFeature<ICoreProvider>(StorageFeatureName))
-			{
-				RowHandle ComponentRow = Storage->LookupMappedRow(ICompatibilityProvider::ObjectMappingDomain, FMapKeyView(Component));
-				Storage->AddColumns<FTypedElementSyncFromWorldTag>(ComponentRow);
-			}
-		}
-#endif
 	}
 }
 
