StaticMesh - Fix SetLODGroup function not clearing the section info map correctly, causing inconsistent results on lower LODs.

We're also now preserving the section info map data on non-generated LODs.

diff --git a/Engine/Source/Runtime/Engine/Private/StaticMesh.cpp b/Engine/Source/Runtime/Engine/Private/StaticMesh.cpp
--- a/Engine/Source/Runtime/Engine/Private/StaticMesh.cpp
+++ b/Engine/Source/Runtime/Engine/Private/StaticMesh.cpp
@@ -5469,12 +5469,14 @@ void UStaticMesh::SetLODGroup(FName NewGroup, bool bRebuildImmediately, bool bAl
 			// Set reduction settings to the defaults.
 			SourceModel.ReductionSettings = GroupSettings.GetDefaultSettings(LODIndex);
 			
-			if (LODIndex != 0)
+			const bool bIsReductionActive = IsReductionActive(LODIndex);
+			if (LODIndex != 0 && bIsReductionActive)
 			{
 				//Reset the section info map
 				if (bResetSectionInfoMap)
 				{
-					for (int32 SectionIndex = 0; SectionIndex < GetSectionInfoMap().GetSectionNumber(LODIndex); ++SectionIndex)
+					const int32 NumberOfSections = GetSectionInfoMap().GetSectionNumber(LODIndex);
+					for (int32 SectionIndex = 0; SectionIndex < NumberOfSections; ++SectionIndex)
 					{
 						GetSectionInfoMap().Remove(LODIndex, SectionIndex);
 					}
