Fix V1043 warnings

#rnx

diff --git a/Engine/Source/Runtime/Engine/Public/kDOP.h b/Engine/Source/Runtime/Engine/Public/kDOP.h
--- a/Engine/Source/Runtime/Engine/Public/kDOP.h
+++ b/Engine/Source/Runtime/Engine/Public/kDOP.h
@@ -166,16 +166,16 @@ inline bool appLineCheckTriangle(const FVector4& Start, const FVector4& End, con
 }
 
 /** ( -0.0001f, -0.0001f, -0.0001f, -0.0001f ) */
-static const VectorRegister GSmallNegativeNumber = MakeVectorRegister(-0.0001f, -0.0001f, -0.0001f, -0.0001f);
+static const VectorRegister GSmallNegativeNumber = MakeVectorRegister(-0.0001f, -0.0001f, -0.0001f, -0.0001f); //-V1043
 //extern const VectorRegister GSmallNegativeNumber;
 
 /** ( 0.0001f, 0.0001f, 0.0001f, 0.0001f ) */
-static const VectorRegister GSmallNumber = MakeVectorRegister( 0.0001f, 0.0001f, 0.0001f, 0.0001f );
+static const VectorRegister GSmallNumber = MakeVectorRegister( 0.0001f, 0.0001f, 0.0001f, 0.0001f ); //-V1043
 //extern const VectorRegister GSmallNumber;
 
-static const VectorRegister GZeroVectorRegister = MakeVectorRegister( 0.f, 0.f, 0.f, 0.f );
+static const VectorRegister GZeroVectorRegister = MakeVectorRegister( 0.f, 0.f, 0.f, 0.f ); //-V1043
 
-static const VectorRegister VectorNegativeOne = MakeVectorRegister( -1.0f, -1.0f, -1.0f, -1.0f );
+static const VectorRegister VectorNegativeOne = MakeVectorRegister( -1.0f, -1.0f, -1.0f, -1.0f ); //-V1043
 
 /**
  * Line vs triangle intersection test. Tests 1 line against 4 triangles at once.
