RemoteObjects:  Distinguishing between Loaning out an Object and Reassigning a Borrowed Object.  This helps clarify some situations, particularly where you're sending an object and maintaining ownership you can just check if you're Loaning it now, because ReassignBorrowed is now a separate case (which means another server has borrowed it, but you never owned it).

#rnx

diff --git a/Engine/Source/Runtime/CoreUObject/Internal/UObject/UObjectMigrationContext.h b/Engine/Source/Runtime/CoreUObject/Internal/UObject/UObjectMigrationContext.h
--- a/Engine/Source/Runtime/CoreUObject/Internal/UObject/UObjectMigrationContext.h
+++ b/Engine/Source/Runtime/CoreUObject/Internal/UObject/UObjectMigrationContext.h
@@ -8,18 +8,19 @@
 
 #define UE_API COREUOBJECT_API
 
+/** Which side of the migration are we on? */
 enum class EObjectMigrationSide : uint8
 {
 	/** We are receiving (migrating-in) an object */
 	Receive,
 
 	/** We are sending (migrating-out) an object */
-	Send,
+	Send
 };
 
 /**
-* Possible types of migration that have happened when you have received an Object
-*/
+ * Possible types of migration that have happened when you have received an Object
+ */
 enum class EObjectMigrationRecvType : uint8
 {
 	Invalid,
@@ -33,11 +34,11 @@ enum class EObjectMigrationRecvType : uint8
 	/** We received objects and must take ownership of the them (we have no choice) */
 	AssignedOwnership
 };
-COREUOBJECT_API const TCHAR* ToString(EObjectMigrationRecvType Value);
+UE_API const TCHAR* ToString(EObjectMigrationRecvType Value);
 
 /**
-	* Possible types of migration that are happening when you are sending an Object
-	*/
+ * Possible types of migration that are happening when you are sending an Object
+ */
 enum class EObjectMigrationSendType : uint8
 {
 	Invalid,
@@ -48,10 +49,13 @@ enum class EObjectMigrationSendType : uint8
 	/** This borrowed object is being returned (sent back) to the server that loaned it out */
 	ReturnBorrowed,
 
+	/** This borrowed object is being forwarded to another server which requested it (it should be returned to the owner, not us) */
+	ReassignBorrowed,
+
 	/** Reassign objects to relinquish ownership and assign them to the destination server */
 	ReassignOwnership
 };
-COREUOBJECT_API const TCHAR* ToString(EObjectMigrationSendType Value);
+UE_API const TCHAR* ToString(EObjectMigrationSendType Value);
 
 /** Structure that holds context for what we're intending to do during migration */
 struct FUObjectMigrationContext final
diff --git a/Engine/Source/Runtime/CoreUObject/Private/UObject/UObjectMigrationContext.cpp b/Engine/Source/Runtime/CoreUObject/Private/UObject/UObjectMigrationContext.cpp
--- a/Engine/Source/Runtime/CoreUObject/Private/UObject/UObjectMigrationContext.cpp
+++ b/Engine/Source/Runtime/CoreUObject/Private/UObject/UObjectMigrationContext.cpp
@@ -27,8 +27,8 @@ EObjectMigrationRecvType FUObjectMigrationContext::GetObjectMigrationRecvType(FR
 	const FRemoteServerId LocalServerId = FRemoteServerId::GetLocalServerId();
 	const bool bAlreadyOwns = UE::RemoteObject::Handle::IsOwned(InObjectId);
 
-	const bool bAssignOwnership = (OwnerServerId == LocalServerId) && !bAlreadyOwns;
-	const bool bReturnLoan = (OwnerServerId == LocalServerId) && bAlreadyOwns;
+	const bool bAssignOwnership = !bAlreadyOwns && (OwnerServerId == LocalServerId);
+	const bool bReturnLoan = bAlreadyOwns && ensureMsgf(OwnerServerId == LocalServerId, TEXT("We are receiving object %s that we own, so the data we received should affirm we are the owner."), *InObjectId.ToString());
 
 	if (bAssignOwnership)
 	{
@@ -80,8 +80,14 @@ EObjectMigrationSendType FUObjectMigrationContext::GetObjectMigrationSendType(FR
 		return EObjectMigrationSendType::ReturnBorrowed;
 	}
 
-	// By default we loan the object to the destination server but ownership does not change
-	return EObjectMigrationSendType::Loan;
+	if (bCurrentlyOwns)
+	{
+		// We loan the object to the destination server but ownership does not change
+		return EObjectMigrationSendType::Loan;
+	}
+
+	// We did not own the object and are sending it onward (forwarding it)
+	return EObjectMigrationSendType::ReassignBorrowed;
 #else
 	return EObjectMigrationSendType::Invalid;
 #endif
@@ -115,6 +121,9 @@ const TCHAR* ToString(EObjectMigrationSendType Value)
 		case EObjectMigrationSendType::ReturnBorrowed:
 			return TEXT("ReturnBorrowed");
 
+		case EObjectMigrationSendType::ReassignBorrowed:
+			return TEXT("ReassignBorrowed");
+
 		case EObjectMigrationSendType::ReassignOwnership:
 			return TEXT("ReassignOwnership");
 
diff --git a/Engine/Source/Runtime/Engine/Private/Actor.cpp b/Engine/Source/Runtime/Engine/Private/Actor.cpp
--- a/Engine/Source/Runtime/Engine/Private/Actor.cpp
+++ b/Engine/Source/Runtime/Engine/Private/Actor.cpp
@@ -1289,6 +1289,7 @@ void AActor::PostMigrate(const FUObjectMigrationContext& MigrationContext)
 	CSV_CUSTOM_STAT(ActorMigration, Total, 1, ECsvCustomStatOp::Accumulate);
 	CSV_CUSTOM_STAT(ActorMigration, Loan, MigrationContext.MigrationSide == EObjectMigrationSide::Send && MigrationContext.GetObjectMigrationSendType(this) == EObjectMigrationSendType::Loan, ECsvCustomStatOp::Accumulate);
 	CSV_CUSTOM_STAT(ActorMigration, ReturnBorrowed, MigrationContext.MigrationSide == EObjectMigrationSide::Send && MigrationContext.GetObjectMigrationSendType(this) == EObjectMigrationSendType::ReturnBorrowed, ECsvCustomStatOp::Accumulate);
+	CSV_CUSTOM_STAT(ActorMigration, ReassignBorrowed, MigrationContext.MigrationSide == EObjectMigrationSide::Send && MigrationContext.GetObjectMigrationSendType(this) == EObjectMigrationSendType::ReassignBorrowed, ECsvCustomStatOp::Accumulate);
 	CSV_CUSTOM_STAT(ActorMigration, ReassignOwnership, MigrationContext.MigrationSide == EObjectMigrationSide::Send && MigrationContext.GetObjectMigrationSendType(this) == EObjectMigrationSendType::ReassignOwnership, ECsvCustomStatOp::Accumulate);
 	CSV_CUSTOM_STAT(ActorMigration, Borrowed, MigrationContext.MigrationSide == EObjectMigrationSide::Receive && MigrationContext.GetObjectMigrationRecvType(this) == EObjectMigrationRecvType::Borrowed, ECsvCustomStatOp::Accumulate);
 	CSV_CUSTOM_STAT(ActorMigration, ReturnedLoan, MigrationContext.MigrationSide == EObjectMigrationSide::Receive && MigrationContext.GetObjectMigrationRecvType(this) == EObjectMigrationRecvType::ReturnedLoan, ECsvCustomStatOp::Accumulate);
@@ -1325,13 +1326,13 @@ void AActor::PostMigrate(const FUObjectMigrationContext& MigrationContext)
 		UE_LOG(LogActor, Verbose, TEXT("%hs: Send %s (LocalSrvId: %s. DstSrvId: %s. HasOwnership: %d. MigrationType: %s)"), __func__, *GetName(), *LocalServerId.ToString(), *MigrationContext.RemoteServerId.ToString(), bHasOwnership, ToString(MigrationSendType));
 		ensureMsgf(MigrationContext.RemoteServerId != LocalServerId, TEXT("%hs: Sending Object %s from %s to %s but destination is the same as the Local Server Id."), __func__, *GetName(), *LocalServerId.ToString(), *MigrationContext.RemoteServerId.ToString());
 
-		const bool bKeepOwnership = (MigrationSendType == EObjectMigrationSendType::Loan) && UE::RemoteObject::Handle::IsOwned(this);
+		// If we're loaning, we're keeping ownership; in all other cases we either don't own, or won't own the object after it's migrated
+		const bool bKeepOwnership = (MigrationSendType == EObjectMigrationSendType::Loan);
 		if (bKeepOwnership)
 		{
 			RemoteObjects_ActorLoanedStates.AddAnnotation(this, { .bActorWasDestroyed = (bool)bActorIsBeingDestroyed, .bTickFunctionsRegistered = (bool)bTickFunctionsRegistered });
 		}
-
-		if (!bKeepOwnership || !bHasOwnership)
+		else 
 		{
 			// Remove the actor from the world on this server as it is about to be migrated to another server.
 			// Migrated actors are unique to the server instance that currently owns their memory and are not replicated to other servers.
diff --git a/Engine/Source/Runtime/Engine/Private/Components/ActorComponent.cpp b/Engine/Source/Runtime/Engine/Private/Components/ActorComponent.cpp
--- a/Engine/Source/Runtime/Engine/Private/Components/ActorComponent.cpp
+++ b/Engine/Source/Runtime/Engine/Private/Components/ActorComponent.cpp
@@ -724,14 +724,16 @@ void UActorComponent::PostMigrate(const struct FUObjectMigrationContext& Migrati
 	if (MigrationContext.MigrationSide == EObjectMigrationSide::Send)
 	{
 		const EObjectMigrationSendType SendType = MigrationContext.GetObjectMigrationSendType(this);
+		UE_LOG(LogActorComponent, VeryVerbose, TEXT("%hs: %s %s"), __func__, ToString(SendType), *GetName());
 
-		const bool bKeepOwnership = (SendType == EObjectMigrationSendType::Loan) && UE::RemoteObject::Handle::IsOwned(this);
+		// If we're loaning, we're keeping ownership; in all other cases we either don't own, or won't own the object after it's migrated
+		const bool bKeepOwnership = (SendType == EObjectMigrationSendType::Loan);
 		if (bKeepOwnership)
 		{
 			checkfSlow(!RemoteObjects_LoanedStates.GetAnnotationMap().Find(this), TEXT("%hs: Expected no existing Loaned State for %s (it should have been consumed when received, or not loaned out twice in a row)"), __func__, *GetName());
 			RemoteObjects_LoanedStates.AddAnnotation(this, NewState);
 		}
-		else // if (!bKeepOwnership)
+		else
 		{
 			// If we're transferring ownership, make sure we don't keep anything registered locally
 			if (IsRegistered())
@@ -746,6 +748,7 @@ void UActorComponent::PostMigrate(const struct FUObjectMigrationContext& Migrati
 	else if (ensureMsgf(MigrationContext.MigrationSide == EObjectMigrationSide::Receive, TEXT("Invalid MigrationContext: Side was not Send or Receive")))
 	{
 		const EObjectMigrationRecvType RecvType = MigrationContext.GetObjectMigrationRecvType(this);
+		UE_LOG(LogActorComponent, VeryVerbose, TEXT("%hs: %s %s"), __func__, ToString(RecvType), *GetName());
 
 		// We are expecting this flow from AActor, but now we must do this per-component in such a way
 		// that we are not mutating our previously deserialized data.
