- Makes material expressions that can't be translated in the old translator report an error during translation.
- IfThenElse, Aggregate and Blend marked to be exposed only when r.Material.Translator.EnableNew cvar is enabled.

#rnx

diff --git a/Engine/Source/Editor/MaterialEditor/Private/MaterialExpressionClasses.cpp b/Engine/Source/Editor/MaterialEditor/Private/MaterialExpressionClasses.cpp
--- a/Engine/Source/Editor/MaterialEditor/Private/MaterialExpressionClasses.cpp
+++ b/Engine/Source/Editor/MaterialEditor/Private/MaterialExpressionClasses.cpp
@@ -105,13 +105,7 @@ void MaterialExpressionClasses::InitMaterialExpressionClasses()
 							continue;
 						}
 
-						// Hide node types related to control flow, unless it's enabled
-						if (Class->HasMetaData("MaterialControlFlow"))
-						{
-							continue;
-						}
-
-						if (!bMaterialTranslatorEnableNew && Class->HasMetaData("MaterialNewHLSLGenerator"))
+						if (!bMaterialTranslatorEnableNew && Class->HasMetaData("NewMaterialTranslator"))
 						{
 							continue;
 						}
diff --git a/Engine/Source/Editor/UnrealEd/Private/Commandlets/DumpMaterialExpressionsCommandlet.cpp b/Engine/Source/Editor/UnrealEd/Private/Commandlets/DumpMaterialExpressionsCommandlet.cpp
--- a/Engine/Source/Editor/UnrealEd/Private/Commandlets/DumpMaterialExpressionsCommandlet.cpp
+++ b/Engine/Source/Editor/UnrealEd/Private/Commandlets/DumpMaterialExpressionsCommandlet.cpp
@@ -106,7 +106,7 @@ void WriteOutMaterialExpressions(FDiagnosticTableWriterCSV& CsvFile)
 			{
 				const bool bClassDeprecated = Class->HasAnyClassFlags(CLASS_Deprecated);
 				const bool bControlFlow = Class->HasMetaData("MaterialControlFlow");
-				const bool bNewHLSLGenerator = Class->HasMetaData("MaterialNewHLSLGenerator");
+				const bool bNewHLSLGenerator = Class->HasMetaData("NewMaterialTranslator");
 
 				// If the expression is listed in the material node creation dropdown menu
 				// See class exclusions in:
diff --git a/Engine/Source/Runtime/Engine/Private/Materials/MaterialExpressions.cpp b/Engine/Source/Runtime/Engine/Private/Materials/MaterialExpressions.cpp
--- a/Engine/Source/Runtime/Engine/Private/Materials/MaterialExpressions.cpp
+++ b/Engine/Source/Runtime/Engine/Private/Materials/MaterialExpressions.cpp
@@ -802,6 +802,16 @@ FString UMaterialExpression::GetAssetPathName() const
 }
 
 #if WITH_EDITOR
+
+int32 UMaterialExpression::Compile(class FMaterialCompiler* Compiler, int32 OutputIndex)
+{
+	if (GetClass()->HasMetaData("NewMaterialTranslator"))
+	{
+		Compiler->Errorf(TEXT("Unsupported expression '%s'."), *GetName());
+	}
+	return INDEX_NONE;
+}
+
 void UMaterialExpression::CopyMaterialExpressions(const TArray<UMaterialExpression*>& SrcExpressions, const TArray<UMaterialExpressionComment*>& SrcExpressionComments, 
 	UMaterial* Material, UMaterialFunction* EditFunction, TArray<UMaterialExpression*>& OutNewExpressions, TArray<UMaterialExpression*>& OutNewComments)
 {
diff --git a/Engine/Source/Runtime/Engine/Public/Materials/MaterialExpression.h b/Engine/Source/Runtime/Engine/Public/Materials/MaterialExpression.h
--- a/Engine/Source/Runtime/Engine/Public/Materials/MaterialExpression.h
+++ b/Engine/Source/Runtime/Engine/Public/Materials/MaterialExpression.h
@@ -267,7 +267,7 @@ class UMaterialExpression : public UObject
 	 * @param	Compiler - UMaterial compiler that knows how to handle this expression.
 	 * @return	Index to the new FMaterialCompiler::CodeChunk entry for this expression
 	 */	
-	virtual int32 Compile(class FMaterialCompiler* Compiler, int32 OutputIndex) { return INDEX_NONE; }
+	ENGINE_API virtual int32 Compile(class FMaterialCompiler* Compiler, int32 OutputIndex);
 	virtual int32 CompilePreview(class FMaterialCompiler* Compiler, int32 OutputIndex) { return Compile(Compiler, OutputIndex); }
 
 	ENGINE_API bool IsUsingNewHLSLGenerator() const;
diff --git a/Engine/Source/Runtime/Engine/Public/Materials/MaterialExpressionAggregate.h b/Engine/Source/Runtime/Engine/Public/Materials/MaterialExpressionAggregate.h
--- a/Engine/Source/Runtime/Engine/Public/Materials/MaterialExpressionAggregate.h
+++ b/Engine/Source/Runtime/Engine/Public/Materials/MaterialExpressionAggregate.h
@@ -28,7 +28,7 @@ struct FMaterialExpressionAggregateEntry
 	FExpressionInput Input;
 };
 
-UCLASS(MinimalAPI)
+UCLASS(MinimalAPI, meta=(NewMaterialTranslator))
 class UMaterialExpressionAggregate : public UMaterialExpression
 {
 	GENERATED_BODY()
diff --git a/Engine/Source/Runtime/Engine/Public/Materials/MaterialExpressionBlend.h b/Engine/Source/Runtime/Engine/Public/Materials/MaterialExpressionBlend.h
--- a/Engine/Source/Runtime/Engine/Public/Materials/MaterialExpressionBlend.h
+++ b/Engine/Source/Runtime/Engine/Public/Materials/MaterialExpressionBlend.h
@@ -16,7 +16,7 @@ enum class EMaterialExpressionBlendMode : uint8
 	UseB,
 };
 
-UCLASS(MinimalAPI, CollapseCategories, HideCategories = Object)
+UCLASS(MinimalAPI, CollapseCategories, HideCategories = Object, meta=(NewMaterialTranslator))
 class UMaterialExpressionBlend : public UMaterialExpression
 {
 	GENERATED_BODY()
diff --git a/Engine/Source/Runtime/Engine/Public/Materials/MaterialExpressionIfThenElse.h b/Engine/Source/Runtime/Engine/Public/Materials/MaterialExpressionIfThenElse.h
--- a/Engine/Source/Runtime/Engine/Public/Materials/MaterialExpressionIfThenElse.h
+++ b/Engine/Source/Runtime/Engine/Public/Materials/MaterialExpressionIfThenElse.h
@@ -10,7 +10,7 @@
 #include "MaterialExpressionIfThenElse.generated.h"
 
 
-UCLASS(collapsecategories, hidecategories = Object, MinimalAPI)
+UCLASS(collapsecategories, hidecategories = Object, MinimalAPI, meta=(NewMaterialTranslator))
 class UMaterialExpressionIfThenElse : public UMaterialExpression
 {
 	GENERATED_BODY()
