Addendum to my previous change: prevents MaterialExpressionOperator from being created when new translator is disabled.

#rnx

diff --git a/Engine/Source/Runtime/Engine/Public/Materials/MaterialExpressionOperator.h b/Engine/Source/Runtime/Engine/Public/Materials/MaterialExpressionOperator.h
--- a/Engine/Source/Runtime/Engine/Public/Materials/MaterialExpressionOperator.h
+++ b/Engine/Source/Runtime/Engine/Public/Materials/MaterialExpressionOperator.h
@@ -108,7 +108,7 @@ struct FMaterialExpressionOperatorInput
 	}
 };
  
-UCLASS(MinimalAPI)
+UCLASS(MinimalAPI, meta=(NewMaterialTranslator))
 class UMaterialExpressionOperator : public UMaterialExpression
 {
 	GENERATED_BODY()
