[NetPhysics] Convert NetworkPhysicsSettingsComponent into a DataAsset

diff --git a/Engine/Source/Runtime/Engine/Private/PhysicsEngine/NetworkPhysicsComponent.cpp b/Engine/Source/Runtime/Engine/Private/PhysicsEngine/NetworkPhysicsComponent.cpp
--- a/Engine/Source/Runtime/Engine/Private/PhysicsEngine/NetworkPhysicsComponent.cpp
+++ b/Engine/Source/Runtime/Engine/Private/PhysicsEngine/NetworkPhysicsComponent.cpp
@@ -501,12 +501,13 @@ void UNetworkPhysicsComponent::InitializeComponent()
 		UNetworkPhysicsSettingsComponent* SettingsComponent = Owner->FindComponentByClass<UNetworkPhysicsSettingsComponent>();
 		if (SettingsComponent)
 		{
-			SetNumberOfInputsToNetwork(SettingsComponent->NetworkPhysicsComponentSettings.GetRedundantInputs() + 1);
-			SetNumberOfRemoteInputsToNetwork(SettingsComponent->NetworkPhysicsComponentSettings.GetRedundantRemoteInputs() + 1);
-			SetNumberOfStatesToNetwork(SettingsComponent->NetworkPhysicsComponentSettings.GetRedundantStates() + 1);
-			bEnableUnreliableFlow = SettingsComponent->NetworkPhysicsComponentSettings.GetEnableUnreliableFlow();
-			bEnableReliableFlow = SettingsComponent->NetworkPhysicsComponentSettings.GetEnableReliableFlow();
-			bValidateDataOnGameThread = SettingsComponent->NetworkPhysicsComponentSettings.GetValidateDataOnGameThread();
+			const FNetworkPhysicsSettingsData& SettingsData = SettingsComponent->GetSettings();
+			SetNumberOfInputsToNetwork(SettingsData.NetworkPhysicsComponentSettings.GetRedundantInputs() + 1);
+			SetNumberOfRemoteInputsToNetwork(SettingsData.NetworkPhysicsComponentSettings.GetRedundantRemoteInputs() + 1);
+			SetNumberOfStatesToNetwork(SettingsData.NetworkPhysicsComponentSettings.GetRedundantStates() + 1);
+			bEnableUnreliableFlow = SettingsData.NetworkPhysicsComponentSettings.GetEnableUnreliableFlow();
+			bEnableReliableFlow = SettingsData.NetworkPhysicsComponentSettings.GetEnableReliableFlow();
+			bValidateDataOnGameThread = SettingsData.NetworkPhysicsComponentSettings.GetValidateDataOnGameThread();
 
 			if (ReplicatedInputs.History)
 			{
@@ -1585,9 +1586,9 @@ void UNetworkPhysicsComponent::UpdateAsyncComponent(const bool bFullUpdate)
 					AsyncInput->ActorName = AActor::GetDebugName(Owner);
 			
 					UNetworkPhysicsSettingsComponent* SettingsComponent = Owner->FindComponentByClass<UNetworkPhysicsSettingsComponent>();
-					if (SettingsComponent && SettingsComponent->GetNetworkPhysicsSettings_Internal())
+					if (SettingsComponent && SettingsComponent->GetNetworkPhysicsSettings_Internal() && SettingsComponent->GetSettings_Internal().IsValid())
 					{
-						AsyncInput->SettingsComponent = SettingsComponent->GetNetworkPhysicsSettings_Internal();
+						AsyncInput->SettingsComponent = SettingsComponent->GetSettings_Internal();
 					}
 				}
 				
@@ -1746,9 +1747,9 @@ void FAsyncNetworkPhysicsComponent::OnPhysicsObjectUnregistered_Internal(Chaos::
 	}
 }
 
-const FNetworkPhysicsSettingsNetworkPhysicsComponent& FAsyncNetworkPhysicsComponent::GetComponentSettings()
+const FNetworkPhysicsSettingsNetworkPhysicsComponent& FAsyncNetworkPhysicsComponent::GetComponentSettings() const
 {
-	return SettingsComponent ? SettingsComponent->Settings.NetworkPhysicsComponentSettings : SettingsNetworkPhysicsComponent_Default;
+	return SettingsComponent.IsValid() ? SettingsComponent.Pin()->NetworkPhysicsComponentSettings : SettingsNetworkPhysicsComponent_Default;
 };
 
 void FAsyncNetworkPhysicsComponent::ConsumeAsyncInput(const int32 PhysicsStep)
@@ -1759,9 +1760,9 @@ void FAsyncNetworkPhysicsComponent::ConsumeAsyncInput(const int32 PhysicsStep)
 
 		/** Onetime setup marshaled data */
 		{
-			if (AsyncInput->SettingsComponent.IsSet())
+			if (AsyncInput->SettingsComponent.IsSet() && (*AsyncInput->SettingsComponent).IsValid())
 			{
-				SettingsComponent = *AsyncInput->SettingsComponent;
+				SettingsComponent = (*AsyncInput->SettingsComponent).Pin();
 			}
 			if (AsyncInput->bIsLocallyControlled.IsSet())
 			{
diff --git a/Engine/Source/Runtime/Engine/Private/PhysicsEngine/NetworkPhysicsSettingsComponent.cpp b/Engine/Source/Runtime/Engine/Private/PhysicsEngine/NetworkPhysicsSettingsComponent.cpp
--- a/Engine/Source/Runtime/Engine/Private/PhysicsEngine/NetworkPhysicsSettingsComponent.cpp
+++ b/Engine/Source/Runtime/Engine/Private/PhysicsEngine/NetworkPhysicsSettingsComponent.cpp
@@ -28,7 +28,8 @@ namespace PhysicsReplicationCVars
 	}
 }
 
-
+/** Static Initialization*/
+const FNetworkPhysicsSettingsData UNetworkPhysicsSettingsComponent::SettingsNetworkPhysicsData_Default = FNetworkPhysicsSettingsData();
 TMap<AActor*, UNetworkPhysicsSettingsComponent*> UNetworkPhysicsSettingsComponent::ObjectToSettings_External = TMap<AActor*, UNetworkPhysicsSettingsComponent*>();
 
 UNetworkPhysicsSettingsComponent::UNetworkPhysicsSettingsComponent()
@@ -67,6 +68,15 @@ void UNetworkPhysicsSettingsComponent::InitializeComponent()
 {
 	Super::InitializeComponent();
 
+	if (!SettingsDataAsset)
+	{
+		UE_LOG(LogChaos, Warning, TEXT("UNetworkPhysicsSettingsComponent doesn't have a DataAsset assigned. Settings directly on the component are deprecated - %s"), *AActor::GetDebugName(GetOwner()));
+		return;
+	}
+
+	// Copy over GT settings into a TSharedPtr to be accessed on the physics thread
+	SettingsDataAsset->InitializeInternalSettings();
+
 	using namespace Chaos;
 	NetworkPhysicsSettings_Internal = nullptr;
 	if (UWorld* World = GetWorld())
@@ -89,7 +99,11 @@ void UNetworkPhysicsSettingsComponent::InitializeComponent()
 								if (FNetworkPhysicsSettingsAsyncInput* AsyncInput = NetworkPhysicsSettings_Internal->GetProducerInputData_External())
 								{
 									AsyncInput->PhysicsObject = PhysicsObject;
-									FNetworkPhysicsSettingsAsync Settings
+									AsyncInput->Settings_Internal = GetSettings_Internal();
+
+#if WITH_EDITORONLY_DATA
+									PRAGMA_DISABLE_DEPRECATION_WARNINGS
+									FNetworkPhysicsSettingsAsync SettingsCopy
 									{
 										.GeneralSettings = GeneralSettings,
 										.DefaultReplicationSettings = DefaultReplicationSettings,
@@ -97,11 +111,13 @@ void UNetworkPhysicsSettingsComponent::InitializeComponent()
 										.ResimulationSettings = ResimulationSettings,
 										.NetworkPhysicsComponentSettings = NetworkPhysicsComponentSettings,
 									};
-									AsyncInput->Settings = Settings;
+									AsyncInput->Settings = SettingsCopy;
+									PRAGMA_ENABLE_DEPRECATION_WARNINGS
+#endif
 								}
 
 								// Apply resimulation error correction settings for render interpolation to the physics proxy
-								ResimulationSettings.ResimulationErrorCorrectionSettings.ApplySettings_External(PhysicsObject);
+								GetSettings().ResimulationSettings.ResimulationErrorCorrectionSettings.ApplySettings_External(PhysicsObject);
 							}
 						}
 					}
@@ -121,13 +137,16 @@ void UNetworkPhysicsSettingsComponent::UninitializeComponent()
 	Super::UninitializeComponent();
 
 	using namespace Chaos;
-	if (UWorld* World = GetWorld())
+	if (NetworkPhysicsSettings_Internal)
 	{
-		if (FPhysScene* PhysScene = World->GetPhysicsScene())
+		if (UWorld* World = GetWorld())
 		{
-			if (Chaos::FPhysicsSolver* Solver = PhysScene->GetSolver())
+			if (FPhysScene* PhysScene = World->GetPhysicsScene())
 			{
-				Solver->UnregisterAndFreeSimCallbackObject_External(NetworkPhysicsSettings_Internal);
+				if (Chaos::FPhysicsSolver* Solver = PhysScene->GetSolver())
+				{
+					Solver->UnregisterAndFreeSimCallbackObject_External(NetworkPhysicsSettings_Internal);
+				}
 			}
 		}
 	}
@@ -146,10 +165,10 @@ void UNetworkPhysicsSettingsComponent::BeginPlay()
 	// Apply overrides on actor
 	if (AActor* Owner = GetOwner())
 	{
-		if ((GeneralSettings.bOverrideSimProxyRepMode || PhysicsReplicationCVars::ResimulationCVars::SimProxyRepMode >= 0)
+		if ((GetSettings().GeneralSettings.bOverrideSimProxyRepMode || PhysicsReplicationCVars::ResimulationCVars::SimProxyRepMode >= 0)
 			&& Owner->GetLocalRole() == ENetRole::ROLE_SimulatedProxy)
 		{
-			EPhysicsReplicationMode RepMode = GeneralSettings.bOverrideSimProxyRepMode ? GeneralSettings.SimProxyRepMode : static_cast<EPhysicsReplicationMode>(PhysicsReplicationCVars::ResimulationCVars::SimProxyRepMode);
+			EPhysicsReplicationMode RepMode = GetSettings().GeneralSettings.bOverrideSimProxyRepMode ? GetSettings().GeneralSettings.SimProxyRepMode : static_cast<EPhysicsReplicationMode>(PhysicsReplicationCVars::ResimulationCVars::SimProxyRepMode);
 			Owner->SetPhysicsReplicationMode(RepMode);
 		}
 	}
@@ -171,7 +190,7 @@ void UNetworkPhysicsSettingsComponent::OnComponentPhysicsStateChanged(UPrimitive
 			}
 
 			// Apply resimulation error correction settings for render interpolation to the physics proxy
-			ResimulationSettings.ResimulationErrorCorrectionSettings.ApplySettings_External(PhysicsObject);
+			GetSettings().ResimulationSettings.ResimulationErrorCorrectionSettings.ApplySettings_External(PhysicsObject);
 		}
 
 		RegisterInPhysicsReplicationLOD();
@@ -180,7 +199,7 @@ void UNetworkPhysicsSettingsComponent::OnComponentPhysicsStateChanged(UPrimitive
 
 void UNetworkPhysicsSettingsComponent::RegisterInPhysicsReplicationLOD()
 {
-	if (!GeneralSettings.bFocalParticleInPhysicsReplicationLOD)
+	if (!GetSettings().GeneralSettings.bFocalParticleInPhysicsReplicationLOD)
 	{
 		return;
 	}
@@ -200,7 +219,7 @@ UNetworkPhysicsSettingsComponent* UNetworkPhysicsSettingsComponent::GetSettingsF
 	return Value ? *Value : nullptr;
 }
 
-void FNetworkPhysicsSettingsResimulationErrorCorrection::ApplySettings_External(Chaos::FPhysicsObjectHandle PhysicsObject)
+void FNetworkPhysicsSettingsResimulationErrorCorrection::ApplySettings_External(Chaos::FPhysicsObjectHandle PhysicsObject) const
 {
 	if (!PhysicsObject)
 	{
@@ -288,15 +307,22 @@ void FNetworkPhysicsSettingsComponentAsync::ConsumeAsyncInput()
 	// Receive data on PT from GT
 	if (const FNetworkPhysicsSettingsAsyncInput* AsyncInput = GetConsumerInput_Internal())
 	{
-		if (AsyncInput->PhysicsObject.IsSet())
+		PRAGMA_DISABLE_DEPRECATION_WARNINGS
+		if (AsyncInput->Settings.IsSet())
 		{
-			PhysicsObject = *AsyncInput->PhysicsObject;
-			RegisterSettingsInPhysicsReplication();
+			Settings = *AsyncInput->Settings;
 		}
+		PRAGMA_ENABLE_DEPRECATION_WARNINGS
 
-		if (AsyncInput->Settings.IsSet())
+		if (AsyncInput->Settings_Internal.IsSet() && (*AsyncInput->Settings_Internal).IsValid())
 		{
-			Settings = *AsyncInput->Settings;
+			Settings_Internal = (*AsyncInput->Settings_Internal).Pin();
+		}
+
+		if (AsyncInput->PhysicsObject.IsSet())
+		{
+			PhysicsObject = *AsyncInput->PhysicsObject;
+			RegisterSettingsInPhysicsReplication();
 		}
 	}
 }
@@ -307,7 +333,7 @@ void FNetworkPhysicsSettingsComponentAsync::RegisterSettingsInPhysicsReplication
 	{
 		if (IPhysicsReplicationAsync* PhysRep = RigidsSolver->GetPhysicsReplication_Internal())
 		{
-			PhysRep->RegisterSettings(PhysicsObject, Settings);
+			PhysRep->RegisterSettings(PhysicsObject, Settings_Internal.ToWeakPtr());
 		}
 	}
 }
diff --git a/Engine/Source/Runtime/Engine/Private/PhysicsEngine/PhysicsReplication.cpp b/Engine/Source/Runtime/Engine/Private/PhysicsEngine/PhysicsReplication.cpp
--- a/Engine/Source/Runtime/Engine/Private/PhysicsEngine/PhysicsReplication.cpp
+++ b/Engine/Source/Runtime/Engine/Private/PhysicsEngine/PhysicsReplication.cpp
@@ -633,12 +633,13 @@ bool FPhysicsReplication::ApplyRigidBodyState(float DeltaSeconds, FBodyInstance*
 	bool bCorrectConnectedBodiesFriction = PhysicsReplicationCVars::DefaultReplicationCVars::bCorrectConnectedBodiesFriction;
 
 	// Assign per-actor settings from NetworkPhysicSettingsComponent if this actor has one
-	if (SettingsCurrent.Get())
+	if (SettingsCurrent.IsValid())
 	{
-		MaxLinearHardSnapDistance = SettingsCurrent.Get()->DefaultReplicationSettings.GetMaxLinearHardSnapDistance(MaxLinearHardSnapDistance);
-		bHardsnapLegacyInPT = SettingsCurrent.Get()->DefaultReplicationSettings.GetHardsnapDefaultLegacyInPT();
-		bCorrectConnectedBodies = SettingsCurrent.Get()->DefaultReplicationSettings.GetCorrectConnectedBodies();
-		bCorrectConnectedBodiesFriction = SettingsCurrent.Get()->DefaultReplicationSettings.GetCorrectConnectedBodiesFriction();
+		const FNetworkPhysicsSettingsData& SettingsData = SettingsCurrent.Pin()->GetSettings();
+		MaxLinearHardSnapDistance = SettingsData.DefaultReplicationSettings.GetMaxLinearHardSnapDistance(MaxLinearHardSnapDistance);
+		bHardsnapLegacyInPT = SettingsData.DefaultReplicationSettings.GetHardsnapDefaultLegacyInPT();
+		bCorrectConnectedBodies = SettingsData.DefaultReplicationSettings.GetCorrectConnectedBodies();
+		bCorrectConnectedBodiesFriction = SettingsData.DefaultReplicationSettings.GetCorrectConnectedBodiesFriction();
 	}
 
 	// Get Current state
@@ -886,19 +887,19 @@ void FPhysicsReplicationAsync::OnPhysicsObjectUnregistered_Internal(Chaos::FCons
 	ObjectToSettings.Remove(PhysicsObject);
 }
 
-void FPhysicsReplicationAsync::RegisterSettings(Chaos::FConstPhysicsObjectHandle PhysicsObject, FNetworkPhysicsSettingsAsync InSettings)
+void FPhysicsReplicationAsync::RegisterSettings(Chaos::FConstPhysicsObjectHandle PhysicsObject, TWeakPtr<const FNetworkPhysicsSettingsData> InSettings)
 {
 	if (PhysicsObject != nullptr)
 	{
-		FNetworkPhysicsSettingsAsync& Settings = ObjectToSettings.FindOrAdd(PhysicsObject);
+		TWeakPtr<const FNetworkPhysicsSettingsData>& Settings = ObjectToSettings.FindOrAdd(PhysicsObject);
 		Settings = InSettings;
 	}
 }
 
 void FPhysicsReplicationAsync::FetchObjectSettings(Chaos::FConstPhysicsObjectHandle PhysicsObject)
 {
-	FNetworkPhysicsSettingsAsync* CustomSettings = ObjectToSettings.Find(PhysicsObject);
-	SettingsCurrent = (CustomSettings != nullptr) ? *CustomSettings : SettingsDefault;
+	TWeakPtr<const FNetworkPhysicsSettingsData>* CustomSettings = ObjectToSettings.Find(PhysicsObject);
+	SettingsCurrent = (CustomSettings && (*CustomSettings).IsValid()) ? *(*CustomSettings).Pin().Get() : SettingsDefault;
 }
 
 void FPhysicsReplicationAsync::OnPostInitialize_Internal()
diff --git a/Engine/Source/Runtime/Engine/Public/Physics/NetworkPhysicsComponent.h b/Engine/Source/Runtime/Engine/Public/Physics/NetworkPhysicsComponent.h
--- a/Engine/Source/Runtime/Engine/Public/Physics/NetworkPhysicsComponent.h
+++ b/Engine/Source/Runtime/Engine/Public/Physics/NetworkPhysicsComponent.h
@@ -1662,7 +1662,7 @@ struct FAsyncNetworkPhysicsComponentInput : public Chaos::FSimCallbackInput
 	TOptional<bool> bCompareStateToTriggerRewind;
 	TOptional<bool> bCompareStateToTriggerRewindIncludeSimProxies;
 	TOptional<bool> bCompareInputToTriggerRewind;
-	TOptional<FNetworkPhysicsSettingsComponentAsync*> SettingsComponent;
+	TOptional<TWeakPtr<const FNetworkPhysicsSettingsData>> SettingsComponent;
 
 	TUniquePtr<Chaos::FBaseRewindHistory> InputData;
 	TUniquePtr<Chaos::FBaseRewindHistory> StateData;
@@ -1801,7 +1801,7 @@ private:
 	Chaos::FPBDRigidsEvolution* GetEvolution();
 
 	/** Get the settings for this NetworkPhysicsComponent */
-	const FNetworkPhysicsSettingsNetworkPhysicsComponent& GetComponentSettings();
+	const FNetworkPhysicsSettingsNetworkPhysicsComponent& GetComponentSettings() const;
 
 	/** Trigger a resimulation on frame */
 	void TriggerResimulation(int32 ResimFrame);
@@ -1832,7 +1832,7 @@ private:
 	int32 NewImportantInputFrame = INT_MAX;
 
 	// Component settings
-	FNetworkPhysicsSettingsComponentAsync* SettingsComponent;
+	TWeakPtr<const FNetworkPhysicsSettingsData> SettingsComponent;
 	static const FNetworkPhysicsSettingsNetworkPhysicsComponent SettingsNetworkPhysicsComponent_Default;
 
 	// Actor component that will be used to fill the histories
diff --git a/Engine/Source/Runtime/Engine/Public/Physics/NetworkPhysicsSettingsComponent.h b/Engine/Source/Runtime/Engine/Public/Physics/NetworkPhysicsSettingsComponent.h
--- a/Engine/Source/Runtime/Engine/Public/Physics/NetworkPhysicsSettingsComponent.h
+++ b/Engine/Source/Runtime/Engine/Public/Physics/NetworkPhysicsSettingsComponent.h
@@ -10,6 +10,7 @@
 #include "Components/ActorComponent.h"
 #include "Chaos/SimCallbackObject.h"
 #include "Curves/CurveFloat.h"
+#include "Engine/DataAsset.h"
 
 #include "NetworkPhysicsSettingsComponent.generated.h"
 
@@ -103,28 +104,28 @@ struct FNetworkPhysicsSettingsDefaultReplication
 	// Overrides CVar: p.MaxLinearHardSnapDistance -- Hardsnap if distance between current position and extrapolated target position is larger than this value.
 	UPROPERTY(config, EditDefaultsOnly, Category = "Overrides", Meta = (EditCondition = "bOverrideMaxLinearHardSnapDistance"))
 	float MaxLinearHardSnapDistance = 400.0f;
-	float GetMaxLinearHardSnapDistance(float DefaultValue) { return bOverrideMaxLinearHardSnapDistance ? MaxLinearHardSnapDistance : DefaultValue; }
+	float GetMaxLinearHardSnapDistance(float DefaultValue) const { return bOverrideMaxLinearHardSnapDistance ? MaxLinearHardSnapDistance : DefaultValue; }
 
 	UPROPERTY(config, EditDefaultsOnly, Category = "Overrides", Meta = (InlineEditConditionToggle))
 	uint32 bOverrideDefaultLegacyHardsnapInPT : 1;
 	// Overrides CVar: p.DefaultReplication.Legacy.HardsnapInPT -- If default replication is used and it's running the legacy flow through Game Thread, allow hardsnapping to be performed on Physics Thread if async physics is enabled.
 	UPROPERTY(config, EditDefaultsOnly, Category = "Overrides", Meta = (EditCondition = "bOverrideDefaultLegacyHardsnapInPT"))
 	bool bHardsnapInPhysicsThread = PhysicsReplicationCVars::DefaultReplicationCVars::bHardsnapLegacyInPT;
-	bool GetHardsnapDefaultLegacyInPT() { return bOverrideDefaultLegacyHardsnapInPT ? bHardsnapInPhysicsThread : PhysicsReplicationCVars::DefaultReplicationCVars::bHardsnapLegacyInPT; }
+	bool GetHardsnapDefaultLegacyInPT() const { return bOverrideDefaultLegacyHardsnapInPT ? bHardsnapInPhysicsThread : PhysicsReplicationCVars::DefaultReplicationCVars::bHardsnapLegacyInPT; }
 
 	UPROPERTY(config, EditDefaultsOnly, Category = "Overrides", Meta = (InlineEditConditionToggle))
 	uint32 bOverrideCorrectConnectedBodies : 1;
 	// Overrides CVar: p.DefaultReplication.CorrectConnectedBodies -- When true, transform corrections will also apply to any connected physics object.
 	UPROPERTY(config, EditDefaultsOnly, Category = "Overrides", Meta = (EditCondition = "bOverrideCorrectConnectedBodies"))
 	bool bCorrectConnectedBodies = PhysicsReplicationCVars::DefaultReplicationCVars::bCorrectConnectedBodies;
-	bool GetCorrectConnectedBodies() { return bOverrideCorrectConnectedBodies ? bCorrectConnectedBodies : PhysicsReplicationCVars::DefaultReplicationCVars::bCorrectConnectedBodies; }
+	bool GetCorrectConnectedBodies() const { return bOverrideCorrectConnectedBodies ? bCorrectConnectedBodies : PhysicsReplicationCVars::DefaultReplicationCVars::bCorrectConnectedBodies; }
 
 	UPROPERTY(config, EditDefaultsOnly, Category = "Overrides", Meta = (InlineEditConditionToggle))
 	uint32 bOverrideCorrectConnectedBodiesFriction : 1;
 	// Overrides CVar: p.DefaultReplication.CorrectConnectedBodiesFriction -- When true, transform correction on any connected physics object will also recalculate their friction.
 	UPROPERTY(config, EditDefaultsOnly, Category = "Overrides", Meta = (EditCondition = "bOverrideCorrectConnectedBodiesFriction"))
 	bool bCorrectConnectedBodiesFriction = PhysicsReplicationCVars::DefaultReplicationCVars::bCorrectConnectedBodiesFriction;
-	bool GetCorrectConnectedBodiesFriction() { return bOverrideCorrectConnectedBodiesFriction ? bCorrectConnectedBodiesFriction : PhysicsReplicationCVars::DefaultReplicationCVars::bCorrectConnectedBodiesFriction; }
+	bool GetCorrectConnectedBodiesFriction() const { return bOverrideCorrectConnectedBodiesFriction ? bCorrectConnectedBodiesFriction : PhysicsReplicationCVars::DefaultReplicationCVars::bCorrectConnectedBodiesFriction; }
 };
 
 USTRUCT()
@@ -137,112 +138,112 @@ struct FNetworkPhysicsSettingsPredictiveInterpolation
 	// Overrides CVar: np2.PredictiveInterpolation.PosCorrectionTimeBase -- Base time to correct positional offset over. RoundTripTime * PosCorrectionTimeMultiplier is added on top of this.
 	UPROPERTY(config, EditDefaultsOnly, Category = "Overrides", Meta = (EditCondition = "bOverridePosCorrectionTimeBase"))
 	float PosCorrectionTimeBase = PhysicsReplicationCVars::PredictiveInterpolationCVars::PosCorrectionTimeBase;
-	float GetPosCorrectionTimeBase() { return bOverridePosCorrectionTimeBase ? PosCorrectionTimeBase : PhysicsReplicationCVars::PredictiveInterpolationCVars::PosCorrectionTimeBase; }
+	float GetPosCorrectionTimeBase() const { return bOverridePosCorrectionTimeBase ? PosCorrectionTimeBase : PhysicsReplicationCVars::PredictiveInterpolationCVars::PosCorrectionTimeBase; }
 
 	UPROPERTY(config, EditDefaultsOnly, Category = "Overrides", Meta = (InlineEditConditionToggle))
 	uint32 bOverridePosCorrectionTimeMin : 1;
 	// Overrides CVar: np2.PredictiveInterpolation.PosCorrectionTimeMin -- Min time to correct positional offset over. DeltaSeconds is added on top of this.
 	UPROPERTY(config, EditDefaultsOnly, Category = "Overrides", Meta = (EditCondition = "bOverridePosCorrectionTimeMin"))
 	float PosCorrectionTimeMin = PhysicsReplicationCVars::PredictiveInterpolationCVars::PosCorrectionTimeMin;
-	float GetPosCorrectionTimeMin() { return bOverridePosCorrectionTimeMin ? PosCorrectionTimeMin : PhysicsReplicationCVars::PredictiveInterpolationCVars::PosCorrectionTimeMin; }
+	float GetPosCorrectionTimeMin() const { return bOverridePosCorrectionTimeMin ? PosCorrectionTimeMin : PhysicsReplicationCVars::PredictiveInterpolationCVars::PosCorrectionTimeMin; }
 
 	UPROPERTY(config, EditDefaultsOnly, Category = "Overrides", Meta = (InlineEditConditionToggle))
 	uint32 bOverridePosCorrectionTimeMultiplier : 1;
 	// Overrides CVar: np2.PredictiveInterpolation.PosCorrectionTimeMultiplier -- Multiplier to adjust how much of RoundTripTime to add to positional offset correction.
 	UPROPERTY(config, EditDefaultsOnly, Category = "Overrides", Meta = (EditCondition = "bOverridePosCorrectionTimeMultiplier"))
 	float PosCorrectionTimeMultiplier = PhysicsReplicationCVars::PredictiveInterpolationCVars::PosCorrectionTimeMultiplier;
-	float GetPosCorrectionTimeMultiplier() { return bOverridePosCorrectionTimeMultiplier ? PosCorrectionTimeMultiplier : PhysicsReplicationCVars::PredictiveInterpolationCVars::PosCorrectionTimeMultiplier; }
+	float GetPosCorrectionTimeMultiplier() const { return bOverridePosCorrectionTimeMultiplier ? PosCorrectionTimeMultiplier : PhysicsReplicationCVars::PredictiveInterpolationCVars::PosCorrectionTimeMultiplier; }
 
 	UPROPERTY(config, EditDefaultsOnly, Category = "Overrides", Meta = (InlineEditConditionToggle))
 	uint32 bOverrideRotCorrectionTimeBase : 1;
 	// Overrides CVar: np2.PredictiveInterpolation.RotCorrectionTimeBase -- Base time to correct rotational offset over. RoundTripTime * RotCorrectionTimeMultiplier is added on top of this.
 	UPROPERTY(config, EditDefaultsOnly, Category = "Overrides", Meta = (EditCondition = "bOverrideRotCorrectionTimeBase"))
 	float RotCorrectionTimeBase = PhysicsReplicationCVars::PredictiveInterpolationCVars::RotCorrectionTimeBase;
-	float GetRotCorrectionTimeBase() { return bOverrideRotCorrectionTimeBase ? RotCorrectionTimeBase : PhysicsReplicationCVars::PredictiveInterpolationCVars::RotCorrectionTimeBase; }
+	float GetRotCorrectionTimeBase() const { return bOverrideRotCorrectionTimeBase ? RotCorrectionTimeBase : PhysicsReplicationCVars::PredictiveInterpolationCVars::RotCorrectionTimeBase; }
 
 	UPROPERTY(config, EditDefaultsOnly, Category = "Overrides", Meta = (InlineEditConditionToggle))
 	uint32 bOverrideRotCorrectionTimeMin : 1;
 	// Overrides CVar: np2.PredictiveInterpolation.RotCorrectionTimeMin -- Min time to correct rotational offset over. DeltaSeconds is added on top of this.
 	UPROPERTY(config, EditDefaultsOnly, Category = "Overrides", Meta = (EditCondition = "bOverrideRotCorrectionTimeMin"))
 	float RotCorrectionTimeMin = PhysicsReplicationCVars::PredictiveInterpolationCVars::RotCorrectionTimeMin;
-	float GetRotCorrectionTimeMin() { return bOverrideRotCorrectionTimeMin ? RotCorrectionTimeMin : PhysicsReplicationCVars::PredictiveInterpolationCVars::RotCorrectionTimeMin; }
+	float GetRotCorrectionTimeMin() const { return bOverrideRotCorrectionTimeMin ? RotCorrectionTimeMin : PhysicsReplicationCVars::PredictiveInterpolationCVars::RotCorrectionTimeMin; }
 
 	UPROPERTY(config, EditDefaultsOnly, Category = "Overrides", Meta = (InlineEditConditionToggle))
 	uint32 bOverrideRotCorrectionTimeMultiplier : 1;
 	// Overrides CVar: np2.PredictiveInterpolation.RotCorrectionTimeMultiplier -- Multiplier to adjust how much of RoundTripTime to add to rotational offset correction.
 	UPROPERTY(config, EditDefaultsOnly, Category = "Overrides", Meta = (EditCondition = "bOverrideRotCorrectionTimeMultiplier"))
 	float RotCorrectionTimeMultiplier = PhysicsReplicationCVars::PredictiveInterpolationCVars::RotCorrectionTimeMultiplier;
-	float GetRotCorrectionTimeMultiplier() { return bOverrideRotCorrectionTimeMultiplier ? RotCorrectionTimeMultiplier : PhysicsReplicationCVars::PredictiveInterpolationCVars::RotCorrectionTimeMultiplier; }
+	float GetRotCorrectionTimeMultiplier() const { return bOverrideRotCorrectionTimeMultiplier ? RotCorrectionTimeMultiplier : PhysicsReplicationCVars::PredictiveInterpolationCVars::RotCorrectionTimeMultiplier; }
 
 	UPROPERTY(config, EditDefaultsOnly, Category = "Overrides", Meta = (InlineEditConditionToggle))
 	uint32 bOverridePosInterpolationTimeMultiplier : 1;
 	// Overrides CVar: np2.PredictiveInterpolation.InterpolationTimeMultiplier -- Multiplier to adjust the interpolation time which is based on the sendrate of state data from the server.
 	UPROPERTY(config, EditDefaultsOnly, Category = "Overrides", Meta = (EditCondition = "bOverridePosInterpolationTimeMultiplier"))
 	float PosInterpolationTimeMultiplier = PhysicsReplicationCVars::PredictiveInterpolationCVars::PosInterpolationTimeMultiplier;
-	float GetPosInterpolationTimeMultiplier() { return bOverridePosInterpolationTimeMultiplier ? PosInterpolationTimeMultiplier : PhysicsReplicationCVars::PredictiveInterpolationCVars::PosInterpolationTimeMultiplier; }
+	float GetPosInterpolationTimeMultiplier() const { return bOverridePosInterpolationTimeMultiplier ? PosInterpolationTimeMultiplier : PhysicsReplicationCVars::PredictiveInterpolationCVars::PosInterpolationTimeMultiplier; }
 
 	UPROPERTY(config, EditDefaultsOnly, Category = "Overrides", Meta = (InlineEditConditionToggle))
 	uint32 bOverrideRotInterpolationTimeMultiplier : 1;
 	// Overrides CVar: np2.PredictiveInterpolation.RotInterpolationTimeMultiplier -- Multiplier to adjust the rotational interpolation time which is based on the sendrate of state data from the server.
 	UPROPERTY(config, EditDefaultsOnly, Category = "Overrides", Meta = (EditCondition = "bOverrideRotInterpolationTimeMultiplier"))
 	float RotInterpolationTimeMultiplier = PhysicsReplicationCVars::PredictiveInterpolationCVars::RotInterpolationTimeMultiplier;
-	float GetRotInterpolationTimeMultiplier() { return bOverrideRotInterpolationTimeMultiplier ? RotInterpolationTimeMultiplier : PhysicsReplicationCVars::PredictiveInterpolationCVars::RotInterpolationTimeMultiplier; }
+	float GetRotInterpolationTimeMultiplier() const { return bOverrideRotInterpolationTimeMultiplier ? RotInterpolationTimeMultiplier : PhysicsReplicationCVars::PredictiveInterpolationCVars::RotInterpolationTimeMultiplier; }
 
 	UPROPERTY(config, EditDefaultsOnly, Category = "Overrides", Meta = (InlineEditConditionToggle))
 	uint32 bOverrideSoftSnapPosStrength : 1;
 	// Overrides CVar: np2.PredictiveInterpolation.SoftSnapPosStrength -- Value in percent between 0.0 - 1.0 representing how much to softsnap each tick of the remaining positional distance.
 	UPROPERTY(config, EditDefaultsOnly, Category = "Overrides", Meta = (EditCondition = "bOverrideSoftSnapPosStrength"))
 	float SoftSnapPosStrength = PhysicsReplicationCVars::PredictiveInterpolationCVars::SoftSnapPosStrength;
-	float GetSoftSnapPosStrength() { return bOverrideSoftSnapPosStrength ? SoftSnapPosStrength : PhysicsReplicationCVars::PredictiveInterpolationCVars::SoftSnapPosStrength; }
+	float GetSoftSnapPosStrength() const { return bOverrideSoftSnapPosStrength ? SoftSnapPosStrength : PhysicsReplicationCVars::PredictiveInterpolationCVars::SoftSnapPosStrength; }
 
 	UPROPERTY(config, EditDefaultsOnly, Category = "Overrides", Meta = (InlineEditConditionToggle))
 	uint32 bOverrideSoftSnapRotStrength : 1;
 	// Overrides CVar: np2.PredictiveInterpolation.SoftSnapRotStrength -- Value in percent between 0.0 - 1.0 representing how much to softsnap each tick of the remaining rotational distance.
 	UPROPERTY(config, EditDefaultsOnly, Category = "Overrides", Meta = (EditCondition = "bOverrideSoftSnapRotStrength"))
 	float SoftSnapRotStrength = PhysicsReplicationCVars::PredictiveInterpolationCVars::SoftSnapRotStrength;
-	float GetSoftSnapRotStrength() { return bOverrideSoftSnapRotStrength ? SoftSnapRotStrength : PhysicsReplicationCVars::PredictiveInterpolationCVars::SoftSnapRotStrength; }
+	float GetSoftSnapRotStrength() const { return bOverrideSoftSnapRotStrength ? SoftSnapRotStrength : PhysicsReplicationCVars::PredictiveInterpolationCVars::SoftSnapRotStrength; }
 
 	UPROPERTY(config, EditDefaultsOnly, Category = "Overrides", Meta = (InlineEditConditionToggle))
 	uint32 bOverrideSoftSnapToSource : 1;
 	// Overrides CVar: np2.PredictiveInterpolation.SoftSnapToSource -- If true, softsnap will be performed towards the source state of the current target instead of the predicted state of the current target.
 	UPROPERTY(config, EditDefaultsOnly, Category = "Overrides", Meta = (EditCondition = "bOverrideSoftSnapToSource"))
 	bool bSoftSnapToSource = PhysicsReplicationCVars::PredictiveInterpolationCVars::bSoftSnapToSource;
-	bool GetSoftSnapToSource() { return bOverrideSoftSnapToSource ? bSoftSnapToSource : PhysicsReplicationCVars::PredictiveInterpolationCVars::bSoftSnapToSource; }
+	bool GetSoftSnapToSource() const { return bOverrideSoftSnapToSource ? bSoftSnapToSource : PhysicsReplicationCVars::PredictiveInterpolationCVars::bSoftSnapToSource; }
 
 	UPROPERTY(config, EditDefaultsOnly, Category = "Overrides", Meta = (InlineEditConditionToggle))
 	uint32 bOverrideDisableSoftSnap : 1;
 	// Overrides CVar: np2.PredictiveInterpolation.DisableSoftSnap -- When true, predictive interpolation will not use softsnap to correct the replication with when velocity fails. Hardsnap will still eventually kick in if replication can't reach the target.
 	UPROPERTY(config, EditDefaultsOnly, Category = "Overrides", Meta = (EditCondition = "bOverrideDisableSoftSnap"))
 	bool bDisableSoftSnap = PhysicsReplicationCVars::PredictiveInterpolationCVars::bDisableSoftSnap;
-	bool GetDisableSoftSnap() { return bOverrideDisableSoftSnap ? bDisableSoftSnap : PhysicsReplicationCVars::PredictiveInterpolationCVars::bDisableSoftSnap; }
+	bool GetDisableSoftSnap() const { return bOverrideDisableSoftSnap ? bDisableSoftSnap : PhysicsReplicationCVars::PredictiveInterpolationCVars::bDisableSoftSnap; }
 
 	UPROPERTY(config, EditDefaultsOnly, Category = "Overrides", Meta = (InlineEditConditionToggle))
 	uint32 bOverrideSkipVelocityRepOnPosEarlyOut : 1;
 	// Overrides CVar: np2.PredictiveInterpolation.SkipVelocityRepOnPosEarlyOut -- If true, don't run linear velocity replication if position can early out but angular can't early out.
 	UPROPERTY(config, EditDefaultsOnly, Category = "Overrides", Meta = (EditCondition = "bOverrideSkipVelocityRepOnPosEarlyOut"))
 	bool bSkipVelocityRepOnPosEarlyOut = PhysicsReplicationCVars::PredictiveInterpolationCVars::bSkipVelocityRepOnPosEarlyOut;
-	bool GetSkipVelocityRepOnPosEarlyOut() { return bOverrideSkipVelocityRepOnPosEarlyOut ? bSkipVelocityRepOnPosEarlyOut : PhysicsReplicationCVars::PredictiveInterpolationCVars::bSkipVelocityRepOnPosEarlyOut; }
+	bool GetSkipVelocityRepOnPosEarlyOut() const { return bOverrideSkipVelocityRepOnPosEarlyOut ? bSkipVelocityRepOnPosEarlyOut : PhysicsReplicationCVars::PredictiveInterpolationCVars::bSkipVelocityRepOnPosEarlyOut; }
 
 	UPROPERTY(config, EditDefaultsOnly, Category = "Overrides", Meta = (InlineEditConditionToggle))
 	uint32 bOverridePostResimWaitForUpdate : 1;
 	// Overrides CVar: np2.PredictiveInterpolation.PostResimWaitForUpdate -- After a resimulation, wait for replicated states that correspond to post-resim state before processing replication again.
 	UPROPERTY(config, EditDefaultsOnly, Category = "Overrides", Meta = (EditCondition = "bOverridePostResimWaitForUpdate"))
 	bool bPostResimWaitForUpdate = PhysicsReplicationCVars::PredictiveInterpolationCVars::bPostResimWaitForUpdate;
-	bool GetPostResimWaitForUpdate() { return bOverridePostResimWaitForUpdate ? bPostResimWaitForUpdate : PhysicsReplicationCVars::PredictiveInterpolationCVars::bPostResimWaitForUpdate; }
+	bool GetPostResimWaitForUpdate() const { return bOverridePostResimWaitForUpdate ? bPostResimWaitForUpdate : PhysicsReplicationCVars::PredictiveInterpolationCVars::bPostResimWaitForUpdate; }
 
 	UPROPERTY(config, EditDefaultsOnly, Category = "Overrides", Meta = (InlineEditConditionToggle))
 	uint32 bOverrideCorrectConnectedBodies : 1;
 	// Overrides CVar: np2.PredictiveInterpolation.CorrectConnectedBodies -- When true, transform corrections will also apply to any connected physics object.
 	UPROPERTY(config, EditDefaultsOnly, Category = "Overrides", Meta = (EditCondition = "bOverrideCorrectConnectedBodies"))
 	bool bCorrectConnectedBodies = PhysicsReplicationCVars::PredictiveInterpolationCVars::bCorrectConnectedBodies;
-	bool GetCorrectConnectedBodies() { return bOverrideCorrectConnectedBodies ? bCorrectConnectedBodies : PhysicsReplicationCVars::PredictiveInterpolationCVars::bCorrectConnectedBodies; }
+	bool GetCorrectConnectedBodies() const { return bOverrideCorrectConnectedBodies ? bCorrectConnectedBodies : PhysicsReplicationCVars::PredictiveInterpolationCVars::bCorrectConnectedBodies; }
 
 	UPROPERTY(config, EditDefaultsOnly, Category = "Overrides", Meta = (InlineEditConditionToggle))
 	uint32 bOverrideCorrectConnectedBodiesFriction : 1;
 	// Overrides CVar: np2.PredictiveInterpolation.CorrectConnectedBodiesFriction -- When true, transform correction on any connected physics object will also recalculate their friction.
 	UPROPERTY(config, EditDefaultsOnly, Category = "Overrides", Meta = (EditCondition = "bOverrideCorrectConnectedBodiesFriction"))
 	bool bCorrectConnectedBodiesFriction = PhysicsReplicationCVars::PredictiveInterpolationCVars::bCorrectConnectedBodiesFriction;
-	bool GetCorrectConnectedBodiesFriction() { return bOverrideCorrectConnectedBodiesFriction ? bCorrectConnectedBodiesFriction : PhysicsReplicationCVars::PredictiveInterpolationCVars::bCorrectConnectedBodiesFriction; }
+	bool GetCorrectConnectedBodiesFriction() const { return bOverrideCorrectConnectedBodiesFriction ? bCorrectConnectedBodiesFriction : PhysicsReplicationCVars::PredictiveInterpolationCVars::bCorrectConnectedBodiesFriction; }
 
 };
 
@@ -260,7 +261,7 @@ struct FNetworkPhysicsSettingsResimulationErrorCorrection
 	GENERATED_BODY()
 
 	/** Apply these settings to Physics Object */
-	void ApplySettings_External(Chaos::FPhysicsObjectHandle PhysicsObject);
+	void ApplySettings_External(Chaos::FPhysicsObjectHandle PhysicsObject) const;
 
 	/** Enable override for post-resimulation error correction settings during render interpolation
 	* NOTE: This currently does not work if the experimental p.RenderInterp.ErrorVelocityCorrection CVar is set to true (false by default) */
@@ -305,56 +306,56 @@ struct FNetworkPhysicsSettingsResimulation
 	uint32 ResimulationErrorThreshold_DEPRECATED;
 
 	UE_DEPRECATED(5.5, "GetResimulationErrorThreshold has been renamed, please use GetResimulationErrorPositionThreshold.")
-	uint32 GetResimulationErrorThreshold(uint32 DefaultValue) { return (uint32)FMath::CeilToInt(GetResimulationErrorPositionThreshold(DefaultValue)); }
+	uint32 GetResimulationErrorThreshold(uint32 DefaultValue) const { return (uint32)FMath::CeilToInt(GetResimulationErrorPositionThreshold(static_cast<float>(DefaultValue))); }
 
 	UPROPERTY(config, EditDefaultsOnly, Category = "Overrides", Meta = (InlineEditConditionToggle))
 	uint32 bOverrideResimulationErrorPositionThreshold : 1;
 	// Overrides Project Settings -> Physics -> Replication -> Physics Prediction -> Resimulation Error Position Threshold -- Distance that the object is allowed to desync from the server before triggering a resimulation, within this threshold runtime correction can be performed if RuntimeCorrectionEnabled is true.
 	UPROPERTY(config, EditDefaultsOnly, Category = "Overrides", Meta = (EditCondition = "bOverrideResimulationErrorPositionThreshold"))
 	float ResimulationErrorPositionThreshold = 10.0f;
-	float GetResimulationErrorPositionThreshold(uint32 DefaultValue) { return bOverrideResimulationErrorPositionThreshold ? ResimulationErrorPositionThreshold : DefaultValue; }
+	float GetResimulationErrorPositionThreshold(float DefaultValue) const { return bOverrideResimulationErrorPositionThreshold ? ResimulationErrorPositionThreshold : DefaultValue; }
 
 	UPROPERTY(config, EditDefaultsOnly, Category = "Overrides", Meta = (InlineEditConditionToggle))
 	uint32 bOverrideResimulationErrorRotationThreshold : 1;
 	// Overrides Project Settings -> Physics -> Replication -> Physics Prediction -> Resimulation Error Rotation Threshold -- Rotation difference in degrees that the object is allowed to desync from the server before triggering a resimulation, within this threshold runtime correction can be performed if RuntimeCorrectionEnabled is true.
 	UPROPERTY(config, EditDefaultsOnly, Category = "Overrides", Meta = (EditCondition = "bOverrideResimulationErrorRotationThreshold"))
 	float ResimulationErrorRotationThreshold = 4.0f;
-	float GetResimulationErrorRotationThreshold(uint32 DefaultValue) { return bOverrideResimulationErrorRotationThreshold ? ResimulationErrorRotationThreshold : DefaultValue; }
+	float GetResimulationErrorRotationThreshold(float DefaultValue) const { return bOverrideResimulationErrorRotationThreshold ? ResimulationErrorRotationThreshold : DefaultValue; }
 
 	UPROPERTY(config, EditDefaultsOnly, Category = "Overrides", Meta = (InlineEditConditionToggle))
 	uint32 bOverrideResimulationErrorLinearVelocityThreshold : 1;
 	// Overrides Project Settings -> Physics -> Replication -> Physics Prediction -> Resimulation Error Linear Velocity Threshold -- Velocity difference in centimeters / second that the object is allowed to desync from the server before triggering a resimulation, within this threshold runtime correction can be performed if RuntimeCorrectionEnabled is true.
 	UPROPERTY(config, EditDefaultsOnly, Category = "Overrides", Meta = (EditCondition = "bOverrideResimulationErrorLinearVelocityThreshold"))
 	float ResimulationErrorLinearVelocityThreshold = 5.0f;
-	float GetResimulationErrorLinearVelocityThreshold(uint32 DefaultValue) { return bOverrideResimulationErrorLinearVelocityThreshold ? ResimulationErrorLinearVelocityThreshold : DefaultValue; }
+	float GetResimulationErrorLinearVelocityThreshold(float DefaultValue) const { return bOverrideResimulationErrorLinearVelocityThreshold ? ResimulationErrorLinearVelocityThreshold : DefaultValue; }
 
 	UPROPERTY(config, EditDefaultsOnly, Category = "Overrides", Meta = (InlineEditConditionToggle))
 	uint32 bOverrideResimulationErrorAngularVelocityThreshold : 1;
 	// Overrides Project Settings -> Physics -> Replication -> Physics Prediction -> Resimulation Error Angular Velocity Threshold -- Degrees / second that the object is allowed to desync from the server before triggering a resimulation, within this threshold runtime correction can be performed if RuntimeCorrectionEnabled is true.
 	UPROPERTY(config, EditDefaultsOnly, Category = "Overrides", Meta = (EditCondition = "bOverrideResimulationErrorAngularVelocityThreshold"))
 	float ResimulationErrorAngularVelocityThreshold = 2.0f;
-	float GetResimulationErrorAngularVelocityThreshold(uint32 DefaultValue) { return bOverrideResimulationErrorAngularVelocityThreshold ? ResimulationErrorAngularVelocityThreshold : DefaultValue; }
+	float GetResimulationErrorAngularVelocityThreshold(float DefaultValue) const { return bOverrideResimulationErrorAngularVelocityThreshold ? ResimulationErrorAngularVelocityThreshold : DefaultValue; }
 
 	UPROPERTY(config, EditDefaultsOnly, Category = "Overrides", Meta = (InlineEditConditionToggle))
 	uint32 bOverrideRuntimeCorrectionEnabled : 1;
 	// Overrides CVar: np2.Resim.RuntimeCorrectionEnabled -- Apply positional and rotational runtime corrections while within resim trigger distance.
 	UPROPERTY(config, EditDefaultsOnly, Category = "Overrides", Meta = (EditCondition = "bOverrideRuntimeCorrectionEnabled"))
 	bool bRuntimeCorrectionEnabled = PhysicsReplicationCVars::ResimulationCVars::bRuntimeCorrectionEnabled;
-	bool GetRuntimeCorrectionEnabled() { return bOverrideRuntimeCorrectionEnabled ? bRuntimeCorrectionEnabled : PhysicsReplicationCVars::ResimulationCVars::bRuntimeCorrectionEnabled; }
+	bool GetRuntimeCorrectionEnabled() const { return bOverrideRuntimeCorrectionEnabled ? bRuntimeCorrectionEnabled : PhysicsReplicationCVars::ResimulationCVars::bRuntimeCorrectionEnabled; }
 
 	UPROPERTY(config, EditDefaultsOnly, Category = "Overrides", Meta = (InlineEditConditionToggle))
 	uint32 bOverrideRuntimeVelocityCorrection : 1;
 	// Overrides CVar: np2.Resim.RuntimeVelocityCorrection -- Apply linear and angular velocity corrections in runtime while within resim trigger distance. Used if RuntimeCorrectionEnabled is true.
 	UPROPERTY(config, EditDefaultsOnly, Category = "Overrides", Meta = (EditCondition = "bOverrideRuntimeVelocityCorrection"))
 	bool bRuntimeVelocityCorrection = PhysicsReplicationCVars::ResimulationCVars::bRuntimeVelocityCorrection;
-	bool GetRuntimeVelocityCorrectionEnabled() { return bOverrideRuntimeVelocityCorrection ? bRuntimeVelocityCorrection : PhysicsReplicationCVars::ResimulationCVars::bRuntimeVelocityCorrection; }
+	bool GetRuntimeVelocityCorrectionEnabled() const { return bOverrideRuntimeVelocityCorrection ? bRuntimeVelocityCorrection : PhysicsReplicationCVars::ResimulationCVars::bRuntimeVelocityCorrection; }
 
 	UPROPERTY(config, EditDefaultsOnly, Category = "Overrides", Meta = (InlineEditConditionToggle))
 	uint32 bOverrideRuntimeCorrectConnectedBodies : 1;
 	// Overrides CVar: np2.Resim.RuntimeCorrectConnectedBodies -- If true runtime position and rotation correction will also shift transform of any connected physics objects. Used if RuntimeCorrectionEnabled is true.
 	UPROPERTY(config, EditDefaultsOnly, Category = "Overrides", Meta = (EditCondition = "bOverrideRuntimeCorrectConnectedBodies"))
 	bool bRuntimeCorrectConnectedBodies = PhysicsReplicationCVars::ResimulationCVars::bRuntimeCorrectConnectedBodies;
-	bool GetRuntimeCorrectConnectedBodies() { return bOverrideRuntimeCorrectConnectedBodies ? bRuntimeCorrectConnectedBodies : PhysicsReplicationCVars::ResimulationCVars::bRuntimeCorrectConnectedBodies; }
+	bool GetRuntimeCorrectConnectedBodies() const { return bOverrideRuntimeCorrectConnectedBodies ? bRuntimeCorrectConnectedBodies : PhysicsReplicationCVars::ResimulationCVars::bRuntimeCorrectConnectedBodies; }
 
 	UPROPERTY(config, EditDefaultsOnly, Category = "Overrides", Meta = (InlineEditConditionToggle))
 	uint32 bOverridePosStabilityMultiplier : 1;
@@ -368,21 +369,21 @@ struct FNetworkPhysicsSettingsResimulation
 	// Overrides CVar: np2.Resim.RotStabilityMultiplier -- Recommended range between 0.0-1.0. Lower value means more stable rotational corrections.
 	UPROPERTY(config, EditDefaultsOnly, Category = "Overrides", Meta = (EditCondition = "bOverrideRotStabilityMultiplier"))
 	float RotStabilityMultiplier = PhysicsReplicationCVars::ResimulationCVars::RotStabilityMultiplier;
-	float GetRotStabilityMultiplier() { return bOverrideRotStabilityMultiplier ? RotStabilityMultiplier : PhysicsReplicationCVars::ResimulationCVars::RotStabilityMultiplier; }
+	float GetRotStabilityMultiplier() const { return bOverrideRotStabilityMultiplier ? RotStabilityMultiplier : PhysicsReplicationCVars::ResimulationCVars::RotStabilityMultiplier; }
 
 	UPROPERTY(config, EditDefaultsOnly, Category = "Overrides", Meta = (InlineEditConditionToggle))
 	uint32 bOverrideVelStabilityMultiplier : 1;
 	// Overrides CVar: np2.Resim.VelStabilityMultiplier -- Recommended range between 0.0-1.0. Lower value means more stable linear velocity corrections.
 	UPROPERTY(config, EditDefaultsOnly, Category = "Overrides", Meta = (EditCondition = "bOverrideVelStabilityMultiplier"))
 	float VelStabilityMultiplier = PhysicsReplicationCVars::ResimulationCVars::VelStabilityMultiplier;
-	float GetVelStabilityMultiplier() { return bOverrideVelStabilityMultiplier ? VelStabilityMultiplier : PhysicsReplicationCVars::ResimulationCVars::VelStabilityMultiplier; }
+	float GetVelStabilityMultiplier() const { return bOverrideVelStabilityMultiplier ? VelStabilityMultiplier : PhysicsReplicationCVars::ResimulationCVars::VelStabilityMultiplier; }
 
 	UPROPERTY(config, EditDefaultsOnly, Category = "Overrides", Meta = (InlineEditConditionToggle))
 	uint32 bOverrideAngVelStabilityMultiplier : 1;
 	// Overrides CVar: np2.Resim.AngVelStabilityMultiplier -- Recommended range between 0.0-1.0. Lower value means more stable angular velocity corrections.
 	UPROPERTY(config, EditDefaultsOnly, Category = "Overrides", Meta = (EditCondition = "bOverrideAngVelStabilityMultiplier"))
 	float AngVelStabilityMultiplier = PhysicsReplicationCVars::ResimulationCVars::AngVelStabilityMultiplier;
-	float GetAngVelStabilityMultiplier() { return bOverrideAngVelStabilityMultiplier ? AngVelStabilityMultiplier : PhysicsReplicationCVars::ResimulationCVars::AngVelStabilityMultiplier; }
+	float GetAngVelStabilityMultiplier() const { return bOverrideAngVelStabilityMultiplier ? AngVelStabilityMultiplier : PhysicsReplicationCVars::ResimulationCVars::AngVelStabilityMultiplier; }
 
 	UPROPERTY(config, EditDefaultsOnly, Category = "Overrides", meta = (DisplayName = "Post-Resimulation Error Correction Render Interpolation Settings"))
 	FNetworkPhysicsSettingsResimulationErrorCorrection ResimulationErrorCorrectionSettings;
@@ -519,6 +520,71 @@ struct FNetworkPhysicsSettingsRenderInterpolation
 };
 */
 
+USTRUCT()
+struct FNetworkPhysicsSettingsData
+{
+	GENERATED_BODY()
+
+	UPROPERTY(EditDefaultsOnly, Category = "Networked Physics Settings")
+	FNetworkPhysicsSettings GeneralSettings;
+
+	UPROPERTY(EditDefaultsOnly, Category = "Networked Physics Settings")
+	FNetworkPhysicsSettingsDefaultReplication DefaultReplicationSettings;
+
+	UPROPERTY(EditDefaultsOnly, Category = "Networked Physics Settings")
+	FNetworkPhysicsSettingsPredictiveInterpolation PredictiveInterpolationSettings;
+
+	UPROPERTY(EditDefaultsOnly, Category = "Networked Physics Settings")
+	FNetworkPhysicsSettingsResimulation ResimulationSettings;
+
+	UPROPERTY(EditDefaultsOnly, Category = "Networked Physics Settings")
+	FNetworkPhysicsSettingsNetworkPhysicsComponent NetworkPhysicsComponentSettings;
+	/*
+	UPROPERTY(EditDefaultsOnly, Category = "Networked Physics Settings")
+	FNetworkPhysicsSettingsRewindData RewindSettings;
+
+	UPROPERTY(EditDefaultsOnly, Category = "Networked Physics Settings")
+	FNetworkPhysicsSettingsRenderInterpolation RenderInterpolationSettings;
+	*/
+};
+
+UCLASS(BlueprintType)
+class UNetworkPhysicsSettingsDataAsset : public UDataAsset
+{
+	GENERATED_BODY()
+
+public:
+	
+	void InitializeInternalSettings()
+	{
+		if (bInitialized)
+		{
+			return;
+		}
+
+		Settings_Internal = MakeShared<FNetworkPhysicsSettingsData>(Settings);
+
+		bInitialized = true;
+	}
+
+	/** Get the settings, on the game thread */
+	const FNetworkPhysicsSettingsData& GetSettings_External() const { return Settings; }
+	
+	/** Get the settings, on the physics thread */
+	TWeakPtr<const FNetworkPhysicsSettingsData> GetSettings_Internal() const { return Settings_Internal.ToWeakPtr(); }
+
+private:
+	bool bInitialized = false;
+
+	/** Game Thread Settings*/
+	// Network Physics Settings
+	UPROPERTY(EditDefaultsOnly, Category = "Networked Physics Settings")
+	FNetworkPhysicsSettingsData Settings;
+
+	/** Physics Thread Settings */
+	TSharedPtr<FNetworkPhysicsSettingsData> Settings_Internal = nullptr;
+};
+
 /** Settings Component for network replicated physics actors
 * Overrides default settings, CVar settings and project settings. */
 UCLASS(BlueprintType, MinimalAPI, meta = (BlueprintSpawnableComponent))
@@ -527,7 +593,10 @@ class UNetworkPhysicsSettingsComponent : public UActorComponent
 	GENERATED_BODY()
 
 public:
+	PRAGMA_DISABLE_DEPRECATION_WARNINGS
 	ENGINE_API UNetworkPhysicsSettingsComponent();
+	~UNetworkPhysicsSettingsComponent() {};
+	PRAGMA_ENABLE_DEPRECATION_WARNINGS
 
 	virtual void OnRegister() override;
 	virtual void OnUnregister() override;
@@ -547,32 +616,41 @@ private:
 	void RegisterInPhysicsReplicationLOD();
 
 public:
-	UPROPERTY(EditDefaultsOnly, Category = "Networked Physics Settings")
+	/** Get the settings, on the game thread */
+	const FNetworkPhysicsSettingsData& GetSettings() const { return SettingsDataAsset ? SettingsDataAsset->GetSettings_External() : SettingsNetworkPhysicsData_Default; }
+
+	/** Get the settings, on the physics thread */
+	TWeakPtr<const FNetworkPhysicsSettingsData> GetSettings_Internal() const { return SettingsDataAsset ? SettingsDataAsset->GetSettings_Internal() : nullptr; }
+
+	UPROPERTY(EditDefaultsOnly, BlueprintReadOnly, Category = "Networked Physics Settings")
+	TObjectPtr<UNetworkPhysicsSettingsDataAsset> SettingsDataAsset;
+
+// Deprecated properties
+//#if WITH_EDITORONLY_DATA // TODO UE 5.7 enable the if-statement to hide the deprecated properties
+	UE_DEPRECATED(5.7, "FNetworkPhysicsSettings GeneralSettings has been deprecated, create a DataAsset of UNetworkPhysicsSettingsDataAsset type and reference that in the UNetworkPhysicsSettingsComponent instead.")
+	UPROPERTY(meta = (DeprecatedProperty, DeprecationMessage = "Deprecated, create a DataAsset of UNetworkPhysicsSettingsDataAsset type and reference that in the UNetworkPhysicsSettingsComponent instead"))
 	FNetworkPhysicsSettings GeneralSettings;
 	
-	UPROPERTY(EditDefaultsOnly, Category = "Networked Physics Settings")
+	UE_DEPRECATED(5.7, "FNetworkPhysicsSettingsDefaultReplication DefaultReplicationSettings has been deprecated, create a DataAsset of UNetworkPhysicsSettingsDataAsset type and reference that in the UNetworkPhysicsSettingsComponent instead.")
+	UPROPERTY(meta = (DeprecatedProperty, DeprecationMessage = "Deprecated, create a DataAsset of UNetworkPhysicsSettingsDataAsset type and reference that in the UNetworkPhysicsSettingsComponent instead"))
 	FNetworkPhysicsSettingsDefaultReplication DefaultReplicationSettings;
 
-	UPROPERTY(EditDefaultsOnly, Category = "Networked Physics Settings")
+	UE_DEPRECATED(5.7, "FNetworkPhysicsSettingsPredictiveInterpolation PredictiveInterpolationSettings has been deprecated, create a DataAsset of UNetworkPhysicsSettingsDataAsset type and reference that in the UNetworkPhysicsSettingsComponent instead.")
+	UPROPERTY(meta = (DeprecatedProperty, DeprecationMessage = "Deprecated, create a DataAsset of UNetworkPhysicsSettingsDataAsset type and reference that in the UNetworkPhysicsSettingsComponent instead"))
 	FNetworkPhysicsSettingsPredictiveInterpolation PredictiveInterpolationSettings;
 	
-	UPROPERTY(EditDefaultsOnly, Category = "Networked Physics Settings")
+	UE_DEPRECATED(5.7, "FNetworkPhysicsSettingsResimulation ResimulationSettings has been deprecated, create a DataAsset of UNetworkPhysicsSettingsDataAsset type and reference that in the UNetworkPhysicsSettingsComponent instead.")
+	UPROPERTY(meta = (DeprecatedProperty, DeprecationMessage = "Deprecated, create a DataAsset of UNetworkPhysicsSettingsDataAsset type and reference that in the UNetworkPhysicsSettingsComponent instead"))
 	FNetworkPhysicsSettingsResimulation ResimulationSettings;
 	
-	UPROPERTY(EditDefaultsOnly, Category = "Networked Physics Settings")
+	UE_DEPRECATED(5.7, "FNetworkPhysicsSettingsNetworkPhysicsComponent NetworkPhysicsComponentSettings has been deprecated, create a DataAsset of UNetworkPhysicsSettingsDataAsset type and reference that in the UNetworkPhysicsSettingsComponent instead.")
+	UPROPERTY(meta = (DeprecatedProperty, DeprecationMessage = "Deprecated, create a DataAsset of UNetworkPhysicsSettingsDataAsset type and reference that in the UNetworkPhysicsSettingsComponent instead"))
 	FNetworkPhysicsSettingsNetworkPhysicsComponent NetworkPhysicsComponentSettings;
-	
-	/*
-	UPROPERTY(EditDefaultsOnly, Category = "Networked Physics Settings")
-	FNetworkPhysicsSettingsRewindData RewindSettings;
-	
-	UPROPERTY(EditDefaultsOnly, Category = "Networked Physics Settings")
-	FNetworkPhysicsSettingsRenderInterpolation RenderInterpolationSettings;
-	*/
+//#endif
 
 private:
 	FNetworkPhysicsSettingsComponentAsync* NetworkPhysicsSettings_Internal;
-
+	static const FNetworkPhysicsSettingsData SettingsNetworkPhysicsData_Default;
 
 	/* --- Static API --- */
 private:
@@ -603,20 +681,28 @@ struct FNetworkPhysicsSettingsAsync
 struct FNetworkPhysicsSettingsAsyncInput : public Chaos::FSimCallbackInput
 {
 	TOptional<Chaos::FConstPhysicsObjectHandle> PhysicsObject;
+	TOptional<TWeakPtr<const FNetworkPhysicsSettingsData>> Settings_Internal;
+
+	UE_DEPRECATED(5.7, "Deprecated, use Settings_Internal instead")
 	TOptional<FNetworkPhysicsSettingsAsync> Settings;
 
 	void Reset()
 	{
 		PhysicsObject.Reset();
+		Settings_Internal.Reset();
+		PRAGMA_DISABLE_DEPRECATION_WARNINGS
 		Settings.Reset();
+		PRAGMA_ENABLE_DEPRECATION_WARNINGS
 	}
 };
 
 class FNetworkPhysicsSettingsComponentAsync : public Chaos::TSimCallbackObject<FNetworkPhysicsSettingsAsyncInput, Chaos::FSimCallbackNoOutput, Chaos::ESimCallbackOptions::Presimulate>
 {
 public:
+	PRAGMA_DISABLE_DEPRECATION_WARNINGS
 	FNetworkPhysicsSettingsComponentAsync();
 	~FNetworkPhysicsSettingsComponentAsync() {};
+	PRAGMA_ENABLE_DEPRECATION_WARNINGS
 
 	virtual void OnPreSimulate_Internal() override;
 
@@ -625,6 +711,9 @@ private:
 	void RegisterSettingsInPhysicsReplication();
 
 public:
+	TSharedPtr<const FNetworkPhysicsSettingsData> Settings_Internal;
+
+	UE_DEPRECATED(5.7, "Deprecated, use Settings_Internal instead")
 	FNetworkPhysicsSettingsAsync Settings;
 
 private:
diff --git a/Engine/Source/Runtime/Engine/Public/PhysicsReplication.h b/Engine/Source/Runtime/Engine/Public/PhysicsReplication.h
--- a/Engine/Source/Runtime/Engine/Public/PhysicsReplication.h
+++ b/Engine/Source/Runtime/Engine/Public/PhysicsReplication.h
@@ -207,15 +207,15 @@ class FPhysicsReplicationAsync : public IPhysicsReplicationAsync,
 	virtual bool ResimulationReplication(Chaos::FPBDRigidParticleHandle* Handle, FReplicatedPhysicsTargetAsync& Target, const float DeltaSeconds);
 
 public:
-	virtual void RegisterSettings(Chaos::FConstPhysicsObjectHandle PhysicsObject, FNetworkPhysicsSettingsAsync InSettings) override;
+	virtual void RegisterSettings(Chaos::FConstPhysicsObjectHandle PhysicsObject, TWeakPtr<const FNetworkPhysicsSettingsData> InSettings) override;
 
 private:
 	float LatencyOneWay;
 	FRigidBodyErrorCorrection ErrorCorrectionDefault;
-	FNetworkPhysicsSettingsAsync SettingsCurrent;
-	FNetworkPhysicsSettingsAsync SettingsDefault;
+	FNetworkPhysicsSettingsData SettingsCurrent;
+	FNetworkPhysicsSettingsData SettingsDefault;
 	TMap<Chaos::FConstPhysicsObjectHandle, FReplicatedPhysicsTargetAsync> ObjectToTarget;
-	TMap<Chaos::FConstPhysicsObjectHandle, FNetworkPhysicsSettingsAsync> ObjectToSettings;
+	TMap<Chaos::FConstPhysicsObjectHandle, TWeakPtr<const FNetworkPhysicsSettingsData>> ObjectToSettings;
 	TArray<const Chaos::Private::FPBDIsland*> ResimIslands;
 	TArray<const Chaos::FGeometryParticleHandle*> ResimIslandsParticles;
 	TArray<int32> ParticlesInResimIslands;
diff --git a/Engine/Source/Runtime/Engine/Public/PhysicsReplicationInterface.h b/Engine/Source/Runtime/Engine/Public/PhysicsReplicationInterface.h
--- a/Engine/Source/Runtime/Engine/Public/PhysicsReplicationInterface.h
+++ b/Engine/Source/Runtime/Engine/Public/PhysicsReplicationInterface.h
@@ -4,10 +4,11 @@
 
 #include "CoreMinimal.h"
 #include "Chaos/PhysicsObject.h"
+#include "Physics/NetworkPhysicsSettingsComponent.h" // Temporary until RegisterSettings with FNetworkPhysicsSettingsAsync param is deprecated
 
 class UPrimitiveComponent;
 struct FRigidBodyState;
-struct FNetworkPhysicsSettingsAsync;
+struct FNetworkPhysicsSettingsData;
 
 class IPhysicsReplication // Game Thread API
 {
@@ -26,5 +27,8 @@ class IPhysicsReplicationAsync // Physics Thread API
 public:
 	virtual ~IPhysicsReplicationAsync() { }
 
-	virtual void RegisterSettings(Chaos::FConstPhysicsObjectHandle PhysicsObject, FNetworkPhysicsSettingsAsync InSettings) = 0;
+	UE_DEPRECATED(5.7, "Deprecated, use the RegisterSettings that passes through a TWeakPtr<FNetworkPhysicsSettingsData> parameter instead ")
+	virtual void RegisterSettings(Chaos::FConstPhysicsObjectHandle PhysicsObject, FNetworkPhysicsSettingsAsync InSettings) { }
+
+	virtual void RegisterSettings(Chaos::FConstPhysicsObjectHandle PhysicsObject, TWeakPtr<const FNetworkPhysicsSettingsData> InSettings) = 0;
 };
