Substrate - make sure opacity code chunk for SubstrateShadingModel is never index_none.

#tests pc editor
#rnx

diff --git a/Engine/Source/Runtime/Engine/Private/Materials/MaterialExpressionSubstrate.cpp b/Engine/Source/Runtime/Engine/Private/Materials/MaterialExpressionSubstrate.cpp
--- a/Engine/Source/Runtime/Engine/Private/Materials/MaterialExpressionSubstrate.cpp
+++ b/Engine/Source/Runtime/Engine/Private/Materials/MaterialExpressionSubstrate.cpp
@@ -306,16 +306,17 @@ int32 UMaterialExpressionSubstrateShadingModels::CompileCommon(class FMaterialCo
 	SubstrateOperator.BSDFRegisteredSharedLocalBasis = NewRegisteredSharedLocalBasis;
 
 	int32 OpacityCodeChunk = INDEX_NONE;
+	const float DefaultOpacity = 1.0f;
 	if (!Compiler->SubstrateSkipsOpacityEvaluation())
 	{
 		// We evaluate opacity only for shading models and blending mode requiring it.
 		// For instance, a translucent shader reading depth for soft fading should no evaluate opacity when an instance forces an opaque mode.
-		LEGACY_DIRECT_ATTRIBUTE_MAPPING(OpacityCodeChunkTmp, MP_Opacity, CompileWithDefaultFloat1(Compiler, Opacity, 1.0f, EditorOnlyData ? &EditorOnlyData->Opacity : nullptr));
+		LEGACY_DIRECT_ATTRIBUTE_MAPPING(OpacityCodeChunkTmp, MP_Opacity, CompileWithDefaultFloat1(Compiler, Opacity, DefaultOpacity, EditorOnlyData ? &EditorOnlyData->Opacity : nullptr));
 		OpacityCodeChunk = OpacityCodeChunkTmp;
 	}
-	else
+	if(OpacityCodeChunk == INDEX_NONE)
 	{
-		OpacityCodeChunk = Compiler->Constant(1.0f);
+		OpacityCodeChunk = Compiler->Constant(DefaultOpacity);
 	}
 
 	LEGACY_DIRECT_ATTRIBUTE_MAPPING(EmissiveCodeChunk, MP_EmissiveColor, CompileWithDefaultFloat3(Compiler, EmissiveColor, 0.0f, 0.0f, 0.0f, EditorOnlyData ? &EditorOnlyData->EmissiveColor : nullptr));
