Unshelved from pending changelist '44108873':

NetDriver
* Fix NetDriver::TickDispatch hitch log not looking at realtime seconds for hitches

#rnx

diff --git a/Engine/Source/Runtime/Engine/Private/NetDriver.cpp b/Engine/Source/Runtime/Engine/Private/NetDriver.cpp
--- a/Engine/Source/Runtime/Engine/Private/NetDriver.cpp
+++ b/Engine/Source/Runtime/Engine/Private/NetDriver.cpp
@@ -2626,19 +2626,15 @@ void UNetDriver::TickDispatch( float DeltaTime )
 	// Manage realtime values
 	{
 		const double CurrentRealtime = FPlatformTime::Seconds();
+		const float DeltaRealtime = CurrentRealtime - LastTickDispatchRealtime;
 		LastTickDispatchRealtime = CurrentRealtime;
 
 		// Check to see if too much time is passing between ticks
 		// Setting this to somewhat large value for now, but small enough to catch blocking calls that are causing timeouts
 		constexpr float TickLogThreshold = 5.0f;
-
-		const float DeltaRealtime = CurrentRealtime - LastTickDispatchRealtime;
 		bDidHitchLastFrame = (DeltaTime > TickLogThreshold || DeltaRealtime > TickLogThreshold);
 
-		if (bDidHitchLastFrame)
-		{
-			UE_LOG( LogNet, Log, TEXT( "UNetDriver::TickDispatch: Very long time between ticks. DeltaTime: %2.2f, Realtime: %2.2f. %s" ), DeltaTime, DeltaRealtime, *GetName() );
-		}
+		UE_CLOG(bDidHitchLastFrame, LogNet, Log, TEXT( "UNetDriver::TickDispatch: Very long time between ticks. DeltaTime: %2.2f, Realtime: %2.2f. %s" ), DeltaTime, DeltaRealtime, *GetName() );
 	}
 
 	// Get new time.
