Substrate - Added missing localization to node tooltip and connector tooltip.

#rnx
#tests PC editor dx12

diff --git a/Engine/Source/Runtime/Engine/Private/Materials/MaterialExpressionSubstrate.cpp b/Engine/Source/Runtime/Engine/Private/Materials/MaterialExpressionSubstrate.cpp
--- a/Engine/Source/Runtime/Engine/Private/Materials/MaterialExpressionSubstrate.cpp
+++ b/Engine/Source/Runtime/Engine/Private/Materials/MaterialExpressionSubstrate.cpp
@@ -1030,13 +1030,13 @@ void UMaterialExpressionSubstrateSlabBSDF::GetCaption(TArray<FString>& OutCaptio
 
 void UMaterialExpressionSubstrateSlabBSDF::GetExpressionToolTip(TArray<FString>& OutToolTip)
 {
-	OutToolTip.Add(TEXT("Substrate Slab BSDF"));
-	OutToolTip.Add(TEXT("Complexity = ") + FSubstrateMaterialComplexity::ToString(GetHighestComplexity().SubstrateMaterialType()));
-	OutToolTip.Add(TEXT("The complexity represents the cost of the shading path (Lighting, Lumen, SSS) the material will follow:"));
-	OutToolTip.Add(TEXT(" - Simple means the Slab only relies on Diffuse, F0 and Roughness. It will follow a fast shading path."));
-	OutToolTip.Add(TEXT(" - Single means the Slab uses more features such as F90, Fuzz, Second Roughness, MFP or SSS. It will follow a more expenssive shading path."));
-	OutToolTip.Add(TEXT(" - Complex means a Slab uses anisotropic lighting, with any of the previous features."));
-	OutToolTip.Add(TEXT(" - Complex Special means the Slab is using more advanced features such as glints or specular LUT. This is the most expenssive shading path."));
+	OutToolTip.Add(FText(LOCTEXT("SlabNodeTT0", "Substrate Slab BSDF")).ToString());
+	OutToolTip.Add(FText(LOCTEXT("SlabNodeTT1", "Complexity = ")).ToString() + FSubstrateMaterialComplexity::ToString(GetHighestComplexity().SubstrateMaterialType()));
+	OutToolTip.Add(FText(LOCTEXT("SlabNodeTT2", "The complexity represents the cost of the shading path (Lighting, Lumen, SSS) the material will follow:")).ToString());
+	OutToolTip.Add(FText(LOCTEXT("SlabNodeTT3", " - Simple means the Slab only relies on Diffuse, F0 and Roughness. It will follow a fast shading path.")).ToString());
+	OutToolTip.Add(FText(LOCTEXT("SlabNodeTT4", " - Single means the Slab uses more features such as F90, Fuzz, Second Roughness, MFP or SSS. It will follow a more expenssive shading path.")).ToString());
+	OutToolTip.Add(FText(LOCTEXT("SlabNodeTT5", " - Complex means a Slab uses anisotropic lighting, with any of the previous features.")).ToString());
+	OutToolTip.Add(FText(LOCTEXT("SlabNodeTT6", " - Complex Special means the Slab is using more advanced features such as glints or specular LUT. This is the most expenssive shading path.")).ToString());
 }
 
 EMaterialValueType UMaterialExpressionSubstrateSlabBSDF::GetOutputValueType(int32 OutputIndex)
@@ -1421,10 +1421,10 @@ void UMaterialExpressionSubstrateSimpleClearCoatBSDF::GetCaption(TArray<FString>
 
 void UMaterialExpressionSubstrateSimpleClearCoatBSDF::GetExpressionToolTip(TArray<FString>& OutToolTip)
 {
-	OutToolTip.Add(TEXT("Substrate Simple Clear Coat BSDF"));
-	OutToolTip.Add(TEXT("It corresponds to the legacy \"Clear Coat\" shading model (also available using the \"Substrate Shading Models\" node)."));
-	OutToolTip.Add(TEXT("It has the complexity \"Single\"."));
-	OutToolTip.Add(TEXT("Clear coat top layer controls are limited to coverage and roughness. It always has a fixed grey scale transmittance and F0=0.04."));
+	OutToolTip.Add(FText(LOCTEXT("SCCNodeTT0", "Substrate Simple Clear Coat BSDF")).ToString());
+	OutToolTip.Add(FText(LOCTEXT("SCCNodeTT1", "It corresponds to the legacy \"Clear Coat\" shading model (also available using the \"Substrate Shading Models\" node).")).ToString());
+	OutToolTip.Add(FText(LOCTEXT("SCCNodeTT2", "It has the complexity \"Single\".")).ToString());
+	OutToolTip.Add(FText(LOCTEXT("SCCNodeTT3", "Clear coat top layer controls are limited to coverage and roughness. It always has a fixed grey scale transmittance and F0=0.04.")).ToString());
 }
 
 EMaterialValueType UMaterialExpressionSubstrateSimpleClearCoatBSDF::GetOutputValueType(int32 OutputIndex)
@@ -3421,10 +3421,10 @@ void UMaterialExpressionSubstrateTransmittanceToMFP::GetConnectorToolTip(int32 I
 		switch (InputIndex)
 		{
 		case 0:
-			ConvertToMultilineToolTip(TEXT("The colored transmittance for a view perpendicular to the surface. The transmittance for other view orientations will automatically be deduced according to surface thickness."), 80, OutToolTip);
+			ConvertToMultilineToolTip(FText(LOCTEXT("TransMFPNodeI0", "The colored transmittance for a view perpendicular to the surface. The transmittance for other view orientations will automatically be deduced according to surface thickness.")).ToString(), 80, OutToolTip);
 			break;
 		case 1:
-			ConvertToMultilineToolTip(TEXT("The thickness (in centimeter) at which the desired colored transmittance is reached. Default thickness: 0.01cm. Another use case example: this node output called thickness can be modulated before it is plugged in a slab node.this can be used to achieve simple scattering/transmittance variation of the same material."), 80, OutToolTip);
+			ConvertToMultilineToolTip(FText(LOCTEXT("TransMFPNodeI1", "The thickness (in centimeter) at which the desired colored transmittance is reached. Default thickness: 0.01cm. Another use case example: this node output called thickness can be modulated before it is plugged in a slab node.this can be used to achieve simple scattering/transmittance variation of the same material.")).ToString(), 80, OutToolTip);
 			break;
 		}
 	}
@@ -3433,10 +3433,10 @@ void UMaterialExpressionSubstrateTransmittanceToMFP::GetConnectorToolTip(int32 I
 		switch (OutputIndex)
 		{
 		case 0:
-			ConvertToMultilineToolTip(TEXT("The Mean Free Path defining the participating media constituting the slab of material (unit = centimeters)."), 80, OutToolTip);
+			ConvertToMultilineToolTip(FText(LOCTEXT("TransMFPNodeO0", "The Mean Free Path defining the participating media constituting the slab of material (unit = centimeters).")).ToString(), 80, OutToolTip);
 			break;
 		case 1:
-			ConvertToMultilineToolTip(TEXT("The thickness of the slab of material (unit = centimeters)."), 80, OutToolTip);
+			ConvertToMultilineToolTip(FText(LOCTEXT("TransMFPNodeO1", "The thickness of the slab of material (unit = centimeters).")).ToString(), 80, OutToolTip);
 			break;
 		}
 	}
@@ -3444,9 +3444,9 @@ void UMaterialExpressionSubstrateTransmittanceToMFP::GetConnectorToolTip(int32 I
 
 void UMaterialExpressionSubstrateTransmittanceToMFP::GetExpressionToolTip(TArray<FString>& OutToolTip)
 {
-	OutToolTip.Add(TEXT("Convert a transmittance color corresponding to a slab of participating media viewed perpendicularly to its surface."));
-	OutToolTip.Add(TEXT("This node directly maps to the Slab MFP input. It is recommended to use it when specifying the colored transmittance of a top layer slab."));
-	OutToolTip.Add(TEXT("For Subsurface scattering, you might prefer to specify the MFP(light mean free path) as world space centimeter directly."));
+	OutToolTip.Add(FText(LOCTEXT("TransMFPNodeTT0", "Convert a transmittance color corresponding to a slab of participating media viewed perpendicularly to its surface.")).ToString());
+	OutToolTip.Add(FText(LOCTEXT("TransMFPNodeTT1", "This node directly maps to the Slab MFP input. It is recommended to use it when specifying the colored transmittance of a top layer slab.")).ToString());
+	OutToolTip.Add(FText(LOCTEXT("TransMFPNodeTT2", "For Subsurface scattering, you might prefer to specify the MFP(light mean free path) as world space centimeter directly.")).ToString());
 }
 #endif // WITH_EDITOR
 
@@ -3526,15 +3526,15 @@ void UMaterialExpressionSubstrateMetalnessToDiffuseAlbedoF0::GetConnectorToolTip
 {
 	switch (OutputIndex)
 	{
-	case 1: ConvertToMultilineToolTip(TEXT("Defines the overall color of the Material. (type = float3, unit = unitless, defaults to 0.18)"), 80, OutToolTip); break;
-	case 2: ConvertToMultilineToolTip(TEXT("Controls how \"metal-like\" your surface looks like. 0 means dielectric, 1 means conductor (type = float, unit = unitless, defaults to 0)"), 80, OutToolTip); break;
-	case 3: ConvertToMultilineToolTip(TEXT("Used to scale the current amount of specularity on non-metallic surfaces and is a value between 0 and 1 (type = float, unit = unitless, defaults to plastic 0.5)"), 80, OutToolTip); break;
+	case 1: ConvertToMultilineToolTip(FText(LOCTEXT("MetalnessNodeO0", "Defines the overall color of the Material. (type = float3, unit = unitless, defaults to 0.18)")).ToString(), 80, OutToolTip); break;
+	case 2: ConvertToMultilineToolTip(FText(LOCTEXT("MetalnessNodeO1", "Controls how \"metal-like\" your surface looks like. 0 means dielectric, 1 means conductor (type = float, unit = unitless, defaults to 0)")).ToString(), 80, OutToolTip); break;
+	case 3: ConvertToMultilineToolTip(FText(LOCTEXT("MetalnessNodeO2", "Used to scale the current amount of specularity on non-metallic surfaces and is a value between 0 and 1 (type = float, unit = unitless, defaults to plastic 0.5)")).ToString(), 80, OutToolTip); break;
 	}
 }
 
 void UMaterialExpressionSubstrateMetalnessToDiffuseAlbedoF0::GetExpressionToolTip(TArray<FString>& OutToolTip)
 {
-	ConvertToMultilineToolTip(TEXT("Convert a metalness parameterization (BaseColor/Specular/Metallic) into DiffuseAlbedo/F0 parameterization."), 80, OutToolTip);
+	ConvertToMultilineToolTip(FText(LOCTEXT("MetalnessNodeTT0", "Convert a metalness parameterization (BaseColor/Specular/Metallic) into DiffuseAlbedo/F0 parameterization.")).ToString(), 80, OutToolTip);
 
 }
 #endif // WITH_EDITOR
@@ -3619,10 +3619,10 @@ void UMaterialExpressionSubstrateHazinessToSecondaryRoughness::GetConnectorToolT
 		switch (InputIndex)
 		{
 		case 0:
-			ConvertToMultilineToolTip(TEXT("The base roughness of the surface. It represented the smoothest part of the reflection."), 80, OutToolTip);
+			ConvertToMultilineToolTip(FText(LOCTEXT("HazinessNodeI0", "The base roughness of the surface. It represented the smoothest part of the reflection.")).ToString(), 80, OutToolTip);
 			break;
 		case 1:
-			ConvertToMultilineToolTip(TEXT("Haziness represent the amount of irregularity of the surface. A high value will lead to a second rough specular lobe causing the surface too look `milky`."), 80, OutToolTip);
+			ConvertToMultilineToolTip(FText(LOCTEXT("HazinessNodeI1", "Haziness represent the amount of irregularity of the surface. A high value will lead to a second rough specular lobe causing the surface too look `milky`.")).ToString(), 80, OutToolTip);
 			break;
 		}
 	}
@@ -3631,10 +3631,10 @@ void UMaterialExpressionSubstrateHazinessToSecondaryRoughness::GetConnectorToolT
 		switch (OutputIndex)
 		{
 		case 0:
-			ConvertToMultilineToolTip(TEXT("The roughness of the second lobe."), 80, OutToolTip);
+			ConvertToMultilineToolTip(FText(LOCTEXT("HazinessNodeO0", "The roughness of the second lobe.")).ToString(), 80, OutToolTip);
 			break;
 		case 1:
-			ConvertToMultilineToolTip(TEXT("The weight of the secondary specular lobe, while the primary specular lobe will have a weight of (1 - SecondRoughnessWeight)."), 80, OutToolTip);
+			ConvertToMultilineToolTip(FText(LOCTEXT("HazinessNodeO1", "The weight of the secondary specular lobe, while the primary specular lobe will have a weight of (1 - SecondRoughnessWeight).")).ToString(), 80, OutToolTip);
 			break;
 		}
 	}
@@ -3642,7 +3642,7 @@ void UMaterialExpressionSubstrateHazinessToSecondaryRoughness::GetConnectorToolT
 
 void UMaterialExpressionSubstrateHazinessToSecondaryRoughness::GetExpressionToolTip(TArray<FString>& OutToolTip)
 {
-	ConvertToMultilineToolTip(TEXT("Compute a second specular lobe roughness from a base surface roughness and haziness. This parameterisation ensure that the haziness makes physically and is perceptually easy to author."), 80, OutToolTip);
+	ConvertToMultilineToolTip(FText(LOCTEXT("HazinessNodeTT0", "Compute a second specular lobe roughness from a base surface roughness and haziness. This parameterisation ensure that the haziness makes physically and is perceptually easy to author.")).ToString(), 80, OutToolTip);
 
 }
 #endif // WITH_EDITOR
@@ -3741,10 +3741,10 @@ void UMaterialExpressionSubstrateThinFilm::GetConnectorToolTip(int32 InputIndex,
 		switch (OutputIndex)
 		{
 		case 0:
-			ConvertToMultilineToolTip(TEXT("F0 accounting for thin film interferences. This is percentage of light reflected as specular from a surface when the view is perpendicular to the surface. (type = float3, unit = unitless, defaults to plastic 0.04)"), 80, OutToolTip);
+			ConvertToMultilineToolTip(FText(LOCTEXT("ThinFilmNodeO0", "F0 accounting for thin film interferences. This is percentage of light reflected as specular from a surface when the view is perpendicular to the surface. (type = float3, unit = unitless, defaults to plastic 0.04)")).ToString(), 80, OutToolTip);
 			break;
 		case 1:
-			ConvertToMultilineToolTip(TEXT("F90 accounting for thin film interferences. the percentage of light reflected as specular from a surface when the view is tangent to the surface. (type = float3, unit = unitless, defaults to 1.0f)."), 80, OutToolTip);
+			ConvertToMultilineToolTip(FText(LOCTEXT("ThinFilmNodeO1", "F90 accounting for thin film interferences. the percentage of light reflected as specular from a surface when the view is tangent to the surface. (type = float3, unit = unitless, defaults to 1.0f).")).ToString(), 80, OutToolTip);
 			break;
 		}
 		return;
@@ -3756,7 +3756,7 @@ void UMaterialExpressionSubstrateThinFilm::GetConnectorToolTip(int32 InputIndex,
 
 void UMaterialExpressionSubstrateThinFilm::GetExpressionToolTip(TArray<FString>& OutToolTip)
 {
-	ConvertToMultilineToolTip(TEXT("Compute the resulting material specular parameter F0 and F90 according to input surface properties as well as the thin film parameters."), 80, OutToolTip);
+	ConvertToMultilineToolTip(FText(LOCTEXT("ThinFilmNodeTT0", "Compute the resulting material specular parameter F0 and F90 according to input surface properties as well as the thin film parameters.")).ToString(), 80, OutToolTip);
 }
 #endif // WITH_EDITOR
 
