Rewind debugger crash fix

diff --git a/Engine/Source/Runtime/Engine/Private/ObjectTrace.cpp b/Engine/Source/Runtime/Engine/Private/ObjectTrace.cpp
--- a/Engine/Source/Runtime/Engine/Private/ObjectTrace.cpp
+++ b/Engine/Source/Runtime/Engine/Private/ObjectTrace.cpp
@@ -544,11 +544,13 @@ void FObjectTrace::OutputInstanceLifetimeEnd(const UObject* InOuterObject, uint6
 	{
 		return;
 	}
-
-	UE_TRACE_LOG(Object, ObjectLifetimeEnd2, ObjectChannel)
-		<< ObjectLifetimeEnd2.Cycle(FPlatformTime::Cycles64())
-		<< ObjectLifetimeEnd2.RecordingTime(FObjectTrace::GetWorldElapsedTime(InOuterObject->GetWorld()))
-		<< ObjectLifetimeEnd2.Id(InInstanceId);
+	if (UWorld* World = InOuterObject->GetWorld())
+	{
+		UE_TRACE_LOG(Object, ObjectLifetimeEnd2, ObjectChannel)
+			<< ObjectLifetimeEnd2.Cycle(FPlatformTime::Cycles64())
+			<< ObjectLifetimeEnd2.RecordingTime(FObjectTrace::GetWorldElapsedTime(World))
+			<< ObjectLifetimeEnd2.Id(InInstanceId);
+	}
 }
 
 void FObjectTrace::OutputObjectLifetimeEnd(const UObject* InObject)
