Animation - removing unnecessary allocations using the default allocator

diff --git a/Engine/Plugins/Animation/ControlRig/Source/ControlRig/Private/ControlRig.cpp b/Engine/Plugins/Animation/ControlRig/Source/ControlRig/Private/ControlRig.cpp
--- a/Engine/Plugins/Animation/ControlRig/Source/ControlRig/Private/ControlRig.cpp
+++ b/Engine/Plugins/Animation/ControlRig/Source/ControlRig/Private/ControlRig.cpp
@@ -899,11 +899,7 @@ bool UControlRig::Execute(const FName& InEventName)
 
 	// The EventQueueToRun should only be modified in URigVMHost::Evaluate_AnyThread
 	// We create a temporary queue for the execution of only this event
-	TArray<FName> LocalEventQueueToRun = EventQueueToRun;
-	if(LocalEventQueueToRun.IsEmpty())
-	{
-		LocalEventQueueToRun = EventQueue;
-	}
+	TConstArrayView<FName> LocalEventQueueToRun = EventQueueToRun.IsEmpty() ? EventQueue : EventQueueToRun;
 
 	const bool bIsEventInQueue = LocalEventQueueToRun.Contains(InEventName);
 	const bool bIsEventFirstInQueue = !LocalEventQueueToRun.IsEmpty() && LocalEventQueueToRun[0] == InEventName; 
@@ -1694,11 +1690,7 @@ bool UControlRig::Execute_Internal(const FName& InEventName)
 
 	// The EventQueueToRun should only be modified in URigVMHost::Evaluate_AnyThread
 	// We create a temporary queue for the execution of only this event
-	TArray<FName> LocalEventQueueToRun = EventQueueToRun;
-	if(LocalEventQueueToRun.IsEmpty())
-	{
-		LocalEventQueueToRun = EventQueue;
-	}
+	TConstArrayView<FName> LocalEventQueueToRun = EventQueueToRun.IsEmpty() ? EventQueue : EventQueueToRun;
 
 	// make sure to initialize here as well - just in case this gets
 	// called without a call to ::Execute. This is already tackled by
diff --git a/Engine/Plugins/Animation/ControlRig/Source/ControlRig/Private/Tools/ControlRigPoseAdapter.cpp b/Engine/Plugins/Animation/ControlRig/Source/ControlRig/Private/Tools/ControlRigPoseAdapter.cpp
--- a/Engine/Plugins/Animation/ControlRig/Source/ControlRig/Private/Tools/ControlRigPoseAdapter.cpp
+++ b/Engine/Plugins/Animation/ControlRig/Source/ControlRig/Private/Tools/ControlRigPoseAdapter.cpp
@@ -337,9 +337,9 @@ void FControlRigPoseAdapter::SetHierarchyCurvesLookup(const TArray<FRigBaseEleme
 	}
 }
 
-void FControlRigPoseAdapter::SetPoseCurveToHierarchyCurve(const TArray<FRigBaseElement*> InHierarchyCurves, const FBlendedCurve& InCurve)
+void FControlRigPoseAdapter::SetPoseCurveToHierarchyCurve(const TArray<FRigBaseElement*>& InHierarchyCurves, const FBlendedCurve& InCurve)
 {
-	PoseCurveToHierarchyCurve.SetNumUninitialized(InCurve.Num());
+	PoseCurveToHierarchyCurve.SetNumUninitialized(InCurve.Num(), EAllowShrinking::No);
 
 	int32 CurveIndex = 0;
 	InCurve.ForEachElement([this, &CurveIndex, &InHierarchyCurves](const UE::Anim::FCurveElement& InCurveElement)
diff --git a/Engine/Plugins/Animation/ControlRig/Source/ControlRig/Private/Units/Highlevel/Hierarchy/RigUnit_ModifyBoneTransforms.cpp b/Engine/Plugins/Animation/ControlRig/Source/ControlRig/Private/Units/Highlevel/Hierarchy/RigUnit_ModifyBoneTransforms.cpp
--- a/Engine/Plugins/Animation/ControlRig/Source/ControlRig/Private/Units/Highlevel/Hierarchy/RigUnit_ModifyBoneTransforms.cpp
+++ b/Engine/Plugins/Animation/ControlRig/Source/ControlRig/Private/Units/Highlevel/Hierarchy/RigUnit_ModifyBoneTransforms.cpp
@@ -7,7 +7,7 @@
 
 FRigUnit_ModifyBoneTransforms_Execute()
 {
-	TArray<FRigUnit_ModifyTransforms_PerItem> ItemsToModify;
+	TArray<FRigUnit_ModifyTransforms_PerItem, TInlineAllocator<4>> ItemsToModify;
 	for (int32 BoneIndex = 0; BoneIndex < BoneToModify.Num(); BoneIndex++)
 	{
 		FRigUnit_ModifyTransforms_PerItem ItemToModify;
diff --git a/Engine/Plugins/Animation/ControlRig/Source/ControlRig/Public/Tools/ControlRigPoseAdapter.h b/Engine/Plugins/Animation/ControlRig/Source/ControlRig/Public/Tools/ControlRigPoseAdapter.h
--- a/Engine/Plugins/Animation/ControlRig/Source/ControlRig/Public/Tools/ControlRigPoseAdapter.h
+++ b/Engine/Plugins/Animation/ControlRig/Source/ControlRig/Public/Tools/ControlRigPoseAdapter.h
@@ -37,7 +37,7 @@ public:
 		return LocalPose;
 	}
 
-	const TArray<FTransform> GetGlobalPose() const
+	const TArray<FTransform>& GetGlobalPose() const
 	{
 		return LocalPose;
 	}
@@ -47,7 +47,7 @@ public:
 		InPose.CopyBonesTo(LocalPose);
 	}
 
-	const TMap<FName, int32> GetHierarchyCurveLookup() const
+	const TMap<FName, int32>& GetHierarchyCurveLookup() const
 	{
 		return HierarchyCurveLookup;
 	}
@@ -64,7 +64,7 @@ public:
 		return PoseCurveToHierarchyCurve;
 	}
 
-	UE_API void SetPoseCurveToHierarchyCurve(const TArray<FRigBaseElement*> InHierarchyCurves, const FBlendedCurve& InCurve);
+	UE_API void SetPoseCurveToHierarchyCurve(const TArray<FRigBaseElement*>& InHierarchyCurves, const FBlendedCurve& InCurve);
 
 	UE_API void UpdateInputOutputMappingIfRequired(UControlRig* InControlRig
 		, URigHierarchy* InHierarchy
diff --git a/Engine/Plugins/Experimental/Animation/HierarchyTableAnimation/Source/Runtime/Private/AnimNode_BlendProfileLayeredBlend.cpp b/Engine/Plugins/Experimental/Animation/HierarchyTableAnimation/Source/Runtime/Private/AnimNode_BlendProfileLayeredBlend.cpp
--- a/Engine/Plugins/Experimental/Animation/HierarchyTableAnimation/Source/Runtime/Private/AnimNode_BlendProfileLayeredBlend.cpp
+++ b/Engine/Plugins/Experimental/Animation/HierarchyTableAnimation/Source/Runtime/Private/AnimNode_BlendProfileLayeredBlend.cpp
@@ -224,8 +224,8 @@ void FAnimNode_BlendProfileLayeredBlend::Evaluate_AnyThread(FPoseContext& Output
 			TargetPoseMaxWeight = FMath::Max(TargetPoseMaxWeight, BoneBlendWeight);
 		}
 
-		const TArray<const FBlendedCurve*> SourceCurves = { &BasePoseData.GetCurve(), &BlendPoseData.GetCurve() };
-		const TArray<float> SourceWeights = { BlendWeight, TargetPoseMaxWeight };
+		const FBlendedCurve* SourceCurves[] = { &BasePoseData.GetCurve(), &BlendPoseData.GetCurve() };
+		const float SourceWeights[] = { BlendWeight, TargetPoseMaxWeight };
 
 		BlendCurves(SourceCurves, SourceWeights, Output.Curve, CurveBlendingOption);
 	}
@@ -275,12 +275,12 @@ void FAnimNode_BlendProfileLayeredBlend::Evaluate_AnyThread(FPoseContext& Output
 		// 2. For each attribute that has a custom weight, i.e. one's that shouldn't be weighted by its attached bone, go and correct the blended value.
 		for (const FBlendProfileStandaloneCachedData::FMaskedAttributeWeight& MaskedAttribute : CachedAttributeMaskWeights)
 		{
-			const TArray<TWeakObjectPtr<UScriptStruct>>& UniqueTypes = OutputAttributes.GetUniqueTypes();
+			const TConstArrayView<TWeakObjectPtr<UScriptStruct>> UniqueTypes = OutputAttributes.GetUniqueTypes();
 
 			for (int32 UniqueTypeIndex = 0; UniqueTypeIndex < UniqueTypes.Num(); ++UniqueTypeIndex)
 			{
-				const TArray<FAttributeId>& AttributeKeys = OutputAttributes.GetKeys(UniqueTypeIndex);
-				const TArray<TWrappedAttribute<FAnimStackAllocator>>& AttributeValues = OutputAttributes.GetValues(UniqueTypeIndex);
+				const TConstArrayView<FAttributeId> AttributeKeys = OutputAttributes.GetKeys(UniqueTypeIndex);
+				const TConstArrayView<TWrappedAttribute<FAnimStackAllocator>> AttributeValues = OutputAttributes.GetValues(UniqueTypeIndex);
 				
 				const TWeakObjectPtr<UScriptStruct> AttributeType = UniqueTypes[UniqueTypeIndex];
 				const IAttributeBlendOperator* Operator = UE::Anim::AttributeTypes::GetTypeOperator(AttributeType);
diff --git a/Engine/Plugins/Experimental/UAF/UAFAnimGraph/Source/UAFAnimGraph/Private/EvaluationVM/Tasks/BlendKeyframesPerBone.cpp b/Engine/Plugins/Experimental/UAF/UAFAnimGraph/Source/UAFAnimGraph/Private/EvaluationVM/Tasks/BlendKeyframesPerBone.cpp
--- a/Engine/Plugins/Experimental/UAF/UAFAnimGraph/Source/UAFAnimGraph/Private/EvaluationVM/Tasks/BlendKeyframesPerBone.cpp
+++ b/Engine/Plugins/Experimental/UAF/UAFAnimGraph/Source/UAFAnimGraph/Private/EvaluationVM/Tasks/BlendKeyframesPerBone.cpp
@@ -379,12 +379,12 @@ void FAnimNextBlendKeyframePerBoneWithScaleTask::Execute(UE::UAF::FEvaluationVM&
 		// 2. For each attribute that has a custom weight, i.e. one's that shouldn't be weighted by its attached bone, go and correct the blended value.
 		for (const FMaskedAttributeWeight& MaskedAttribute : AttributeMaskWeights)
 		{
-			const TArray<TWeakObjectPtr<UScriptStruct>>& UniqueTypes = OutputAttributes.GetUniqueTypes();
+			const TConstArrayView<TWeakObjectPtr<UScriptStruct>> UniqueTypes = OutputAttributes.GetUniqueTypes();
 
 			for (int32 UniqueTypeIndex = 0; UniqueTypeIndex < UniqueTypes.Num(); ++UniqueTypeIndex)
 			{
-				const TArray<FAttributeId>& AttributeKeys = OutputAttributes.GetKeys(UniqueTypeIndex);
-				const TArray<TWrappedAttribute<FAnimStackAllocator>>& AttributeValues = OutputAttributes.GetValues(UniqueTypeIndex);
+				const TConstArrayView<FAttributeId> AttributeKeys = OutputAttributes.GetKeys(UniqueTypeIndex);
+				const TConstArrayView<TWrappedAttribute<FAnimStackAllocator>> AttributeValues = OutputAttributes.GetValues(UniqueTypeIndex);
 				
 				const TWeakObjectPtr<UScriptStruct> AttributeType = UniqueTypes[UniqueTypeIndex];
 				const IAttributeBlendOperator* Operator = UE::Anim::AttributeTypes::GetTypeOperator(AttributeType);
diff --git a/Engine/Plugins/FX/Niagara/Source/NiagaraEditor/Private/Widgets/SNiagaraDistributionEditor.cpp b/Engine/Plugins/FX/Niagara/Source/NiagaraEditor/Private/Widgets/SNiagaraDistributionEditor.cpp
--- a/Engine/Plugins/FX/Niagara/Source/NiagaraEditor/Private/Widgets/SNiagaraDistributionEditor.cpp
+++ b/Engine/Plugins/FX/Niagara/Source/NiagaraEditor/Private/Widgets/SNiagaraDistributionEditor.cpp
@@ -581,7 +581,16 @@ public:
 		}
 	}
 
+	UE_DEPRECATED(5.6, "Use version taking a TAdderReserverRef")
 	virtual TArray<FRichCurveEditInfoConst> GetCurves() const override { return ConstCurveInfos; }
+	virtual void GetCurves(TAdderReserverRef<FRichCurveEditInfoConst> Curves) const override
+	{
+		for (const FRichCurveEditInfoConst& ConstCurve: ConstCurveInfos)
+		{
+			Curves.Add(ConstCurve);
+		}
+	}
+
 	virtual TArray<FRichCurveEditInfo> GetCurves() override { return CurveInfos; }
 
 	virtual void ModifyOwner() override
diff --git a/Engine/Plugins/FX/Niagara/Source/NiagaraEditorWidgets/Private/DetailCustomizations/NiagaraDataInterfaceCurveDetails.cpp b/Engine/Plugins/FX/Niagara/Source/NiagaraEditorWidgets/Private/DetailCustomizations/NiagaraDataInterfaceCurveDetails.cpp
--- a/Engine/Plugins/FX/Niagara/Source/NiagaraEditorWidgets/Private/DetailCustomizations/NiagaraDataInterfaceCurveDetails.cpp
+++ b/Engine/Plugins/FX/Niagara/Source/NiagaraEditorWidgets/Private/DetailCustomizations/NiagaraDataInterfaceCurveDetails.cpp
@@ -1081,11 +1081,20 @@ public:
 	}
 
 	/** FCurveOwnerInterface */
+	UE_DEPRECATED(5.6, "Use version taking a TAdderReserverRef")
 	virtual TArray<FRichCurveEditInfoConst> GetCurves() const override
 	{
 		return ConstCurves;
 	}
 
+	virtual void GetCurves(TAdderReserverRef<FRichCurveEditInfoConst> OutCurves) const override
+	{
+		for (const FRichCurveEditInfoConst& ConstCurve: ConstCurves)
+		{
+			OutCurves.Add(ConstCurve);
+		}
+	}
+
 	virtual TArray<FRichCurveEditInfo> GetCurves() override
 	{
 		return Curves;
diff --git a/Engine/Plugins/Runtime/RigVM/Source/RigVMEditor/Private/Widgets/SRigVMGraphPinCurveFloat.cpp b/Engine/Plugins/Runtime/RigVM/Source/RigVMEditor/Private/Widgets/SRigVMGraphPinCurveFloat.cpp
--- a/Engine/Plugins/Runtime/RigVM/Source/RigVMEditor/Private/Widgets/SRigVMGraphPinCurveFloat.cpp
+++ b/Engine/Plugins/Runtime/RigVM/Source/RigVMEditor/Private/Widgets/SRigVMGraphPinCurveFloat.cpp
@@ -59,6 +59,11 @@ TArray<FRichCurveEditInfoConst> SRigVMGraphPinCurveFloat::GetCurves() const
 	return Curves;
 }
 
+void SRigVMGraphPinCurveFloat::GetCurves(TAdderReserverRef<FRichCurveEditInfoConst> Curves) const
+{
+	Curves.Add(Curve.GetRichCurveConst());
+}
+
 TArray<FRichCurveEditInfo> SRigVMGraphPinCurveFloat::GetCurves()
 {
 	TArray<FRichCurveEditInfo> Curves;
diff --git a/Engine/Plugins/Runtime/RigVM/Source/RigVMEditor/Public/Widgets/SRigVMGraphPinCurveFloat.h b/Engine/Plugins/Runtime/RigVM/Source/RigVMEditor/Public/Widgets/SRigVMGraphPinCurveFloat.h
--- a/Engine/Plugins/Runtime/RigVM/Source/RigVMEditor/Public/Widgets/SRigVMGraphPinCurveFloat.h
+++ b/Engine/Plugins/Runtime/RigVM/Source/RigVMEditor/Public/Widgets/SRigVMGraphPinCurveFloat.h
@@ -20,7 +20,9 @@ public:
 	UE_API void Construct(const FArguments& InArgs, UEdGraphPin* InGraphPinObj);
 
 	/** FCurveOwnerInterface interface */
+	UE_DEPRECATED(5.6, "Use version taking a TAdderReserverRef")
 	UE_API virtual TArray<FRichCurveEditInfoConst> GetCurves() const override;
+	UE_API virtual void GetCurves(TAdderReserverRef<FRichCurveEditInfoConst> Curves) const override;
 	UE_API virtual TArray<FRichCurveEditInfo> GetCurves() override;
 	UE_API virtual void ModifyOwner() override;
 	UE_API virtual TArray<const UObject*> GetOwners() const override;
diff --git a/Engine/Source/Editor/AnimGraph/Private/PoseDriverDetails.cpp b/Engine/Source/Editor/AnimGraph/Private/PoseDriverDetails.cpp
--- a/Engine/Source/Editor/AnimGraph/Private/PoseDriverDetails.cpp
+++ b/Engine/Source/Editor/AnimGraph/Private/PoseDriverDetails.cpp
@@ -884,6 +884,14 @@ TArray<FRichCurveEditInfoConst> SPDD_TargetRow::GetCurves() const
 	return Curves;
 }
 
+void SPDD_TargetRow::GetCurves(TAdderReserverRef<FRichCurveEditInfoConst> Curves) const
+{
+	if (FPoseDriverTarget* Target = GetTarget())
+	{
+		Curves.Add(&(Target->CustomCurve));
+	}
+}
+
 TArray<FRichCurveEditInfo> SPDD_TargetRow::GetCurves()
 {
 	TArray<FRichCurveEditInfo> Curves;
diff --git a/Engine/Source/Editor/AnimGraph/Private/PoseDriverDetails.h b/Engine/Source/Editor/AnimGraph/Private/PoseDriverDetails.h
--- a/Engine/Source/Editor/AnimGraph/Private/PoseDriverDetails.h
+++ b/Engine/Source/Editor/AnimGraph/Private/PoseDriverDetails.h
@@ -140,7 +140,9 @@ public:
 	bool IsSoloTarget() const;
 
 	/** FCurveOwnerInterface interface */
+	UE_DEPRECATED(5.6, "Use version taking a TAdderReserverRef")
 	virtual TArray<FRichCurveEditInfoConst> GetCurves() const override;
+	virtual void GetCurves(TAdderReserverRef<FRichCurveEditInfoConst> Curves) const override;
 	virtual TArray<FRichCurveEditInfo> GetCurves() override;
 	virtual void ModifyOwner() override;
 	virtual TArray<const UObject*> GetOwners() const override;
diff --git a/Engine/Source/Editor/CurveTableEditor/Private/CurveTableEditorHandle.cpp b/Engine/Source/Editor/CurveTableEditor/Private/CurveTableEditorHandle.cpp
--- a/Engine/Source/Editor/CurveTableEditor/Private/CurveTableEditorHandle.cpp
+++ b/Engine/Source/Editor/CurveTableEditor/Private/CurveTableEditorHandle.cpp
@@ -49,6 +49,14 @@ TArray<FRichCurveEditInfoConst> FCurveTableEditorHandle::GetCurves() const
 	return Curves;
 }
 
+void FCurveTableEditorHandle::GetCurves(TAdderReserverRef<FRichCurveEditInfoConst> Curves) const
+{
+	if (const FRealCurve* Curve = GetCurve())
+	{
+		Curves.Add(FRichCurveEditInfoConst(Curve, RowName));
+	}
+}
+
 TArray<FRichCurveEditInfo> FCurveTableEditorHandle::GetCurves()
 {
 	TArray<FRichCurveEditInfo> Curves;
diff --git a/Engine/Source/Editor/CurveTableEditor/Private/CurveTableEditorHandle.h b/Engine/Source/Editor/CurveTableEditor/Private/CurveTableEditorHandle.h
--- a/Engine/Source/Editor/CurveTableEditor/Private/CurveTableEditorHandle.h
+++ b/Engine/Source/Editor/CurveTableEditor/Private/CurveTableEditorHandle.h
@@ -36,7 +36,9 @@ struct FCurveTableEditorHandle : public FCurveOwnerInterface
 	FName RowName;
 
 	//~ Begin FCurveOwnerInterface Interface.
+	UE_DEPRECATED(5.6, "Use version taking a TAdderReserverRef")
 	virtual TArray<FRichCurveEditInfoConst> GetCurves() const override;
+	virtual void GetCurves(TAdderReserverRef<FRichCurveEditInfoConst> Curves) const override;
 	virtual TArray<FRichCurveEditInfo> GetCurves() override;
 	virtual void ModifyOwner() override;
 	virtual void MakeTransactional() override;
diff --git a/Engine/Source/Editor/DetailCustomizations/Private/AnimTrailNodeDetails.cpp b/Engine/Source/Editor/DetailCustomizations/Private/AnimTrailNodeDetails.cpp
--- a/Engine/Source/Editor/DetailCustomizations/Private/AnimTrailNodeDetails.cpp
+++ b/Engine/Source/Editor/DetailCustomizations/Private/AnimTrailNodeDetails.cpp
@@ -97,6 +97,11 @@ TArray<FRichCurveEditInfoConst> FAnimTrailNodeDetails::FTrailRelaxCurveEditor::G
 	return Curves;
 }
 
+void FAnimTrailNodeDetails::FTrailRelaxCurveEditor::GetCurves(TAdderReserverRef<FRichCurveEditInfoConst> Curves) const
+{
+	Curves.Add(&GraphNodeOwner->Node.TrailRelaxationSpeed.EditorCurveData);
+}
+
 TArray<FRichCurveEditInfo> FAnimTrailNodeDetails::FTrailRelaxCurveEditor::GetCurves()
 {
 	TArray<FRichCurveEditInfo> Curves;
diff --git a/Engine/Source/Editor/DetailCustomizations/Private/AnimTrailNodeDetails.h b/Engine/Source/Editor/DetailCustomizations/Private/AnimTrailNodeDetails.h
--- a/Engine/Source/Editor/DetailCustomizations/Private/AnimTrailNodeDetails.h
+++ b/Engine/Source/Editor/DetailCustomizations/Private/AnimTrailNodeDetails.h
@@ -27,7 +27,9 @@ public:
 	{
 		FTrailRelaxCurveEditor(UAnimGraphNode_Trail* InGraphNode = nullptr, TSharedPtr<IPropertyHandle> InTrailRelaxCurveHandle = TSharedPtr<IPropertyHandle>());
 		/** FCurveOwnerInterface interface */
+		UE_DEPRECATED(5.6, "Use version taking a TAdderReserverRef")
 		virtual TArray<FRichCurveEditInfoConst> GetCurves() const override;
+		virtual void GetCurves(TAdderReserverRef<FRichCurveEditInfoConst> Curves) const override;
 		virtual TArray<FRichCurveEditInfo> GetCurves() override;
 		virtual void ModifyOwner() override;
 		virtual TArray<const UObject*> GetOwners() const override;
diff --git a/Engine/Source/Editor/DetailCustomizations/Private/CurveColorCustomization.cpp b/Engine/Source/Editor/DetailCustomizations/Private/CurveColorCustomization.cpp
--- a/Engine/Source/Editor/DetailCustomizations/Private/CurveColorCustomization.cpp
+++ b/Engine/Source/Editor/DetailCustomizations/Private/CurveColorCustomization.cpp
@@ -269,6 +269,14 @@ TArray<FRichCurveEditInfoConst> FCurveColorCustomization::GetCurves() const
 	return Curves;
 }
 
+void FCurveColorCustomization::GetCurves(TAdderReserverRef<FRichCurveEditInfoConst> Curves) const
+{
+	Curves.Add(FRichCurveEditInfoConst(&RuntimeCurve->ColorCurves[0], RedCurveName));
+	Curves.Add(FRichCurveEditInfoConst(&RuntimeCurve->ColorCurves[1], GreenCurveName));
+	Curves.Add(FRichCurveEditInfoConst(&RuntimeCurve->ColorCurves[2], BlueCurveName));
+	Curves.Add(FRichCurveEditInfoConst(&RuntimeCurve->ColorCurves[3], AlphaCurveName));
+}
+
 TArray<FRichCurveEditInfo> FCurveColorCustomization::GetCurves()
 {
 	TArray<FRichCurveEditInfo> Curves;
diff --git a/Engine/Source/Editor/DetailCustomizations/Private/CurveColorCustomization.h b/Engine/Source/Editor/DetailCustomizations/Private/CurveColorCustomization.h
--- a/Engine/Source/Editor/DetailCustomizations/Private/CurveColorCustomization.h
+++ b/Engine/Source/Editor/DetailCustomizations/Private/CurveColorCustomization.h
@@ -39,7 +39,9 @@ public:
 	virtual void CustomizeChildren( TSharedRef<IPropertyHandle> InStructPropertyHandle, IDetailChildrenBuilder& StructBuilder, IPropertyTypeCustomizationUtils& StructCustomizationUtils ) override;
 
 	/** FCurveOwnerInterface interface */
+	UE_DEPRECATED(5.6, "Use version taking a TAdderReserverRef")
 	virtual TArray<FRichCurveEditInfoConst> GetCurves() const override;
+	virtual void GetCurves(TAdderReserverRef<FRichCurveEditInfoConst> Curves) const override;
 	virtual TArray<FRichCurveEditInfo> GetCurves() override;
 	virtual void ModifyOwner() override;
 	virtual TArray<const UObject*> GetOwners() const override;
diff --git a/Engine/Source/Editor/DetailCustomizations/Private/CurveStructCustomization.cpp b/Engine/Source/Editor/DetailCustomizations/Private/CurveStructCustomization.cpp
--- a/Engine/Source/Editor/DetailCustomizations/Private/CurveStructCustomization.cpp
+++ b/Engine/Source/Editor/DetailCustomizations/Private/CurveStructCustomization.cpp
@@ -283,6 +283,17 @@ TArray<FRichCurveEditInfoConst> FCurveStructCustomization::GetCurves() const
 	return Curves;
 }
 
+void FCurveStructCustomization::GetCurves(TAdderReserverRef<FRichCurveEditInfoConst> Curves) const
+{
+	FRichCurveEditInfoConst CurveEditInfo = FRichCurveEditInfoConst(&RuntimeCurve->EditorCurveData);
+	if (StructPropertyHandle.IsValid())
+	{
+		// Use the DisplayName of the property in place of the CurveName (which is None after creation above)
+		CurveEditInfo.CurveName = FName(StructPropertyHandle->GetPropertyDisplayName().ToString());
+	}
+	Curves.Add(CurveEditInfo);
+}
+
 TArray<FRichCurveEditInfo> FCurveStructCustomization::GetCurves()
 {
 	TArray<FRichCurveEditInfo> Curves;
diff --git a/Engine/Source/Editor/DetailCustomizations/Private/CurveStructCustomization.h b/Engine/Source/Editor/DetailCustomizations/Private/CurveStructCustomization.h
--- a/Engine/Source/Editor/DetailCustomizations/Private/CurveStructCustomization.h
+++ b/Engine/Source/Editor/DetailCustomizations/Private/CurveStructCustomization.h
@@ -39,7 +39,9 @@ public:
 	virtual void CustomizeChildren( TSharedRef<IPropertyHandle> InStructPropertyHandle, IDetailChildrenBuilder& StructBuilder, IPropertyTypeCustomizationUtils& StructCustomizationUtils ) override;
 
 	/** FCurveOwnerInterface interface */
+	UE_DEPRECATED(5.6, "Use version taking a TAdderReserverRef")
 	virtual TArray<FRichCurveEditInfoConst> GetCurves() const override;
+	virtual void GetCurves(TAdderReserverRef<FRichCurveEditInfoConst> Curves) const override;
 	virtual TArray<FRichCurveEditInfo> GetCurves() override;
 	virtual void ModifyOwner() override;
 	virtual TArray<const UObject*> GetOwners() const override;
diff --git a/Engine/Source/Editor/DetailCustomizations/Private/CurveVectorCustomization.cpp b/Engine/Source/Editor/DetailCustomizations/Private/CurveVectorCustomization.cpp
--- a/Engine/Source/Editor/DetailCustomizations/Private/CurveVectorCustomization.cpp
+++ b/Engine/Source/Editor/DetailCustomizations/Private/CurveVectorCustomization.cpp
@@ -274,6 +274,13 @@ TArray<FRichCurveEditInfoConst> FCurveVectorCustomization::GetCurves() const
 	return Curves;
 }
 
+void FCurveVectorCustomization::GetCurves(TAdderReserverRef<FRichCurveEditInfoConst> Curves) const
+{
+	Curves.Add(FRichCurveEditInfoConst(&RuntimeCurve->VectorCurves[0], XCurveName));
+	Curves.Add(FRichCurveEditInfoConst(&RuntimeCurve->VectorCurves[1], YCurveName));
+	Curves.Add(FRichCurveEditInfoConst(&RuntimeCurve->VectorCurves[2], ZCurveName));
+}
+
 TArray<FRichCurveEditInfo> FCurveVectorCustomization::GetCurves()
 {
 	TArray<FRichCurveEditInfo> Curves;
diff --git a/Engine/Source/Editor/DetailCustomizations/Private/CurveVectorCustomization.h b/Engine/Source/Editor/DetailCustomizations/Private/CurveVectorCustomization.h
--- a/Engine/Source/Editor/DetailCustomizations/Private/CurveVectorCustomization.h
+++ b/Engine/Source/Editor/DetailCustomizations/Private/CurveVectorCustomization.h
@@ -35,7 +35,9 @@ public:
 	virtual void CustomizeChildren( TSharedRef<IPropertyHandle> InStructPropertyHandle, IDetailChildrenBuilder& StructBuilder, IPropertyTypeCustomizationUtils& StructCustomizationUtils ) override;
 
 	/** FCurveOwnerInterface interface */
+	UE_DEPRECATED(5.6, "Use version taking a TAdderReserverRef")
 	virtual TArray<FRichCurveEditInfoConst> GetCurves() const override;
+	virtual void GetCurves(TAdderReserverRef<FRichCurveEditInfoConst> Curves) const override;
 	virtual TArray<FRichCurveEditInfo> GetCurves() override;
 	virtual void ModifyOwner() override;
 	virtual TArray<const UObject*> GetOwners() const override;
diff --git a/Engine/Source/Runtime/AnimGraphRuntime/Private/AnimNodes/AnimNode_BlendSpacePlayer.cpp b/Engine/Source/Runtime/AnimGraphRuntime/Private/AnimNodes/AnimNode_BlendSpacePlayer.cpp
--- a/Engine/Source/Runtime/AnimGraphRuntime/Private/AnimNodes/AnimNode_BlendSpacePlayer.cpp
+++ b/Engine/Source/Runtime/AnimGraphRuntime/Private/AnimNodes/AnimNode_BlendSpacePlayer.cpp
@@ -199,7 +199,7 @@ const FBlendSampleData* FAnimNode_BlendSpacePlayerBase::GetHighestWeightedSample
 
 void FAnimNode_BlendSpacePlayerBase::Reinitialize(bool bResetTime)
 {
-	BlendSampleDataCache.Empty();
+	BlendSampleDataCache.Reset();
 	if(bResetTime)
 	{
 		float CurrentStartPosition = GetStartPosition();
diff --git a/Engine/Source/Runtime/AnimGraphRuntime/Private/AnimNodes/AnimNode_LayeredBoneBlend.cpp b/Engine/Source/Runtime/AnimGraphRuntime/Private/AnimNodes/AnimNode_LayeredBoneBlend.cpp
--- a/Engine/Source/Runtime/AnimGraphRuntime/Private/AnimNodes/AnimNode_LayeredBoneBlend.cpp
+++ b/Engine/Source/Runtime/AnimGraphRuntime/Private/AnimNodes/AnimNode_LayeredBoneBlend.cpp
@@ -221,13 +221,13 @@ void FAnimNode_LayeredBoneBlend::Evaluate_AnyThread(FPoseContext& Output)
 		// evaluate children
 		BasePose.Evaluate(BasePoseContext);
 
-		TArray<FCompactPose> TargetBlendPoses;
+		TArray<FCompactPose, TMemStackAllocator<>> TargetBlendPoses;
 		TargetBlendPoses.SetNum(NumPoses);
 
-		TArray<FBlendedCurve> TargetBlendCurves;
+		TArray<FBlendedCurve, TMemStackAllocator<>> TargetBlendCurves;
 		TargetBlendCurves.SetNum(NumPoses);
 
-		TArray<UE::Anim::FStackAttributeContainer> TargetBlendAttributes;
+		TArray<UE::Anim::FStackAttributeContainer, TMemStackAllocator<>> TargetBlendAttributes;
 		TargetBlendAttributes.SetNum(NumPoses);
 
 		for (int32 ChildIndex = 0; ChildIndex < NumPoses; ++ChildIndex)
diff --git a/Engine/Source/Runtime/AnimGraphRuntime/Private/BoneControllers/AnimNode_ResetRoot.cpp b/Engine/Source/Runtime/AnimGraphRuntime/Private/BoneControllers/AnimNode_ResetRoot.cpp
--- a/Engine/Source/Runtime/AnimGraphRuntime/Private/BoneControllers/AnimNode_ResetRoot.cpp
+++ b/Engine/Source/Runtime/AnimGraphRuntime/Private/BoneControllers/AnimNode_ResetRoot.cpp
@@ -32,25 +32,20 @@ void FAnimNode_ResetRoot::EvaluateSkeletalControl_AnyThread(FComponentSpacePoseC
 	DECLARE_SCOPE_HIERARCHICAL_COUNTER_ANIMNODE(EvaluateSkeletalControl_AnyThread)
 	ANIM_MT_SCOPE_CYCLE_COUNTER_VERBOSE(ResetRoot, !IsInGameThread());
 
-	const FBoneContainer& BoneContainer = Output.Pose.GetPose().GetBoneContainer();
-	FCSPose<FCompactPose>& CSPose = Output.Pose;
+	check(OutBoneTransforms.IsEmpty());
 
-	// Get all Children of the Root in Component Space. We want to preserve these in Component Space, after the Root has been reset.
-	TArray<FTransform> ChildrenTransform;
-	ChildrenTransform.Reserve(RootChildren.Num());
-	for (int32 Index = 0; Index < RootChildren.Num(); Index++)
-	{
-		ChildrenTransform.Add(CSPose.GetComponentSpaceTransform(RootChildren[Index]));
-	}
+	FCSPose<FCompactPose>& CSPose = Output.Pose;
+	const FBoneContainer& BoneContainer = CSPose.GetPose().GetBoneContainer();
 
 	// Reset Root
 	const FCompactPoseBoneIndex RootBoneIndex = FCompactPoseBoneIndex(0);
-	FTransform RootTransform = BoneContainer.GetRefPoseTransform(RootBoneIndex);
+	const FTransform RootTransform = BoneContainer.GetRefPoseTransform(RootBoneIndex);
 	
+	OutBoneTransforms.Reserve(RootChildren.Num() + 1);
 	OutBoneTransforms.Add(FBoneTransform(RootBoneIndex, RootTransform));
 	for (int32 Index = 0; Index < RootChildren.Num(); Index++)
 	{
-		OutBoneTransforms.Add(FBoneTransform(RootChildren[Index], ChildrenTransform[Index]));
+		OutBoneTransforms.Add(FBoneTransform(RootChildren[Index], CSPose.GetComponentSpaceTransform(RootChildren[Index])));
 	}
 
 	OutBoneTransforms.Sort(FCompareBoneTransformIndex());
diff --git a/Engine/Source/Runtime/Engine/Classes/Animation/AttributeBlendData.h b/Engine/Source/Runtime/Engine/Classes/Animation/AttributeBlendData.h
--- a/Engine/Source/Runtime/Engine/Classes/Animation/AttributeBlendData.h
+++ b/Engine/Source/Runtime/Engine/Classes/Animation/AttributeBlendData.h
@@ -114,10 +114,10 @@ namespace UE
 			struct FAttributeSet
 			{
 				/** Pointers to attribute values */
-				TArray<const uint8*> DataPtrs;
+				TArray<const uint8*, FAnimStackAllocator> DataPtrs;
 
 				/** Weight indices used to map to weight data */
-				TArray<int32> WeightIndices;
+				TArray<int32, FAnimStackAllocator> WeightIndices;
 				
 				/** Identifier of the attribute */
 				const FAttributeId* Identifier;
@@ -141,8 +141,8 @@ namespace UE
 			};
 
 			/** Processed unique and sets of attributes */
-			TArray<FAttributeSet> AttributeSets;
-			TArray<FUniqueAttribute> UniqueAttributes;
+			TArray<FAttributeSet, FAnimStackAllocator> AttributeSets;
+			TArray<FUniqueAttribute, FAnimStackAllocator> UniqueAttributes;
 
 			/** Container level weight data */
 			float UniformWeight;
@@ -155,7 +155,7 @@ namespace UE
 
 			/** Bone level weight data */
 			TArrayView<const FPerBoneBlendWeight> PerBoneWeights;
-			TArray<int32> HighestBoneWeightedIndices;
+			TArray<int32, FAnimStackAllocator> HighestBoneWeightedIndices;
 			inline const TArrayView<const float>& GetBoneWeights() const { return Weights; }
 			bool bPerBoneFilter = false;
 			
@@ -396,7 +396,7 @@ namespace UE
 			{
 				friend struct FAttributeBlendData;
 			protected:
-				TSingleIterator(const FAttributeBlendData& InData, const TArray<FUniqueAttribute>& InAttributes) : Data(InData), AttributesView(InAttributes), CurrentIndex(-1) {}
+				TSingleIterator(const FAttributeBlendData& InData, TArrayView<const FUniqueAttribute> InAttributes) : Data(InData), AttributesView(InAttributes), CurrentIndex(-1) {}
 			public:
 				/** Cycle through to next unique attribute, returns false if the end was reached */
 				bool Next()
@@ -468,7 +468,7 @@ namespace UE
 			{
 				friend struct FAttributeBlendData;
 			protected:
-				TSingleRawIterator(const FAttributeBlendData& InData, const TArray<FUniqueAttribute>& InAttributes) : Data(InData), AttributesView(InAttributes), CurrentIndex(-1) {}
+				TSingleRawIterator(const FAttributeBlendData& InData, TArrayView<const FUniqueAttribute> InAttributes) : Data(InData), AttributesView(InAttributes), CurrentIndex(-1) {}
 			public:
 				/** Cycle through to next unique attribute, returns false if the end was reached */
 				bool Next()
diff --git a/Engine/Source/Runtime/Engine/Classes/Curves/CurveBase.h b/Engine/Source/Runtime/Engine/Classes/Curves/CurveBase.h
--- a/Engine/Source/Runtime/Engine/Classes/Curves/CurveBase.h
+++ b/Engine/Source/Runtime/Engine/Classes/Curves/CurveBase.h
@@ -59,12 +59,27 @@ public:
 
 	//~ FCurveOwnerInterface interface
 
+	UE_DEPRECATED(5.6, "Use version taking a TAdderReserverRef")
 	virtual TArray<FRichCurveEditInfoConst> GetCurves() const override
 	{
 		TArray<FRichCurveEditInfoConst> Curves;
 		return Curves;
 	}
 
+	virtual void GetCurves(TAdderReserverRef<FRichCurveEditInfoConst> Curves) const override
+	{
+		// @todo: remove this entire implementation, since the default behavior should be to not return anything,
+		// once the GetCurves 5.6 deprecation gets removed
+		PRAGMA_DISABLE_DEPRECATION_WARNINGS
+		const TArray<FRichCurveEditInfoConst> TempCurves = GetCurves();
+		PRAGMA_ENABLE_DEPRECATION_WARNINGS
+
+		for (const FRichCurveEditInfoConst& TempCurve : TempCurves)
+		{
+			Curves.Add(TempCurve);
+		}
+	}
+
 	virtual TArray<FRichCurveEditInfo> GetCurves() override
 	{
 		TArray<FRichCurveEditInfo> Curves;
diff --git a/Engine/Source/Runtime/Engine/Classes/Curves/CurveFloat.h b/Engine/Source/Runtime/Engine/Classes/Curves/CurveFloat.h
--- a/Engine/Source/Runtime/Engine/Classes/Curves/CurveFloat.h
+++ b/Engine/Source/Runtime/Engine/Classes/Curves/CurveFloat.h
@@ -44,7 +44,9 @@ class UCurveFloat : public UCurveBase
 	ENGINE_API float GetFloatValue(float InTime) const;
 
 	// Begin FCurveOwnerInterface
+	UE_DEPRECATED(5.6, "Use version taking a TAdderReserverRef")
 	ENGINE_API virtual TArray<FRichCurveEditInfoConst> GetCurves() const override;
+	ENGINE_API virtual void GetCurves(TAdderReserverRef<FRichCurveEditInfoConst> Curves) const override;
 	ENGINE_API virtual TArray<FRichCurveEditInfo> GetCurves() override;
 	ENGINE_API virtual bool IsValidCurve( FRichCurveEditInfo CurveInfo ) override;
 
diff --git a/Engine/Source/Runtime/Engine/Classes/Curves/CurveLinearColor.h b/Engine/Source/Runtime/Engine/Classes/Curves/CurveLinearColor.h
--- a/Engine/Source/Runtime/Engine/Classes/Curves/CurveLinearColor.h
+++ b/Engine/Source/Runtime/Engine/Classes/Curves/CurveLinearColor.h
@@ -34,7 +34,9 @@ class UCurveLinearColor : public UCurveBase
 	FRichCurve FloatCurves[4];
 
 	// Begin FCurveOwnerInterface
+	UE_DEPRECATED(5.6, "Use version taking a TAdderReserverRef")
 	ENGINE_API virtual TArray<FRichCurveEditInfoConst> GetCurves() const override;
+	ENGINE_API virtual void GetCurves(TAdderReserverRef<FRichCurveEditInfoConst> Curves) const override;
 	ENGINE_API virtual TArray<FRichCurveEditInfo> GetCurves() override;
 	virtual bool IsLinearColorCurve() const override { return true; }
 
diff --git a/Engine/Source/Runtime/Engine/Classes/Curves/CurveOwnerInterface.h b/Engine/Source/Runtime/Engine/Classes/Curves/CurveOwnerInterface.h
--- a/Engine/Source/Runtime/Engine/Classes/Curves/CurveOwnerInterface.h
+++ b/Engine/Source/Runtime/Engine/Classes/Curves/CurveOwnerInterface.h
@@ -5,6 +5,7 @@
 #include "CoreMinimal.h"
 #include "Curves/RichCurve.h"
 #include "UObject/PackageReload.h"
+#include "Containers/AdderRef.h"
 
 /**
  * Interface you implement if you want the CurveEditor to be able to edit curves on you.
@@ -15,9 +16,12 @@ public:
 
 	virtual ~FCurveOwnerInterface() { }
 
-	/** Returns set of curves to edit. Must not release the curves while being edited. */
+	UE_DEPRECATED(5.6, "Use version taking a TAdderReserverRef")
 	virtual TArray<FRichCurveEditInfoConst> GetCurves() const = 0;
 
+	/** Returns set of curves to edit. Must not release the curves while being edited. */
+	virtual void GetCurves(TAdderReserverRef<FRichCurveEditInfoConst> Curves) const = 0;
+
 	/** Returns set of curves to query. Must not release the curves while being edited. */
 	virtual TArray<FRichCurveEditInfo> GetCurves() = 0;
 
diff --git a/Engine/Source/Runtime/Engine/Classes/Curves/CurveVector.h b/Engine/Source/Runtime/Engine/Classes/Curves/CurveVector.h
--- a/Engine/Source/Runtime/Engine/Classes/Curves/CurveVector.h
+++ b/Engine/Source/Runtime/Engine/Classes/Curves/CurveVector.h
@@ -40,7 +40,9 @@ class UCurveVector : public UCurveBase
 	ENGINE_API FVector GetVectorValue(float InTime) const;
 
 	// Begin FCurveOwnerInterface
-	virtual TArray<FRichCurveEditInfoConst> GetCurves() const override;
+	UE_DEPRECATED(5.6, "Use version taking a TAdderReserverRef")
+	ENGINE_API virtual TArray<FRichCurveEditInfoConst> GetCurves() const override;
+	ENGINE_API virtual void GetCurves(TAdderReserverRef<FRichCurveEditInfoConst> Curves) const override;
 	virtual TArray<FRichCurveEditInfo> GetCurves() override;
 
 	/** Determine if Curve is the same */
diff --git a/Engine/Source/Runtime/Engine/Classes/Engine/CurveTable.h b/Engine/Source/Runtime/Engine/Classes/Engine/CurveTable.h
--- a/Engine/Source/Runtime/Engine/Classes/Engine/CurveTable.h
+++ b/Engine/Source/Runtime/Engine/Classes/Engine/CurveTable.h
@@ -101,7 +101,9 @@ public:
 	//~ End  UObject Interface
 
 	//~ Begin FCurveOwnerInterface Interface.
+	UE_DEPRECATED(5.6, "Use version taking a TAdderReserverRef")
 	virtual TArray<FRichCurveEditInfoConst> GetCurves() const override;
+	virtual void GetCurves(TAdderReserverRef<FRichCurveEditInfoConst> Curves) const override;
 	virtual TArray<FRichCurveEditInfo> GetCurves() override;
 	virtual void ModifyOwner() override;
 	virtual void MakeTransactional() override;
diff --git a/Engine/Source/Runtime/Engine/Private/Animation/AnimInstanceProxy.cpp b/Engine/Source/Runtime/Engine/Private/Animation/AnimInstanceProxy.cpp
--- a/Engine/Source/Runtime/Engine/Private/Animation/AnimInstanceProxy.cpp
+++ b/Engine/Source/Runtime/Engine/Private/Animation/AnimInstanceProxy.cpp
@@ -1555,8 +1555,9 @@ void FAnimInstanceProxy::SlotEvaluatePoseWithBlendProfiles(const FName& SlotNode
 	TArray<float>& PerBoneWeightTotals = BlendProfileScratchData.PerBoneWeightTotals;
 	TArray<float>& PerBoneWeightTotalsAdditive = BlendProfileScratchData.PerBoneWeightTotalsAdditive;
 	TArray<float>& BoneBlendProfileScales = BlendProfileScratchData.BoneBlendProfileScales;
+	check(BlendProfileScratchData.IsEmpty());
+
 	const int32 NumBones = RequiredBones->GetCompactPoseNumBones();
-	PerBoneWeightTotals.Reset(NumBones);
 	PerBoneWeightTotals.AddUninitialized(NumBones);
 	for (int32 BoneIndex = 0; BoneIndex < NumBones; ++BoneIndex)
 	{
@@ -1579,14 +1580,6 @@ void FAnimInstanceProxy::SlotEvaluatePoseWithBlendProfiles(const FName& SlotNode
 	TArray<const FCompactPose*, TInlineAllocator<8>>& BlendingPoses = BlendProfileScratchData.BlendingPoses;
 	TArray<const FBlendedCurve*, TInlineAllocator<8>>& BlendingCurves = BlendProfileScratchData.BlendingCurves;
 	TArray<const UE::Anim::FStackAttributeContainer*, TInlineAllocator<8>>& BlendingAttributes = BlendProfileScratchData.BlendingAttributes;
-	Poses.Reset();
-	AdditivePoses.Reset();
-	PoseIndices.Reset();
-	AdditivePoseIndices.Reset();
-	BlendingWeights.Reset();
-	BlendingPoses.Reset();
-	BlendingCurves.Reset();
-	BlendingAttributes.Reset();
 
 	// Gather all the poses we're interested in.
 	int32 CurrentPoseIndex = 0;
@@ -1637,8 +1630,8 @@ void FAnimInstanceProxy::SlotEvaluatePoseWithBlendProfiles(const FName& SlotNode
 			}
 
 			PerBoneWeights.AddDefaulted();
-			PerBoneWeights[CurrentPoseIndex].Reset(NumBones);
 			PerBoneWeights[CurrentPoseIndex].AddUninitialized(NumBones);
+
 			const float PoseWeight = EvalState.BlendInfo.GetBlendedValue();
 			if (BlendProfile)
 			{
@@ -1868,6 +1861,8 @@ void FAnimInstanceProxy::SlotEvaluatePoseWithBlendProfiles(const FName& SlotNode
 		OutBlendedAnimationPoseData.GetCurve().Accumulate(AdditiveAnimationPoseData.GetCurve(), AdditivePose.Weight);
 		UE::Anim::Attributes::AccumulateAttributes(AdditiveAnimationPoseData.GetAttributes(), OutBlendedAnimationPoseData.GetAttributes(), AdditivePose.Weight, AdditivePose.AdditiveType);
 	}
+
+	BlendProfileScratchData.Reset();
 }
 
 void FAnimInstanceProxy::SlotEvaluatePose(const FName& SlotNodeName, const FAnimationPoseData& SourceAnimationPoseData, float InSourceWeight, FAnimationPoseData& OutBlendedAnimationPoseData, float InBlendWeight, float InTotalNodeWeight)
@@ -1909,12 +1904,7 @@ void FAnimInstanceProxy::SlotEvaluatePose(const FName& SlotNodeName, const FAnim
 	TArray<const FCompactPose*, TInlineAllocator<8>>& BlendingPoses = BlendProfileScratchData.BlendingPoses;
 	TArray<const FBlendedCurve*, TInlineAllocator<8>>& BlendingCurves = BlendProfileScratchData.BlendingCurves;
 	TArray<const UE::Anim::FStackAttributeContainer*, TInlineAllocator<8>>& BlendingAttributes = BlendProfileScratchData.BlendingAttributes;
-	NonAdditivePoses.Reset();
-	AdditivePoses.Reset();
-	BlendingWeights.Reset();
-	BlendingPoses.Reset();
-	BlendingCurves.Reset();
-	BlendingAttributes.Reset();
+	check(BlendProfileScratchData.IsEmpty());
 
 	// first pass we go through collect weights and valid montages. 
 #if DEBUG_MONTAGEINSTANCE_WEIGHT
@@ -1930,6 +1920,8 @@ void FAnimInstanceProxy::SlotEvaluatePose(const FName& SlotNodeName, const FAnim
 			BlendedPose = SourcePose;
 			BlendedCurve = SourceCurve;
 			BlendedAttributes = SourceAttributes;
+
+			BlendProfileScratchData.Reset();
 			return;
 		}
 
@@ -2082,6 +2074,8 @@ void FAnimInstanceProxy::SlotEvaluatePose(const FName& SlotNodeName, const FAnim
 			FAnimationRuntime::AccumulateAdditivePose(OutBlendedAnimationPoseData, AdditiveAnimationPoseData, AdditivePose.Weight, AdditivePose.AdditiveType);
 		}
 	}
+
+	BlendProfileScratchData.Reset();
 }
 
 //to debug montage weight
diff --git a/Engine/Source/Runtime/Engine/Private/Animation/AnimNode_StateMachine.cpp b/Engine/Source/Runtime/Engine/Private/Animation/AnimNode_StateMachine.cpp
--- a/Engine/Source/Runtime/Engine/Private/Animation/AnimNode_StateMachine.cpp
+++ b/Engine/Source/Runtime/Engine/Private/Animation/AnimNode_StateMachine.cpp
@@ -611,7 +611,7 @@ void FAnimNode_StateMachine::Update_AnyThread(const FAnimationUpdateContext& Con
 			// if we find an inactive one, remove all older transitions and break out
 			if (!ActiveTransitionArray[Index].bActive)
 			{
-				ActiveTransitionArray.RemoveAt(0, Index+1);
+				ActiveTransitionArray.RemoveAt(0, Index+1, EAllowShrinking::No);
 				break;
 			}
 		}
@@ -1027,7 +1027,7 @@ void FAnimNode_StateMachine::Evaluate_AnyThread(FPoseContext& Output)
 		{
 			delete CachedPose;
 		}
-		StateCachedPoses.Empty();
+		StateCachedPoses.Reset();
 	}
 	else if (!IsAConduitState(CurrentState))
 	{
diff --git a/Engine/Source/Runtime/Engine/Private/Animation/AnimationRuntime.cpp b/Engine/Source/Runtime/Engine/Private/Animation/AnimationRuntime.cpp
--- a/Engine/Source/Runtime/Engine/Private/Animation/AnimationRuntime.cpp
+++ b/Engine/Source/Runtime/Engine/Private/Animation/AnimationRuntime.cpp
@@ -595,7 +595,7 @@ void FAnimationRuntime::BlendTwoPosesTogetherPerBone(const FAnimationPoseData& S
 }
 
 template <int32 TRANSFORM_BLEND_MODE>
-void BlendPosePerBone(const TArray<FBoneIndexType>& RequiredBoneIndices, const TArray<int32>& PerBoneIndices, const FBlendSampleData& BlendSampleDataCache, FTransformArrayA2& ResultAtoms, const FTransformArrayA2& SourceAtoms)
+void BlendPosePerBone(const TArrayView<int32> RequiredBoneIndices, const TArrayView<int32> PerBoneIndices, const FBlendSampleData& BlendSampleDataCache, FTransformArrayA2& ResultAtoms, const FTransformArrayA2& SourceAtoms)
 {
 	const float BlendWeight = BlendSampleDataCache.GetClampedWeight();
 	for (int32 i = 0; i < RequiredBoneIndices.Num(); ++i)
@@ -614,7 +614,7 @@ void BlendPosePerBone(const TArray<FBoneIndexType>& RequiredBoneIndices, const T
 }
 
 template <int32 TRANSFORM_BLEND_MODE>
-void BlendPosePerBone(const TArray<int32>& PerBoneIndices, const FBlendSampleData& BlendSampleDataCache, FCompactPose& ResultPose, const FCompactPose& SourcePose)
+void BlendPosePerBone(const TArrayView<int32> PerBoneIndices, const FBlendSampleData& BlendSampleDataCache, FCompactPose& ResultPose, const FCompactPose& SourcePose)
 {
 	const float BlendWeight = BlendSampleDataCache.GetClampedWeight();
 	for (FCompactPoseBoneIndex BoneIndex : SourcePose.ForEachBoneIndex())
@@ -654,9 +654,11 @@ void FAnimationRuntime::BlendPosesTogetherPerBone(TArrayView<const FCompactPose>
 	FBlendedCurve& OutCurve = OutAnimationPoseData.GetCurve();
 	UE::Anim::FStackAttributeContainer& OutAttributes = OutAnimationPoseData.GetAttributes();
 
+	FMemMark Mark(FMemStack::Get());
+	TArray<int32, TMemStackAllocator<>> PerBoneIndices;
+	PerBoneIndices.SetNumUninitialized(OutPose.GetNumBones());
+
 	const FBoneContainer& RequiredBones = OutPose.GetBoneContainer();
-	TArray<int32> PerBoneIndices;
-	PerBoneIndices.AddUninitialized(OutPose.GetNumBones());
 	TSharedPtr<IInterpolationIndexProvider::FPerBoneInterpolationData> Data = InterpolationIndexProvider->GetPerBoneInterpolationData(RequiredBones.GetSkeletonAsset());
 	for(FCompactPoseBoneIndex BoneIndex : OutPose.ForEachBoneIndex())
 	{
@@ -674,8 +676,8 @@ void FAnimationRuntime::BlendPosesTogetherPerBone(TArrayView<const FCompactPose>
 
 	if (SourceCurves.Num() > 0)
 	{
-		TArray<float, TInlineAllocator<16>> SourceWeights;
-		SourceWeights.AddUninitialized(BlendSampleDataCache.Num());
+		TArray<float, TMemStackAllocator<>> SourceWeights;
+		SourceWeights.SetNumUninitialized(BlendSampleDataCache.Num());
 		for (int32 CacheIndex = 0; CacheIndex < BlendSampleDataCache.Num(); ++CacheIndex)
 		{
 			SourceWeights[CacheIndex] = BlendSampleDataCache[CacheIndex].TotalWeight;
@@ -713,9 +715,12 @@ void FAnimationRuntime::BlendPosesTogetherPerBone(TArrayView<const FCompactPose>
 	FBlendedCurve& OutCurve = OutAnimationPoseData.GetCurve();
 	UE::Anim::FStackAttributeContainer& OutAttributes = OutAnimationPoseData.GetAttributes();
 
+	FMemMark Mark(FMemStack::Get());
+	TArray<int32, TMemStackAllocator<>> PerBoneIndices;
+	PerBoneIndices.SetNumUninitialized(OutPose.GetNumBones());
+
 	const FBoneContainer& RequiredBones = OutPose.GetBoneContainer();
-	TArray<int32> PerBoneIndices;
-	PerBoneIndices.AddUninitialized(OutPose.GetNumBones());
+
 	TSharedPtr<IInterpolationIndexProvider::FPerBoneInterpolationData> Data = InterpolationIndexProvider->GetPerBoneInterpolationData(OutPose.GetBoneContainer().GetSkeletonAsset());
 	for(FCompactPoseBoneIndex BoneIndex : OutPose.ForEachBoneIndex())
 	{
@@ -1758,7 +1763,7 @@ void FAnimationRuntime::EnsureParentsPresent(TArray<FBoneIndexType>& BoneIndices
 void FAnimationRuntime::ExcludeBonesWithNoParents(const TArray<int32>& BoneIndices, const FReferenceSkeleton& RefSkeleton, TArray<int32>& FilteredRequiredBones)
 {
 	// Filter list, we only want bones that have their parents present in this array.
-	FilteredRequiredBones.Empty(BoneIndices.Num());
+	FilteredRequiredBones.Reset(BoneIndices.Num());
 
 	for (int32 Index=0; Index<BoneIndices.Num(); Index++)
 	{
@@ -1784,7 +1789,7 @@ void FAnimationRuntime::ExcludeBonesWithNoParents(const TArray<int32>& BoneIndic
 	}
 }
 
-void FAnimationRuntime::UpdateDesiredBoneWeight(const TArray<FPerBoneBlendWeight>& SrcBoneBlendWeights, TArray<FPerBoneBlendWeight>& TargetBoneBlendWeights, const TArray<float>& BlendWeights)
+void FAnimationRuntime::UpdateDesiredBoneWeight(const TArrayView<FPerBoneBlendWeight> SrcBoneBlendWeights, TArrayView<FPerBoneBlendWeight> TargetBoneBlendWeights, const TArrayView<float> BlendWeights)
 {
 	// in the future, cache this outside
 	ensure (TargetBoneBlendWeights.Num() == SrcBoneBlendWeights.Num());
@@ -1833,27 +1838,16 @@ extern "C" const uint8* GetTransformFromArray(const uint8 *BlendPoseBase, const
 	return reinterpret_cast<const uint8*>(BlendPose);
 }
 
-void FAnimationRuntime::BlendPosesPerBoneFilter(
-	struct FCompactPose& BasePose, 
-	const TArray<struct FCompactPose>& BlendPoses, 
-	struct FBlendedCurve& BaseCurve, 
-	const TArray<struct FBlendedCurve>& BlendedCurves, 
-	struct FCompactPose& OutPose, 
-	struct FBlendedCurve& OutCurve, 
-	TArray<FPerBoneBlendWeight>& BoneBlendWeights, 
-	EBlendPosesPerBoneFilterFlags BlendFlags,
-	ECurveBlendOption::Type CurveBlendOption)
-{
-	UE::Anim::FStackAttributeContainer TempAttributes;
-	FAnimationPoseData AnimationPoseData = {OutPose, OutCurve, TempAttributes };
-	
-	BlendPosesPerBoneFilter(BasePose, BlendPoses, BaseCurve, BlendedCurves, TempAttributes, {}, AnimationPoseData, BoneBlendWeights, BlendFlags, CurveBlendOption);
-}
-
-void FAnimationRuntime::BlendPosesPerBoneFilter(FCompactPose& BasePose, const TArray<FCompactPose>& BlendPoses, FBlendedCurve& BaseCurve, const TArray<FBlendedCurve>& BlendedCurves,
-	UE::Anim::FStackAttributeContainer& BaseAttributes,
-	const TArray<UE::Anim::FStackAttributeContainer>& BlendAttributes,
-	FAnimationPoseData& OutAnimationPoseData, TArray<FPerBoneBlendWeight>& BoneBlendWeights, EBlendPosesPerBoneFilterFlags BlendFlags, enum ECurveBlendOption::Type CurveBlendOption)
+void FAnimationRuntime::BlendPosesPerBoneFilter(FCompactPose& BasePose,
+		const TConstArrayView<FCompactPose> BlendPoses,
+		FBlendedCurve& BaseCurve,
+		const TConstArrayView<FBlendedCurve> BlendedCurves,
+		UE::Anim::FStackAttributeContainer& BaseAttributes,
+		const TArrayView<UE::Anim::FStackAttributeContainer> BlendAttributes,
+		FAnimationPoseData& OutAnimationPoseData,
+		const TArrayView<FPerBoneBlendWeight> BoneBlendWeights,
+		EBlendPosesPerBoneFilterFlags BlendFlags,
+		enum ECurveBlendOption::Type CurveBlendOption)
 {
 	SCOPE_CYCLE_COUNTER(STAT_BlendPosesPerBoneFilter);
 
@@ -2682,7 +2676,7 @@ const FTransform& FAnimationRuntime::GetComponentSpaceTransformWithCache(const F
 
 void FAnimationRuntime::FillUpComponentSpaceTransforms(const FReferenceSkeleton& RefSkeleton, const TArrayView<const FTransform> &BoneSpaceTransforms, TArray<FTransform> &ComponentSpaceTransforms)
 {
-	ComponentSpaceTransforms.Empty(BoneSpaceTransforms.Num());
+	ComponentSpaceTransforms.Reset(BoneSpaceTransforms.Num());
 	ComponentSpaceTransforms.AddUninitialized(BoneSpaceTransforms.Num());
 
 	// initialize to identity since some of them don't have tracks
diff --git a/Engine/Source/Runtime/Engine/Private/Animation/AttributeBlendData.cpp b/Engine/Source/Runtime/Engine/Private/Animation/AttributeBlendData.cpp
--- a/Engine/Source/Runtime/Engine/Private/Animation/AttributeBlendData.cpp
+++ b/Engine/Source/Runtime/Engine/Private/Animation/AttributeBlendData.cpp
@@ -230,8 +230,8 @@ namespace UE { namespace Anim {
 		const int32 TypeIndex = AttributeContainers.FindTypeIndex(AttributeScriptStruct);
 		if (TypeIndex != INDEX_NONE)
 		{
-			const TArray<TWrappedAttribute<FAnimStackAllocator>>& ValuesArray = AttributeContainers.GetValues(TypeIndex);
-			const TArray<FAttributeId>& AttributeIdentifiers = AttributeContainers.GetKeys(TypeIndex);
+			const TConstArrayView<TWrappedAttribute<FAnimStackAllocator>> ValuesArray = AttributeContainers.GetValues(TypeIndex);
+			const TConstArrayView<FAttributeId> AttributeIdentifiers = AttributeContainers.GetKeys(TypeIndex);
 			
 			for (int32 AttributeIndex = 0; AttributeIndex < AttributeIdentifiers.Num(); ++AttributeIndex)
 			{
diff --git a/Engine/Source/Runtime/Engine/Private/Animation/AttributesRuntime.cpp b/Engine/Source/Runtime/Engine/Private/Animation/AttributesRuntime.cpp
--- a/Engine/Source/Runtime/Engine/Private/Animation/AttributesRuntime.cpp
+++ b/Engine/Source/Runtime/Engine/Private/Animation/AttributesRuntime.cpp
@@ -104,7 +104,7 @@ void Attributes::OverrideAttributes(const FStackAttributeContainer& SourceAttrib
 	UniqueStructs.Append(SourceAttributes.GetUniqueTypes());
 	UniqueStructs.Append(OutAttributes.GetUniqueTypes());
 
-	const TArray<const FStackAttributeContainer*> Array = { &SourceAttributes };
+	const FStackAttributeContainer* Array[] = { &SourceAttributes };
 	for (const TWeakObjectPtr<UScriptStruct>& WeakScriptStruct : UniqueStructs)
 	{
 		const IAttributeBlendOperator* Operator = AttributeTypes::GetTypeOperator(WeakScriptStruct);
@@ -123,7 +123,7 @@ void Attributes::AccumulateAttributes(const FStackAttributeContainer& SourceAttr
 	UniqueStructs.Append(SourceAttributes.GetUniqueTypes());
 	UniqueStructs.Append(OutAttributes.GetUniqueTypes());
 
-	const TArray<const FStackAttributeContainer*> Array = { &SourceAttributes };
+	const FStackAttributeContainer* Array[] = { &SourceAttributes };
 	for (const TWeakObjectPtr<UScriptStruct>& WeakScriptStruct : UniqueStructs)
 	{
 		const IAttributeBlendOperator* Operator = AttributeTypes::GetTypeOperator(WeakScriptStruct);
@@ -139,7 +139,7 @@ void Attributes::ConvertToAdditive(const FStackAttributeContainer& BaseAttribute
 	TSet<TWeakObjectPtr<UScriptStruct>, DefaultKeyFuncs<TWeakObjectPtr<UScriptStruct>>, TInlineSetAllocator<4>> UniqueStructs;
 	UniqueStructs.Append(BaseAttributes.GetUniqueTypes());
 
-	const TArray<const FStackAttributeContainer*> Array = { &BaseAttributes };
+	const FStackAttributeContainer* Array[] = { &BaseAttributes };
 	for (const TWeakObjectPtr<UScriptStruct>& WeakScriptStruct : UniqueStructs)
 	{
 		const IAttributeBlendOperator* Operator = AttributeTypes::GetTypeOperator(WeakScriptStruct);
@@ -361,17 +361,17 @@ void Attributes::BlendAttributesPerBone(TArrayView<const FStackAttributeContaine
 	}
 	else
 	{
-		TArray<float> Weights;
-		for (const FBlendSampleData& SampleData : BlendSampleDataCache)
+		TArrayView<float> Weights((float*)FMemory_Alloca(BlendSampleDataCache.Num() * sizeof(float)), BlendSampleDataCache.Num());
+		for (int32 WeightIndex = 0; WeightIndex < BlendSampleDataCache.Num(); ++WeightIndex)
 		{
-			Weights.Add(SampleData.GetClampedWeight());
+			Weights[WeightIndex] = BlendSampleDataCache[WeightIndex].GetClampedWeight();
 		}
 
 		for (const TWeakObjectPtr<UScriptStruct>& WeakScriptStruct : UniqueStructs)
 		{
 			const IAttributeBlendOperator* Operator = AttributeTypes::GetTypeOperator(WeakScriptStruct);
 
-			FAttributeBlendData AttributeBlendData = FAttributeBlendData::PerContainerWeighted(SourceAttributes, MakeArrayView(Weights), WeakScriptStruct.Get());
+			FAttributeBlendData AttributeBlendData = FAttributeBlendData::PerContainerWeighted(SourceAttributes, Weights, WeakScriptStruct.Get());
 			Operator->Blend(AttributeBlendData, &OutAttributes);
 		}
 	}
@@ -401,24 +401,23 @@ void Attributes::BlendAttributesPerBone(TArrayView<const FStackAttributeContaine
 	}
 	else
 	{
-		TArray<float> Weights;
-		for (const int32& SampleIndex : BlendSampleDataCacheIndices)
+		TArrayView<float> Weights((float*)FMemory_Alloca(BlendSampleDataCacheIndices.Num() * sizeof(float)), BlendSampleDataCacheIndices.Num());
+		for (int32 WeightIndex = 0; WeightIndex < BlendSampleDataCacheIndices.Num(); ++WeightIndex)
 		{
-			const FBlendSampleData& SampleData = BlendSampleDataCache[SampleIndex];
-			Weights.Add(SampleData.GetClampedWeight());
+			Weights[WeightIndex] = BlendSampleDataCache[BlendSampleDataCacheIndices[WeightIndex]].GetClampedWeight();
 		}
 
 		for (const TWeakObjectPtr<UScriptStruct>& WeakScriptStruct : UniqueStructs)
 		{
 			const IAttributeBlendOperator* Operator = AttributeTypes::GetTypeOperator(WeakScriptStruct);
 
-			FAttributeBlendData AttributeBlendData = FAttributeBlendData::PerContainerWeighted(SourceAttributes, MakeArrayView(Weights), WeakScriptStruct.Get());
+			FAttributeBlendData AttributeBlendData = FAttributeBlendData::PerContainerWeighted(SourceAttributes, Weights, WeakScriptStruct.Get());
 			Operator->Blend(AttributeBlendData, &OutAttributes);
 		}
 	}
 }
 
-void Attributes::BlendAttributesPerBoneFilter(const FStackAttributeContainer& BaseAttributes, const TArrayView<const FStackAttributeContainer> BlendAttributes, const TArray<FPerBoneBlendWeight>& BoneBlendWeights, FStackAttributeContainer& OutAttributes)
+void Attributes::BlendAttributesPerBoneFilter(const FStackAttributeContainer& BaseAttributes, const TArrayView<const FStackAttributeContainer> BlendAttributes, const TConstArrayView<FPerBoneBlendWeight> BoneBlendWeights, FStackAttributeContainer& OutAttributes)
 {
 	QUICK_SCOPE_CYCLE_COUNTER(STAT_BlendAttributesPerBoneFilter);
 
@@ -479,14 +478,16 @@ void Attributes::MirrorAttributes(FStackAttributeContainer& CustomAttributes, co
 		return true;
 	};
 
+	// @todo: initialize this array with a properly sized TInlineAllocator
+	TArray<int32> SortedBoneIndices;
 	for (const TWeakObjectPtr<UScriptStruct>& WeakScriptStruct : CustomAttributes.GetUniqueTypes())
 	{
 		const int32 TypeIndex = CustomAttributes.FindTypeIndex(WeakScriptStruct.Get());
 		if (TypeIndex != INDEX_NONE)
 		{
 			MirrorSets.Reset();
-			const TArray<int32>& UniqueBoneIndices = CustomAttributes.GetUniqueTypedBoneIndices(TypeIndex);
-			TArray<int32> SortedBoneIndices(UniqueBoneIndices); 
+			const TConstArrayView<int32> UniqueBoneIndices = CustomAttributes.GetUniqueTypedBoneIndices(TypeIndex);
+			SortedBoneIndices = UniqueBoneIndices; 
 			SortedBoneIndices.Sort();
 
 			if(CompactPoseMirrorBones.Num())
@@ -526,8 +527,8 @@ void Attributes::MirrorAttributes(FStackAttributeContainer& CustomAttributes, co
 			    }
 			}
 
-			TArray<TWrappedAttribute<FAnimStackAllocator>>& ValuesArray = FStackAttributeContainerAccessor::GetValues(CustomAttributes, TypeIndex);
-			const TArray<FAttributeId>& AttributeIdentifiers = CustomAttributes.GetKeys(TypeIndex);
+			TArrayView<TWrappedAttribute<FAnimStackAllocator>> ValuesArray = FStackAttributeContainerAccessor::GetValues(CustomAttributes, TypeIndex);
+			const TConstArrayView<FAttributeId> AttributeIdentifiers = CustomAttributes.GetKeys(TypeIndex);
 
 			// gather attributes that are on mirrored bones 
 			for (int32 AttributeIndex = 0; AttributeIndex < AttributeIdentifiers.Num(); ++AttributeIndex)
diff --git a/Engine/Source/Runtime/Engine/Private/Animation/BlendProfileScratchData.h b/Engine/Source/Runtime/Engine/Private/Animation/BlendProfileScratchData.h
--- a/Engine/Source/Runtime/Engine/Private/Animation/BlendProfileScratchData.h
+++ b/Engine/Source/Runtime/Engine/Private/Animation/BlendProfileScratchData.h
@@ -16,6 +16,9 @@ namespace UE { namespace Anim { struct FStackAttributeContainer; } }
  */
 struct FBlendProfileScratchData : public TThreadSingleton<FBlendProfileScratchData>
 {
+	// @todo: TArray<TArray<>> allocations are not preserved during resets. We should either implement a new TArrayOfArrays container
+	// or change the deallocation policy perhaps via template specialization?
+
 	TArray<TArray<float>> PerBoneWeights;		// A set of bone weights, per montage instance. Index this like [PoseIndex][CompactPoseBoneIndex].
 	TArray<float> PerBoneWeightTotals;			// The bone weight totals for non-additive poses, used for normalizing weights.
 	TArray<float> PerBoneWeightTotalsAdditive;	// The bone weight totals for additive poses, used for normalizing weights.
@@ -31,4 +34,36 @@ struct FBlendProfileScratchData : public TThreadSingleton<FBlendProfileScratchDa
 	TArray<const FCompactPose*, TInlineAllocator<8>> BlendingPoses;		// The non additive poses to blend.
 	TArray<const FBlendedCurve*, TInlineAllocator<8>> BlendingCurves;	// The curves to blend.
 	TArray<const UE::Anim::FStackAttributeContainer*, TInlineAllocator<8>> BlendingAttributes;	// The attributes to blend.
+
+	void Reset()
+	{
+		PerBoneWeights.Reset();
+		PerBoneWeightTotals.Reset();
+		PerBoneWeightTotalsAdditive.Reset();
+		BoneBlendProfileScales.Reset();
+		Poses.Reset();
+		AdditivePoses.Reset();
+		PoseIndices.Reset();
+		AdditivePoseIndices.Reset();
+		BlendingWeights.Reset();
+		BlendingPoses.Reset();
+		BlendingCurves.Reset();
+		BlendingAttributes.Reset();
+	}
+
+	bool IsEmpty() const
+	{
+		return PerBoneWeights.IsEmpty() &&
+			PerBoneWeightTotals.IsEmpty() &&
+			PerBoneWeightTotalsAdditive.IsEmpty() &&
+			BoneBlendProfileScales.IsEmpty() &&
+			Poses.IsEmpty() &&
+			AdditivePoses.IsEmpty() &&
+			PoseIndices.IsEmpty() &&
+			AdditivePoseIndices.IsEmpty() &&
+			BlendingWeights.IsEmpty() &&
+			BlendingPoses.IsEmpty() &&
+			BlendingCurves.IsEmpty() &&
+			BlendingAttributes.IsEmpty();
+	}
 };
diff --git a/Engine/Source/Runtime/Engine/Private/Animation/Skeleton.cpp b/Engine/Source/Runtime/Engine/Private/Animation/Skeleton.cpp
--- a/Engine/Source/Runtime/Engine/Private/Animation/Skeleton.cpp
+++ b/Engine/Source/Runtime/Engine/Private/Animation/Skeleton.cpp
@@ -2220,7 +2220,7 @@ UBlendProfile* USkeleton::GetBlendProfile(const FName& InProfileName)
 {
 	TObjectPtr<UBlendProfile>* FoundProfile = BlendProfiles.FindByPredicate([InProfileName](const UBlendProfile* Profile)
 	{
-		return Profile->GetName() == InProfileName.ToString();
+		return Profile->GetFName() == InProfileName;
 	});
 
 	if(FoundProfile)
diff --git a/Engine/Source/Runtime/Engine/Private/AnimationAsset.cpp b/Engine/Source/Runtime/Engine/Private/AnimationAsset.cpp
--- a/Engine/Source/Runtime/Engine/Private/AnimationAsset.cpp
+++ b/Engine/Source/Runtime/Engine/Private/AnimationAsset.cpp
@@ -800,7 +800,8 @@ void FBlendSampleData::NormalizeDataWeight(TArray<FBlendSampleData>& SampleDataL
 	check(SampleDataList.Num() > 0);
 	const int32 NumBones = SampleDataList[0].PerBoneBlendData.Num();
 
-	TArray<float> PerBoneTotalSums;
+	FMemMark Mark(FMemStack::Get());
+	TArray<float, TMemStackAllocator<>> PerBoneTotalSums;
 	PerBoneTotalSums.AddZeroed(NumBones);
 
 	for (int32 PoseIndex = 0; PoseIndex < SampleDataList.Num(); PoseIndex++)
diff --git a/Engine/Source/Runtime/Engine/Private/CurveTable.cpp b/Engine/Source/Runtime/Engine/Private/CurveTable.cpp
--- a/Engine/Source/Runtime/Engine/Private/CurveTable.cpp
+++ b/Engine/Source/Runtime/Engine/Private/CurveTable.cpp
@@ -1017,6 +1017,14 @@ TArray<FRichCurveEditInfoConst> UCurveTable::GetCurves() const
 	return Curves;
 }
 
+void UCurveTable::GetCurves(TAdderReserverRef<FRichCurveEditInfoConst> Curves) const
+{
+	for (auto Iter = RowMap.CreateConstIterator(); Iter; ++Iter)
+	{
+		Curves.Add(FRichCurveEditInfoConst(Iter.Value(), Iter.Key()));
+	}
+}
+
 TArray<FRichCurveEditInfo> UCurveTable::GetCurves()
 {
 	TArray<FRichCurveEditInfo> Curves;
diff --git a/Engine/Source/Runtime/Engine/Private/Curves/CurveBase.cpp b/Engine/Source/Runtime/Engine/Private/Curves/CurveBase.cpp
--- a/Engine/Source/Runtime/Engine/Private/Curves/CurveBase.cpp
+++ b/Engine/Source/Runtime/Engine/Private/Curves/CurveBase.cpp
@@ -20,7 +20,8 @@ UCurveBase::UCurveBase(const FObjectInitializer& ObjectInitializer)
 
 void UCurveBase::GetTimeRange(float& MinTime, float& MaxTime) const
 {
-	TArray<FRichCurveEditInfoConst> Curves = GetCurves();
+	TArray<FRichCurveEditInfoConst, TInlineAllocator<16>> Curves;
+	GetCurves(Curves);
 
 	if (Curves.Num() > 0)
 	{
@@ -48,7 +49,8 @@ void UCurveBase::GetTimeRange(double& MinTime, double& MaxTime) const
 
 void UCurveBase::GetValueRange(float& MinValue, float& MaxValue) const
 {
-	TArray<FRichCurveEditInfoConst> Curves = GetCurves();
+	TArray<FRichCurveEditInfoConst, TInlineAllocator<16>> Curves;
+	GetCurves(Curves);
 
 	if (Curves.Num() > 0)
 	{
diff --git a/Engine/Source/Runtime/Engine/Private/Curves/CurveFloat.cpp b/Engine/Source/Runtime/Engine/Private/Curves/CurveFloat.cpp
--- a/Engine/Source/Runtime/Engine/Private/Curves/CurveFloat.cpp
+++ b/Engine/Source/Runtime/Engine/Private/Curves/CurveFloat.cpp
@@ -54,6 +54,11 @@ TArray<FRichCurveEditInfoConst> UCurveFloat::GetCurves() const
 	return Curves;
 }
 
+void UCurveFloat::GetCurves(TAdderReserverRef<FRichCurveEditInfoConst> Curves) const
+{
+	Curves.Add(FRichCurveEditInfoConst(&FloatCurve));
+}
+
 TArray<FRichCurveEditInfo> UCurveFloat::GetCurves()
 {
 	TArray<FRichCurveEditInfo> Curves;
diff --git a/Engine/Source/Runtime/Engine/Private/Curves/CurveLinearColor.cpp b/Engine/Source/Runtime/Engine/Private/Curves/CurveLinearColor.cpp
--- a/Engine/Source/Runtime/Engine/Private/Curves/CurveLinearColor.cpp
+++ b/Engine/Source/Runtime/Engine/Private/Curves/CurveLinearColor.cpp
@@ -192,6 +192,14 @@ TArray<FRichCurveEditInfoConst> UCurveLinearColor::GetCurves() const
 	return Curves;
 }
 
+void UCurveLinearColor::GetCurves(TAdderReserverRef<FRichCurveEditInfoConst> Curves) const 
+{
+	Curves.Add(FRichCurveEditInfoConst(&FloatCurves[0], RedCurveName));
+	Curves.Add(FRichCurveEditInfoConst(&FloatCurves[1], GreenCurveName));
+	Curves.Add(FRichCurveEditInfoConst(&FloatCurves[2], BlueCurveName));
+	Curves.Add(FRichCurveEditInfoConst(&FloatCurves[3], AlphaCurveName));
+}
+
 TArray<FRichCurveEditInfo> UCurveLinearColor::GetCurves()
 {
 	TArray<FRichCurveEditInfo> Curves;
diff --git a/Engine/Source/Runtime/Engine/Private/Curves/CurveVector.cpp b/Engine/Source/Runtime/Engine/Private/Curves/CurveVector.cpp
--- a/Engine/Source/Runtime/Engine/Private/Curves/CurveVector.cpp
+++ b/Engine/Source/Runtime/Engine/Private/Curves/CurveVector.cpp
@@ -83,6 +83,13 @@ TArray<FRichCurveEditInfoConst> UCurveVector::GetCurves() const
 	return Curves;
 }
 
+void UCurveVector::GetCurves(TAdderReserverRef<FRichCurveEditInfoConst> Curves) const 
+{
+	Curves.Add(FRichCurveEditInfoConst(&FloatCurves[0], XCurveName));
+	Curves.Add(FRichCurveEditInfoConst(&FloatCurves[1], YCurveName));
+	Curves.Add(FRichCurveEditInfoConst(&FloatCurves[2], ZCurveName));
+}
+
 TArray<FRichCurveEditInfo> UCurveVector::GetCurves() 
 {
 	TArray<FRichCurveEditInfo> Curves;
diff --git a/Engine/Source/Runtime/Engine/Public/Animation/ActiveStateMachineScope.h b/Engine/Source/Runtime/Engine/Public/Animation/ActiveStateMachineScope.h
--- a/Engine/Source/Runtime/Engine/Public/Animation/ActiveStateMachineScope.h
+++ b/Engine/Source/Runtime/Engine/Public/Animation/ActiveStateMachineScope.h
@@ -30,7 +30,7 @@ struct FEncounteredStateMachineStack
 		int32 StateMachineIndex = INDEX_NONE;
 		int32 StateIndex = INDEX_NONE;;
 	};
-	TArray<FStateMachineEntry, TInlineAllocator<4>> StateStack;
+	TArray<FStateMachineEntry, TInlineAllocator<4, TMemStackAllocator<>>> StateStack;
 	static ENGINE_API FEncounteredStateMachineStack InitStack(const FEncounteredStateMachineStack& ParentStack, int32 InStateMachineIndex, int32 InStateIndex);
 	static ENGINE_API FEncounteredStateMachineStack InitStack(int32 InStateMachineIndex, int32 InStateIndex);
 };
diff --git a/Engine/Source/Runtime/Engine/Public/Animation/AnimCurveTypes.h b/Engine/Source/Runtime/Engine/Public/Animation/AnimCurveTypes.h
--- a/Engine/Source/Runtime/Engine/Public/Animation/AnimCurveTypes.h
+++ b/Engine/Source/Runtime/Engine/Public/Animation/AnimCurveTypes.h
@@ -1177,6 +1177,17 @@ struct FRawCurveTracks
 #if WITH_EDITORONLY_DATA
 		VectorCurves.Empty();
 		TransformCurves.Empty();
+#endif
+	}
+	/** 
+	 * Clear all keys preserving allocations
+	 */
+	void Reset()
+	{
+		FloatCurves.Reset();
+#if WITH_EDITORONLY_DATA
+		VectorCurves.Reset();
+		TransformCurves.Reset();
 #endif
 	}
 
diff --git a/Engine/Source/Runtime/Engine/Public/Animation/AttributesRuntime.h b/Engine/Source/Runtime/Engine/Public/Animation/AttributesRuntime.h
--- a/Engine/Source/Runtime/Engine/Public/Animation/AttributesRuntime.h
+++ b/Engine/Source/Runtime/Engine/Public/Animation/AttributesRuntime.h
@@ -81,7 +81,7 @@ namespace UE
 			static ENGINE_API void BlendAttributesPerBone(TArrayView<const FStackAttributeContainer> SourceAttributes, TArrayView<const int32> PerBoneInterpolationIndices, TArrayView<const FBlendSampleData> BlendSampleDataCache, TArrayView<const int32> BlendSampleDataCacheIndices, FStackAttributeContainer& OutAttributes);
 
 			/* Blend custom attribute values from N set of inputs, using bone filter pose weights */
-			static ENGINE_API void BlendAttributesPerBoneFilter(const FStackAttributeContainer& BaseAttributes, const TArrayView<const FStackAttributeContainer> BlendAttributes, const TArray<FPerBoneBlendWeight>& BoneBlendWeights, FStackAttributeContainer& OutAttributes);
+			static ENGINE_API void BlendAttributesPerBoneFilter(const FStackAttributeContainer& BaseAttributes, const TArrayView<const FStackAttributeContainer> BlendAttributes, const TConstArrayView<FPerBoneBlendWeight> BoneBlendWeights, FStackAttributeContainer& OutAttributes);
 
 			/** Add any new or override existing custom attributes */
 			static ENGINE_API void OverrideAttributes(const FStackAttributeContainer& SourceAttributes, FStackAttributeContainer& OutAttributes, float Weight);
diff --git a/Engine/Source/Runtime/Engine/Public/AnimationRuntime.h b/Engine/Source/Runtime/Engine/Public/AnimationRuntime.h
--- a/Engine/Source/Runtime/Engine/Public/AnimationRuntime.h
+++ b/Engine/Source/Runtime/Engine/Public/AnimationRuntime.h
@@ -367,31 +367,22 @@ public:
 	* having multiple bone names with multiple weights, and filtering through which one is correct one
 	* I assume all those things should be determined before coming here and this only cares about weights
 	**/
-	UE_DEPRECATED(4.26, "Please use the BlendPosesPerBoneFilterwith different signature.")
 	static ENGINE_API void BlendPosesPerBoneFilter(
 		FCompactPose& BasePose,
-		const TArray<FCompactPose>& BlendPoses,
+		const TConstArrayView<FCompactPose> BlendPoses,
 		FBlendedCurve& BaseCurve,
-		const TArray<FBlendedCurve>& BlendCurves,
-		FCompactPose& OutPose,
-		FBlendedCurve& OutCurve,
-		TArray<FPerBoneBlendWeight>& BoneBlendWeights,
-		EBlendPosesPerBoneFilterFlags blendFlags,
-		enum ECurveBlendOption::Type CurveBlendOption);
-
-	static ENGINE_API void BlendPosesPerBoneFilter(
-		FCompactPose& BasePose,
-		const TArray<FCompactPose>& BlendPoses,
-		FBlendedCurve& BaseCurve,
-		const TArray<FBlendedCurve>& BlendCurves,
-		UE::Anim::FStackAttributeContainer& CustomAttributes,
-		const TArray<UE::Anim::FStackAttributeContainer>& BlendAttributes,
+		const TConstArrayView<FBlendedCurve> BlendedCurves,
+		UE::Anim::FStackAttributeContainer& BaseAttributes,
+		const TArrayView<UE::Anim::FStackAttributeContainer> BlendAttributes,
 		FAnimationPoseData& OutAnimationPoseData,
-		TArray<FPerBoneBlendWeight>& BoneBlendWeights,
-		EBlendPosesPerBoneFilterFlags blendFlags,
+		const TArrayView<FPerBoneBlendWeight> BoneBlendWeights,
+		EBlendPosesPerBoneFilterFlags BlendFlags,
 		enum ECurveBlendOption::Type CurveBlendOption);
 
-	static ENGINE_API void UpdateDesiredBoneWeight(const TArray<FPerBoneBlendWeight>& SrcBoneBlendWeights, TArray<FPerBoneBlendWeight>& TargetBoneBlendWeights, const TArray<float>& BlendWeights);
+	static ENGINE_API void UpdateDesiredBoneWeight(
+		const TArrayView<FPerBoneBlendWeight> SrcBoneBlendWeights,
+		TArrayView<FPerBoneBlendWeight> TargetBoneBlendWeights,
+		const TArrayView<float> BlendWeights);
 
 	/**
 	 *	Create Mask Weight for skeleton joints, not per mesh or per required bones
diff --git a/Engine/Source/Runtime/Engine/Public/BonePose.h b/Engine/Source/Runtime/Engine/Public/BonePose.h
--- a/Engine/Source/Runtime/Engine/Public/BonePose.h
+++ b/Engine/Source/Runtime/Engine/Public/BonePose.h
@@ -596,7 +596,8 @@ void FCSPose<PoseType>::CalculateComponentSpaceTransform(BoneIndexType BoneIndex
 	checkSlow(Pose.IsValid());
 	check(ComponentSpaceFlags[BoneIndex] == 0);
 
-	TArray<BoneIndexType> BoneIndexStack;
+	FMemMark Mark(FMemStack::Get());
+	TArray<BoneIndexType, TMemStackAllocator<>> BoneIndexStack;
 	BoneIndexStack.Reserve(ComponentSpaceFlags.Num());
 	BoneIndexStack.Add(BoneIndex);	//Add a dummy index to avoid last element checks in the loop
 
@@ -769,7 +770,8 @@ void FCSPose<PoseType>::LocalBlendCSBoneTransforms(TConstArrayView<FBoneTransfor
 		BoneMask.Reset();
 		BoneMask.AddZeroed(Pose.GetNumBones());
 
-		TArray<struct FBoneTransform> LocalBoneTransforms;
+		FMemMark Mark(FMemStack::Get());
+		TArray<FBoneTransform, TMemStackAllocator<>> LocalBoneTransforms;
 		LocalBoneTransforms.SetNumUninitialized(BoneTransforms.Num());
 
 		// First, save the current local-space poses for the modified bones
diff --git a/Engine/Source/Runtime/Experimental/Chaos/Public/Chaos/TriangleMeshImplicitObject.h b/Engine/Source/Runtime/Experimental/Chaos/Public/Chaos/TriangleMeshImplicitObject.h
--- a/Engine/Source/Runtime/Experimental/Chaos/Public/Chaos/TriangleMeshImplicitObject.h
+++ b/Engine/Source/Runtime/Experimental/Chaos/Public/Chaos/TriangleMeshImplicitObject.h
@@ -238,8 +238,7 @@ namespace Chaos
 			{
 				return;
 			}
-			TArray<int32> NodeIndexStack;
-			NodeIndexStack.Reserve(16);
+			TArray<int32, TInlineAllocator<16>> NodeIndexStack;
 			NodeIndexStack.Push(0);
 			while (NodeIndexStack.Num())
 			{
diff --git a/Engine/Source/Runtime/Experimental/Dataflow/Engine/Private/Dataflow/DataflowColorRamp.cpp b/Engine/Source/Runtime/Experimental/Dataflow/Engine/Private/Dataflow/DataflowColorRamp.cpp
--- a/Engine/Source/Runtime/Experimental/Dataflow/Engine/Private/Dataflow/DataflowColorRamp.cpp
+++ b/Engine/Source/Runtime/Experimental/Dataflow/Engine/Private/Dataflow/DataflowColorRamp.cpp
@@ -9,6 +9,16 @@ TArray<FRichCurveEditInfoConst> FDataflowColorCurveOwner::GetCurves() const
 	return ConstCurves;
 }
 
+void FDataflowColorCurveOwner::GetCurves(TAdderReserverRef<FRichCurveEditInfoConst> OutCurves) const
+{
+	check(OutCurves.Num() == 0);
+	OutCurves.Reserve(ConstCurves.Num());
+	for (const FRichCurveEditInfoConst& ConstCurve : ConstCurves)
+	{
+		OutCurves.Add(ConstCurve);
+	}
+}
+
 TArray<FRichCurveEditInfo> FDataflowColorCurveOwner::GetCurves()
 {
 	return Curves;
diff --git a/Engine/Source/Runtime/Experimental/Dataflow/Engine/Public/Dataflow/DataflowColorRamp.h b/Engine/Source/Runtime/Experimental/Dataflow/Engine/Public/Dataflow/DataflowColorRamp.h
--- a/Engine/Source/Runtime/Experimental/Dataflow/Engine/Public/Dataflow/DataflowColorRamp.h
+++ b/Engine/Source/Runtime/Experimental/Dataflow/Engine/Public/Dataflow/DataflowColorRamp.h
@@ -29,6 +29,7 @@ public:
 
 	/** implemenmt all the necessary virtual function from FCurveOwnerInterface */
 	DATAFLOWENGINE_API virtual TArray<FRichCurveEditInfoConst> GetCurves() const override;
+	DATAFLOWENGINE_API virtual void GetCurves(TAdderReserverRef<FRichCurveEditInfoConst> Curves) const override;
 	DATAFLOWENGINE_API virtual TArray<FRichCurveEditInfo> GetCurves() override;
 	DATAFLOWENGINE_API virtual void ModifyOwner() override;
 	DATAFLOWENGINE_API virtual TArray<const UObject*> GetOwners() const override;
