Remove Child Element system from object tracing system that's no longer being used

diff --git a/Engine/Plugins/Animation/GameplayInsights/Source/GameplayInsights/Private/GameplayAnalyzer.cpp b/Engine/Plugins/Animation/GameplayInsights/Source/GameplayInsights/Private/GameplayAnalyzer.cpp
--- a/Engine/Plugins/Animation/GameplayInsights/Source/GameplayInsights/Private/GameplayAnalyzer.cpp
+++ b/Engine/Plugins/Animation/GameplayInsights/Source/GameplayInsights/Private/GameplayAnalyzer.cpp
@@ -36,10 +36,6 @@ void FGameplayAnalyzer::OnAnalysisBegin(const FOnAnalysisContext& Context)
 	Builder.RouteEvent(RouteId_PropertiesEnd, "Object", "PropertiesEnd");
 	Builder.RouteEvent(RouteId_PropertyValue2, "Object", "PropertyValue2");
 	Builder.RouteEvent(RouteId_PropertyValue3, "Object", "PropertyValue3");
-	Builder.RouteEvent(RouteId_ChildElement, "Object", "ChildElement");
-	Builder.RouteEvent(RouteId_ChildElementLifetimeBegin, "Object", "ChildElementLifetimeBegin");
-	Builder.RouteEvent(RouteId_ChildElementLifetimeEnd, "Object", "ChildElementLifetimeEnd");
-	Builder.RouteEvent(RouteId_ChildElementEvent, "Object", "ChildElementEvent");
 }
 
 bool FGameplayAnalyzer::OnEvent(uint16 RouteId, EStyle Style, const FOnEventContext& Context)
@@ -233,45 +229,6 @@ bool FGameplayAnalyzer::OnEvent(uint16 RouteId, EStyle Style, const FOnEventCont
 		GameplayProvider.AppendPropertyValue(ObjectId, Context.EventTime.AsSeconds(Cycle), Cycle, ParentId, TypeId, NameId, ParentNameId, Value);
 		break;
 	}
-	case RouteId_ChildElement:
-	{
-		uint64 Id = EventData.GetValue<uint64>("Id");
-		uint64 ChildElementId = EventData.GetValue<uint64>("ChildElementId");
-		uint64 StructTypeId = EventData.GetValue<uint64>("StructTypeId", 0);
-		uint64 OuterId = EventData.GetValue<uint64>("OuterId");
-		uint64 OuterChildElementId = EventData.GetValue<uint64>("OuterChildElementId");
-		FString ElementName;
-		EventData.GetString("Name", ElementName);
-		GameplayProvider.AppendObject(RewindDebugger::FObjectId(Id, ChildElementId), RewindDebugger::FObjectId(OuterId, OuterChildElementId), StructTypeId, *ElementName, /*PathName*/{});
-		break;
-	}
-	case RouteId_ChildElementLifetimeBegin:
-	{
-		uint64 Cycle = EventData.GetValue<uint64>("Cycle");
-		double RecordingTime = EventData.GetValue<double>("RecordingTime");
-		uint64 Id = EventData.GetValue<uint64>("Id");
-		uint64 ChildElementId = EventData.GetValue<uint64>("ChildElementId");
-		GameplayProvider.AppendObjectLifetimeBegin(RewindDebugger::FObjectId(Id, ChildElementId), Context.EventTime.AsSeconds(Cycle), RecordingTime);
-		break;
-	}
-	case RouteId_ChildElementLifetimeEnd:
-	{
-		uint64 Cycle = EventData.GetValue<uint64>("Cycle");
-		double RecordingTime = EventData.GetValue<double>("RecordingTime");
-		uint64 Id = EventData.GetValue<uint64>("Id");
-		uint64 ChildElementId = EventData.GetValue<uint64>("ChildElementId");
-		GameplayProvider.AppendObjectLifetimeEnd(RewindDebugger::FObjectId(Id, ChildElementId), Context.EventTime.AsSeconds(Cycle), RecordingTime);
-		break;
-	}
-	case RouteId_ChildElementEvent:
-	{
-		uint64 Cycle = EventData.GetValue<uint64>("Cycle");
-		uint64 Id = EventData.GetValue<uint64>("Id");
-		uint64 ChildElementId = EventData.GetValue<uint64>("ChildElementId");
-		FString Event = TraceServices::FTraceAnalyzerUtils::LegacyAttachmentString<TCHAR>("Event", Context);
-		GameplayProvider.AppendObjectEvent(RewindDebugger::FObjectId(Id, ChildElementId), Context.EventTime.AsSeconds(Cycle), *Event);
-		break;
-	}
 	}
 
 	return true;
diff --git a/Engine/Plugins/Animation/GameplayInsights/Source/GameplayInsights/Private/GameplayAnalyzer.h b/Engine/Plugins/Animation/GameplayInsights/Source/GameplayInsights/Private/GameplayAnalyzer.h
--- a/Engine/Plugins/Animation/GameplayInsights/Source/GameplayInsights/Private/GameplayAnalyzer.h
+++ b/Engine/Plugins/Animation/GameplayInsights/Source/GameplayInsights/Private/GameplayAnalyzer.h
@@ -38,10 +38,6 @@ private:
 		RouteId_PropertyValue2,
 		RouteId_PropertyValue3,
 		RouteId_RecordingInfo,
-		RouteId_ChildElement,
-		RouteId_ChildElementLifetimeBegin,
-		RouteId_ChildElementLifetimeEnd,
-		RouteId_ChildElementEvent,
 	};
 
 	TraceServices::IAnalysisSession& Session;
diff --git a/Engine/Plugins/Animation/GameplayInsights/Source/GameplayInsights/Private/GameplayProvider.cpp b/Engine/Plugins/Animation/GameplayInsights/Source/GameplayInsights/Private/GameplayProvider.cpp
--- a/Engine/Plugins/Animation/GameplayInsights/Source/GameplayInsights/Private/GameplayProvider.cpp
+++ b/Engine/Plugins/Animation/GameplayInsights/Source/GameplayInsights/Private/GameplayProvider.cpp
@@ -545,13 +545,9 @@ void FGameplayProvider::AppendObjectLifetimeEnd(const RewindDebugger::FObjectId&
 		// do not remove from ActiveObjectRecordingLifetimes - lifetimes can be queried by Object Id in GetObjectRecordingLifetime
 	}
 
-	// For now, we don't have a need to notify EndPlay on child element so let's avoid it
-	if (!InObjectId.IsChildElement())
+	if (const int32* ObjectInfoIndex = ObjectIdToIndexMap.Find(InObjectId))
 	{
-		if (const int32* ObjectInfoIndex = ObjectIdToIndexMap.Find(InObjectId))
-		{
-			OnObjectEndPlayDelegate.Broadcast(InObjectId.GetMainId(), InProfileTime, ObjectInfos[*ObjectInfoIndex]);
-		}
+		OnObjectEndPlayDelegate.Broadcast(InObjectId.GetMainId(), InProfileTime, ObjectInfos[*ObjectInfoIndex]);
 	}
 }
 
@@ -585,8 +581,7 @@ void FGameplayProvider::AppendObjectEvent(const RewindDebugger::FObjectId& InObj
 	Message.Name = Session.StoreString(InEventName);
 
 	// For now, we don't have a need to notify EndPlay on child element so let's avoid it
-	if (Message.Name == EndPlayEvent
-		&& !InObjectId.IsChildElement())
+	if (Message.Name == EndPlayEvent)
 	{
 		if (const int32* ObjectInfoIndex = ObjectIdToIndexMap.Find(InObjectId))
 		{
diff --git a/Engine/Plugins/Animation/GameplayInsights/Source/RewindDebugger/Private/RewindDebugger.cpp b/Engine/Plugins/Animation/GameplayInsights/Source/RewindDebugger/Private/RewindDebugger.cpp
--- a/Engine/Plugins/Animation/GameplayInsights/Source/RewindDebugger/Private/RewindDebugger.cpp
+++ b/Engine/Plugins/Animation/GameplayInsights/Source/RewindDebugger/Private/RewindDebugger.cpp
@@ -311,8 +311,7 @@ void FRewindDebugger::GetTargetObjectIds(TArray<RewindDebugger::FObjectId>& OutT
 		{
 			GameplayProvider->EnumerateObjects(CurrentTraceRange.GetLowerBoundValue(), CurrentTraceRange.GetUpperBoundValue(), [this, &OutTargetObjectIds](const FObjectInfo& InObjectInfo)
 				{
-					// Skip children
-					if (!InObjectInfo.GetId().IsChildElement() && SelectedObjectName.Get() == InObjectInfo.Name)
+					if (SelectedObjectName.Get() == InObjectInfo.Name)
 					{
 						OutTargetObjectIds.Add(InObjectInfo.GetId());
 					}
diff --git a/Engine/Plugins/Animation/GameplayInsights/Source/RewindDebugger/Private/RewindDebuggerObjectTrack.cpp b/Engine/Plugins/Animation/GameplayInsights/Source/RewindDebugger/Private/RewindDebuggerObjectTrack.cpp
--- a/Engine/Plugins/Animation/GameplayInsights/Source/RewindDebugger/Private/RewindDebuggerObjectTrack.cpp
+++ b/Engine/Plugins/Animation/GameplayInsights/Source/RewindDebugger/Private/RewindDebuggerObjectTrack.cpp
@@ -26,13 +26,6 @@ static bool GetTypeHierarchyNames(const RewindDebugger::FObjectId InObjectId, co
 	const FObjectInfo& ObjectInfo = GameplayProvider->GetObjectInfo(InObjectId);
 	uint64 ClassId = ObjectInfo.ClassId;
 	
-	if (InObjectId.IsChildElement() && ClassId == FObjectId::InvalidId)
-	{
-		// for child elements with no traced ClassId (actually a StructId in this case... pending rename), use the main object ClassId (backwards compatibility with previous behavior)
-		const FObjectInfo& BaseObjectInfo = GameplayProvider->GetObjectInfo(InObjectId.GetMainId());
-		ClassId = BaseObjectInfo.ClassId;
-	}
-	
 	while (ClassId != FObjectId::InvalidId)
 	{
 		const FClassInfo& ClassInfo = GameplayProvider->GetClassInfo(ClassId);
@@ -56,12 +49,8 @@ FRewindDebuggerObjectTrack::FRewindDebuggerObjectTrack(const FObjectId& InObject
 
 	if (TArray<const TCHAR*> TypeNames; GetTypeHierarchyNames(InObjectId, *Session, TypeNames))
 	{
-		FRewindDebuggerTrackCreators::EnumerateCreators([&Creators = ChildTrackCreators, &TypeNames, IsChildElementTrack = InObjectId.IsChildElement()](const IRewindDebuggerTrackCreator* Creator)
+		FRewindDebuggerTrackCreators::EnumerateCreators([&Creators = ChildTrackCreators, &TypeNames](const IRewindDebuggerTrackCreator* Creator)
 			{
-				if (IsChildElementTrack && !Creator->IsChildElementSupported())
-				{
-					return;
-				}
 
 				if (TypeNames.FindByPredicate([TargetTypeName = Creator->GetTargetTypeName()](const TCHAR* TypeName) { return TargetTypeName == TypeName; }))
 				{
@@ -282,11 +271,6 @@ uint64 FRewindDebuggerObjectTrack::GetObjectIdInternal() const
 	return ObjectId.GetMainId();
 }
 
-uint64 FRewindDebuggerObjectTrack::GetChildElementIdInternal() const
-{
-	return ObjectId.GetChildElementId();
-}
-
 bool FRewindDebuggerObjectTrack::HasDebugDataInternal() const
 {
 	if (PrimaryChildTrack.IsValid())
diff --git a/Engine/Plugins/Animation/GameplayInsights/Source/RewindDebugger/Private/RewindDebuggerObjectTrack.h b/Engine/Plugins/Animation/GameplayInsights/Source/RewindDebugger/Private/RewindDebuggerObjectTrack.h
--- a/Engine/Plugins/Animation/GameplayInsights/Source/RewindDebugger/Private/RewindDebuggerObjectTrack.h
+++ b/Engine/Plugins/Animation/GameplayInsights/Source/RewindDebugger/Private/RewindDebuggerObjectTrack.h
@@ -37,7 +37,6 @@ private:
 	virtual FSlateIcon GetIconInternal() override;
 	virtual FText GetDisplayNameInternal() const override;
 	virtual uint64 GetObjectIdInternal() const override;
-	virtual uint64 GetChildElementIdInternal() const override;
 	virtual bool HasDebugDataInternal() const override;
 	virtual bool HandleDoubleClickInternal() override;
 
diff --git a/Engine/Plugins/Animation/GameplayInsights/Source/RewindDebugger/Private/SRewindDebugger.cpp b/Engine/Plugins/Animation/GameplayInsights/Source/RewindDebugger/Private/SRewindDebugger.cpp
--- a/Engine/Plugins/Animation/GameplayInsights/Source/RewindDebugger/Private/SRewindDebugger.cpp
+++ b/Engine/Plugins/Animation/GameplayInsights/Source/RewindDebugger/Private/SRewindDebugger.cpp
@@ -405,14 +405,12 @@ void SRewindDebugger::Construct(const FArguments& InArgs, TSharedRef<FUICommandL
 																		FText ReadableName = (*Tracks)[0]->GetDisplayName();
 #if OBJECT_TRACE_ENABLED
 																		const RewindDebugger::FObjectId ObjectId = (*Tracks)[0]->GetAssociatedObjectId();
-																		if (!ObjectId.IsChildElement())
+
+																		if (UObject* Object = FObjectTrace::GetObjectFromId(ObjectId.GetMainId()))
 																		{
-																			if (UObject* Object = FObjectTrace::GetObjectFromId(ObjectId.GetMainId()))
+																			if (const AActor* Actor = Cast<AActor>(Object))
 																			{
-																				if (const AActor* Actor = Cast<AActor>(Object))
-																				{
-																					ReadableName = FText::FromString(Actor->GetActorLabel());
-																				}
+																				ReadableName = FText::FromString(Actor->GetActorLabel());
 																			}
 																		}
 #endif // OBJECT_TRACE_ENABLED
diff --git a/Engine/Plugins/Runtime/StateTree/Source/StateTreeEditorModule/Private/Debugger/StateTreeRewindDebuggerTrack.cpp b/Engine/Plugins/Runtime/StateTree/Source/StateTreeEditorModule/Private/Debugger/StateTreeRewindDebuggerTrack.cpp
--- a/Engine/Plugins/Runtime/StateTree/Source/StateTreeEditorModule/Private/Debugger/StateTreeRewindDebuggerTrack.cpp
+++ b/Engine/Plugins/Runtime/StateTree/Source/StateTreeEditorModule/Private/Debugger/StateTreeRewindDebuggerTrack.cpp
@@ -113,11 +113,6 @@ uint64 FRewindDebuggerTrack::GetObjectIdInternal() const
 	return ObjectId.GetMainId();
 }
 
-uint64 FRewindDebuggerTrack::GetChildElementIdInternal() const
-{
-	return ObjectId.GetChildElementId();
-}
-
 FInstanceEventCollection* FRewindDebuggerTrack::GetOrCreateEventCollection(const FStateTreeInstanceDebugId InstanceId)
 {
 	ensureMsgf(EventCollection.IsValid(), TEXT("Event collection is expected to be initialized"));
diff --git a/Engine/Plugins/Runtime/StateTree/Source/StateTreeEditorModule/Private/Debugger/StateTreeRewindDebuggerTrack.h b/Engine/Plugins/Runtime/StateTree/Source/StateTreeEditorModule/Private/Debugger/StateTreeRewindDebuggerTrack.h
--- a/Engine/Plugins/Runtime/StateTree/Source/StateTreeEditorModule/Private/Debugger/StateTreeRewindDebuggerTrack.h
+++ b/Engine/Plugins/Runtime/StateTree/Source/StateTreeEditorModule/Private/Debugger/StateTreeRewindDebuggerTrack.h
@@ -53,7 +53,6 @@ private:
 	virtual FSlateIcon GetIconInternal() override;
 	virtual FName GetNameInternal() const override;
 	virtual uint64 GetObjectIdInternal() const override;
-	virtual uint64 GetChildElementIdInternal() const override;
 	virtual bool HandleDoubleClickInternal() override;
 	//~ End RewindDebugger::FRewindDebuggerTrack interface
 
diff --git a/Engine/Source/Editor/RewindDebuggerInterface/Public/IRewindDebuggerTrackCreator.h b/Engine/Source/Editor/RewindDebuggerInterface/Public/IRewindDebuggerTrackCreator.h
--- a/Engine/Source/Editor/RewindDebuggerInterface/Public/IRewindDebuggerTrackCreator.h
+++ b/Engine/Source/Editor/RewindDebuggerInterface/Public/IRewindDebuggerTrackCreator.h
@@ -94,12 +94,6 @@ PRAGMA_ENABLE_DEPRECATION_WARNINGS
 		return GetTrackTypesInternal(Types);
 	};
 
-	/** @return Whether the track created can also represent non-UObject based children */
-	bool IsChildElementSupported() const
-	{
-		return IsChildElementSupportedInternal();
-	}
-
 	/** @return Whether the track created will be used by its parent track to replace its representation and timeline values */
 	bool IsCreatingPrimaryChildTrack() const
 	{
@@ -154,11 +148,6 @@ private:
 		return TSharedPtr<FRewindDebuggerTrack>();
 	}
 
-	virtual bool IsChildElementSupportedInternal() const
-	{
-		return false;
-	}
-
 	virtual bool IsCreatingPrimaryChildTrackInternal() const
 	{
 		return false;
diff --git a/Engine/Source/Editor/RewindDebuggerInterface/Public/RewindDebuggerTrack.h b/Engine/Source/Editor/RewindDebuggerInterface/Public/RewindDebuggerTrack.h
--- a/Engine/Source/Editor/RewindDebuggerInterface/Public/RewindDebuggerTrack.h
+++ b/Engine/Source/Editor/RewindDebuggerInterface/Public/RewindDebuggerTrack.h
@@ -143,7 +143,7 @@ public:
 	/** insights object id for an object associated with this track */
 	FObjectId GetAssociatedObjectId() const
 	{
-		return FObjectId(GetObjectIdInternal(), GetChildElementIdInternal());
+		return FObjectId(GetObjectIdInternal());
 	}
 
 	/** iterates over all sub-tracks of this track and call Iterator function */
@@ -243,11 +243,6 @@ private:
 		return {};
 	}
 
-	virtual uint64 GetChildElementIdInternal() const
-	{
-		return INDEX_NONE;
-	}
-
 	virtual bool HasDebugDataInternal() const
 	{
 		return true;
diff --git a/Engine/Source/Editor/RewindDebuggerInterface/Public/RewindDebuggerTypes.h b/Engine/Source/Editor/RewindDebuggerInterface/Public/RewindDebuggerTypes.h
--- a/Engine/Source/Editor/RewindDebuggerInterface/Public/RewindDebuggerTypes.h
+++ b/Engine/Source/Editor/RewindDebuggerInterface/Public/RewindDebuggerTypes.h
@@ -18,9 +18,8 @@ struct FObjectId
 
 	FObjectId() = default;
 
-	explicit FObjectId(const uint64 MainId, const uint64 ChildElementId = InvalidId)
+	explicit FObjectId(const uint64 MainId)
 		: ObjectId(MainId)
-		, ChildElementId(ChildElementId)
 	{
 	}
 
@@ -33,24 +32,6 @@ struct FObjectId
 		return ObjectId;
 	}
 
-	/**
-	 * Optional identifier to allow multiple non-UObject based child elements of a given UObject to be represented in the traces.
-	 * @return Identifier of the child element if specified, InvalidId otherwise.
-	 * @see GetObjectId
-	 */
-	uint64 GetChildElementId() const
-	{
-		return ChildElementId;
-	}
-
-	/**
-	 * @return Whether the current gameplay object identifier is used to represent a child element.
-	 */
-	bool IsChildElement() const
-	{
-		return ChildElementId != InvalidId;
-	}
-
 	/**
 	 * @return Whether the current gameplay object identifier was set to represent a UObject, or one of its children.
 	 */
@@ -63,7 +44,7 @@ struct FObjectId
 
 	bool operator<(const FObjectId& Other) const
 	{
-		return ObjectId < Other.ObjectId && ChildElementId < Other.ChildElementId;
+		return ObjectId < Other.ObjectId;
 	}
 
 	friend FString LexToString(const FObjectId& ObjectIdentifier)
@@ -73,22 +54,16 @@ struct FObjectId
 			return TEXT("Invalid");
 		}
 
-		if (ObjectIdentifier.ChildElementId == InvalidId)
-		{
-			return FString::Printf(TEXT("[%llX]"), ObjectIdentifier.ObjectId);
-		}
-
-		return FString::Printf(TEXT("[%llX : %llX]"), ObjectIdentifier.ObjectId, ObjectIdentifier.ChildElementId);
+		return FString::Printf(TEXT("[%llX]"), ObjectIdentifier.ObjectId);
 	}
 
 	friend uint32 GetTypeHash(const FObjectId Identifier)
 	{
-		return HashCombine(GetTypeHash(Identifier.ObjectId), GetTypeHash(Identifier.ChildElementId));
+		return GetTypeHash(Identifier.ObjectId);
 	}
 
 private:
 	uint64 ObjectId = InvalidId;
-	uint64 ChildElementId = InvalidId;
 };
 
 }
\ No newline at end of file
diff --git a/Engine/Source/Runtime/Engine/Private/ObjectTrace.cpp b/Engine/Source/Runtime/Engine/Private/ObjectTrace.cpp
--- a/Engine/Source/Runtime/Engine/Private/ObjectTrace.cpp
+++ b/Engine/Source/Runtime/Engine/Private/ObjectTrace.cpp
@@ -95,36 +95,6 @@ UE_TRACE_EVENT_BEGIN(Object, View)
 	UE_TRACE_EVENT_FIELD(float, AspectRatio)
 UE_TRACE_EVENT_END()
 
-UE_TRACE_EVENT_BEGIN(Object, ChildElement)
-	UE_TRACE_EVENT_FIELD(uint64, Id)
-	UE_TRACE_EVENT_FIELD(uint64, ChildElementId)
-	UE_TRACE_EVENT_FIELD(uint64, OuterId)
-	UE_TRACE_EVENT_FIELD(uint64, OuterChildElementId)
-	UE_TRACE_EVENT_FIELD(uint64, StructTypeId)
-	UE_TRACE_EVENT_FIELD(UE::Trace::WideString, Name)
-UE_TRACE_EVENT_END()
-
-UE_TRACE_EVENT_BEGIN(Object, ChildElementLifetimeBegin)
-	UE_TRACE_EVENT_FIELD(uint64, Cycle)
-	UE_TRACE_EVENT_FIELD(double, RecordingTime)
-	UE_TRACE_EVENT_FIELD(uint64, Id)
-	UE_TRACE_EVENT_FIELD(uint64, ChildElementId)
-UE_TRACE_EVENT_END()
-
-UE_TRACE_EVENT_BEGIN(Object, ChildElementLifetimeEnd)
-	UE_TRACE_EVENT_FIELD(uint64, Cycle)
-	UE_TRACE_EVENT_FIELD(double, RecordingTime)
-	UE_TRACE_EVENT_FIELD(uint64, Id)
-	UE_TRACE_EVENT_FIELD(uint64, ChildElementId)
-UE_TRACE_EVENT_END()
-
-UE_TRACE_EVENT_BEGIN(Object, ChildElementEvent)
-	UE_TRACE_EVENT_FIELD(uint64, Cycle)
-	UE_TRACE_EVENT_FIELD(uint64, Id)
-	UE_TRACE_EVENT_FIELD(uint64, ChildElementId)
-	UE_TRACE_EVENT_FIELD(UE::Trace::WideString, Event)
-UE_TRACE_EVENT_END()
-
 namespace UE::ObjectTrace
 {
 bool IsTracingDisabledForObject(const UObject* InObject)
@@ -157,13 +127,12 @@ struct FObjectId
 	static constexpr uint64 InvalidId = 0;
 
 	uint64 ObjectId = InvalidId;
-	uint64 ChildElementId = InvalidId;
 
 	bool operator==(const FObjectId& Other) const = default;
 
 	friend uint32 GetTypeHash(const FObjectId Identifier)
 	{
-		return HashCombine(GetTypeHash(Identifier.ObjectId), GetTypeHash(Identifier.ChildElementId));
+		return GetTypeHash(Identifier.ObjectId);
 	}
 };
 
@@ -608,79 +577,4 @@ void FObjectTrace::OutputWorld(const UWorld* InWorld)
 	TRACE_OBJECT(InWorld);
 }
 
-void FObjectTrace::OutputChildElement(const UObject* InObject, const uint64 InChildElementId, const UObject* InOuterObject, const uint64 InOuterChildElementId, const UStruct* InStructType, const FString& InName)
-{
-	if (UE::ObjectTrace::IsTracingDisabledForObject(InObject))
-	{
-		return;
-	}
-	
-	const uint64 ObjectId = GetObjectId(InObject);
-	if (!AddObjectToTracedSet(UE::ObjectTrace::FObjectId{ObjectId, InChildElementId}))
-	{
-		return;
-	}
-
-	OutputObject(InOuterObject);
-	// intentionally not tracing the outer child element from here because we don't know what its parent should be
-	//   - game systems must ensure they trace the whole child element hierarchy
-	// 	OutputChildElement(InOuterObject, InOuterChildElementId);
-
-	FString Name = InName;
-	if (Name.IsEmpty())
-	{
-		check(InObject);
-		Name = InObject->GetName();
-		Name += FString::Printf(TEXT(" : [0x%llX]"), InChildElementId);
-	}
-	
-	TRACE_TYPE(InStructType);
-
-	UE_TRACE_LOG(Object, ChildElement, ObjectChannel)
-		<< ChildElement.Id(ObjectId)
-		<< ChildElement.ChildElementId(InChildElementId)
-		<< ChildElement.OuterId(GetObjectId(InOuterObject))
-		<< ChildElement.OuterChildElementId(InOuterChildElementId)
-		<< ChildElement.StructTypeId(GetObjectId(InStructType))
-		<< ChildElement.Name(*Name, Name.Len());
-
-	UE_TRACE_LOG(Object, ChildElementLifetimeBegin, ObjectChannel)
-		<< ChildElementLifetimeBegin.Cycle(FPlatformTime::Cycles64())
-		<< ChildElementLifetimeBegin.RecordingTime(FObjectTrace::GetWorldElapsedTime(InObject->GetWorld()))
-		<< ChildElementLifetimeBegin.Id(ObjectId)
-		<< ChildElementLifetimeBegin.ChildElementId(InChildElementId);
-}
-
-void FObjectTrace::OutputChildElementLifetimeEnd(const UObject* InObject, const uint64 InChildElementId)
-{
-	if (UE::ObjectTrace::IsTracingDisabledForObject(InObject))
-	{
-		return;
-	}
-
-	TRACE_OBJECT(InObject);
-
-	UE_TRACE_LOG(Object, ChildElementLifetimeEnd, ObjectChannel)
-		<< ChildElementLifetimeEnd.Cycle(FPlatformTime::Cycles64())
-		<< ChildElementLifetimeEnd.RecordingTime(FObjectTrace::GetWorldElapsedTime(InObject->GetWorld()))
-		<< ChildElementLifetimeEnd.Id(GetObjectId(InObject))
-		<< ChildElementLifetimeEnd.ChildElementId(InChildElementId);
-}
-
-void FObjectTrace::OutputChildElementEvent(const UObject* InObject, const uint64 InChildElementId, const TCHAR* InEvent)
-{
-	if (UE::ObjectTrace::IsTracingDisabledForObject(InObject))
-	{
-		return;
-	}
-
-	TRACE_OBJECT(InObject);
-
-	UE_TRACE_LOG(Object, ChildElementEvent, ObjectChannel)
-		<< ChildElementEvent.Cycle(FPlatformTime::Cycles64())
-		<< ChildElementEvent.Id(GetObjectId(InObject))
-		<< ChildElementEvent.ChildElementId(InChildElementId)
-		<< ChildElementEvent.Event(InEvent);
-}
-
 #endif
diff --git a/Engine/Source/Runtime/Engine/Public/ObjectTrace.h b/Engine/Source/Runtime/Engine/Public/ObjectTrace.h
--- a/Engine/Source/Runtime/Engine/Public/ObjectTrace.h
+++ b/Engine/Source/Runtime/Engine/Public/ObjectTrace.h
@@ -125,16 +125,6 @@ struct FObjectTrace
 
 	/** Helper function to output a world */
 	ENGINE_API static void OutputWorld(const UWorld* InWorld);
-
-	/** Helper function to output a child element of a UObject */
-	ENGINE_API static void OutputChildElement(const UObject* InObject, uint64 InChildElementId,
-			const UObject* InOuterObject = nullptr, uint64 InOuterChildElementId = InvalidObjectId, const UStruct* InStructType = nullptr, const FString& InName = "");
-
-	/** Helper function to output destruction event of a child element */
-	ENGINE_API static void OutputChildElementLifetimeEnd(const UObject* InObject, uint64 InChildElementId);
-
-	/** Helper function to output an event for a child element */
-	ENGINE_API static void OutputChildElementEvent(const UObject* InObject, uint64 InChildElementId, const TCHAR* InEvent);
 };
 
 #define TRACE_TYPE(Type) \
@@ -172,24 +162,6 @@ struct FObjectTrace
 #define TRACE_PAWN_POSSESS(Controller, Pawn)\
 	if (CAN_TRACE_OBJECT(Controller) && (Pawn==nullptr || (CAN_TRACE_OBJECT(Pawn)))) { FObjectTrace::OutputPawnPossess(Controller, Pawn); }
 
-#define TRACE_CHILD_ELEMENT(Object, ChildElementId) \
-	if (CAN_TRACE_OBJECT(Object)) { FObjectTrace::OutputChildElement(Object, ChildElementId); }
-
-#define TRACE_CHILD_ELEMENT_WITH_OUTER(Object, ChildElementId, OuterObject, OuterChildElementId, StructType, Name) \
-	if (CAN_TRACE_OBJECT(Object)) { FObjectTrace::OutputChildElement(Object, ChildElementId, OuterObject, OuterChildElementId, StructType, Name); }
-
-#define TRACE_CHILD_ELEMENT_LIFETIME_BEGIN(Object, ChildElementId) \
-	if (CAN_TRACE_OBJECT(Object)) { FObjectTrace::OutputChildElement(Object, ChildElementId, Object, FObjectTrace::InvalidObjectId); }
-
-#define TRACE_CHILD_ELEMENT_WITH_OUTER_LIFETIME_BEGIN(Object, ChildElementId, OuterObject, OuterChildElementId) \
-	if (CAN_TRACE_OBJECT(Object)) { FObjectTrace::OutputChildElement(Object, ChildElementId, OuterObject, OuterChildElementId); }
-
-#define TRACE_CHILD_ELEMENT_LIFETIME_END(Object, ChildElementId) \
-	if (CAN_TRACE_OBJECT(Object)) { FObjectTrace::OutputChildElementLifetimeEnd(Object, ChildElementId); }
-
-#define TRACE_CHILD_ELEMENT_EVENT(Object, ChildElementId, Event) \
-	if (CAN_TRACE_OBJECT(Object)) { FObjectTrace::OutputChildElementEvent(Object, ChildElementId, Event); }
-
 #else
 
 #define TRACE_OBJECT_EVENT(Object, Event) \
@@ -207,29 +179,6 @@ struct FObjectTrace
 #define TRACE_PAWN_POSSESS(Controller, Pawn)\
 	FObjectTrace::OutputPawnPossess(Controller, Pawn);
 
-#define TRACE_CHILD_ELEMENT(Object, ChildElementId) \
-	FObjectTrace::OutputChildElement(Object, ChildElementId, Object, FObjectTrace::InvalidObjectId);
-
-
-#define TRACE_CHILD_ELEMENT_WITH_OUTER_(Object, ChildElementId, OuterObject, OuterChildElementId, Type, Name) \
-	FObjectTrace::OutputChildElement(Object, ChildElementId, OuterObject, OuterChildElementId, Type, Name);
-
-#define TRACE_CHILD_ELEMENT_WITH_OUTER(Object, ChildElementId, OuterObject, OuterChildElementId) \
-	FObjectTrace::OutputChildElement(Object, ChildElementId, OuterObject, OuterChildElementId);
-
-#define TRACE_CHILD_ELEMENT_LIFETIME_BEGIN(Object, ChildElementId) \
-	FObjectTrace::OutputChildElement(Object, ChildElementId, Object, FObjectTrace::InvalidObjectId);
-
-#define TRACE_CHILD_ELEMENT_WITH_OUTER_LIFETIME_BEGIN(Object, ChildElementId, OuterObject, OuterChildElementId) \
-	FObjectTrace::OutputChildElement(Object, ChildElementId, OuterObject, OuterChildElementId);
-
-#define TRACE_CHILD_ELEMENT_LIFETIME_END(Object, ChildElementId) \
-	FObjectTrace::OutputChildElementLifetimeEnd(Object, ChildElementId);
-
-#define TRACE_CHILD_ELEMENT_EVENT(Object, ChildElementId, Event) \
-	FObjectTrace::OutputChildElementEvent(Object, ChildElementId, Event);
-
-
 #endif
 	
 #define UNCONDITIONAL_TRACE_OBJECT_EVENT(Object, Event) \
