Added temporary workaround for bad pointer access in FNetworkObjectList::Remove, avoid retrieving the value that we don't end up using anyway since bTrackDormantObjectsByLevel is disabled anyway and the package name isn't used.

diff --git a/Engine/Source/Runtime/Engine/Private/NetworkObjectList.cpp b/Engine/Source/Runtime/Engine/Private/NetworkObjectList.cpp
--- a/Engine/Source/Runtime/Engine/Private/NetworkObjectList.cpp
+++ b/Engine/Source/Runtime/Engine/Private/NetworkObjectList.cpp
@@ -144,7 +144,7 @@ void FNetworkObjectList::Remove(AActor* const Actor)
 		return;
 	}
 
-	const FName PackageName = Actor->GetLevel()->GetPackage()->GetFName();
+	const FName PackageName = (UE::Net::Private::bTrackDormantObjectsByLevel ? Actor->GetLevel()->GetPackage()->GetFName() : FName());
 
 	FNetworkObjectInfo* NetworkObjectInfo = NetworkObjectInfoPtr->Get();
 
