Change Material Attribute Definition Map to explicitly map MCT_Float1 to avoid clashes where MCT_Float may be acting as a mask for all float array lengths.

diff --git a/Engine/Source/Runtime/Engine/Private/Materials/MaterialAttributeDefinitionMap.cpp b/Engine/Source/Runtime/Engine/Private/Materials/MaterialAttributeDefinitionMap.cpp
--- a/Engine/Source/Runtime/Engine/Private/Materials/MaterialAttributeDefinitionMap.cpp
+++ b/Engine/Source/Runtime/Engine/Private/Materials/MaterialAttributeDefinitionMap.cpp
@@ -330,27 +330,27 @@ void FMaterialAttributeDefinitionMap::InitializeAttributeMap()
 
 	// Basic attributes
 	Add(FGuid(0x69B8D336, 0x16ED4D49, 0x9AA49729, 0x2F050F7A), TEXT("BaseColor"),		MP_BaseColor,		MCT_Float3,	FVector4(0,0,0,0),	SF_Pixel);
-	Add(FGuid(0x57C3A161, 0x7F064296, 0xB00B24A5, 0xA496F34C), TEXT("Metallic"),		MP_Metallic,		MCT_Float,	FVector4(0,0,0,0),	SF_Pixel);
-	Add(FGuid(0x9FDAB399, 0x25564CC9, 0x8CD2D572, 0xC12C8FED), TEXT("Specular"),		MP_Specular,		MCT_Float,	FVector4(.5,0,0,0), SF_Pixel);
-	Add(FGuid(0xD1DD967C, 0x4CAD47D3, 0x9E6346FB, 0x08ECF210), TEXT("Roughness"),		MP_Roughness,		MCT_Float,	FVector4(.5,0,0,0), SF_Pixel);
-	Add(FGuid(0x55E2B4FB, 0xC1C54DB2, 0x9F11875F, 0x7231EB1E), TEXT("Anisotropy"),		MP_Anisotropy,		MCT_Float,	FVector4(0,0,0,0),  SF_Pixel);
+	Add(FGuid(0x57C3A161, 0x7F064296, 0xB00B24A5, 0xA496F34C), TEXT("Metallic"),		MP_Metallic,		MCT_Float1,	FVector4(0,0,0,0),	SF_Pixel);
+	Add(FGuid(0x9FDAB399, 0x25564CC9, 0x8CD2D572, 0xC12C8FED), TEXT("Specular"),		MP_Specular,		MCT_Float1,	FVector4(.5,0,0,0), SF_Pixel);
+	Add(FGuid(0xD1DD967C, 0x4CAD47D3, 0x9E6346FB, 0x08ECF210), TEXT("Roughness"),		MP_Roughness,		MCT_Float1,	FVector4(.5,0,0,0), SF_Pixel);
+	Add(FGuid(0x55E2B4FB, 0xC1C54DB2, 0x9F11875F, 0x7231EB1E), TEXT("Anisotropy"),		MP_Anisotropy,		MCT_Float1,	FVector4(0,0,0,0),  SF_Pixel);
 	Add(FGuid(0xB769B54D, 0xD08D4440, 0xABC21BA6, 0xCD27D0E2), TEXT("EmissiveColor"),	MP_EmissiveColor,	MCT_Float3,	FVector4(0,0,0,0),	SF_Pixel);
-	Add(FGuid(0xB8F50FBA, 0x2A754EC1, 0x9EF672CF, 0xEB27BF51), TEXT("Opacity"),			MP_Opacity,			MCT_Float,	FVector4(1,0,0,0),	SF_Pixel);
-	Add(FGuid(0x679FFB17, 0x2BB5422C, 0xAD520483, 0x166E0C75), TEXT("OpacityMask"),		MP_OpacityMask,		MCT_Float,	FVector4(1,0,0,0),	SF_Pixel);
+	Add(FGuid(0xB8F50FBA, 0x2A754EC1, 0x9EF672CF, 0xEB27BF51), TEXT("Opacity"),			MP_Opacity,			MCT_Float1,	FVector4(1,0,0,0),	SF_Pixel);
+	Add(FGuid(0x679FFB17, 0x2BB5422C, 0xAD520483, 0x166E0C75), TEXT("OpacityMask"),		MP_OpacityMask,		MCT_Float1,	FVector4(1,0,0,0),	SF_Pixel);
 	Add(FGuid(0x0FA2821A, 0x200F4A4A, 0xB719B789, 0xC1259C64), TEXT("Normal"),			MP_Normal,			MCT_Float3,	FVector4(0,0,1,0),	SF_Pixel);
 	Add(FGuid(0xD5F8E9CF, 0xCDC3468D, 0xB10E4465, 0x596A7BBA), TEXT("Tangent"),			MP_Tangent,			MCT_Float3,	FVector4(1,0,0,0),	SF_Pixel);
 
 	// Advanced attributes
 	Add(FGuid(0xF905F895, 0xD5814314, 0x916D2434, 0x8C40CE9E), TEXT("WorldPositionOffset"),		MP_WorldPositionOffset,		MCT_Float3,	FVector4(0,0,0,0),	SF_Vertex);
-	Add(FGuid(0x199A7166, 0xC67041DC, 0xA68EAD0D, 0x7017D0AD), TEXT("Displacement"),			MP_Displacement,			MCT_Float,	FVector4(0.5,0,0,0), SF_Pixel);
+	Add(FGuid(0x199A7166, 0xC67041DC, 0xA68EAD0D, 0x7017D0AD), TEXT("Displacement"),			MP_Displacement,			MCT_Float1,	FVector4(0.5,0,0,0), SF_Pixel);
 	Add(FGuid(0x5B8FC679, 0x51CE4082, 0x9D777BEE, 0xF4F72C44), TEXT("SubsurfaceColor"),			MP_SubsurfaceColor,			MCT_Float3,	FVector4(1,1,1,0),	SF_Pixel);
-	Add(FGuid(0x9E502E69, 0x3C8F48FA, 0x94645CFD, 0x28E5428D), TEXT("CustomData0"),				MP_CustomData0,				MCT_Float,	FVector4(1,0,0,0),	SF_Pixel);
-	Add(FGuid(0xBE4F2FFD, 0x12FC4296, 0xB0124EEA, 0x12C28D92), TEXT("CustomData1"),				MP_CustomData1,				MCT_Float,	FVector4(.1,0,0,0),	SF_Pixel);
-	Add(FGuid(0xE8EBD0AD, 0xB1654CBE, 0xB079C3A8, 0xB39B9F15), TEXT("AmbientOcclusion"),		MP_AmbientOcclusion,		MCT_Float,	FVector4(1,0,0,0),	SF_Pixel);
+	Add(FGuid(0x9E502E69, 0x3C8F48FA, 0x94645CFD, 0x28E5428D), TEXT("CustomData0"),				MP_CustomData0,				MCT_Float1,	FVector4(1,0,0,0),	SF_Pixel);
+	Add(FGuid(0xBE4F2FFD, 0x12FC4296, 0xB0124EEA, 0x12C28D92), TEXT("CustomData1"),				MP_CustomData1,				MCT_Float1,	FVector4(.1,0,0,0),	SF_Pixel);
+	Add(FGuid(0xE8EBD0AD, 0xB1654CBE, 0xB079C3A8, 0xB39B9F15), TEXT("AmbientOcclusion"),		MP_AmbientOcclusion,		MCT_Float1,	FVector4(1,0,0,0),	SF_Pixel);
 	Add(FGuid(0xD0B0FA03, 0x14D74455, 0xA851BAC5, 0x81A0788B), TEXT("Refraction"),				MP_Refraction,				MCT_Float3,	FVector4(1,0,0,0),	SF_Pixel);
-	Add(FGuid(0x0AC97EC3, 0xE3D047BA, 0xB610167D, 0xC4D919FF), TEXT("PixelDepthOffset"),		MP_PixelDepthOffset,		MCT_Float,	FVector4(0,0,0,0),	SF_Pixel);
+	Add(FGuid(0x0AC97EC3, 0xE3D047BA, 0xB610167D, 0xC4D919FF), TEXT("PixelDepthOffset"),		MP_PixelDepthOffset,		MCT_Float1,	FVector4(0,0,0,0),	SF_Pixel);
 	Add(FGuid(0xD9423FFF, 0xD77E4D82, 0x8FF9CF5E, 0x055D1255), TEXT("ShadingModel"),			MP_ShadingModel,			MCT_ShadingModel, FVector4(0, 0, 0, 0), SF_Pixel, INDEX_NONE, false, &CompileShadingModelBlendFunction);
-	Add(FGuid(0x42BDD2E0, 0xBE714189, 0xA0984BC3, 0xDD0BE872), TEXT("SurfaceThickness"),		MP_SurfaceThickness,		MCT_Float,  FVector4(SUBSTRATE_LAYER_DEFAULT_THICKNESS_CM, 0, 0, 0), SF_Pixel);
+	Add(FGuid(0x42BDD2E0, 0xBE714189, 0xA0984BC3, 0xDD0BE872), TEXT("SurfaceThickness"),		MP_SurfaceThickness,		MCT_Float1,  FVector4(SUBSTRATE_LAYER_DEFAULT_THICKNESS_CM, 0, 0, 0), SF_Pixel);
 	Add(FGuid(0x5973A03E, 0x13A74E08, 0x92D0CEDD, 0xF2936CF8), TEXT("FrontMaterial"),			MP_FrontMaterial,			MCT_Substrate, FVector4(0,0,0,0),	SF_Pixel, INDEX_NONE, false, &CompileSubstrateBlendFunction);
 
 	// Used when compiling material with execution pins, which are compiling all attributes together
@@ -371,19 +371,19 @@ void FMaterialAttributeDefinitionMap::InitializeAttributeMap()
 	Add(FGuid(0xE89CBD84, 0x62EA48BE, 0x80F88521, 0x2B0C403C), TEXT("SpecularColor"),	MP_SpecularColor,	MCT_Float3, FVector4(0,0,0,0), SF_Pixel, INDEX_NONE, bHideAttribute);
 
 	// Debug attributes
-	Add(FGuid(0x5BF6BA94, 0xA3264629, 0xA253A05B, 0x0EABBB86), TEXT("Missing"), MP_MAX, MCT_Float, FVector4(0,0,0,0), SF_Pixel, INDEX_NONE, bHideAttribute);
+	Add(FGuid(0x5BF6BA94, 0xA3264629, 0xA253A05B, 0x0EABBB86), TEXT("Missing"), MP_MAX, MCT_Float1, FVector4(0,0,0,0), SF_Pixel, INDEX_NONE, bHideAttribute);
 
 	// Removed attributes
 	Add(FGuid(0x2091ECA2, 0xB59248EE, 0x8E2CD578, 0xD371926D), TEXT("WorldDisplacement"), MP_WorldDisplacement_DEPRECATED, MCT_Float3, FVector4(0, 0, 0, 0), SF_Vertex, INDEX_NONE, bHideAttribute);
-	Add(FGuid(0xA0119D44, 0xC456450D, 0x9C39C933, 0x1F72D8D1), TEXT("TessellationMultiplier"), MP_TessellationMultiplier_DEPRECATED, MCT_Float, FVector4(1, 0, 0, 0), SF_Vertex, INDEX_NONE, bHideAttribute);
+	Add(FGuid(0xA0119D44, 0xC456450D, 0x9C39C933, 0x1F72D8D1), TEXT("TessellationMultiplier"), MP_TessellationMultiplier_DEPRECATED, MCT_Float1, FVector4(1, 0, 0, 0), SF_Vertex, INDEX_NONE, bHideAttribute);
 
 	// UMaterialExpression custom outputs
 	AddCustomAttribute(FGuid(0xfbd7b46e, 0xb1234824, 0xbde76b23, 0x609f984c), UMaterialExpressionBentNormalCustomOutput::StaticClass(), "BentNormal", "GetBentNormal", 0, MCT_Float3, FVector4(0, 0, 1, 0), SF_Pixel);
 	AddCustomAttribute(FGuid(0xAA3D5C04, 0x16294716, 0xBBDEC869, 0x6A27DD72), UMaterialExpressionClearCoatNormalCustomOutput::StaticClass(), "ClearCoatBottomNormal", "ClearCoatBottomNormal", 0, MCT_Float3, FVector4(0, 0, 1, 0), SF_Pixel);
 	AddCustomAttribute(FGuid(0x8EAB2CB2, 0x73634A24, 0x8CD14F47, 0x3F9C8E55), UMaterialExpressionTangentOutput::StaticClass(), "CustomEyeTangent", "GetTangentOutput", 0, MCT_Float3, FVector4(0, 0, 0, 0), SF_Pixel);
 	AddCustomAttribute(FGuid(0xF2D8C70E, 0x42ECA0D1, 0x4652D0AD, 0xB785A065), UMaterialExpressionThinTranslucentMaterialOutput::StaticClass(), "TransmittanceColor", "GetThinTranslucentMaterialOutput", 0, MCT_Float3, FVector4(0.5, 0.5, 0.5, 0), SF_Pixel);
-	AddCustomAttribute(FGuid(0xAAA1D0E0, 0xE312458E, 0xF3556BAD, 0xA01277CD), UMaterialExpressionThinTranslucentMaterialOutput::StaticClass(), "ThinTranslucentSurfaceCoverage", "GetThinTranslucentMaterialOutput", 1, MCT_Float, FVector4(1, 0, 0, 0), SF_Pixel);
-	AddCustomAttribute(FGuid(0xDB380E51, 0xBDED47B6, 0xB185A8BE, 0x8EACBFA4), UMaterialExpressionFirstPersonOutput::StaticClass(), "FirstPersonInterpolationAlpha", "GetFirstPersonOutput", 0, MCT_Float, FVector4(1, 0, 0, 0), SF_Vertex, true /*bNeedsPreviousFrameEvaluation*/);
+	AddCustomAttribute(FGuid(0xAAA1D0E0, 0xE312458E, 0xF3556BAD, 0xA01277CD), UMaterialExpressionThinTranslucentMaterialOutput::StaticClass(), "ThinTranslucentSurfaceCoverage", "GetThinTranslucentMaterialOutput", 1, MCT_Float1, FVector4(1, 0, 0, 0), SF_Pixel);
+	AddCustomAttribute(FGuid(0xDB380E51, 0xBDED47B6, 0xB185A8BE, 0x8EACBFA4), UMaterialExpressionFirstPersonOutput::StaticClass(), "FirstPersonInterpolationAlpha", "GetFirstPersonOutput", 0, MCT_Float1, FVector4(1, 0, 0, 0), SF_Vertex, true /*bNeedsPreviousFrameEvaluation*/);
 	AddCustomAttribute(FGuid(0x6B1237CA, 0x1B8C43C2, 0xA9C8A020, 0x430F1258), UMaterialExpressionMaterialCache::StaticClass(), "MaterialCache", "MaterialCache", 0, MCT_Float4, FVector4(0, 0, 0, 0), SF_Pixel);
 }
 
