Update the NetGUID cache with actor component's when borrowing and returning a borrowed actor.

diff --git a/Engine/Source/Runtime/Engine/Classes/Engine/NetDriver.h b/Engine/Source/Runtime/Engine/Classes/Engine/NetDriver.h
--- a/Engine/Source/Runtime/Engine/Classes/Engine/NetDriver.h
+++ b/Engine/Source/Runtime/Engine/Classes/Engine/NetDriver.h
@@ -1048,6 +1048,20 @@ public:
 	/** Creates a channel of each type that is set as bIniitalServer for the given connection. */
 	ENGINE_API void CreateInitialServerChannels(UNetConnection* ClientConnection);
 	
+	/** 
+	 * Utility function that adds an object and, if an actor or component, their sub-objects to the driver's NetGUID cache.
+	 * 
+	 * Only currently supports generic replication.
+	 */
+	ENGINE_API void AddObjectToNetworkIdCache(UObject* Object);
+
+	/** 
+	 * Utility function that removes an object and, if an actor or component, their sub-objects from the driver's NetGUID cache.
+	 *
+	 * Only currently supports generic replication.
+	 */
+	ENGINE_API void RemoveObjectFromNetworkIdCache(UObject* Object);
+
 protected:
 	/** 
 	 * Register each metric used by the networking system. 
diff --git a/Engine/Source/Runtime/Engine/Private/Actor.cpp b/Engine/Source/Runtime/Engine/Private/Actor.cpp
--- a/Engine/Source/Runtime/Engine/Private/Actor.cpp
+++ b/Engine/Source/Runtime/Engine/Private/Actor.cpp
@@ -1326,6 +1326,15 @@ void AActor::PostMigrate(const FUObjectMigrationContext& MigrationContext)
 		UE_LOG(LogActor, Verbose, TEXT("%hs: Send %s (LocalSrvId: %s. DstSrvId: %s. HasOwnership: %d. MigrationType: %s)"), __func__, *GetName(), *LocalServerId.ToString(), *MigrationContext.RemoteServerId.ToString(), bHasOwnership, ToString(MigrationSendType));
 		ensureMsgf(MigrationContext.RemoteServerId != LocalServerId, TEXT("%hs: Sending Object %s from %s to %s but destination is the same as the Local Server Id."), __func__, *GetName(), *LocalServerId.ToString(), *MigrationContext.RemoteServerId.ToString());
 
+		// Remove the actor and it's sub-objects from the network id cache since these objects are now considered remote. It's assumed that the
+		// network id cache will borrow any objects that are not in the cache but are considered remote.
+		//
+		// Loan is an exception because the server is still retaining ownership of this object and can still be dereferenced.
+		if (MigrationSendType != EObjectMigrationSendType::Loan)
+		{
+			GetWorld()->GetNetDriver()->RemoveObjectFromNetworkIdCache(this);
+		}
+
 		// If we're loaning, we're keeping ownership; in all other cases we either don't own, or won't own the object after it's migrated
 		const bool bKeepOwnership = (MigrationSendType == EObjectMigrationSendType::Loan);
 		if (bKeepOwnership)
@@ -1373,18 +1382,6 @@ void AActor::PostMigrate(const FUObjectMigrationContext& MigrationContext)
 			// Unregister the actor's tick function
 			RegisterAllActorTickFunctions(/*bRegisterTickFunctions*/ false, /*bIncludeComponents*/ true);
 		}
-
-		if (MigrationSendType == EObjectMigrationSendType::ReturnBorrowed)
-		{
-			if (GetWorld()->GetNetDriver()->GuidCache->RemoveNetGUID(this))
-			{
-				UE_LOG(LogActor, Verbose, TEXT("Deregistering loaned object %s from the NetGUID cache."), *this->GetName());
-			}
-			else
-			{
-				UE_LOG(LogActor, Verbose, TEXT("Loaned object %s isn't registered in the NetGUID cache."), *GetName());
-			}
-		}
 	}
 	else if (MigrationContext.MigrationSide == EObjectMigrationSide::Receive)
 	{
@@ -1404,6 +1401,10 @@ void AActor::PostMigrate(const FUObjectMigrationContext& MigrationContext)
 		{
 			constexpr bool bOnlyAddIfUnique = true;
 			GetLevel()->TryAddActorToList(this, bOnlyAddIfUnique);
+
+			// Add the actor and it's sub-objects to the network id cache since these objects are no longer considered remote and can be derferenced
+			// through the network id cache (e.g. by RPC argument and object deserialization).
+			GetWorld()->GetNetDriver()->AddObjectToNetworkIdCache(this);
 		}
 
 		// If we're being passed ownership, register us with the global systems
@@ -1434,16 +1435,6 @@ void AActor::PostMigrate(const FUObjectMigrationContext& MigrationContext)
 					ensureMsgf(IsPendingKillPending(), TEXT("%hs: We thought %s was destroyed while on loan, but our received state does not match a destroyed Actor"), __func__, *GetName());
 					UE_LOG(LogActor, Warning, TEXT("%hs: Destroying %s back on owning server (it was destroyed while loaned-out)"), __func__, *GetName());
 
-					// Remove the object from the NetGUID cache (if it's there) before destroying the actor below.
-					if (GetWorld()->GetNetDriver()->GuidCache->RemoveNetGUID(this))
-					{
-						UE_LOG(LogActor, Warning, TEXT("Removed destroyed object %s from the NetGUID cache."), *GetName());
-					}
-					else
-					{
-						UE_LOG(LogActor, Warning, TEXT("Destroyed object %s is not in the NetGUID cache."), *GetName());
-					}
-
 					// Destroy it again on the owning server, as it needs to be deregistered from our systems
 					bActorIsBeingDestroyed = false;
 					ClearGarbage();
@@ -1457,7 +1448,6 @@ void AActor::PostMigrate(const FUObjectMigrationContext& MigrationContext)
 				bTickFunctionsRegistered = LoanedState.bTickFunctionsRegistered;
 				RegisterAllActorTickFunctions(!bTickFunctionsRegistered, bRegisterComponents);
 			}
-
 		}
 		else if (MigrationRecvType == EObjectMigrationRecvType::Borrowed)
 		{
@@ -1475,35 +1465,6 @@ void AActor::PostMigrate(const FUObjectMigrationContext& MigrationContext)
 			}
 		}
 
-		if (MigrationRecvType == EObjectMigrationRecvType::Borrowed)
-		{
-			TSharedPtr<FNetGUIDCache>& GuidCache = GetWorld()->GetNetDriver()->GuidCache;
-			const bool bIsStatic = !GuidCache->IsDynamicObject(this);
-			FRemoteObjectId RemoteObjectId = FRemoteObjectId(this);
-			FNetworkGUID NetGUID = FNetGUIDCache::GetNetworkGUIDFromRemoteObjectId(RemoteObjectId, bIsStatic);
-
-			UE_LOG(LogActor, Verbose, TEXT("Registering borrowed object %s (NetGUID=%s RemoteObjectId=%s) in the NetGUID cache."),
-				   *this->GetName(),
-				   *NetGUID.ToString(),
-				   *RemoteObjectId.ToString());
-
-			if (!GuidCache->GetNetGUID(this).IsValid())
-			{
-				if (!bActorIsBeingDestroyed)
-				{
-					GuidCache->RegisterNetGUID_Server(NetGUID, this);
-				}
-				else
-				{
-					UE_LOG(LogActor, Verbose, TEXT("Borrowed object %s cannot be added to the NetGUID cache because it's being destroyed."), *GetName());
-				}
-			}
-			else
-			{
-				UE_LOG(LogActor, Verbose, TEXT("Borrowed object %s is already registered in the NetGUID cache."), *GetName());
-			}
-		}
-
 		ActorHasBegunPlay = EActorBeginPlayState::HasBegunPlay;
 	}
 }
diff --git a/Engine/Source/Runtime/Engine/Private/Components/ActorComponent.cpp b/Engine/Source/Runtime/Engine/Private/Components/ActorComponent.cpp
--- a/Engine/Source/Runtime/Engine/Private/Components/ActorComponent.cpp
+++ b/Engine/Source/Runtime/Engine/Private/Components/ActorComponent.cpp
@@ -726,6 +726,15 @@ void UActorComponent::PostMigrate(const struct FUObjectMigrationContext& Migrati
 		const EObjectMigrationSendType SendType = MigrationContext.GetObjectMigrationSendType(this);
 		UE_LOG(LogActorComponent, VeryVerbose, TEXT("%hs: %s %s"), __func__, ToString(SendType), *GetName());
 
+		// Remove the actor and it's sub-objects from the network id cache since these objects are now considered remote. It's assumed that the
+		// network id cache will borrow any objects that are not in the cache but are considered remote.
+		//
+		// Loan is an exception because the server is still retaining ownership of this object and can still be dereferenced.
+		if (SendType != EObjectMigrationSendType::Loan)
+		{
+			GetWorld()->GetNetDriver()->RemoveObjectFromNetworkIdCache(this);
+		}
+
 		// If we're loaning, we're keeping ownership; in all other cases we either don't own, or won't own the object after it's migrated
 		const bool bKeepOwnership = (SendType == EObjectMigrationSendType::Loan);
 		if (bKeepOwnership)
@@ -750,6 +759,10 @@ void UActorComponent::PostMigrate(const struct FUObjectMigrationContext& Migrati
 		const EObjectMigrationRecvType RecvType = MigrationContext.GetObjectMigrationRecvType(this);
 		UE_LOG(LogActorComponent, VeryVerbose, TEXT("%hs: %s %s"), __func__, ToString(RecvType), *GetName());
 
+		// Add the actor and it's sub-objects to the network id cache since these objects are no longer considered remote and can be derferenced
+		// through the network id cache (e.g. by RPC argument and object deserialization).
+		GetWorld()->GetNetDriver()->AddObjectToNetworkIdCache(this);
+
 		// We are expecting this flow from AActor, but now we must do this per-component in such a way
 		// that we are not mutating our previously deserialized data.
 		// RegisterAllComponents();
diff --git a/Engine/Source/Runtime/Engine/Private/NetDriver.cpp b/Engine/Source/Runtime/Engine/Private/NetDriver.cpp
--- a/Engine/Source/Runtime/Engine/Private/NetDriver.cpp
+++ b/Engine/Source/Runtime/Engine/Private/NetDriver.cpp
@@ -7922,6 +7922,104 @@ void UNetDriver::CreateInitialServerChannels(UNetConnection* ClientConnection)
 	}
 }
 
+void UNetDriver::AddObjectToNetworkIdCache(UObject* Object)
+{
+	if (!Object)
+	{
+		return;
+	}
+
+	if (!GuidCache)
+	{
+		return;
+	}
+
+	FNetworkGUID NetGUID = GuidCache->GetOrAssignNetGUID(Object);
+	UE_LOG(LogNet, Verbose, TEXT("Adding object %s (NetGUID=%s) to the NetGUID cache."), *Object->GetName(), *NetGUID.ToString());
+
+	auto AddSubObjectsToNetGUIDCache = [this](const UE::Net::FSubObjectRegistry* SubObjects)
+	{
+		for (const UE::Net::FSubObjectRegistry::FEntry& SubObjectInfo : SubObjects->GetRegistryList())
+		{
+			UObject* SubObject = SubObjectInfo.GetSubObject();
+			if (SubObject)
+			{
+				FNetworkGUID SubObjectNetGUID = GuidCache->GetOrAssignNetGUID(SubObject);
+				UE_LOG(LogNet, Verbose, TEXT("Adding sub-object %s (NetGUID=%s) to the NetGUID cache."), *SubObject->GetName(), *SubObjectNetGUID.ToString());
+			}
+		}
+	};
+
+	if (AActor* Actor = Cast<AActor>(Object))
+	{
+		const UE::Net::FSubObjectRegistry& ActorSubObjects = UE::Net::FSubObjectRegistryGetter::GetSubObjects(Actor);
+		AddSubObjectsToNetGUIDCache(&ActorSubObjects);
+	}
+	else if (UActorComponent* ActorComponent = Cast<UActorComponent>(Object))
+	{
+		const UE::Net::FSubObjectRegistry* ActorComponentSubObjects = UE::Net::FSubObjectRegistryGetter::GetSubObjectsOfActorComponent(ActorComponent->GetOwner(), ActorComponent);
+		if (ActorComponentSubObjects)
+		{
+			AddSubObjectsToNetGUIDCache(ActorComponentSubObjects);
+		}
+	}
+}
+
+void UNetDriver::RemoveObjectFromNetworkIdCache(UObject* Object)
+{
+	if (!Object)
+	{
+		return;
+	}
+
+	if (!GuidCache)
+	{
+		return;
+	}
+
+	if (GuidCache->RemoveNetGUID(Object))
+	{
+		UE_LOG(LogNet, Verbose, TEXT("Removed object %s from the NetGUID cache."), *Object->GetName());
+	}
+	else
+	{
+		UE_LOG(LogNet, Verbose, TEXT("Cannot find object %s in the NetGUID cache."), *Object->GetName());
+	}
+
+	auto RemoveSubObjectsFromNetGUIDCache = [this](const UE::Net::FSubObjectRegistry* SubObjects)
+	{
+		for (const UE::Net::FSubObjectRegistry::FEntry& SubObjectInfo : SubObjects->GetRegistryList())
+		{
+			UObject* SubObject = SubObjectInfo.GetSubObject();
+			if (SubObject)
+			{
+				if (GuidCache->RemoveNetGUID(SubObject))
+				{
+					UE_LOG(LogNet, Verbose, TEXT("Removed sub-object %s from the NetGUID cache."), *SubObject->GetName());
+				}
+				else
+				{
+					UE_LOG(LogNet, Verbose, TEXT("Cannot find sub-object %s in the NetGUID cache."), *SubObject->GetName());
+				}
+			}
+		}
+	};
+
+	if (AActor* Actor = Cast<AActor>(Object))
+	{
+		const UE::Net::FSubObjectRegistry& ActorSubObjects = UE::Net::FSubObjectRegistryGetter::GetSubObjects(Actor);
+		RemoveSubObjectsFromNetGUIDCache(&ActorSubObjects);
+	}
+	else if (UActorComponent* ActorComponent = Cast<UActorComponent>(Object))
+	{
+		const UE::Net::FSubObjectRegistry* ActorComponentSubObjects = UE::Net::FSubObjectRegistryGetter::GetSubObjectsOfActorComponent(ActorComponent->GetOwner(), ActorComponent);
+		if (ActorComponentSubObjects)
+		{
+			RemoveSubObjectsFromNetGUIDCache(ActorComponentSubObjects);
+		}
+	}
+}
+
 bool UNetDriver::ShouldReplicateFunction(AActor* Actor, UFunction* Function) const
 {
 	return (Actor && Actor->GetNetDriverName() == NetDriverName);
