Fix ActorInstanceHandle equality check for instances with no managers

#rnx

diff --git a/Engine/Source/Runtime/Engine/Private/Engine/ActorInstanceHandle.cpp b/Engine/Source/Runtime/Engine/Private/Engine/ActorInstanceHandle.cpp
--- a/Engine/Source/Runtime/Engine/Private/Engine/ActorInstanceHandle.cpp
+++ b/Engine/Source/Runtime/Engine/Private/Engine/ActorInstanceHandle.cpp
@@ -383,7 +383,7 @@ bool FActorInstanceHandle::operator==(const FActorInstanceHandle& Other) const
 
 	// If both handles point to the same Manager/InstanceIndex pair, we can assume they'll point
 	// to the same ReferenceObject.
-	if (ManagerInterface == Other.ManagerInterface && InstanceIndex == Other.InstanceIndex)
+	if ((InstanceIndex != INDEX_NONE && Other.InstanceIndex != INDEX_NONE) && (ManagerInterface == Other.ManagerInterface && InstanceIndex == Other.InstanceIndex))
 	{
 		return true;
 	}
