Fix bug in operator Select handling in new material translator causing wrong order of rows and columns.

#rnx

diff --git a/Engine/Source/Runtime/Engine/Private/Materials/MaterialIREmitter.cpp b/Engine/Source/Runtime/Engine/Private/Materials/MaterialIREmitter.cpp
--- a/Engine/Source/Runtime/Engine/Private/Materials/MaterialIREmitter.cpp
+++ b/Engine/Source/Runtime/Engine/Private/Materials/MaterialIREmitter.cpp
@@ -474,6 +474,28 @@ FValueRef FEmitter::InputDefaultFloat4(const FExpressionInput* Input, UE::Math::
 	return Value ? Value : Value.To(ConstantFloat4(Default));
 }
 
+FValueRef FEmitter::EmitAttributeDefaultValue(const FMaterialAggregateAttribute& Attribute)
+{
+	switch (Attribute.Type)
+	{
+	case EMaterialAggregateAttributeType::Bool1: return ConstantBool((bool)Attribute.DefaultValue.X);
+	case EMaterialAggregateAttributeType::Bool2: return ConstantBool2((bool)Attribute.DefaultValue.X, (bool)Attribute.DefaultValue.Y);
+	case EMaterialAggregateAttributeType::Bool3: return ConstantBool3((bool)Attribute.DefaultValue.X, (bool)Attribute.DefaultValue.Y, (bool)Attribute.DefaultValue.Z);
+	case EMaterialAggregateAttributeType::Bool4: return ConstantBool4((bool)Attribute.DefaultValue.X, (bool)Attribute.DefaultValue.Y, (bool)Attribute.DefaultValue.Z, (bool)Attribute.DefaultValue.W);
+	case EMaterialAggregateAttributeType::ShadingModel:
+	case EMaterialAggregateAttributeType::UInt1: return ConstantInt((TInteger)Attribute.DefaultValue.X);
+	case EMaterialAggregateAttributeType::UInt2: return ConstantInt2({ (TInteger)Attribute.DefaultValue.X, (TInteger)Attribute.DefaultValue.Y });
+	case EMaterialAggregateAttributeType::UInt3: return ConstantInt3({ (TInteger)Attribute.DefaultValue.X, (TInteger)Attribute.DefaultValue.Y, (TInteger)Attribute.DefaultValue.Z });
+	case EMaterialAggregateAttributeType::UInt4: return ConstantInt4({ (TInteger)Attribute.DefaultValue.X, (TInteger)Attribute.DefaultValue.Y, (TInteger)Attribute.DefaultValue.Z, (TInteger)Attribute.DefaultValue.W });
+	case EMaterialAggregateAttributeType::Float1: return ConstantFloat((TFloat)Attribute.DefaultValue.X);
+	case EMaterialAggregateAttributeType::Float2: return ConstantFloat2({ (TFloat)Attribute.DefaultValue.X, (TFloat)Attribute.DefaultValue.Y });
+	case EMaterialAggregateAttributeType::Float3: return ConstantFloat3({ (TFloat)Attribute.DefaultValue.X, (TFloat)Attribute.DefaultValue.Y, (TFloat)Attribute.DefaultValue.Z });
+	case EMaterialAggregateAttributeType::Float4: return ConstantFloat4(UE::Math::TVector4<TFloat>{ Attribute.DefaultValue.X, Attribute.DefaultValue.Y, Attribute.DefaultValue.Z, Attribute.DefaultValue.W });
+	case EMaterialAggregateAttributeType::Aggregate: return Aggregate(Attribute.Aggregate);
+	default: UE_MIR_UNREACHABLE();
+	}
+}
+
 FValueRef FEmitter::CheckTypeIsKind(FValueRef Value, ETypeKind Kind)
 {
 	if (Value.IsValid() && Value->Type.Is(Kind))
@@ -1244,28 +1266,19 @@ FValueRef FEmitter::Vector4(FValueRef InX, FValueRef InY, FValueRef InZ, FValueR
 
 FValueRef FEmitter::Aggregate(const UMaterialAggregate* InAggregate)
 {
-	return Aggregate(InAggregate, {}, TConstArrayView<FValueRef>{});
-}
-
-static FValueRef EmitAttributeDefaultValue(FEmitter& Emitter, const FMaterialAggregateAttribute& Attribute)
-{
-	switch (Attribute.Type)
+	if (InAggregate == UMaterialAggregate::GetMaterialAttributes())
 	{
-		case EMaterialAggregateAttributeType::Bool1: return Emitter.ConstantBool((bool)Attribute.DefaultValue.X);
-		case EMaterialAggregateAttributeType::Bool2: return Emitter.ConstantBool2((bool)Attribute.DefaultValue.X, (bool)Attribute.DefaultValue.Y);
-		case EMaterialAggregateAttributeType::Bool3: return Emitter.ConstantBool3((bool)Attribute.DefaultValue.X, (bool)Attribute.DefaultValue.Y, (bool)Attribute.DefaultValue.Z); 
-		case EMaterialAggregateAttributeType::Bool4: return Emitter.ConstantBool4((bool)Attribute.DefaultValue.X, (bool)Attribute.DefaultValue.Y, (bool)Attribute.DefaultValue.Z, (bool)Attribute.DefaultValue.W);
-		case EMaterialAggregateAttributeType::ShadingModel:
-		case EMaterialAggregateAttributeType::UInt1: return Emitter.ConstantInt((TInteger)Attribute.DefaultValue.X);
-		case EMaterialAggregateAttributeType::UInt2: return Emitter.ConstantInt2({ (TInteger)Attribute.DefaultValue.X, (TInteger)Attribute.DefaultValue.Y });
-		case EMaterialAggregateAttributeType::UInt3: return Emitter.ConstantInt3({ (TInteger)Attribute.DefaultValue.X, (TInteger)Attribute.DefaultValue.Y, (TInteger)Attribute.DefaultValue.Z });
-		case EMaterialAggregateAttributeType::UInt4: return Emitter.ConstantInt4({ (TInteger)Attribute.DefaultValue.X, (TInteger)Attribute.DefaultValue.Y, (TInteger)Attribute.DefaultValue.Z, (TInteger)Attribute.DefaultValue.W });
-		case EMaterialAggregateAttributeType::Float1: return Emitter.ConstantFloat((TFloat)Attribute.DefaultValue.X);
-		case EMaterialAggregateAttributeType::Float2: return Emitter.ConstantFloat2({ (TFloat)Attribute.DefaultValue.X, (TFloat)Attribute.DefaultValue.Y });
-		case EMaterialAggregateAttributeType::Float3: return Emitter.ConstantFloat3({ (TFloat)Attribute.DefaultValue.X, (TFloat)Attribute.DefaultValue.Y, (TFloat)Attribute.DefaultValue.Z });
-		case EMaterialAggregateAttributeType::Float4: return Emitter.ConstantFloat4(UE::Math::TVector4<TFloat>{ Attribute.DefaultValue.X, Attribute.DefaultValue.Y, Attribute.DefaultValue.Z, Attribute.DefaultValue.W });
-		case EMaterialAggregateAttributeType::Aggregate: return Emitter.Aggregate(Attribute.Aggregate);
-		default: UE_MIR_UNREACHABLE();
+		FValueRef AggregateValues[MP_MAX];
+		TConstArrayView<EMaterialProperty> Properties = InAggregate->GetMaterialAttributesProperties();
+		for (int32 Index = 0; Index < Properties.Num(); Index++)
+		{
+			AggregateValues[Index] = MIR::Internal::CreateMaterialAttributeDefaultValue(*this, Material, Properties[Index]);
+		}
+		return Aggregate(InAggregate, {}, { AggregateValues, Properties.Num() });
+	}
+	else
+	{
+		return Aggregate(InAggregate, {}, TConstArrayView<FValueRef>{});
 	}
 }
 
@@ -1303,7 +1316,7 @@ FValueRef FEmitter::Aggregate(const UMaterialAggregate* InAggregate, FValueRef I
 		else
 		{
 			// ...otherwise construct the default value as indicated in the attribute.
-			Components[i] = EmitAttributeDefaultValue(*this, InAggregate->Attributes[i]);
+			Components[i] = EmitAttributeDefaultValue(InAggregate->Attributes[i]);
 		}
 	}
 
@@ -1351,7 +1364,7 @@ FValueRef FEmitter::Aggregate(const UMaterialAggregate* InAggregate, FValueRef I
 		else
 		{
 			// ...otherwise construct the default value.
-			Components[i] = EmitAttributeDefaultValue(*this, InAggregate->Attributes[i]);
+			Components[i] = EmitAttributeDefaultValue(InAggregate->Attributes[i]);
 		}
 	}
 
@@ -2165,12 +2178,12 @@ static FType ValidateOperatorAndGetResultType(FEmitter& Emitter, EOperator Op, F
 		int32 MaxNumComponents = FMath::Max3(FirstArgumentType.GetPrimitive().NumComponents(), BPrimitiveType.NumComponents(), CPrimitiveType.NumComponents());
 
 		// Cast the first argument (the boolean condition) to a bool vector of the maximum number of components.
-		Arguments[0] = Emitter.Cast(Arguments[0], FType::MakePrimitive(EScalarKind::Bool, MaxNumComponents, 1));
+		Arguments[0] = Emitter.Cast(Arguments[0], FType::MakeVector(EScalarKind::Bool, MaxNumComponents));
 
 		// Compute the common type between the second and third argument types with a number of components equal to the max of all three.
 		FType CommonTypeBetweenSecondAndThirdArguments = Emitter.TryGetCommonType(
-			FType::MakePrimitive(BPrimitiveType.ScalarKind, MaxNumComponents, 1),
-			FType::MakePrimitive(CPrimitiveType.ScalarKind, MaxNumComponents, 1));
+			FType::MakeVector(BPrimitiveType.ScalarKind, MaxNumComponents),
+			FType::MakeVector(CPrimitiveType.ScalarKind, MaxNumComponents));
 		
 		// Getting the common type should always be possible.
 		check(CommonTypeBetweenSecondAndThirdArguments);
