AssetManager: Small optimization to BuildChunkMap by allowing GetPrimaryAssetSetChunkIds to take a TArray of managers rather than a TSet.
#rnx

diff --git a/Engine/Source/Runtime/Engine/Classes/Engine/AssetManager.h b/Engine/Source/Runtime/Engine/Classes/Engine/AssetManager.h
--- a/Engine/Source/Runtime/Engine/Classes/Engine/AssetManager.h
+++ b/Engine/Source/Runtime/Engine/Classes/Engine/AssetManager.h
@@ -887,7 +887,10 @@ public:
 	virtual int32 GetContentEncryptionGroupChunkID(FName InGroupName) const { return INDEX_NONE; }
 
 	/** Returns the list of chunks assigned to the list of primary assets, which is usually a manager list. This is called by GetPackageChunkIds */
-	ENGINE_API virtual bool GetPrimaryAssetSetChunkIds(const TSet<FPrimaryAssetId>& PrimaryAssetSet, const ITargetPlatform* TargetPlatform, TArrayView<const int32> ExistingChunkList, TArray<int32>& OutChunkList) const;
+	ENGINE_API virtual bool GetPrimaryAssetSetChunkIds(const TSet<FPrimaryAssetId>& PrimaryAssetSet, const ITargetPlatform* TargetPlatform,
+		TArrayView<const int32> ExistingChunkList, TArray<int32>& OutChunkList) const;
+	ENGINE_API virtual bool GetPrimaryAssetSetChunkIds(TConstArrayView<FPrimaryAssetId> PrimaryAssets, const ITargetPlatform* TargetPlatform,
+		TArrayView<const int32> ExistingChunkList, TArray<int32>& OutChunkList) const;
 
 	/** Resets all asset manager data, called in the editor to reinitialize the config */
 	ENGINE_API virtual void ReinitializeFromConfig();
diff --git a/Engine/Source/Runtime/Engine/Private/AssetManager.cpp b/Engine/Source/Runtime/Engine/Private/AssetManager.cpp
--- a/Engine/Source/Runtime/Engine/Private/AssetManager.cpp
+++ b/Engine/Source/Runtime/Engine/Private/AssetManager.cpp
@@ -5086,11 +5086,18 @@ void UAssetManager::UpdateCachedChunkMapAfterCook(const TSet<FName>& AllCookedPa
 	CachedChunkMap = BuildChunkMap(AllCookedPackages, StartupPackages, AllCookedPackages);
 }
 
-bool UAssetManager::GetPrimaryAssetSetChunkIds(const TSet<FPrimaryAssetId>& PrimaryAssetSet, const class ITargetPlatform* TargetPlatform, TArrayView<const int32> ExistingChunkList, TArray<int32>& OutChunkList) const
+bool UAssetManager::GetPrimaryAssetSetChunkIds(const TSet<FPrimaryAssetId>& PrimaryAssetSet, const class ITargetPlatform* TargetPlatform,
+	TArrayView<const int32> ExistingChunkList, TArray<int32>& OutChunkList) const
+{
+	return GetPrimaryAssetSetChunkIds(PrimaryAssetSet.Array(), TargetPlatform, ExistingChunkList, OutChunkList);
+}
+
+bool UAssetManager::GetPrimaryAssetSetChunkIds(TConstArrayView<FPrimaryAssetId> PrimaryAssets, const class ITargetPlatform* TargetPlatform,
+	TArrayView<const int32> ExistingChunkList, TArray<int32>& OutChunkList) const
 {
 	bool bFoundAny = false;
 	int32 HighestChunk = 0;
-	for (const FPrimaryAssetId& PrimaryAssetId : PrimaryAssetSet)
+	for (const FPrimaryAssetId& PrimaryAssetId : PrimaryAssets)
 	{
 		FPrimaryAssetRules Rules = GetPrimaryAssetRules(PrimaryAssetId);
 
