New Material Translator:  21 new expressions implemented.  Includes distance field and sky atmosphere related expressions, various expressions composed of math ops, and a couple trivial external code expressions.

#rnx

diff --git a/Engine/Config/BaseMaterialExpressions.ini b/Engine/Config/BaseMaterialExpressions.ini
--- a/Engine/Config/BaseMaterialExpressions.ini
+++ b/Engine/Config/BaseMaterialExpressions.ini
@@ -93,6 +93,7 @@ bShaderInputData=true
 
 [/Script/Engine.MaterialExpressionAtmosphericFogColor]
 +MenuCategories=NSLOCTEXT("MaterialExpression", "Atmosphere", "Atmosphere")
++ExternalCodeIdentifiers="SkyAtmosphereAerialPerspective"
 
 [/Script/Engine.MaterialExpressionAtmosphericLightColor]
 +MenuCategories=NSLOCTEXT("MaterialExpression", "Utility", "Utility")
@@ -395,6 +396,7 @@ bShaderInputData=true
 
 [/Script/Engine.MaterialExpressionLightmapUVs]
 +MenuCategories=NSLOCTEXT("MaterialExpression", "Coordinates", "Coordinates")
++ExternalCodeIdentifiers="LightmapUVs"
 bShaderInputData=true
 
 [/Script/Engine.MaterialExpressionLightmassReplace]
@@ -769,27 +771,35 @@ bShaderInputData=true
 
 [/Script/Engine.MaterialExpressionSkyAtmosphereLightDirection]
 +MenuCategories=NSLOCTEXT("MaterialExpression", "Sky", "Sky")
++ExternalCodeIdentifiers="SkyAtmosphereLightDirection"
 
 [/Script/Engine.MaterialExpressionSkyAtmosphereLightIlluminance]
 +MenuCategories=NSLOCTEXT("MaterialExpression", "Sky", "Sky")
++ExternalCodeIdentifiers="SkyAtmosphereLightIlluminance"
 
 [/Script/Engine.MaterialExpressionSkyAtmosphereLightIlluminanceOnGround]
 +MenuCategories=NSLOCTEXT("MaterialExpression", "Sky", "Sky")
++ExternalCodeIdentifiers="SkyAtmosphereLightIlluminanceOnGround"
 
 [/Script/Engine.MaterialExpressionSkyAtmosphereLightDiskLuminance]
 +MenuCategories=NSLOCTEXT("MaterialExpression", "Sky", "Sky")
++ExternalCodeIdentifiers="SkyAtmosphereLightDiskLuminance"
 
 [/Script/Engine.MaterialExpressionSkyAtmosphereAerialPerspective]
 +MenuCategories=NSLOCTEXT("MaterialExpression", "Sky", "Sky")
++ExternalCodeIdentifiers="SkyAtmosphereAerialPerspective"
 
 [/Script/Engine.MaterialExpressionSkyAtmosphereDistantLightScatteredLuminance]
 +MenuCategories=NSLOCTEXT("MaterialExpression", "Sky", "Sky")
++ExternalCodeIdentifiers="SkyAtmosphereDistantLightScatteredLuminance"
 
 [/Script/Engine.MaterialExpressionSkyAtmosphereViewLuminance]
 +MenuCategories=NSLOCTEXT("MaterialExpression", "Sky", "Sky")
++ExternalCodeIdentifiers="SkyAtmosphereViewLuminance"
 
 [/Script/Engine.MaterialExpressionSkyLightEnvMapSample]
 +MenuCategories=NSLOCTEXT("MaterialExpression", "Sky", "Sky")
++ExternalCodeIdentifiers="SkyLightEnvMapSample"
 
 [/Script/Engine.MaterialExpressionSobol]
 +MenuCategories=NSLOCTEXT("MaterialExpression", "Utility", "Utility")
@@ -1011,6 +1021,7 @@ bShaderInputData=true
 +ExternalCodeDeclarations=(ReturnType=Float,			Name="IsFirstPerson",					Definition="(IsFirstPerson(Parameters) ? 1.0f : 0.0f)", Derivative=Zero, bIsInlined=true, Domains=(MD_Surface, MD_Volume))
 +ExternalCodeDeclarations=(ReturnType=Float,			Name="IsOrthographic",					Definition="IsOrthoProjectionFloat()", Derivative=Zero, bIsInlined=true)
 +ExternalCodeDeclarations=(ReturnType=Float3,			Name="LightVector",						Definition="Parameters.LightVector", ShaderFrequency=Pixel|Compute, bIsInlined=true, Domains=(MD_LightFunction, MD_DeferredDecal))
++ExternalCodeDeclarations=(ReturnType=Float2,			Name="LightmapUVs",						Definition="GetLightmapUVs(Parameters)", DefinitionDDX="GetLightmapUVs_DDX(Parameters)", DefinitionDDY="GetLightmapUVs_DDY(Parameters)", Derivative=Valid, EnvironmentDefines=((Name="LIGHTMAP_UV_ACCESS")), ShaderFrequency=Pixel|Compute, MinimumFeatureLevel=SM5)
 +ExternalCodeDeclarations=(ReturnType=Float,			Name="MeshPaintTextureCoordinateIndex",	Definition="GetMeshPaintTextureCoordinateIndex(GetPrimitiveData(Parameters))", bIsInlined=true)
 +ExternalCodeDeclarations=(ReturnType=TextureMeshPaint,	Name="MeshPaintTextureDescriptor",		Definition="GetMeshPaintTextureDescriptor(GetPrimitiveData(Parameters))", Derivative=Zero, bIsInlined=true)
 +ExternalCodeDeclarations=(ReturnType=Float3,			Name="ObjectBounds",					Definition="float3(GetPrimitiveData(Parameters).ObjectBoundsX, GetPrimitiveData(Parameters).ObjectBoundsY, GetPrimitiveData(Parameters).ObjectBoundsZ)", Derivative=Zero, bIsInlined=true, Domains=(MD_Surface, MD_Volume))
@@ -1051,6 +1062,15 @@ bShaderInputData=true
 +ExternalCodeDeclarations=(ReturnType=Float3,	Name="GetCloudEmptySpaceSkippingSphereCenterWorldPosition",	Definition="MaterialExpressionCloudEmptySpaceSkippingSphereCenterWorldPosition(Parameters)")
 +ExternalCodeDeclarations=(ReturnType=Float1,	Name="GetCloudEmptySpaceSkippingSphereRadius",				Definition="MaterialExpressionCloudEmptySpaceSkippingSphereRadius(Parameters)")
 
++ExternalCodeDeclarations=(ReturnType=Float4,	Name="SkyAtmosphereAerialPerspective",						Definition="MaterialExpressionSkyAtmosphereAerialPerspective(Parameters, $0)", EnvironmentDefines=((Name="MATERIAL_SKY_ATMOSPHERE")))
++ExternalCodeDeclarations=(ReturnType=Float3,	Name="SkyAtmosphereDistantLightScatteredLuminance",			Definition="MaterialExpressionSkyAtmosphereDistantLightScatteredLuminance(Parameters)", EnvironmentDefines=((Name="MATERIAL_SKY_ATMOSPHERE")))
++ExternalCodeDeclarations=(ReturnType=Float3,	Name="SkyAtmosphereLightDirection",							Definition="MaterialExpressionSkyAtmosphereLightDirection(Parameters, $0)", EnvironmentDefines=((Name="MATERIAL_SKY_ATMOSPHERE")))
++ExternalCodeDeclarations=(ReturnType=Float3,	Name="SkyAtmosphereLightDiskLuminance",						Definition="MaterialExpressionSkyAtmosphereLightDiskLuminance(Parameters, $0, $1)", EnvironmentDefines=((Name="MATERIAL_SKY_ATMOSPHERE")))
++ExternalCodeDeclarations=(ReturnType=Float3,	Name="SkyAtmosphereLightIlluminance",						Definition="MaterialExpressionSkyAtmosphereLightIlluminance(Parameters, $0, $1)", EnvironmentDefines=((Name="MATERIAL_SKY_ATMOSPHERE")))
++ExternalCodeDeclarations=(ReturnType=Float3,	Name="SkyAtmosphereLightIlluminanceOnGround",				Definition="MaterialExpressionSkyAtmosphereLightIlluminanceOnGround(Parameters, $0)", EnvironmentDefines=((Name="MATERIAL_SKY_ATMOSPHERE")))
++ExternalCodeDeclarations=(ReturnType=Float3,	Name="SkyAtmosphereViewLuminance",							Definition="MaterialExpressionSkyAtmosphereViewLuminance(Parameters, $0)", EnvironmentDefines=((Name="MATERIAL_SKY_ATMOSPHERE")))
++ExternalCodeDeclarations=(ReturnType=Float3,	Name="SkyLightEnvMapSample",								Definition="MaterialExpressionSkyLightEnvMapSample($0, $1)", Domains=(MD_Surface), EnvironmentDefines=((Name="MATERIAL_SKY_ATMOSPHERE")))
+
 +ExternalCodeDeclarations=(ReturnType=Bool,		Name="LightmassReplace",							Definition="GetLightmassReplaceState()")
 +ExternalCodeDeclarations=(ReturnType=Bool,		Name="MeshPaintTextureReplace",						Definition="GetMeshPaintTextureDescriptorIsValid(GetPrimitiveData(Parameters))")
 +ExternalCodeDeclarations=(ReturnType=Bool,		Name="NaniteReplace",								Definition="GetNaniteReplaceState()")
diff --git a/Engine/Source/Runtime/Engine/Private/Materials/MaterialExpressionsIR.cpp b/Engine/Source/Runtime/Engine/Private/Materials/MaterialExpressionsIR.cpp
--- a/Engine/Source/Runtime/Engine/Private/Materials/MaterialExpressionsIR.cpp
+++ b/Engine/Source/Runtime/Engine/Private/Materials/MaterialExpressionsIR.cpp
@@ -256,6 +256,8 @@
 
 using FValueRef = MIR::FValueRef;
 
+static FName NAME_CameraVector("CameraVector");
+
 /* Constants */
 
 void UMaterialExpression::Build(MIR::FEmitter& Em)
@@ -2277,7 +2279,7 @@ void UMaterialExpressionObjectPositionWS::Build(MIR::FEmitter& Em)
 	}
 }
 
-void UMaterialExpressionWorldPosition::Build(MIR::FEmitter& Em)
+static FValueRef EmitWorldPosition(MIR::FEmitter& Em, EWorldPositionIncludedOffsets WorldPositionShaderOffset)
 {
 	// Make sure EWorldPositionIncludedOffsets and corresponding elements of EExternalInput stay in sync, so this enum addition is valid
 	static_assert(WPT_MAX == 4 && WPT_Default == 0);
@@ -2293,6 +2295,33 @@ void UMaterialExpressionWorldPosition::Build(MIR::FEmitter& Em)
 		WorldPosition = Em.CastToFloatKind(WorldPosition);
 	}
 
+	return WorldPosition;
+}
+
+// Adapted from FHLSLMaterialTranslator::GetWorldPositionOrDefault.  WorldPosition input is optional, a default is provided if not set.
+static FValueRef EmitWorldPositionOrDefault(MIR::FEmitter& Em, FValueRef WorldPosition, EPositionOrigin PositionOrigin)
+{
+	if (PositionOrigin != EPositionOrigin::Absolute && PositionOrigin != EPositionOrigin::CameraRelative)
+	{
+		Em.Error(TEXT("Invalid EPositionOrigin enum value."));
+		return Em.Poison();
+	}
+	if (WorldPosition)
+	{
+		// Sanitize the explicitly provided input to the correct vector type if needed.
+		return Em.Cast(WorldPosition, PositionOrigin == EPositionOrigin::Absolute ? MIR::FType::MakeDoubleVector(3) : MIR::FType::MakeFloatVector(3));
+	}
+	else
+	{
+		// Return default world position.
+		return EmitWorldPosition(Em, PositionOrigin == EPositionOrigin::CameraRelative ? EWorldPositionIncludedOffsets::WPT_CameraRelative : EWorldPositionIncludedOffsets::WPT_Default);
+	}
+}
+
+void UMaterialExpressionWorldPosition::Build(MIR::FEmitter& Em)
+{
+	FValueRef WorldPosition = EmitWorldPosition(Em, WorldPositionShaderOffset);
+
 	Em.Output(0, WorldPosition);
 	Em.Output(1, Em.Swizzle(WorldPosition, MIR::FSwizzleMask(MIR::EVectorComponent::X, MIR::EVectorComponent::Y)));
 	Em.Output(2, Em.Subscript(WorldPosition, 2));
@@ -2810,4 +2839,309 @@ void UMaterialExpressionCollectionTransform::Build(MIR::FEmitter& Em)
 	Em.Output(0, Result);
 }
 
+void UMaterialExpressionAtmosphericFogColor::Build(MIR::FEmitter& Em)
+{
+	// This node is deprecated in favor of UMaterialExpressionSkyAtmosphereAerialPerspective, and falls through to the newer expression
+	FValueRef PositionValue = EmitWorldPositionOrDefault(Em, Em.TryInput(&WorldPosition), WorldPositionOriginType);
+	Em.Output(0, EmitInlineHLSL(Em, *this, 0, { PositionValue }));
+}
+
+void UMaterialExpressionBlackBody::Build(MIR::FEmitter& Em)
+{
+	Em.Output(0, Em.InlineHLSL(MIR::FType::MakeFloatVector(3), TEXTVIEW("MaterialExpressionBlackBody($0)"), { Em.Input(&Temp) }));
+}
+
+void UMaterialExpressionDepthFade::Build(MIR::FEmitter& Em)
+{
+	// Scales Opacity by a Linear fade based on SceneDepth, from 0 at PixelDepth to 1 at FadeDistance
+	// Result = Opacity * saturate((SceneDepth - PixelDepth) / max(FadeDistance, DELTA))
+	FValueRef OpacityValue = Em.InputDefaultFloat(&InOpacity, OpacityDefault);
+	FValueRef FadeDistanceValue = Em.Max(Em.InputDefaultFloat(&FadeDistance, FadeDistanceDefault), Em.ConstantFloat(UE_DELTA));
+
+	static FName NAME_PixelDepth("PixelDepth");
+	FValueRef PixelDepth = EmitInlineHLSL(Em, NAME_PixelDepth, {});
+	// On mobile scene depth is limited to 65500 
+	// to avoid false fading on objects that are close or exceed this limit we clamp pixel depth to (65500 - FadeDistance)
+	if (Em.GetFeatureLevel() <= ERHIFeatureLevel::ES3_1)
+	{
+		PixelDepth = Em.Min(PixelDepth, Em.Subtract(Em.ConstantFloat(65500.f), FadeDistanceValue));
+	}
+
+	// We need a dependency on EScreenTexture::SceneDepth, so the value analyzer can see it, even though it's technically not used in the code.
+	FValueRef SceneDepth = Em.InlineHLSL(MIR::FType::MakeFloatScalar(), TEXTVIEW("CalcSceneDepth(ScreenAlignedPosition(GetScreenPosition(Parameters)))"), { Em.ScreenTexture(MIR::EScreenTexture::SceneDepth) });
+
+	Em.Output(0, Em.Multiply(OpacityValue, Em.Saturate(Em.Divide(Em.Subtract(SceneDepth, PixelDepth), FadeDistanceValue))));
+}
+
+void UMaterialExpressionDeriveNormalZ::Build(MIR::FEmitter& Em)
+{
+	// z = sqrt(saturate(1 - ( x * x + y * y)));
+	FValueRef InputVector = Em.Cast(Em.Input(&InXY), MIR::FType::MakeFloatVector(2));
+	FValueRef DotResult = Em.Dot(InputVector, InputVector);
+	FValueRef InnerResult = Em.Subtract(Em.ConstantFloat(1.0f), DotResult);
+	FValueRef SaturatedInnerResult = Em.Saturate(InnerResult);
+	FValueRef DerivedZ = Em.Sqrt(SaturatedInnerResult);
+	
+	Em.Output(0, Em.Vector3(Em.Subscript(InputVector, 0), Em.Subscript(InputVector, 1), DerivedZ));
+}
+
+void UMaterialExpressionDistanceFieldApproxAO::Build(MIR::FEmitter& Em)
+{
+	FValueRef PositionValue = EmitWorldPositionOrDefault(Em, Em.TryInput(&Position), WorldPositionOriginType);
+
+	FValueRef NormalValue = Em.TryInput(&Normal);
+	if (!NormalValue)
+	{
+		static FName NAME_VertexNormal("VertexNormal");
+		NormalValue = EmitInlineHLSL(Em, NAME_VertexNormal, {});
+	}
+
+	FValueRef BaseDistanceValue = Em.InputDefaultFloat(&BaseDistance, BaseDistanceDefault);
+
+	int32 NumStepsClamped = FMath::Clamp(NumSteps, 1, 4);
+	float StepScaleClamped = FMath::Max(StepScaleDefault, 1.0f);
+
+	FValueRef NumStepsConst = Em.ConstantInt(NumStepsClamped);
+	FValueRef NumStepsMinus1Const = Em.ConstantInt(NumStepsClamped - 1);
+	FValueRef StepScaleConst = Em.ConstantFloat(StepScaleClamped);
+
+	FValueRef StepDistance;
+	FValueRef DistanceBias;
+	FValueRef MaxDistance;
+
+	if (NumSteps == 1)
+	{
+		StepDistance = Em.ConstantFloat(0);
+		DistanceBias = BaseDistanceValue;
+		MaxDistance = BaseDistanceValue;
+	}
+	else
+	{
+		FValueRef RadiusValue = Em.InputDefaultFloat(&Radius, RadiusDefault);
+
+		StepDistance = Em.Divide(Em.Subtract(RadiusValue, BaseDistanceValue), Em.Subtract(Em.Pow(StepScaleConst, NumStepsMinus1Const), Em.ConstantFloat(1.0f)));
+		DistanceBias = Em.Subtract(BaseDistanceValue, StepDistance);
+		MaxDistance = RadiusValue;
+	}
+
+	// Last input tells value analyzer that this expression uses the global distance field
+	Em.Output(0, Em.InlineHLSL(MIR::FType::MakeFloatScalar(), TEXTVIEW("CalculateDistanceFieldApproxAO($0, $1, $2, $3, $4, $5, $6)"),
+		{
+			PositionValue,
+			Em.Cast(NormalValue, MIR::FType::MakeFloatVector(3)),
+			NumStepsConst,
+			Em.Cast(StepDistance, MIR::FType::MakeFloatScalar()),
+			StepScaleConst,
+			Em.Cast(DistanceBias, MIR::FType::MakeFloatScalar()),
+			Em.Cast(MaxDistance, MIR::FType::MakeFloatScalar()),
+			Em.ExternalInput(MIR::EExternalInput::GlobalDistanceField)
+		}));
+}
+
+void UMaterialExpressionDistanceFieldGradient::Build(MIR::FEmitter& Em)
+{
+	// Last input tells value analyzer that this expression uses the global distance field
+	FValueRef PositionValue = EmitWorldPositionOrDefault(Em, Em.TryInput(&Position), WorldPositionOriginType);
+	Em.Output(0, Em.InlineHLSL(MIR::FType::MakeFloatVector(3), TEXTVIEW("GetDistanceFieldGradientGlobal($0)"), { PositionValue, Em.ExternalInput(MIR::EExternalInput::GlobalDistanceField) }));
+}
+
+void UMaterialExpressionDistanceToNearestSurface::Build(MIR::FEmitter& Em)
+{
+	// Last input tells value analyzer that this expression uses the global distance field
+	FValueRef PositionValue = EmitWorldPositionOrDefault(Em, Em.TryInput(&Position), WorldPositionOriginType);
+	Em.Output(0, Em.InlineHLSL(MIR::FType::MakeFloatScalar(), TEXTVIEW("GetDistanceToNearestSurfaceGlobal($0)"), { PositionValue, Em.ExternalInput(MIR::EExternalInput::GlobalDistanceField) }));
+}
+
+void UMaterialExpressionFresnel::Build(MIR::FEmitter& Em)
+{
+	// pow(1 - max(0,Normal dot Camera),Exponent) * (1 - BaseReflectFraction) + BaseReflectFraction
+	//
+	FValueRef NormalArg = Em.TryInput(&Normal);
+	if (NormalArg)
+	{
+		NormalArg = Em.Cast(NormalArg, MIR::FType::MakeFloatVector(3));
+	}
+	else
+	{
+		NormalArg = Em.InlineHLSL(MIR::FType::MakeFloatVector(3), TEXTVIEW("Parameters.WorldNormal"), {}, MIR::EValueFlags::None, MIR::EGraphProperties::ReadsPixelNormal);
+	}
+
+	FValueRef DotArg = Em.Dot(NormalArg, EmitInlineHLSL(Em, NAME_CameraVector, {}));
+	FValueRef MaxArg = Em.Max(Em.ConstantFloat(0.f), DotArg);
+	FValueRef MinusArg = Em.Subtract(Em.ConstantFloat(1.f), MaxArg);
+	FValueRef ExponentArg = Em.InputDefaultFloat(&ExponentIn, Exponent);
+	// Compiler->Power got changed to call PositiveClampedPow instead of ClampedPow
+	// Manually implement ClampedPow to maintain backwards compatibility in the case where the input normal is not normalized (length > 1)
+	FValueRef AbsBaseArg = Em.Max(Em.Abs(MinusArg), Em.ConstantFloat(UE_KINDA_SMALL_NUMBER));
+	FValueRef PowArg = Em.Pow(AbsBaseArg, ExponentArg);
+	FValueRef BaseReflectFractionArg = Em.InputDefaultFloat(&BaseReflectFractionIn, BaseReflectFraction);
+	FValueRef ScaleArg = Em.Multiply(PowArg, Em.Subtract(Em.ConstantFloat(1.f), BaseReflectFractionArg));
+	
+	Em.Output(0, Em.Add(ScaleArg, BaseReflectFractionArg));
+}
+
+void UMaterialExpressionReflectionVectorWS::Build(MIR::FEmitter& Em)
+{
+	FValueRef NormalValue = Em.TryInput(&CustomWorldNormal);
+
+	if (NormalValue)
+	{
+		NormalValue = Em.Cast(NormalValue, MIR::FType::MakeFloatVector(3));
+
+		// Ported from HLSL utility function ReflectionAboutCustomWorldNormal
+		if (bNormalizeCustomWorldNormal)
+		{
+			NormalValue = Em.Multiply(NormalValue, Em.Rsqrt(Em.Dot(NormalValue, NormalValue)));
+		}
+
+		// Normal * dot(Normal, CameraVector) * 2.0 - CameraVector;
+		FValueRef CameraVector = EmitInlineHLSL(Em, NAME_CameraVector, {});
+		Em.Output(0, Em.Subtract(Em.Multiply(NormalValue, Em.Multiply(Em.Dot(NormalValue, CameraVector), Em.ConstantFloat(2.0f))), CameraVector));
+	}
+	else
+	{
+		static FName NAME_ReflectionVector("ReflectionVector");
+		Em.Output(0, EmitInlineHLSL(Em, NAME_ReflectionVector, {}));
+	}
+}
+
+void UMaterialExpressionRotateAboutAxis::Build(MIR::FEmitter& Em)
+{
+	FValueRef Angle = Em.Multiply(Em.Subscript(Em.Input(&RotationAngle), 0), Em.ConstantFloat(2.0f * (float)UE_PI / Period));
+	FValueRef Axis = Em.Cast(Em.Input(&NormalizedRotationAxis), MIR::FType::MakeFloatVector(3));
+	FValueRef PosOnAxis = Em.Input(&PivotPoint);
+	FValueRef Pos = Em.Input(&Position);
+
+	// Math adapted from RotateAboutAxis, but simplified and optimized slightly.  Note that the function returns an offset to
+	// the rotated position, not an absolute position, and so the offset will be non-LWC.  This initial subtraction is LWC aware,
+	// but we can then use float operations for the remainder (the LWC RotateAboutAxis HLSL function does the same).
+	FValueRef PosOffset = Em.Cast(Em.Subtract(Pos, PosOnAxis), MIR::FType::MakeFloatVector(3));
+
+	// Construct orthogonal axes in the plane of rotation.  The UAxis is computed by subtracting the projection of
+	// PosOffset along the Axis vector.
+	FValueRef UAxis = Em.Subtract(PosOffset, Em.Multiply(Axis, Em.Dot(Axis, PosOffset)));
+	FValueRef VAxis = Em.Cross(Axis, UAxis);
+
+	// Rotate the orthogonal axes
+	FValueRef CosAngle = Em.Cos(Angle);
+	FValueRef SinAngle = Em.Sin(Angle);
+	FValueRef R = Em.Add(Em.Multiply(UAxis, CosAngle), Em.Multiply(VAxis, SinAngle));
+
+	// Return the offset from the original position to the rotated position.  The original position in this context
+	// is the pre-rotation axis vector.
+	Em.Output(0, Em.Subtract(R, UAxis));
+}
+
+void UMaterialExpressionRotator::Build(MIR::FEmitter& Em)
+{
+	FValueRef TimeValue = Em.TryInput(&Time);
+	if (!TimeValue)
+	{
+		TimeValue = Em.InlineHLSL(MIR::FType::MakeFloatScalar(), TEXTVIEW("View.<PREV>GameTime"), {}, MIR::EValueFlags::SubstituteTags);
+	}
+	TimeValue = Em.Multiply(TimeValue, Em.ConstantFloat(Speed));
+
+	FValueRef BaseCoordinate = Em.TryInput(&Coordinate);
+	if (!BaseCoordinate)
+	{
+		BaseCoordinate = Em.ExternalInput(MIR::TexCoordIndexToExternalInput(ConstCoordinate));
+	}
+	BaseCoordinate = Em.Subtract(BaseCoordinate, Em.ConstantFloat2({ CenterX, CenterY }));
+
+	FValueRef CosValue = Em.Cos(TimeValue);
+	FValueRef SinValue = Em.Sin(TimeValue);
+
+	FValueRef Arg1 = Em.Add(Em.Subtract(Em.Multiply(CosValue, Em.Subscript(BaseCoordinate, 0)), Em.Multiply(SinValue, Em.Subscript(BaseCoordinate, 1))), Em.ConstantFloat(CenterX));		// cos*U - sin*V + CenterX
+	FValueRef Arg2 = Em.Add(Em.Add     (Em.Multiply(SinValue, Em.Subscript(BaseCoordinate, 0)), Em.Multiply(CosValue, Em.Subscript(BaseCoordinate, 1))), Em.ConstantFloat(CenterY));		// sin*U + cos*V + CenterY
+
+	TOptional<MIR::FPrimitive> BaseType = BaseCoordinate->Type.AsPrimitive();
+	if (BaseType && BaseType->NumColumns >= 3)
+	{
+		Em.Output(0, Em.Vector3(Arg1, Arg2, Em.Subscript(BaseCoordinate, 2)));
+	}
+	else
+	{
+		Em.Output(0, Em.Vector2(Arg1, Arg2));
+	}
+}
+
+void UMaterialExpressionSkyAtmosphereAerialPerspective::Build(MIR::FEmitter& Em)
+{
+	FValueRef PositionValue = EmitWorldPositionOrDefault(Em, Em.TryInput(&WorldPosition), WorldPositionOriginType);
+	Em.Output(0, EmitInlineHLSL(Em, *this, 0, { PositionValue }));
+}
+
+void UMaterialExpressionSkyAtmosphereLightDirection::Build(MIR::FEmitter& Em)
+{
+	Em.Output(0, EmitInlineHLSL(Em, *this, 0, { Em.ConstantInt(LightIndex) }));
+}
+
+void UMaterialExpressionSkyAtmosphereLightDiskLuminance::Build(MIR::FEmitter& Em)
+{
+	FValueRef CosHalfDiskRadius = Em.TryInput(&DiskAngularDiameterOverride);
+	if (CosHalfDiskRadius)
+	{
+		// Convert from apex angle (angular diameter) to cosine of the disk radius.
+		CosHalfDiskRadius = Em.Cos(Em.Multiply(Em.ConstantFloat(0.5f * float(UE_PI) / 180.0f), CosHalfDiskRadius));
+	}
+	else
+	{
+		CosHalfDiskRadius = Em.ConstantFloat(-1.0f);
+	}
+	Em.Output(0, EmitInlineHLSL(Em, *this, 0, { Em.ConstantInt(LightIndex), CosHalfDiskRadius }));
+}
+
+void UMaterialExpressionSkyAtmosphereLightIlluminance::Build(MIR::FEmitter& Em)
+{
+	FValueRef PositionValue = EmitWorldPositionOrDefault(Em, Em.TryInput(&WorldPosition), WorldPositionOriginType);
+	Em.Output(0, EmitInlineHLSL(Em, *this, 0, { PositionValue, Em.ConstantInt(LightIndex) }));
+}
+
+void UMaterialExpressionSkyAtmosphereLightIlluminanceOnGround::Build(MIR::FEmitter& Em)
+{
+	Em.Output(0, EmitInlineHLSL(Em, *this, 0, { Em.ConstantInt(LightIndex) }));
+}
+
+void UMaterialExpressionSkyAtmosphereViewLuminance::Build(MIR::FEmitter& Em)
+{
+	FValueRef WorldDirectionValue = Em.TryInput(&WorldDirection);
+	if (!WorldDirectionValue)
+	{
+		WorldDirectionValue = Em.Multiply(Em.ConstantFloat(-1.0f), EmitInlineHLSL(Em, NAME_CameraVector, {}));
+	}
+	Em.Output(0, EmitInlineHLSL(Em, *this, 0, { WorldDirectionValue }));
+}
+
+void UMaterialExpressionSkyLightEnvMapSample::Build(MIR::FEmitter& Em)
+{
+	if (Material->bIsSky)
+	{
+		UE_LOG(LogMaterial, Warning, TEXT("Using SkyLightEnvMapSample from a IsSky material can result in visual artifact. For instance, if the previous frame capture was super bright, it might leak onto a new frame, e.g. transtion from menu to game."));
+	}
+
+	FValueRef DirectionValue = Em.InputDefaultFloat3(&Direction, FVector3f(0.0f, 0.0f, 1.0f));
+	FValueRef RoughnessValue = Em.InputDefaultFloat(&Roughness, 0.0f);
+
+	Em.Output(0, EmitInlineHLSL(Em, *this, 0, { DirectionValue, RoughnessValue }));
+}
+
+void UMaterialExpressionSphereMask::Build(MIR::FEmitter& Em)
+{
+	FValueRef Arg1 = Em.Input(&A);
+	FValueRef Arg2 = Em.Input(&B);
+	UE_MIR_CHECKPOINT(Em);
+
+	// 1.0f / max(0.00001f, Radius)
+	FValueRef ArgInvRadius = Em.Divide(Em.ConstantFloat(1.0f), Em.Max(Em.ConstantFloat(0.00001f), Em.InputDefaultFloat(&Radius, AttenuationRadius)));
+
+	// 1.0f / max(0.00001fp4, 1.0f - Hardness)
+	FValueRef ArgInvHardness = Em.Divide(Em.ConstantFloat(1.0f), Em.Max(Em.ConstantFloat(0.00001f), Em.Subtract(Em.ConstantFloat(1.0f), Em.InputDefaultFloat(&Hardness, HardnessPercent * 0.01f))));
+
+	FValueRef Distance = Em.Length(Em.Subtract(Arg1, Arg2));
+	FValueRef NormalizeDistance = Em.Multiply(Distance, ArgInvRadius);
+	FValueRef NegNormalizedDistance = Em.Subtract(Em.ConstantFloat(1.0f), NormalizeDistance);
+	FValueRef MaskUnclamped = Em.Multiply(NegNormalizedDistance, ArgInvHardness);
+	Em.Output(0, Em.Saturate(MaskUnclamped));
+}
+
 #endif // WITH_EDITOR
diff --git a/Engine/Source/Runtime/Engine/Private/Materials/MaterialExternalCodeRegistry.h b/Engine/Source/Runtime/Engine/Private/Materials/MaterialExternalCodeRegistry.h
--- a/Engine/Source/Runtime/Engine/Private/Materials/MaterialExternalCodeRegistry.h
+++ b/Engine/Source/Runtime/Engine/Private/Materials/MaterialExternalCodeRegistry.h
@@ -10,6 +10,7 @@
 #include "Materials/HLSLMaterialDerivativeAutogen.h"
 #include "Materials/MaterialExpressionViewProperty.h"
 #include "RHIDefinitions.h"
+#include "RHIFeatureLevel.h"
 
 #include "MaterialExternalCodeRegistry.generated.h"
 
@@ -23,6 +24,18 @@ enum class EMaterialShaderFrequency : uint8
 };
 ENUM_CLASS_FLAGS(EMaterialShaderFrequency);
 
+UENUM()
+enum class EMaterialFeatureLevel : uint8
+{
+	ES2_REMOVED = ERHIFeatureLevel::ES2_REMOVED,
+	ES3_1 = ERHIFeatureLevel::ES3_1,
+	SM4_REMOVED = ERHIFeatureLevel::SM4_REMOVED,
+	SM5 = ERHIFeatureLevel::SM5,
+	SM6 = ERHIFeatureLevel::SM6,
+	Num
+};
+static_assert((uint8)ERHIFeatureLevel::Num == (uint8)EMaterialFeatureLevel::Num);
+
 /**
  * This has to be kept in sync with EMaterialValueType.
  * To consolidate those enums, EMaterialValueType must be made an 'enum class' but it has to be deprecated first.
@@ -140,6 +153,10 @@ public:
 	UPROPERTY(Config)
 	TArray<FMaterialExternalCodeEnvironmentDefine> EnvironmentDefines;
 
+	/** Minimum required feature level for this external code. */
+	UPROPERTY(Config)
+	EMaterialFeatureLevel MinimumFeatureLevel = (EMaterialFeatureLevel)0;		// Default is lowest feature level (zero), whatever that is in the enum
+
 public:
 	/**
 	 * Function return type as material value type.
diff --git a/Engine/Source/Runtime/Engine/Private/Materials/MaterialIR.cpp b/Engine/Source/Runtime/Engine/Private/Materials/MaterialIR.cpp
--- a/Engine/Source/Runtime/Engine/Private/Materials/MaterialIR.cpp
+++ b/Engine/Source/Runtime/Engine/Private/Materials/MaterialIR.cpp
@@ -466,6 +466,7 @@ const TCHAR* LexToString(EExternalInput Input)
 		case EExternalInput::ObjectPosition_CameraRelative: return TEXT("ObjectPosition_CameraRelative");
 		case EExternalInput::ViewMaterialTextureMipBias: return TEXT("ViewMaterialTextureMipBias");
 		case EExternalInput::ViewMaterialTextureDerivativeMultiply: return TEXT("ViewMaterialTextureDerivativeMultiply");
+		case EExternalInput::GlobalDistanceField: return TEXT("GlobalDistanceField");
 		default: UE_MIR_UNREACHABLE();
 	}
 }
@@ -481,6 +482,7 @@ FType GetExternalInputType(EExternalInput Id)
 	{
 		case EExternalInput::ViewMaterialTextureMipBias:
 		case EExternalInput::ViewMaterialTextureDerivativeMultiply:
+		case EExternalInput::GlobalDistanceField:						// Type is arbitrary, value not actually used, only as a flag for value analyzer
 			 return FType::MakeFloatScalar();
 
 		case EExternalInput::ActorPosition_Absolute:
diff --git a/Engine/Source/Runtime/Engine/Private/Materials/MaterialIRToHLSLTranslator.cpp b/Engine/Source/Runtime/Engine/Private/Materials/MaterialIRToHLSLTranslator.cpp
--- a/Engine/Source/Runtime/Engine/Private/Materials/MaterialIRToHLSLTranslator.cpp
+++ b/Engine/Source/Runtime/Engine/Private/Materials/MaterialIRToHLSLTranslator.cpp
@@ -547,6 +547,7 @@ struct FPrivate : FMaterialIRToHLSLTranslation
 				case MIR::EExternalInput::ObjectPosition_CameraRelative: Code = TEXTVIEW("GetObjectTranslatedWorldPosition(Parameters)"); break;
 				case MIR::EExternalInput::ViewMaterialTextureMipBias: Code = TEXTVIEW("View.MaterialTextureMipBias"); break;
 				case MIR::EExternalInput::ViewMaterialTextureDerivativeMultiply: Code = TEXTVIEW("View.MaterialTextureDerivativeMultiply"); break;
+				case MIR::EExternalInput::GlobalDistanceField: Code = TEXTVIEW(""); break;			// Used as flag for value analyzer, doesn't generate code
 				default: UE_MIR_UNREACHABLE();
 			}
 			Printer << Code;
diff --git a/Engine/Source/Runtime/Engine/Private/Materials/MaterialIRValueAnalyzer.cpp b/Engine/Source/Runtime/Engine/Private/Materials/MaterialIRValueAnalyzer.cpp
--- a/Engine/Source/Runtime/Engine/Private/Materials/MaterialIRValueAnalyzer.cpp
+++ b/Engine/Source/Runtime/Engine/Private/Materials/MaterialIRValueAnalyzer.cpp
@@ -51,6 +51,17 @@ static void AnalyzeExternalInput(FMaterialIRValueAnalyzer& Analyzer, MIR::FExter
 	{
 		Analyzer.CompilationOutput->SetIsSceneTextureUsed(PPI_SceneDepth);
 	}
+
+	if (ExternalInput->Id == MIR::EExternalInput::GlobalDistanceField)
+	{
+		if (!FDataDrivenShaderPlatformInfo::GetSupportsDistanceFields(Analyzer.Module->GetShaderPlatform()))
+		{
+			FString ShaderPlatformName = FDataDrivenShaderPlatformInfo::GetName(Analyzer.Module->GetShaderPlatform()).ToString();
+			Analyzer.Module->AddError(nullptr, FString::Printf(TEXT("Node not supported in shader platform %s. The node requires DistanceField support."), *ShaderPlatformName));
+		}
+
+		Analyzer.CompilationOutput->bUsesGlobalDistanceField = true;
+	}
 }
 
 static EMaterialShaderFrequency MapToMaterialShaderFrequencyOrAny(MIR::EStage Stage)
@@ -93,6 +104,20 @@ static void AnalyzeInlineHLSL(FMaterialIRValueAnalyzer& Analyzer, MIR::FInlineHL
 			FName AssetPathName = Analyzer.Material->GetOutermost()->GetFName();
 			Analyzer.Module->AddError(nullptr, MaterialExpressionUtils::FormatUnsupportedMaterialDomainError(*InlineHLSL->ExternalCodeDeclaration, AssetPathName));
 		}
+
+		// Cast from material feature level enum to RHI feature level enum
+		ERHIFeatureLevel::Type MinimumFeatureLevel = (ERHIFeatureLevel::Type)InlineHLSL->ExternalCodeDeclaration->MinimumFeatureLevel;
+
+		if (Analyzer.Module->GetFeatureLevel() < (ERHIFeatureLevel::Type)InlineHLSL->ExternalCodeDeclaration->MinimumFeatureLevel)
+		{
+			FString FeatureLevelName;
+			FString MinimumFeatureLevelName;
+			GetFeatureLevelName(Analyzer.Module->GetFeatureLevel(), FeatureLevelName);
+			GetFeatureLevelName(MinimumFeatureLevel, MinimumFeatureLevelName);
+
+			Analyzer.Module->AddError(nullptr, FString::Printf(TEXT("Node %s requires feature level %s.  Current feature level is %s."),
+				*InlineHLSL->ExternalCodeDeclaration->Name.ToString(), *MinimumFeatureLevelName, *FeatureLevelName));
+		}
 	}
 }
 
diff --git a/Engine/Source/Runtime/Engine/Public/Materials/MaterialExpressionAtmosphericFogColor.h b/Engine/Source/Runtime/Engine/Public/Materials/MaterialExpressionAtmosphericFogColor.h
--- a/Engine/Source/Runtime/Engine/Public/Materials/MaterialExpressionAtmosphericFogColor.h
+++ b/Engine/Source/Runtime/Engine/Public/Materials/MaterialExpressionAtmosphericFogColor.h
@@ -6,10 +6,11 @@
 #include "UObject/ObjectMacros.h"
 #include "MaterialExpressionIO.h"
 #include "Materials/MaterialExpression.h"
+#include "Materials/MaterialExpressionExternalCodeBase.h"
 #include "MaterialExpressionAtmosphericFogColor.generated.h"
 
 UCLASS(collapsecategories, hidecategories=Object, DisplayName = "Atmospheric Fog Color (deprecated)")
-class UMaterialExpressionAtmosphericFogColor : public UMaterialExpression
+class UMaterialExpressionAtmosphericFogColor : public UMaterialExpressionExternalCodeBase
 {
 	GENERATED_UCLASS_BODY()
 
@@ -22,6 +23,7 @@ class UMaterialExpressionAtmosphericFogColor : public UMaterialExpression
 
 	//~ Begin UMaterialExpression Interface
 #if WITH_EDITOR
+	virtual void Build(MIR::FEmitter& Emitter) override;
 	virtual FName GetInputName(int32 InputIndex) const override;
 	virtual void PostEditChangeProperty(FPropertyChangedEvent& PropertyChangedEvent) override;
 	virtual int32 Compile(class FMaterialCompiler* Compiler, int32 OutputIndex) override;
diff --git a/Engine/Source/Runtime/Engine/Public/Materials/MaterialExpressionBlackBody.h b/Engine/Source/Runtime/Engine/Public/Materials/MaterialExpressionBlackBody.h
--- a/Engine/Source/Runtime/Engine/Public/Materials/MaterialExpressionBlackBody.h
+++ b/Engine/Source/Runtime/Engine/Public/Materials/MaterialExpressionBlackBody.h
@@ -22,6 +22,7 @@ public:
 
 	//~ Begin UMaterialExpression Interface
 #if WITH_EDITOR
+	virtual void Build(MIR::FEmitter& Emitter) override;
 	virtual int32 Compile(class FMaterialCompiler* Compiler, int32 OutputIndex) override;
 	virtual void GetCaption(TArray<FString>& OutCaptions) const override;
 
diff --git a/Engine/Source/Runtime/Engine/Public/Materials/MaterialExpressionDepthFade.h b/Engine/Source/Runtime/Engine/Public/Materials/MaterialExpressionDepthFade.h
--- a/Engine/Source/Runtime/Engine/Public/Materials/MaterialExpressionDepthFade.h
+++ b/Engine/Source/Runtime/Engine/Public/Materials/MaterialExpressionDepthFade.h
@@ -47,6 +47,7 @@ class UMaterialExpressionDepthFade : public UMaterialExpression
 			return Super::GetInputName(InputIndex);
 		}
 	}
+	virtual void Build(MIR::FEmitter& Emitter) override;
 	virtual int32 Compile(class FMaterialCompiler* Compiler, int32 OutputIndex) override;
 	virtual void GetCaption(TArray<FString>& OutCaptions) const override
 	{
diff --git a/Engine/Source/Runtime/Engine/Public/Materials/MaterialExpressionDeriveNormalZ.h b/Engine/Source/Runtime/Engine/Public/Materials/MaterialExpressionDeriveNormalZ.h
--- a/Engine/Source/Runtime/Engine/Public/Materials/MaterialExpressionDeriveNormalZ.h
+++ b/Engine/Source/Runtime/Engine/Public/Materials/MaterialExpressionDeriveNormalZ.h
@@ -21,6 +21,7 @@ public:
 
 	//~ Begin UMaterialExpression Interface
 #if WITH_EDITOR
+	virtual void Build(MIR::FEmitter& Emitter) override;
 	virtual int32 Compile(class FMaterialCompiler* Compiler, int32 OutputIndex) override;
 	virtual void GetCaption(TArray<FString>& OutCaptions) const override;
 
diff --git a/Engine/Source/Runtime/Engine/Public/Materials/MaterialExpressionDistanceFieldApproxAO.h b/Engine/Source/Runtime/Engine/Public/Materials/MaterialExpressionDistanceFieldApproxAO.h
--- a/Engine/Source/Runtime/Engine/Public/Materials/MaterialExpressionDistanceFieldApproxAO.h
+++ b/Engine/Source/Runtime/Engine/Public/Materials/MaterialExpressionDistanceFieldApproxAO.h
@@ -48,6 +48,7 @@ public:
 
 	//~ Begin UMaterialExpression Interface
 #if WITH_EDITOR
+	virtual void Build(MIR::FEmitter& Emitter) override;
 	virtual FName GetInputName(int32 InputIndex) const override;
 	virtual void PostEditChangeProperty(FPropertyChangedEvent& PropertyChangedEvent) override;
 	virtual int32 Compile(class FMaterialCompiler* Compiler, int32 OutputIndex) override;
diff --git a/Engine/Source/Runtime/Engine/Public/Materials/MaterialExpressionDistanceFieldGradient.h b/Engine/Source/Runtime/Engine/Public/Materials/MaterialExpressionDistanceFieldGradient.h
--- a/Engine/Source/Runtime/Engine/Public/Materials/MaterialExpressionDistanceFieldGradient.h
+++ b/Engine/Source/Runtime/Engine/Public/Materials/MaterialExpressionDistanceFieldGradient.h
@@ -25,6 +25,7 @@ public:
 
 	//~ Begin UMaterialExpression Interface
 #if WITH_EDITOR
+	virtual void Build(MIR::FEmitter& Emitter) override;
 	virtual FName GetInputName(int32 InputIndex) const override;
 	virtual void PostEditChangeProperty(FPropertyChangedEvent& PropertyChangedEvent) override;
 	virtual int32 Compile(class FMaterialCompiler* Compiler, int32 OutputIndex) override;
diff --git a/Engine/Source/Runtime/Engine/Public/Materials/MaterialExpressionDistanceToNearestSurface.h b/Engine/Source/Runtime/Engine/Public/Materials/MaterialExpressionDistanceToNearestSurface.h
--- a/Engine/Source/Runtime/Engine/Public/Materials/MaterialExpressionDistanceToNearestSurface.h
+++ b/Engine/Source/Runtime/Engine/Public/Materials/MaterialExpressionDistanceToNearestSurface.h
@@ -25,6 +25,7 @@ public:
 
 	//~ Begin UMaterialExpression Interface
 #if WITH_EDITOR
+	virtual void Build(MIR::FEmitter& Emitter) override;
 	virtual FName GetInputName(int32 InputIndex) const override;
 	virtual void PostEditChangeProperty(FPropertyChangedEvent& PropertyChangedEvent) override;
 	virtual int32 Compile(class FMaterialCompiler* Compiler, int32 OutputIndex) override;
diff --git a/Engine/Source/Runtime/Engine/Public/Materials/MaterialExpressionFresnel.h b/Engine/Source/Runtime/Engine/Public/Materials/MaterialExpressionFresnel.h
--- a/Engine/Source/Runtime/Engine/Public/Materials/MaterialExpressionFresnel.h
+++ b/Engine/Source/Runtime/Engine/Public/Materials/MaterialExpressionFresnel.h
@@ -46,6 +46,7 @@ public:
 
 	//~ Begin UMaterialExpression Interface
 #if WITH_EDITOR
+	virtual void Build(MIR::FEmitter& Emitter) override;
 	virtual int32 Compile(class FMaterialCompiler* Compiler, int32 OutputIndex) override;
 	virtual void GetCaption(TArray<FString>& OutCaptions) const override
 	{
diff --git a/Engine/Source/Runtime/Engine/Public/Materials/MaterialExpressionLightmapUVs.h b/Engine/Source/Runtime/Engine/Public/Materials/MaterialExpressionLightmapUVs.h
--- a/Engine/Source/Runtime/Engine/Public/Materials/MaterialExpressionLightmapUVs.h
+++ b/Engine/Source/Runtime/Engine/Public/Materials/MaterialExpressionLightmapUVs.h
@@ -10,10 +10,11 @@
 #include "CoreMinimal.h"
 #include "UObject/ObjectMacros.h"
 #include "Materials/MaterialExpression.h"
+#include "Materials/MaterialExpressionExternalCodeBase.h"
 #include "MaterialExpressionLightmapUVs.generated.h"
 
 UCLASS(collapsecategories, hidecategories=Object)
-class UMaterialExpressionLightmapUVs : public UMaterialExpression
+class UMaterialExpressionLightmapUVs : public UMaterialExpressionExternalCodeBase
 {
 	GENERATED_UCLASS_BODY()
 
diff --git a/Engine/Source/Runtime/Engine/Public/Materials/MaterialExpressionReflectionVectorWS.h b/Engine/Source/Runtime/Engine/Public/Materials/MaterialExpressionReflectionVectorWS.h
--- a/Engine/Source/Runtime/Engine/Public/Materials/MaterialExpressionReflectionVectorWS.h
+++ b/Engine/Source/Runtime/Engine/Public/Materials/MaterialExpressionReflectionVectorWS.h
@@ -26,6 +26,7 @@ public:
 
 	//~ Begin UMaterialExpression Interface
 #if WITH_EDITOR
+	virtual void Build(MIR::FEmitter& Emitter) override;
 	virtual int32 Compile(class FMaterialCompiler* Compiler, int32 OutputIndex) override;
 	virtual void GetCaption(TArray<FString>& OutCaptions) const override;
 #endif
diff --git a/Engine/Source/Runtime/Engine/Public/Materials/MaterialExpressionRotateAboutAxis.h b/Engine/Source/Runtime/Engine/Public/Materials/MaterialExpressionRotateAboutAxis.h
--- a/Engine/Source/Runtime/Engine/Public/Materials/MaterialExpressionRotateAboutAxis.h
+++ b/Engine/Source/Runtime/Engine/Public/Materials/MaterialExpressionRotateAboutAxis.h
@@ -34,6 +34,7 @@ public:
 
 	//~ Begin UMaterialExpression Interface
 #if WITH_EDITOR
+	virtual void Build(MIR::FEmitter& Emitter) override;
 	virtual int32 Compile(class FMaterialCompiler* Compiler, int32 OutputIndex) override;
 	virtual void GetCaption(TArray<FString>& OutCaptions) const override;
 #endif
diff --git a/Engine/Source/Runtime/Engine/Public/Materials/MaterialExpressionRotator.h b/Engine/Source/Runtime/Engine/Public/Materials/MaterialExpressionRotator.h
--- a/Engine/Source/Runtime/Engine/Public/Materials/MaterialExpressionRotator.h
+++ b/Engine/Source/Runtime/Engine/Public/Materials/MaterialExpressionRotator.h
@@ -35,6 +35,7 @@ class UMaterialExpressionRotator : public UMaterialExpression
 
 	//~ Begin UMaterialExpression Interface
 #if WITH_EDITOR
+	virtual void Build(MIR::FEmitter& Emitter) override;
 	virtual int32 Compile(class FMaterialCompiler* Compiler, int32 OutputIndex) override;
 	virtual void GetCaption(TArray<FString>& OutCaptions) const override;
 	virtual bool NeedsRealtimePreview() override { return Time.Expression==NULL && Speed != 0.f; }
diff --git a/Engine/Source/Runtime/Engine/Public/Materials/MaterialExpressionSkyAtmosphereLightDirection.h b/Engine/Source/Runtime/Engine/Public/Materials/MaterialExpressionSkyAtmosphereLightDirection.h
--- a/Engine/Source/Runtime/Engine/Public/Materials/MaterialExpressionSkyAtmosphereLightDirection.h
+++ b/Engine/Source/Runtime/Engine/Public/Materials/MaterialExpressionSkyAtmosphereLightDirection.h
@@ -6,10 +6,11 @@
 #include "CoreMinimal.h"
 #include "UObject/ObjectMacros.h"
 #include "Materials/MaterialExpression.h"
+#include "Materials/MaterialExpressionExternalCodeBase.h"
 #include "MaterialExpressionSkyAtmosphereLightDirection.generated.h"
 
 UCLASS()
-class UMaterialExpressionSkyAtmosphereLightDirection : public UMaterialExpression
+class UMaterialExpressionSkyAtmosphereLightDirection : public UMaterialExpressionExternalCodeBase
 {
 	GENERATED_BODY()
 
@@ -21,6 +22,7 @@ public:
 
 	//~ Begin UMaterialExpression Interface
 #if WITH_EDITOR
+	virtual void Build(MIR::FEmitter& Emitter) override;
 	virtual int32 Compile(class FMaterialCompiler* Compiler, int32 OutputIndex) override;
 	virtual void GetCaption(TArray<FString>& OutCaptions) const override;
 
diff --git a/Engine/Source/Runtime/Engine/Public/Materials/MaterialExpressionSkyAtmosphereLightIlluminance.h b/Engine/Source/Runtime/Engine/Public/Materials/MaterialExpressionSkyAtmosphereLightIlluminance.h
--- a/Engine/Source/Runtime/Engine/Public/Materials/MaterialExpressionSkyAtmosphereLightIlluminance.h
+++ b/Engine/Source/Runtime/Engine/Public/Materials/MaterialExpressionSkyAtmosphereLightIlluminance.h
@@ -6,10 +6,11 @@
 #include "CoreMinimal.h"
 #include "UObject/ObjectMacros.h"
 #include "Materials/MaterialExpression.h"
+#include "Materials/MaterialExpressionExternalCodeBase.h"
 #include "MaterialExpressionSkyAtmosphereLightIlluminance.generated.h"
 
 UCLASS()
-class UMaterialExpressionSkyAtmosphereLightIlluminance : public UMaterialExpression
+class UMaterialExpressionSkyAtmosphereLightIlluminance : public UMaterialExpressionExternalCodeBase
 {
 	GENERATED_BODY()
 
@@ -29,6 +30,7 @@ public:
 
 	//~ Begin UMaterialExpression Interface
 #if WITH_EDITOR
+	virtual void Build(MIR::FEmitter& Emitter) override;
 	virtual FName GetInputName(int32 InputIndex) const override;
 	virtual void PostEditChangeProperty(FPropertyChangedEvent& PropertyChangedEvent) override;
 	virtual int32 Compile(class FMaterialCompiler* Compiler, int32 OutputIndex) override;
@@ -41,7 +43,7 @@ public:
 
 
 UCLASS()
-class UMaterialExpressionSkyAtmosphereLightIlluminanceOnGround : public UMaterialExpression
+class UMaterialExpressionSkyAtmosphereLightIlluminanceOnGround : public UMaterialExpressionExternalCodeBase
 {
 	GENERATED_BODY()
 
@@ -53,6 +55,7 @@ public:
 
 	//~ Begin UMaterialExpression Interface
 #if WITH_EDITOR
+	virtual void Build(MIR::FEmitter& Emitter) override;
 	virtual int32 Compile(class FMaterialCompiler* Compiler, int32 OutputIndex) override;
 	virtual void GetCaption(TArray<FString>& OutCaptions) const override;
 	virtual void GetExpressionToolTip(TArray<FString>& OutToolTip) override;
@@ -64,7 +67,7 @@ public:
 
 
 UCLASS()
-class UMaterialExpressionSkyAtmosphereLightDiskLuminance : public UMaterialExpression
+class UMaterialExpressionSkyAtmosphereLightDiskLuminance : public UMaterialExpressionExternalCodeBase
 {
 	GENERATED_BODY()
 
@@ -80,6 +83,7 @@ public:
 
 	//~ Begin UMaterialExpression Interface
 #if WITH_EDITOR
+	virtual void Build(MIR::FEmitter& Emitter) override;
 	virtual int32 Compile(class FMaterialCompiler* Compiler, int32 OutputIndex) override;
 	virtual void GetCaption(TArray<FString>& OutCaptions) const override;
 #endif
@@ -89,7 +93,7 @@ public:
 
 
 UCLASS()
-class UMaterialExpressionSkyAtmosphereAerialPerspective : public UMaterialExpression
+class UMaterialExpressionSkyAtmosphereAerialPerspective : public UMaterialExpressionExternalCodeBase
 {
 	GENERATED_BODY()
 
@@ -106,6 +110,7 @@ public:
 
 	//~ Begin UMaterialExpression Interface
 #if WITH_EDITOR
+	virtual void Build(MIR::FEmitter& Emitter) override;
 	virtual FName GetInputName(int32 InputIndex) const override;
 	virtual void PostEditChangeProperty(FPropertyChangedEvent& PropertyChangedEvent) override;
 	virtual int32 Compile(class FMaterialCompiler* Compiler, int32 OutputIndex) override;
@@ -118,7 +123,7 @@ public:
 
 
 UCLASS()
-class UMaterialExpressionSkyAtmosphereDistantLightScatteredLuminance : public UMaterialExpression
+class UMaterialExpressionSkyAtmosphereDistantLightScatteredLuminance : public UMaterialExpressionExternalCodeBase
 {
 	GENERATED_BODY()
 
diff --git a/Engine/Source/Runtime/Engine/Public/Materials/MaterialExpressionSkyAtmosphereViewLuminance.h b/Engine/Source/Runtime/Engine/Public/Materials/MaterialExpressionSkyAtmosphereViewLuminance.h
--- a/Engine/Source/Runtime/Engine/Public/Materials/MaterialExpressionSkyAtmosphereViewLuminance.h
+++ b/Engine/Source/Runtime/Engine/Public/Materials/MaterialExpressionSkyAtmosphereViewLuminance.h
@@ -6,10 +6,11 @@
 #include "CoreMinimal.h"
 #include "UObject/ObjectMacros.h"
 #include "Materials/MaterialExpression.h"
+#include "Materials/MaterialExpressionExternalCodeBase.h"
 #include "MaterialExpressionSkyAtmosphereViewLuminance.generated.h"
 
 UCLASS()
-class UMaterialExpressionSkyAtmosphereViewLuminance : public UMaterialExpression
+class UMaterialExpressionSkyAtmosphereViewLuminance : public UMaterialExpressionExternalCodeBase
 {
 	GENERATED_BODY()
 
@@ -21,6 +22,7 @@ public:
 
 	//~ Begin UMaterialExpression Interface
 #if WITH_EDITOR
+	virtual void Build(MIR::FEmitter& Emitter) override;
 	virtual int32 Compile(class FMaterialCompiler* Compiler, int32 OutputIndex) override;
 	virtual void GetCaption(TArray<FString>& OutCaptions) const override;
 #endif
diff --git a/Engine/Source/Runtime/Engine/Public/Materials/MaterialExpressionSkyLightEnvMapSample.h b/Engine/Source/Runtime/Engine/Public/Materials/MaterialExpressionSkyLightEnvMapSample.h
--- a/Engine/Source/Runtime/Engine/Public/Materials/MaterialExpressionSkyLightEnvMapSample.h
+++ b/Engine/Source/Runtime/Engine/Public/Materials/MaterialExpressionSkyLightEnvMapSample.h
@@ -6,10 +6,11 @@
 #include "CoreMinimal.h"
 #include "UObject/ObjectMacros.h"
 #include "Materials/MaterialExpression.h"
+#include "Materials/MaterialExpressionExternalCodeBase.h"
 #include "MaterialExpressionSkyLightEnvMapSample.generated.h"
 
 UCLASS()
-class UMaterialExpressionSkyLightEnvMapSample : public UMaterialExpression
+class UMaterialExpressionSkyLightEnvMapSample : public UMaterialExpressionExternalCodeBase
 {
 	GENERATED_BODY()
 
@@ -23,6 +24,7 @@ public:
 
 	//~ Begin UMaterialExpression Interface
 #if WITH_EDITOR
+	virtual void Build(MIR::FEmitter& Emitter) override;
 	virtual int32 Compile(class FMaterialCompiler* Compiler, int32 OutputIndex) override;
 	virtual void GetCaption(TArray<FString>& OutCaptions) const override;
 
diff --git a/Engine/Source/Runtime/Engine/Public/Materials/MaterialExpressionSphereMask.h b/Engine/Source/Runtime/Engine/Public/Materials/MaterialExpressionSphereMask.h
--- a/Engine/Source/Runtime/Engine/Public/Materials/MaterialExpressionSphereMask.h
+++ b/Engine/Source/Runtime/Engine/Public/Materials/MaterialExpressionSphereMask.h
@@ -43,6 +43,7 @@ public:
 
 	//~ Begin UMaterialExpression Interface
 #if WITH_EDITOR
+	virtual void Build(MIR::FEmitter& Emitter) override;
 	virtual int32 Compile(class FMaterialCompiler* Compiler, int32 OutputIndex) override;
 	virtual void GetCaption(TArray<FString>& OutCaptions) const override;
 	virtual FText GetKeywords() const override { return FText::FromString(TEXT("sphere")); }
diff --git a/Engine/Source/Runtime/Engine/Public/Materials/MaterialIR.h b/Engine/Source/Runtime/Engine/Public/Materials/MaterialIR.h
--- a/Engine/Source/Runtime/Engine/Public/Materials/MaterialIR.h
+++ b/Engine/Source/Runtime/Engine/Public/Materials/MaterialIR.h
@@ -406,6 +406,7 @@ enum class EExternalInput
 	ObjectPosition_CameraRelative,
 	ViewMaterialTextureMipBias,
 	ViewMaterialTextureDerivativeMultiply,
+	GlobalDistanceField,					// Flag to mark that global distance fields are used, for value analyzer step.  Doesn't generate code.
 
 	Count,
 };
