Animation - fix for use-after-free on BonesToConvert in LocalBlendCSBoneTransforms (and similar issue that could happen for BoneMask)

diff --git a/Engine/Source/Runtime/Engine/Public/BonePose.h b/Engine/Source/Runtime/Engine/Public/BonePose.h
--- a/Engine/Source/Runtime/Engine/Public/BonePose.h
+++ b/Engine/Source/Runtime/Engine/Public/BonePose.h
@@ -410,6 +410,16 @@ struct FCSPose
 	typedef typename PoseType::BoneIndexType BoneIndexType;
 	typedef typename PoseType::Allocator AllocatorType;
 	
+	// overriding ctor, copy, move to allow deprecations.
+	PRAGMA_DISABLE_DEPRECATION_WARNINGS
+	FCSPose() = default;
+	~FCSPose() = default;
+	FCSPose(const FCSPose& Other) = default;
+	FCSPose(FCSPose&& Other) = default;
+	FCSPose& operator=(const FCSPose& Other) = default;
+	FCSPose& operator=(FCSPose&& Other) = default;
+	PRAGMA_ENABLE_DEPRECATION_WARNINGS
+
 	// Init Pose
 	void InitPose(const FBoneContainer* InBoneContainer)
 	{
@@ -454,8 +464,11 @@ struct FCSPose
 	{
 		Pose.Empty();
 		ComponentSpaceFlags.Empty();
+
+		PRAGMA_DISABLE_DEPRECATION_WARNINGS
 		BonesToConvert.Empty();
 		BoneMask.Empty();
+		PRAGMA_ENABLE_DEPRECATION_WARNINGS
 	}
 
 	const PoseType& GetPose() const { return Pose; }
@@ -532,10 +545,10 @@ protected:
 	// Flags to track each bones current state (0 means local pose, 1 means component space pose)
 	TCustomBoneIndexArray<uint8, BoneIndexType, AllocatorType> ComponentSpaceFlags;
 
-	// Cached bone mask array to avoid reallocations
+	UE_DEPRECATED(5.7, "Property no longer required. We now allocate it via TMemStackAllocator directly.")
 	TCustomBoneIndexArray<uint8, BoneIndexType, AllocatorType> BoneMask;
 
-	// Cached conversion array for this pose, to save on allocations each frame
+	UE_DEPRECATED(5.7, "Property no longer required. We now allocate it via TMemStackAllocator directly.")
 	TArray<FCompactPoseBoneIndex, AllocatorType> BonesToConvert;
 };
 
@@ -674,7 +687,9 @@ void FCSPose<PoseType>::SafeSetCSBoneTransforms(TConstArrayView<FBoneTransform>
 {
 	checkSlow(Pose.IsValid());
 
-	BonesToConvert.Reset();
+	FMemMark Mark(FMemStack::Get());
+	TArray<FCompactPoseBoneIndex, TMemStackAllocator<>> StackBonesToConvert;
+	StackBonesToConvert.Reserve(BoneTransforms.Num());
 
 	// Minimum bone index, we don't need to look at bones prior to this in the pose
 	const int32 MinIndex = BoneTransforms[0].BoneIndex.GetInt();
@@ -684,13 +699,13 @@ void FCSPose<PoseType>::SafeSetCSBoneTransforms(TConstArrayView<FBoneTransform>
 	{
 		if(ComponentSpaceFlags[Transform.BoneIndex] == 1)
 		{
-			BonesToConvert.Add(Transform.BoneIndex);
+			StackBonesToConvert.Add(Transform.BoneIndex);
 		}
 	}
 
 	// Store the beginning of the child transforms, below we don't need to convert any bone added
 	// from BoneTransforms because they're about to be overwritten
-	const int32 FirstChildTransform = BonesToConvert.Num();
+	const int32 FirstChildTransform = StackBonesToConvert.Num();
 
 	FCompactPoseBoneIndexIterator Iter = FCompactPoseBoneIndexIterator(MinIndex);
 	FCompactPoseBoneIndexIterator EndIter = Pose.MakeEndIter();
@@ -701,19 +716,19 @@ void FCSPose<PoseType>::SafeSetCSBoneTransforms(TConstArrayView<FBoneTransform>
 		const FCompactPoseBoneIndex BoneIndex = *Iter;
 		const FCompactPoseBoneIndex ParentIndex = Pose.GetParentBoneIndex(BoneIndex);
 
-		if(ComponentSpaceFlags[BoneIndex] == 1 && BonesToConvert.Contains(ParentIndex))
+		if(ComponentSpaceFlags[BoneIndex] == 1 && StackBonesToConvert.Contains(ParentIndex))
 		{
-			BonesToConvert.AddUnique(BoneIndex);
+			StackBonesToConvert.AddUnique(BoneIndex);
 		}
 	}
 
 	// Convert the bones, we walk backwards to process children first, the pose iteration above is sorted
 	// so we know we already have the right order. We also stop when we get to the bones contained in
 	// BoneTransforms because we're about to overwrite them anyway
-	const int32 NumToConvert = BonesToConvert.Num();
+	const int32 NumToConvert = StackBonesToConvert.Num();
 	for(int32 Idx = NumToConvert - 1; Idx >= FirstChildTransform; --Idx)
 	{
-		ConvertBoneToLocalSpace(BonesToConvert[Idx]);
+		ConvertBoneToLocalSpace(StackBonesToConvert[Idx]);
 	}
 
 	// Finally copy our Component Space transforms
@@ -767,10 +782,10 @@ void FCSPose<PoseType>::LocalBlendCSBoneTransforms(TConstArrayView<FBoneTransfor
 	{
 		// Bone Mask to keep track of which bones have to be converted to local space.
 		// This is basically BoneTransforms bones and their children.
-		BoneMask.Reset();
-		BoneMask.AddZeroed(Pose.GetNumBones());
-
 		FMemMark Mark(FMemStack::Get());
+		TCustomBoneIndexArray<uint8, BoneIndexType, TMemStackAllocator<>> StackBoneMask;
+		StackBoneMask.AddZeroed(Pose.GetNumBones());
+
 		TArray<FBoneTransform, TMemStackAllocator<>> LocalBoneTransforms;
 		LocalBoneTransforms.SetNumUninitialized(BoneTransforms.Num());
 
@@ -783,7 +798,7 @@ void FCSPose<PoseType>::LocalBlendCSBoneTransforms(TConstArrayView<FBoneTransfor
 			LocalBoneTransforms[Index].Transform = GetLocalSpaceTransform(BoneIndex);
 			LocalBoneTransforms[Index].BoneIndex = BoneIndex;
 
-			BoneMask[BoneIndex] = 1;
+			StackBoneMask[BoneIndex] = 1;
 		}
 
 		// Next, update the pose as if Alpha = 1.0
@@ -796,7 +811,7 @@ void FCSPose<PoseType>::LocalBlendCSBoneTransforms(TConstArrayView<FBoneTransfor
 			// Propagate our BoneMask to children.
 			if (ParentIndex != INDEX_NONE)
 			{
-				BoneMask[BoneIndex] |= BoneMask[ParentIndex];
+				StackBoneMask[BoneIndex] |= StackBoneMask[ParentIndex];
 			}
 		}
 
@@ -805,7 +820,7 @@ void FCSPose<PoseType>::LocalBlendCSBoneTransforms(TConstArrayView<FBoneTransfor
 			if (!BoneIndex.IsRootBone())
 			{
 				// If this bone has to be converted to Local Space...
-				if (BoneMask[BoneIndex] != 0)
+				if (StackBoneMask[BoneIndex] != 0)
 				{
 					// .. And is not currently in Local Space, then convert it.
 					ConvertBoneToLocalSpace(BoneIndex);
