[NetPhysics] Fix issues with InputDecay after recent refactors to support Iris and convert settings to DataAsset.

#rnx

diff --git a/Engine/Source/Runtime/Engine/Private/PhysicsEngine/NetworkPhysicsComponent.cpp b/Engine/Source/Runtime/Engine/Private/PhysicsEngine/NetworkPhysicsComponent.cpp
--- a/Engine/Source/Runtime/Engine/Private/PhysicsEngine/NetworkPhysicsComponent.cpp
+++ b/Engine/Source/Runtime/Engine/Private/PhysicsEngine/NetworkPhysicsComponent.cpp
@@ -2122,6 +2122,9 @@ void FAsyncNetworkPhysicsComponent::OnPreProcessInputs_Internal(const int32 Phys
 				// Calculate input decay if we are resimulating and we don't have up to date inputs
 				if (bIsSolverResim)
 				{
+					// If the input was extrapolated, set it back to the input frame it was extrapolated from (the last in the input history) so that we can calculate input decay during resimulation
+					PhysicsData->LocalFrame = FMath::Min(PhysicsData->LocalFrame, InputHistory->GetLatestFrame());
+
 					if (PhysicsData->LocalFrame < PhysicsStep)
 					{
 						const float InputDecay = GetCurrentInputDecay(PhysicsData);
diff --git a/Engine/Source/Runtime/Engine/Private/PhysicsEngine/NetworkPhysicsSettingsComponent.cpp b/Engine/Source/Runtime/Engine/Private/PhysicsEngine/NetworkPhysicsSettingsComponent.cpp
--- a/Engine/Source/Runtime/Engine/Private/PhysicsEngine/NetworkPhysicsSettingsComponent.cpp
+++ b/Engine/Source/Runtime/Engine/Private/PhysicsEngine/NetworkPhysicsSettingsComponent.cpp
@@ -136,6 +136,11 @@ void UNetworkPhysicsSettingsComponent::UninitializeComponent()
 {
 	Super::UninitializeComponent();
 
+	if (SettingsDataAsset)
+	{
+		SettingsDataAsset->MarkUninitialized();
+	}
+
 	using namespace Chaos;
 	if (NetworkPhysicsSettings_Internal)
 	{
diff --git a/Engine/Source/Runtime/Engine/Public/Physics/NetworkPhysicsSettingsComponent.h b/Engine/Source/Runtime/Engine/Public/Physics/NetworkPhysicsSettingsComponent.h
--- a/Engine/Source/Runtime/Engine/Public/Physics/NetworkPhysicsSettingsComponent.h
+++ b/Engine/Source/Runtime/Engine/Public/Physics/NetworkPhysicsSettingsComponent.h
@@ -392,16 +392,20 @@ struct FNetworkPhysicsSettingsResimulation
 USTRUCT()
 struct FNetworkPhysicsSettingsNetworkPhysicsComponent
 {
-	FNetworkPhysicsSettingsNetworkPhysicsComponent()
-	: InputDecayCurve()
-	{
-		InputDecayCurve.GetRichCurve()->AddKey(0.0f, 0.0f);
-		InputDecayCurve.GetRichCurve()->AddKey(0.5f, 1.0f);
-		InputDecayCurve.GetRichCurve()->AddKey(1.0f, 1.0f);
-	};
+	FNetworkPhysicsSettingsNetworkPhysicsComponent() { }
 
 	GENERATED_BODY()
 
+	void Initialize()
+	{
+		if (InputDecayCurve.EditorCurveData.GetNumKeys() == 0)
+		{
+			InputDecayCurve.GetRichCurve()->AddKey(0.0f, 0.0f);
+			InputDecayCurve.GetRichCurve()->AddKey(0.5f, 1.0f);
+			InputDecayCurve.GetRichCurve()->AddKey(1.0f, 1.0f);
+		}
+	}
+
 	UPROPERTY(config, EditDefaultsOnly, Category = "Overrides", Meta = (InlineEditConditionToggle))
 	uint32 bOverrideRedundantInputs : 1;
 	/* Overrides CVar : np2.Resim.RedundantInputs -- How many extra inputs to send with each unreliable network message, to account for packetloss.From owning client to server and server to owning client.
@@ -553,6 +557,15 @@ class UNetworkPhysicsSettingsDataAsset : public UDataAsset
 {
 	GENERATED_BODY()
 
+	virtual void PostInitProperties() override
+	{
+		Super::PostInitProperties();
+		if (!HasAnyFlags(RF_ClassDefaultObject))
+		{
+			Settings.NetworkPhysicsComponentSettings.Initialize();
+		}
+	}
+
 public:
 	
 	void InitializeInternalSettings()
@@ -567,6 +580,11 @@ public:
 		bInitialized = true;
 	}
 
+	void MarkUninitialized()
+	{
+		bInitialized = false;
+	}
+
 	/** Get the settings, on the game thread */
 	const FNetworkPhysicsSettingsData& GetSettings_External() const { return Settings; }
 	
