Remove enumeration material parameter.
We will support the enumeration display support on scalar parameters instead.
No deprecation required since this hasn't been in a UE release yet.
Any internal project usage has been found and removed.
#rnx

diff --git a/Engine/Config/BaseMaterialExpressions.ini b/Engine/Config/BaseMaterialExpressions.ini
--- a/Engine/Config/BaseMaterialExpressions.ini
+++ b/Engine/Config/BaseMaterialExpressions.ini
@@ -1076,9 +1076,6 @@ bShaderInputData=true
 +ExternalCodeDeclarations=(ReturnType=Float4x4,		Name="ViewToTranslatedWorld",		Definition="ResolvedView.<PREV>ViewToTranslatedWorld", Derivative=Zero, bIsInlined=true)
 +ExternalCodeDeclarations=(ReturnType=LWCMatrix,	Name="WorldToParticle",				Definition="Parameters.Particle.WorldToParticle", Derivative=Zero, bIsInlined=true, EnvironmentDefines=((Name="NEEDS_PARTICLE_WORLD_TO_LOCAL")))
 
-[/Script/Engine.MaterialExpressionEnumerationParameter]
-+MenuCategories=NSLOCTEXT("MaterialExpression", "Parameters", "Parameters")
-
 [/Script/Engine.MaterialExpressionAggregate]
 +MenuCategories=NSLOCTEXT("MaterialExpression", "Aggregates", "Aggregates")
 
diff --git a/Engine/Source/Editor/MaterialEditor/Private/MaterialEditingLibrary.cpp b/Engine/Source/Editor/MaterialEditor/Private/MaterialEditingLibrary.cpp
--- a/Engine/Source/Editor/MaterialEditor/Private/MaterialEditingLibrary.cpp
+++ b/Engine/Source/Editor/MaterialEditor/Private/MaterialEditingLibrary.cpp
@@ -24,7 +24,6 @@
 #include "Materials/MaterialExpressionDynamicParameter.h"
 #include "Materials/MaterialParameterCollection.h"
 #include "Materials/MaterialExpressionMaterialFunctionCall.h"
-#include "Materials/MaterialEnumeration.h"
 #include "MaterialEditor/MaterialEditorInstanceConstant.h"
 #include "MaterialStatsCommon.h"
 #include "Particles/ParticleSystemComponent.h"
@@ -1273,27 +1272,6 @@ bool UMaterialEditingLibrary::SetMaterialInstanceStaticSwitchParameterValue(UMat
 	return bResult;
 }
 
-FName UMaterialEditingLibrary::GetMaterialInstanceEnumerationParameterValue(UMaterialInstanceConstant* Instance, FName ParameterName, EMaterialParameterAssociation Association)
-{
-	FName Result = {};
-	if (Instance)
-	{
-		Instance->GetEnumerationParameterValue(FHashedMaterialParameterInfo(ParameterName, Association, Association == EMaterialParameterAssociation::LayerParameter ? 0 : INDEX_NONE), Result);
-	}
-	return Result;
-}
-
-bool UMaterialEditingLibrary::SetMaterialInstanceEnumerationParameterValue(UMaterialInstanceConstant* Instance, FName ParameterName, FName Value, EMaterialParameterAssociation Association)
-{
-	bool bResult = false;
-	if (Instance)
-	{
-		Instance->SetEnumerationParameterValueEditorOnly(FMaterialParameterInfo(ParameterName, Association, Association == EMaterialParameterAssociation::LayerParameter ? 0 : INDEX_NONE), Value);
-		UMaterialEditingLibrary::UpdateMaterialInstance(Instance);
-	}
-	return bResult;
-}
-
 void UMaterialEditingLibrary::UpdateMaterialInstance(UMaterialInstanceConstant* Instance)
 {
 	if (Instance)
@@ -1396,22 +1374,6 @@ void UMaterialEditingLibrary::GetStaticSwitchParameterNames(UMaterialInterface*
 	}
 }
 
-void UMaterialEditingLibrary::GetEnumerationParameterNames(UMaterialInterface* Material, TArray<FName>& ParameterNames)
-{
-	ParameterNames.Empty();
-	if (Material)
-	{
-		TArray<FMaterialParameterInfo> MaterialInfo;
-		TArray<FGuid> MaterialGuids;
-		Material->GetAllEnumerationParameterInfo(MaterialInfo, MaterialGuids);
-
-		for (const FMaterialParameterInfo& Info : MaterialInfo)
-		{
-			ParameterNames.Add(Info.Name);
-		}
-	}
-}
-
 static bool GetParameterSource(UMaterialInterface* Material, const TArray<FMaterialParameterInfo>& Info, const TArray<FGuid>& Guids, const FName& ParameterName, FSoftObjectPath& OutParameterSource)
 {
 	bool bResult = false;
@@ -1480,46 +1442,6 @@ bool UMaterialEditingLibrary::GetStaticSwitchParameterSource(UMaterialInterface*
 	return false;
 }
 
-bool UMaterialEditingLibrary::GetEnumerationParameterSource(UMaterialInterface* Material, const FName ParameterName, FSoftObjectPath& ParameterSource)
-{
-	if (Material)
-	{
-		TArray<FMaterialParameterInfo> MaterialInfo;
-		TArray<FGuid> MaterialGuids;
-		Material->GetAllEnumerationParameterInfo(MaterialInfo, MaterialGuids);
-		return GetParameterSource(Material, MaterialInfo, MaterialGuids, ParameterName, ParameterSource);
-	}
-	return false;
-}
-
-bool UMaterialEditingLibrary::GetEnumerationParameterEntryNames(UMaterialInterface* Material, FName ParameterName, EMaterialParameterAssociation Association, TArray<FName>& OutEntryNames)
-{
-	if (!Material)
-	{
-		return false;
-	}
-
-	FHashedMaterialParameterInfo ParameterInfo(ParameterName, Association, Association == EMaterialParameterAssociation::LayerParameter ? 0 : INDEX_NONE);
-	
-	FMaterialParameterMetadata ParameterMetadata;
-	if (!Material->GetCachedExpressionData().GetParameterValue(EMaterialParameterType::Enumeration, ParameterInfo, ParameterMetadata))
-	{
-		return false;
-	}
-
-	if (!ParameterMetadata.Enumeration)
-	{
-		return false;
-	}
-
-	Cast<IMaterialEnumerationProvider>(ParameterMetadata.Enumeration)->ForEachEntry([&OutEntryNames](FName EntryName, int32 EntryValue)
-	{
-		OutEntryNames.Add(EntryName);
-	});
-
-	return true;
-}
-
 FMaterialStatistics UMaterialEditingLibrary::GetStatistics(class UMaterialInterface* Material)
 {
 	FMaterialStatistics Result;
diff --git a/Engine/Source/Editor/MaterialEditor/Private/MaterialPropertyHelpers.cpp b/Engine/Source/Editor/MaterialEditor/Private/MaterialPropertyHelpers.cpp
--- a/Engine/Source/Editor/MaterialEditor/Private/MaterialPropertyHelpers.cpp
+++ b/Engine/Source/Editor/MaterialEditor/Private/MaterialPropertyHelpers.cpp
@@ -447,7 +447,6 @@ FReply FMaterialPropertyHelpers::OnClickedSaveNewFunctionInstance(class UMateria
 					ChildInstance->RuntimeVirtualTextureParameterValues = EditedInstance->RuntimeVirtualTextureParameterValues;
 					ChildInstance->SparseVolumeTextureParameterValues = EditedInstance->SparseVolumeTextureParameterValues;
 					ChildInstance->FontParameterValues = EditedInstance->FontParameterValues;
-					ChildInstance->EnumerationParameterValues = EditedInstance->EnumerationParameterValues;
 
 					const FStaticParameterSet& StaticParameters = EditedInstance->GetStaticParameters();
 					ChildInstance->StaticSwitchParameterValues = StaticParameters.StaticSwitchParameters;
@@ -542,7 +541,6 @@ FReply FMaterialPropertyHelpers::OnClickedSaveNewLayerInstance(class UMaterialFu
 					ChildInstance->RuntimeVirtualTextureParameterValues = EditedInstance->RuntimeVirtualTextureParameterValues;
 					ChildInstance->SparseVolumeTextureParameterValues = EditedInstance->SparseVolumeTextureParameterValues;
 					ChildInstance->FontParameterValues = EditedInstance->FontParameterValues;
-					ChildInstance->EnumerationParameterValues = EditedInstance->EnumerationParameterValues;
 
 					const FStaticParameterSet& StaticParameters = EditedInstance->GetStaticParameters();
 					ChildInstance->StaticSwitchParameterValues = StaticParameters.StaticSwitchParameters;
diff --git a/Engine/Source/Editor/MaterialEditor/Public/MaterialEditingLibrary.h b/Engine/Source/Editor/MaterialEditor/Public/MaterialEditingLibrary.h
--- a/Engine/Source/Editor/MaterialEditor/Public/MaterialEditingLibrary.h
+++ b/Engine/Source/Editor/MaterialEditor/Public/MaterialEditingLibrary.h
@@ -344,14 +344,6 @@ public:
 	UFUNCTION(BlueprintCallable, Category = "MaterialEditing")
 	static UE_API bool SetMaterialInstanceStaticSwitchParameterValue(UMaterialInstanceConstant* Instance, FName ParameterName, bool Value, EMaterialParameterAssociation Association = EMaterialParameterAssociation::GlobalParameter, bool bUpdateMaterialInstance = true);
 
-	/** Get the current enumeration parameter value from a Material Instance */
-	UFUNCTION(BlueprintPure, Category = "MaterialEditing")
-	static UE_API FName GetMaterialInstanceEnumerationParameterValue(UMaterialInstanceConstant* Instance, FName ParameterName, EMaterialParameterAssociation Association = EMaterialParameterAssociation::GlobalParameter);
-
-	/** Set the enumeration parameter value for a Material Instance */
-	UFUNCTION(BlueprintCallable, Category = "MaterialEditing")
-	static UE_API bool SetMaterialInstanceEnumerationParameterValue(UMaterialInstanceConstant* Instance, FName ParameterName, FName Value, EMaterialParameterAssociation Association = EMaterialParameterAssociation::GlobalParameter);
-
 	/** Called after making modifications to a Material Instance to recompile shaders etc. */
 	UFUNCTION(BlueprintCallable, Category = "MaterialEditing")
 	static UE_API void UpdateMaterialInstance(UMaterialInstanceConstant* Instance);
@@ -376,10 +368,6 @@ public:
 	UFUNCTION(BlueprintCallable, Category = "MaterialEditing")
 	static UE_API void GetStaticSwitchParameterNames(UMaterialInterface* Material, TArray<FName>& ParameterNames);
 	
-	/** Gets all enumeration parameter names */
-	UFUNCTION(BlueprintCallable, Category = "MaterialEditing")
-	static UE_API void GetEnumerationParameterNames(UMaterialInterface* Material, TArray<FName>& ParameterNames);
-
 	/**
 	*	Returns the path of the asset where the parameter originated, as well as true/false if it was found
 	*	@param	Material	The material or material instance you want to look up a parameter from
@@ -420,27 +408,6 @@ public:
 	UFUNCTION(BlueprintCallable, Category = "MaterialEditing")
 	static UE_API bool GetStaticSwitchParameterSource(UMaterialInterface* Material, const FName ParameterName, FSoftObjectPath& ParameterSource);
 	
-	/**
-	*	Returns the path of the asset where the parameter originated, as well as true/false if it was found
-	*	@param	Material	The material or material instance you want to look up a parameter from
-	*	@param	ParameterName		The parameter name
-	*	@param	ParameterSource		The soft object path of the asset the parameter originates in
-	*	@return	Whether or not the parameter was found in this material
-	*/
-	UFUNCTION(BlueprintCallable, Category = "MaterialEditing")
-	static UE_API bool GetEnumerationParameterSource(UMaterialInterface* Material, const FName ParameterName, FSoftObjectPath& ParameterSource);
-	
-	/**
-	*	Returns the array of entry names in the enumeration definition referenced by specified enumeration parameter.
-	*	@param	Material				The material or material instance you want to look up a parameter from
-	*	@param	ParameterName			The parameter name
-	*   @param  Association				The parameter association
-	*	@param	OutEntryNames	The output array of enumeration definition entry names.
-	*	@return	Whether or not the parameter was found in this material
-	*/
-	UFUNCTION(BlueprintCallable, Category = "MaterialEditing")
-	static UE_API bool GetEnumerationParameterEntryNames(UMaterialInterface* Material, FName ParameterName, EMaterialParameterAssociation Association, TArray<FName>& OutEntryNames);
-
 	/** Returns statistics about the given material */
 	UFUNCTION(BlueprintCallable, Category = "MaterialEditing")
 	static UE_API FMaterialStatistics GetStatistics(UMaterialInterface* Material);
diff --git a/Engine/Source/Editor/UnrealEd/Classes/MaterialEditor/DEditorEnumerationParameterValue.h b/Engine/Source/Editor/UnrealEd/Classes/MaterialEditor/DEditorEnumerationParameterValue.h
deleted file mode 100644
diff --git a/Engine/Source/Editor/UnrealEd/Private/EditorParameters.cpp b/Engine/Source/Editor/UnrealEd/Private/EditorParameters.cpp
--- a/Engine/Source/Editor/UnrealEd/Private/EditorParameters.cpp
+++ b/Engine/Source/Editor/UnrealEd/Private/EditorParameters.cpp
@@ -14,8 +14,6 @@
 #include "MaterialEditor/DEditorTextureCollectionParameterValue.h"
 #include "MaterialEditor/DEditorParameterCollectionParameterValue.h"
 #include "MaterialEditor/DEditorVectorParameterValue.h"
-#include "MaterialEditor/DEditorEnumerationParameterValue.h"
-#include "Materials/MaterialEnumeration.h"
 #include "Materials/MaterialParameters.h"
 #include "Math/Color.h"
 #include "Math/Vector4.h"
@@ -91,28 +89,6 @@ UDEditorSparseVolumeTextureParameterValue::UDEditorSparseVolumeTextureParameterV
 {
 }
 
-UDEditorEnumerationParameterValue::UDEditorEnumerationParameterValue(const FObjectInitializer& ObjectInitializer)
-	: Super(ObjectInitializer)
-{
-}
-
-TArray<FString> UDEditorEnumerationParameterValue::GetEntryOptions() const
-{
-	TArray<FString> Options;
-	if (Enumeration)
-	{
-		Cast<IMaterialEnumerationProvider>(Enumeration.Get())->ForEachEntry([&Options](FName InEntryName, int32 InEntryValue)
-		{
-			Options.Add(InEntryName.ToString());
-		});
-	}
-	if (Options.IsEmpty())
-	{
-		Options.Add(TEXT("None"));
-	}
-	return Options;
-}
-
 static UDEditorParameterValue* CreateParameter_Scalar(UObject* Owner, const FMaterialParameterMetadata& Meta)
 {
 	UDEditorScalarParameterValue* Parameter = NewObject<UDEditorScalarParameterValue>(Owner);
@@ -239,17 +215,6 @@ static UDEditorParameterValue* CreateParameter_StaticComponentMask(UObject* Owne
 	return Parameter;
 }
 
-static UDEditorParameterValue* CreateParameter_Enumeration(UObject* Owner, const FMaterialParameterMetadata& Meta)
-{
-	UDEditorEnumerationParameterValue* Parameter = NewObject<UDEditorEnumerationParameterValue>(Owner);
-	if (Meta.Value.Type == EMaterialParameterType::Enumeration)
-	{
-		Parameter->Enumeration = Meta.Enumeration;
-		Parameter->ParameterValue = Meta.Value.EnumEntry;
-	}
-	return Parameter;
-}
-
 UDEditorParameterValue* UDEditorParameterValue::Create(UObject* Owner,
 	EMaterialParameterType Type,
 	const FMaterialParameterInfo& ParameterInfo,
@@ -269,7 +234,6 @@ UDEditorParameterValue* UDEditorParameterValue::Create(UObject* Owner,
 	case EMaterialParameterType::StaticSwitch: Parameter = CreateParameter_StaticSwitch(Owner, Meta); break;
 	case EMaterialParameterType::ParameterCollection: Parameter = CreateParameter_ParameterCollection(Owner, Meta); break;
 	case EMaterialParameterType::StaticComponentMask: Parameter = CreateParameter_StaticComponentMask(Owner, Meta); break;
-	case EMaterialParameterType::Enumeration: Parameter = CreateParameter_Enumeration(Owner, Meta); break;
 	default: break;
 	}
 
diff --git a/Engine/Source/Editor/UnrealEd/Private/PreviewMaterial.cpp b/Engine/Source/Editor/UnrealEd/Private/PreviewMaterial.cpp
--- a/Engine/Source/Editor/UnrealEd/Private/PreviewMaterial.cpp
+++ b/Engine/Source/Editor/UnrealEd/Private/PreviewMaterial.cpp
@@ -1039,7 +1039,6 @@ void UMaterialEditorInstanceConstant::ApplySourceFunctionChanges()
 		SourceFunction->RuntimeVirtualTextureParameterValues = SourceInstance->RuntimeVirtualTextureParameterValues;
 		SourceFunction->SparseVolumeTextureParameterValues = SourceInstance->SparseVolumeTextureParameterValues;
 		SourceFunction->FontParameterValues = SourceInstance->FontParameterValues;
-		SourceFunction->EnumerationParameterValues = SourceInstance->EnumerationParameterValues;
 
 		SourceFunction->StaticSwitchParameterValues = SourceInstance->GetStaticParameters().StaticSwitchParameters;
 		const FStaticParameterSetEditorOnlyData& StaticParameters = SourceInstance->GetEditorOnlyStaticParameters();
diff --git a/Engine/Source/Runtime/Engine/Private/Materials/HLSLMaterialTranslator.cpp b/Engine/Source/Runtime/Engine/Private/Materials/HLSLMaterialTranslator.cpp
--- a/Engine/Source/Runtime/Engine/Private/Materials/HLSLMaterialTranslator.cpp
+++ b/Engine/Source/Runtime/Engine/Private/Materials/HLSLMaterialTranslator.cpp
@@ -535,7 +535,6 @@ static inline EMaterialValueType GetMaterialValueType(EMaterialParameterType Typ
 	case EMaterialParameterType::Vector: return MCT_Float4;
 	case EMaterialParameterType::DoubleVector: return MCT_LWCVector4;
 	case EMaterialParameterType::StaticSwitch: return MCT_Bool;
-	case EMaterialParameterType::Enumeration: return MCT_UInt1;
 	default: checkNoEntry(); return MCT_Unknown;
 	}
 }
diff --git a/Engine/Source/Runtime/Engine/Private/Materials/MaterialCachedData.cpp b/Engine/Source/Runtime/Engine/Private/Materials/MaterialCachedData.cpp
--- a/Engine/Source/Runtime/Engine/Private/Materials/MaterialCachedData.cpp
+++ b/Engine/Source/Runtime/Engine/Private/Materials/MaterialCachedData.cpp
@@ -49,7 +49,6 @@
 #include "MaterialShared.h"
 #include "MaterialCache/MaterialCacheMaterial.h"
 #include "Materials/MaterialExpressionMaterialCache.h"
-#include "Materials/MaterialEnumeration.h"
 
 #include UE_INLINE_GENERATED_CPP_BY_NAME(MaterialCachedData)
 
@@ -276,11 +275,6 @@ bool FMaterialCachedExpressionData::AddParameter(const FMaterialParameterInfo& P
 			EditorOnlyData->StaticComponentMaskValues.Insert(ParameterMeta.Value.AsStaticComponentMask(), Index);
 			break;
 
-		case EMaterialParameterType::Enumeration:
-			Enumerations.Insert(ParameterMeta.Enumeration, Index);
-			EnumerationEntryValues.Insert(ParameterMeta.Value.EnumEntry, Index);
-			break;
-
 		default:
 			checkNoEntry();
 			break;
@@ -412,17 +406,6 @@ bool FMaterialCachedExpressionData::AddParameter(const FMaterialParameterInfo& P
 			}
 			break;
 
-		case EMaterialParameterType::Enumeration:
-			bSameValue = Enumerations[Index] == ParameterMeta.Enumeration && EnumerationEntryValues[Index] == ParameterMeta.Value.EnumEntry;
-			if (!bSameValue)
-			{
-				OutErrorMessage = ReportError(
-					FString::Printf(TEXT("%s(%s)"), *Enumerations[Index].ToString(), *EnumerationEntryValues[Index].ToString()),
-					FString::Printf(TEXT("%s(%s)"), *TSoftObjectPtr<UObject>(ParameterMeta.Enumeration).ToString(), *ParameterMeta.Value.EnumEntry.ToString()),
-					*PreviousEditorInfo, EditorInfo);
-			}
-			break;
-
 		default:
 			bSameValue = true;
 			break;
@@ -947,10 +930,6 @@ void FMaterialCachedExpressionData::GetParameterValueByIndex(EMaterialParameterT
 		}
 		break;
 #endif // WITH_EDITORONLY_DATA
-	case EMaterialParameterType::Enumeration:
-		OutResult.Enumeration = Enumerations[ParameterIndex].LoadSynchronous();
-		OutResult.Value = FMaterialParameterValue(EnumerationEntryValues[ParameterIndex]);
-		break;
 	default:
 		checkNoEntry();
 		break;
diff --git a/Engine/Source/Runtime/Engine/Private/Materials/MaterialEnumeration.cpp b/Engine/Source/Runtime/Engine/Private/Materials/MaterialEnumeration.cpp
--- a/Engine/Source/Runtime/Engine/Private/Materials/MaterialEnumeration.cpp
+++ b/Engine/Source/Runtime/Engine/Private/Materials/MaterialEnumeration.cpp
@@ -10,62 +10,3 @@ int32 IMaterialEnumerationProvider::GetValueOrDefault(FName EntryName, int32 Def
 	ResolveValue(EntryName, Value, DefaultValue);
 	return Value;
 }
-
-bool UMaterialEnumeration::ResolveValue(FName EntryName, int32& OutValue, int32 DefaultValue) const
-{
-	if (SortedEntries.Num() != Entries.Num())
-	{
-		for (const FMaterialEnumerationEntry& Entry : Entries)
-		{
-			if (Entry.Name == EntryName)
-			{
-				OutValue = Entry.Value;
-				return true;
-			}
-		}
-	}
-	else
-	{
-		int32 Index = Algo::BinarySearchBy(SortedEntries, EntryName, &FMaterialEnumerationEntry::Name, FNameFastLess{});
-		if (Index != INDEX_NONE)
-		{
-			OutValue = SortedEntries[Index].Value;
-			return true;
-		}
-	}
-	OutValue = DefaultValue;
-	return false;
-}
-
-void UMaterialEnumeration::ForEachEntry(TFunctionRef<void (FName Name, int32 Value)> Iterator) const
-{
-	for (const FMaterialEnumerationEntry& Entry : Entries)
-	{
-		Iterator(Entry.Name, Entry.Value);
-	}
-}
-
-void UMaterialEnumeration::SortEntries()
-{
-	SortedEntries = Entries;
-	SortedEntries.Sort([](const FMaterialEnumerationEntry& A, const FMaterialEnumerationEntry& B)
-	{
-		return A.Name.FastLess(B.Name);
-	});
-}
-
-void UMaterialEnumeration::PostLoad()
-{
-	Super::PostLoad();
-	SortEntries();
-}
-
-#if WITH_EDITOR
-
-void UMaterialEnumeration::PostEditChangeProperty(FPropertyChangedEvent& PropertyChangedEvent)
-{
-	Super::PostEditChangeProperty(PropertyChangedEvent);
-	SortEntries();
-}
-
-#endif
diff --git a/Engine/Source/Runtime/Engine/Private/Materials/MaterialExpressions.cpp b/Engine/Source/Runtime/Engine/Private/Materials/MaterialExpressions.cpp
--- a/Engine/Source/Runtime/Engine/Private/Materials/MaterialExpressions.cpp
+++ b/Engine/Source/Runtime/Engine/Private/Materials/MaterialExpressions.cpp
@@ -177,7 +177,6 @@
 #include "Materials/MaterialExpressionReroute.h"
 #include "Materials/MaterialExpressionNamedReroute.h"
 #include "Materials/MaterialExpressionScalarParameter.h"
-#include "Materials/MaterialExpressionEnumerationParameter.h"
 #include "Materials/MaterialExpressionSetMaterialAttributes.h"
 #include "Materials/MaterialExpressionShadowReplace.h"
 #include "Materials/MaterialExpressionSign.h"
@@ -15360,17 +15359,6 @@ void UMaterialFunctionInstance::UpdateParameterSet()
 						}
 					}
 				}
-				else if (const UMaterialExpressionEnumerationParameter* EnumerationParameter = Cast<const UMaterialExpressionEnumerationParameter>(FunctionExpression))
-				{
-					for (FEnumerationParameterValue& EnumerationParameterValue : EnumerationParameterValues)
-					{
-						if (EnumerationParameterValue.ExpressionGUID == EnumerationParameter->ExpressionGUID)
-						{
-							EnumerationParameterValue.ParameterInfo.Name = EnumerationParameter->ParameterName;
-							break;
-						}
-					}
-				}
 			}
 		}
 	}
@@ -15387,7 +15375,6 @@ void UMaterialFunctionInstance::OverrideMaterialInstanceParameterValues(UMateria
 	Instance->RuntimeVirtualTextureParameterValues = RuntimeVirtualTextureParameterValues;
 	Instance->SparseVolumeTextureParameterValues = SparseVolumeTextureParameterValues;
 	Instance->FontParameterValues = FontParameterValues;
-	Instance->EnumerationParameterValues = EnumerationParameterValues;
 
 	// Static parameters
 	FStaticParameterSet StaticParametersOverride = Instance->GetStaticParameters();
@@ -15612,7 +15599,6 @@ bool UMaterialFunctionInstance::GetParameterOverrideValue(EMaterialParameterType
 	case EMaterialParameterType::StaticSwitch: bResult = GameThread_GetParameterValue(StaticSwitchParameterValues, ParameterInfo, OutResult); break;
 	case EMaterialParameterType::ParameterCollection: bResult = GameThread_GetParameterValue(ParameterCollectionParameterValues, ParameterInfo, OutResult); break;
 	case EMaterialParameterType::StaticComponentMask: bResult = GameThread_GetParameterValue(StaticComponentMaskParameterValues, ParameterInfo, OutResult); break;
-	case EMaterialParameterType::Enumeration: bResult = GameThread_GetParameterValue(EnumerationParameterValues, ParameterInfo, OutResult); break;
 	default: checkNoEntry(); break;
 	}
 
@@ -24412,105 +24398,6 @@ void UMaterialExpressionOperator::PostLoad()
 	UMaterialExpression::PostLoad();
 }
 
-///////////////////////////////////////////////////////////////////////////////
-// UMaterialExpressionEnumerationParameter
-///////////////////////////////////////////////////////////////////////////////
-
-int32 UMaterialExpressionEnumerationParameter::Compile(class FMaterialCompiler* Compiler, int32 OutputIndex)
-{
-	if (!Enumeration)
-	{
-		return Compiler->Errorf(TEXT("Enum Parameter '%s': Enumeration asset missing."), *ParameterName.ToString());
-	}
-	int32 ResolvedDefaultValue = Enumeration->GetValueOrDefault(DefaultEntry, 0);
-	return Compiler->NumericParameter(EMaterialParameterType::Enumeration, ParameterName, ResolvedDefaultValue);
-}
-
-void UMaterialExpressionEnumerationParameter::GetCaption(TArray<FString>& OutCaptions) const
-{
-	OutCaptions.Add(FString::Printf(TEXT("Param (%s)"), *DefaultEntry.ToString())); 
-	OutCaptions.Add(FString::Printf(TEXT("'%s'"), *ParameterName.ToString())); 
-}
-
-EMaterialValueType UMaterialExpressionEnumerationParameter::GetOutputValueType(int32 InIndex)
-{
-	return MCT_UInt1;
-}
-
-bool UMaterialExpressionEnumerationParameter::GetParameterValue(FMaterialParameterMetadata& OutMeta) const
-{
-	OutMeta.Value = FMaterialParameterValue(DefaultEntry);
-	OutMeta.Enumeration = Enumeration.GetObject();
-	return Super::GetParameterValue(OutMeta);
-}
-
-bool UMaterialExpressionEnumerationParameter::SetParameterValue(const FName& Name, const FMaterialParameterMetadata& Meta, EMaterialExpressionSetParameterValueFlags Flags)
-{
-	if (Meta.Value.Type == EMaterialParameterType::Enumeration)
-	{
-		if (SetParameterValue(Name, Meta.Enumeration.Get(), Meta.Value.EnumEntry, Flags))
-		{
-			if (EnumHasAnyFlags(Flags, EMaterialExpressionSetParameterValueFlags::AssignGroupAndSortPriority))
-			{
-				Group = Meta.Group;
-				SortPriority = Meta.SortPriority;
-			}
-			return true;
-		}
-	}
-	return false;
-}
-
-void UMaterialExpressionEnumerationParameter::PostEditChangeProperty(FPropertyChangedEvent& PropertyChangedEvent)
-{
-	const FName PropertyName = PropertyChangedEvent.GetMemberPropertyName();
-	if (PropertyName == GET_MEMBER_NAME_STRING_CHECKED(ThisClass, Enumeration)
-		|| PropertyName == GET_MEMBER_NAME_STRING_CHECKED(ThisClass, DefaultEntry))
-	{
-		int32 DefaultEntryValue;
-		if (!Enumeration || !Enumeration->ResolveValue(DefaultEntry, DefaultEntryValue))
-		{
-			DefaultEntryValue = 0;
-			DefaultEntry = NAME_None;
-		}
-		FEditorSupportDelegates::NumericParameterDefaultChanged.Broadcast(this, EMaterialParameterType::Enumeration, ParameterName, DefaultEntryValue);
-	}
-	Super::PostEditChangeProperty(PropertyChangedEvent);
-}
-
-bool UMaterialExpressionEnumerationParameter::SetParameterValue(const FName& InParameterName, UObject* InEnumeration, FName InDefaultEntry, EMaterialExpressionSetParameterValueFlags Flags)
-{
-	if (InParameterName == ParameterName)
-	{
-		Enumeration = InEnumeration;
-		DefaultEntry = InDefaultEntry;
-		if (EnumHasAnyFlags(Flags, EMaterialExpressionSetParameterValueFlags::SendPostEditChangeProperty))
-		{
-			SendPostEditChangeProperty(this, GET_MEMBER_NAME_STRING_CHECKED(ThisClass, Enumeration));
-			SendPostEditChangeProperty(this, GET_MEMBER_NAME_STRING_CHECKED(ThisClass, DefaultEntry));
-		}
-		return true;
-	}
-	return false;
-}
-
-TArray<FString> UMaterialExpressionEnumerationParameter::GetEntryOptions() const
-{
-	TArray<FString> Options;
-	if (Enumeration)
-	{
-		Enumeration->ForEachEntry([&Options](FName InEntryName, int32 InEntryValue)
-		{
-			Options.Add(InEntryName.ToString());
-		});
-	}
-	if (Options.IsEmpty())
-	{
-		Options.Add(TEXT("None"));
-	}
-	return Options;
-}
-
 #endif // WITH_EDITOR
 
 ///////////////////////////////////////////////////////////////////////////////
diff --git a/Engine/Source/Runtime/Engine/Private/Materials/MaterialExpressionsIR.cpp b/Engine/Source/Runtime/Engine/Private/Materials/MaterialExpressionsIR.cpp
--- a/Engine/Source/Runtime/Engine/Private/Materials/MaterialExpressionsIR.cpp
+++ b/Engine/Source/Runtime/Engine/Private/Materials/MaterialExpressionsIR.cpp
@@ -75,7 +75,6 @@
 #include "Materials/MaterialExpressionDivide.h"
 #include "Materials/MaterialExpressionDotProduct.h"
 #include "Materials/MaterialExpressionDoubleVectorParameter.h"
-#include "Materials/MaterialExpressionEnumerationParameter.h"
 #include "Materials/MaterialExpressionDynamicParameter.h"
 #include "Materials/MaterialExpressionExponential.h"
 #include "Materials/MaterialExpressionExponential2.h"
@@ -656,22 +655,6 @@ void UMaterialExpressionShadowReplace::Build(MIR::FEmitter& Em)
 
 #undef REPLACE_INPUTS
 
-void UMaterialExpressionEnumerationParameter::Build(MIR::FEmitter& Em)
-{
-	if (!Enumeration)
-	{
-		return Em.Errorf(TEXT("Enumeration definition unspecified in enumeration parameter '%s'."), *ParameterName.ToString());
-	}
-
-	FMaterialParameterMetadata Metadata;
-	if (!GetParameterValue(Metadata))
-	{
-		return Em.Errorf(TEXT("Could fetch parameter metadata from enumeration parameter '%s'."), *ParameterName.ToString());
-	}
-
-	Em.Output(0, Em.Parameter(ParameterName, Metadata));
-}
-
 void UMaterialExpressionAppendVector::Build(MIR::FEmitter& Em)
 {
 	FValueRef AVal = Em.CheckIsScalarOrVector(Em.Input(&A));
diff --git a/Engine/Source/Runtime/Engine/Private/Materials/MaterialIREmitter.cpp b/Engine/Source/Runtime/Engine/Private/Materials/MaterialIREmitter.cpp
--- a/Engine/Source/Runtime/Engine/Private/Materials/MaterialIREmitter.cpp
+++ b/Engine/Source/Runtime/Engine/Private/Materials/MaterialIREmitter.cpp
@@ -901,18 +901,13 @@ FValueRef FEmitter::Parameter(FName Name, FMaterialParameterMetadata& Metadata,
 	switch (Metadata.Value.Type)
 	{
 		case EMaterialParameterType::Scalar:
-		case EMaterialParameterType::Enumeration:
 		{
 			if (Metadata.PrimitiveDataIndex != INDEX_NONE)
 			{
 				return CustomPrimitiveData(Metadata.PrimitiveDataIndex);
 			}
 
-			FType Type = Metadata.Value.Type == EMaterialParameterType::Enumeration
-				? FType::MakeIntScalar()
-				: FType::MakeFloatScalar();
-
-			FUniformParameter Prototype = MakePrototype<FUniformParameter>(Type);
+			FUniformParameter Prototype = MakePrototype<FUniformParameter>(FType::MakeFloatScalar());
 			Prototype.ParameterIdInModule = RegisterParameter(Module, Info, Metadata);
 			return EmitPrototype(*this, Prototype);
 		}
@@ -1532,7 +1527,7 @@ static T FoldScalarOperator(FEmitter& Emitter, EOperator Operator, T A, T B, T C
 	}
 }
 
-// It tries to apply a known identity of specified operator, e.g. "x + 0 = x ∀ x ∈ R".
+// It tries to apply a known identity of specified operator, e.g. "x + 0 = x ? x ? R".
 // If it returns a value, the operation has been "folded" and the returned value is the
 // result (in the example above, it would return "x").
 // If it returns null, the end result could not be inferred, but the operator could have
@@ -3573,7 +3568,7 @@ static FValue* DifferentiateOperator(FEmitter& E, FOperator* Op, EDerivativeAxis
 					   E.Add(E.Multiply(dH, E.Subtract(G, F)), 
 						   E.Multiply(H, E.Subtract(dG, dF))));
 
-		// d/dx select(F, g(x), h(x)) ≈ select(F, g'(x), h'(x))
+		// d/dx select(F, g(x), h(x)) ˜ select(F, g'(x), h'(x))
 		case TO_Select:
 			return E.Select(F, dG, dH);
 
diff --git a/Engine/Source/Runtime/Engine/Private/Materials/MaterialIRValueAnalyzer.cpp b/Engine/Source/Runtime/Engine/Private/Materials/MaterialIRValueAnalyzer.cpp
--- a/Engine/Source/Runtime/Engine/Private/Materials/MaterialIRValueAnalyzer.cpp
+++ b/Engine/Source/Runtime/Engine/Private/Materials/MaterialIRValueAnalyzer.cpp
@@ -14,7 +14,6 @@
 #include "Materials/MaterialIRInternal.h"
 #include "Materials/MaterialInsights.h"
 #include "Materials/Material.h"
-#include "Materials/MaterialEnumeration.h"
 #include "MaterialCachedData.h"
 #include "MaterialDomain.h"
 #include "MaterialShared.h"
@@ -482,7 +481,6 @@ static void AnalyzePrimitiveUniformParameter(FMaterialIRValueAnalyzer& Analyzer,
 		case EMaterialParameterType::Scalar: DefaultValue = ParameterMetadata.Value.AsScalar(); break;
 		case EMaterialParameterType::Vector: DefaultValue = ParameterMetadata.Value.AsLinearColor(); break;
 		case EMaterialParameterType::DoubleVector: DefaultValue = ParameterMetadata.Value.AsVector4d(); break;
-		case EMaterialParameterType::Enumeration: DefaultValue = Cast<IMaterialEnumerationProvider>(ParameterMetadata.Enumeration)->GetValueOrDefault(ParameterMetadata.Value.EnumEntry, 0); break;
 		default: UE_MIR_TODO();
 	}
 
diff --git a/Engine/Source/Runtime/Engine/Private/Materials/MaterialInstance.cpp b/Engine/Source/Runtime/Engine/Private/Materials/MaterialInstance.cpp
--- a/Engine/Source/Runtime/Engine/Private/Materials/MaterialInstance.cpp
+++ b/Engine/Source/Runtime/Engine/Private/Materials/MaterialInstance.cpp
@@ -30,7 +30,6 @@
 #include "Materials/MaterialExpressionMaterialAttributeLayers.h"
 #include "Materials/MaterialExpressionRuntimeVirtualTextureSampleParameter.h"
 #include "Materials/MaterialExpressionSparseVolumeTextureSample.h"
-#include "Materials/MaterialExpressionEnumerationParameter.h"
 #include "Materials/MaterialExpressionStaticComponentMaskParameter.h"
 #include "Materials/MaterialExpressionMaterialFunctionCall.h"
 #include "Materials/MaterialInstanceConstant.h"
@@ -351,7 +350,6 @@ bool FMaterialInstanceResource::GetParameterValue(EMaterialParameterType Type, c
 			case EMaterialParameterType::RuntimeVirtualTexture: bResult = RenderThread_GetParameterValue<const URuntimeVirtualTexture*>(ParameterInfo, OutValue); break;
 			case EMaterialParameterType::SparseVolumeTexture: bResult = RenderThread_GetParameterValue<const USparseVolumeTexture*>(ParameterInfo, OutValue); break;
 			case EMaterialParameterType::ParameterCollection: bResult = RenderThread_GetParameterValue<const UMaterialParameterCollection*>(ParameterInfo, OutValue); break;
-			case EMaterialParameterType::Enumeration: bResult = RenderThread_GetParameterValue<FName>(ParameterInfo, OutValue); break;
 			default: ensure(false); break; // other parameter types are not expected on the render thread
 		}
 	}
@@ -517,7 +515,6 @@ void FMaterialInstanceResource::InitMIParameters(FMaterialInstanceParameterSet&
 	ParameterSet.RuntimeVirtualTextureParameters.Sort(SortMaterialInstanceParametersPredicate<const URuntimeVirtualTexture*>);
 	ParameterSet.SparseVolumeTextureParameters.Sort(SortMaterialInstanceParametersPredicate<const USparseVolumeTexture*>);
 	ParameterSet.ParameterCollectionParameters.Sort(SortMaterialInstanceParametersPredicate<const UMaterialParameterCollection*>);
-	ParameterSet.EnumerationParameters.Sort(SortMaterialInstanceParametersPredicate<FName>);
 
 	StaticSwitchParameterArray.Array = MoveTemp(ParameterSet.StaticSwitchParameters);
 	ScalarParameterArray.Array = MoveTemp(ParameterSet.ScalarParameters);
@@ -528,7 +525,6 @@ void FMaterialInstanceResource::InitMIParameters(FMaterialInstanceParameterSet&
 	RuntimeVirtualTextureParameterArray.Array = MoveTemp(ParameterSet.RuntimeVirtualTextureParameters);
 	SparseVolumeTextureParameterArray.Array = MoveTemp(ParameterSet.SparseVolumeTextureParameters);
 	ParameterCollectionParameterArray.Array = MoveTemp(ParameterSet.ParameterCollectionParameters);
-	EnumerationParameterArray.Array = MoveTemp(ParameterSet.EnumerationParameters);
 	UserSceneTextureOverrides = MoveTemp(ParameterSet.UserSceneTextureOverrides);
 	PostProcessBlendableOverrides = ParameterSet.PostProcessBlendableOverrides;
 
@@ -542,7 +538,6 @@ void FMaterialInstanceResource::InitMIParameters(FMaterialInstanceParameterSet&
 	RuntimeVirtualTextureParameterArray.HashAddAllItems();
 	SparseVolumeTextureParameterArray.HashAddAllItems();
 	ParameterCollectionParameterArray.HashAddAllItems();
-	EnumerationParameterArray.HashAddAllItems();
 }
 
 void FMaterialInstanceResource::EvaluateParameterCollections(FRHICommandListBase& RHICmdList, FUniformExpressionCache& OutUniformExpressionCache, const FMaterialRenderContext& Context, FUniformExpressionCacheAsyncUpdater* Updater) const
@@ -730,7 +725,6 @@ void UMaterialInstance::SwapLayerParameterIndices(int32 OriginalIndex, int32 New
 		SwapLayerParameterIndicesArray(SparseVolumeTextureParameterValues, OriginalIndex, NewIndex);
 		SwapLayerParameterIndicesArray(FontParameterValues, OriginalIndex, NewIndex);
 		SwapLayerParameterIndicesArray(StaticParametersRuntime.StaticSwitchParameters, OriginalIndex, NewIndex);
-		SwapLayerParameterIndicesArray(EnumerationParameterValues, OriginalIndex, NewIndex);
 		if (EditorOnly)
 		{
 			SwapLayerParameterIndicesArray(EditorOnly->StaticParameters.StaticComponentMaskParameters, OriginalIndex, NewIndex);
@@ -751,7 +745,6 @@ void UMaterialInstance::RemoveLayerParameterIndex(int32 Index)
 	RemoveLayerParameterIndicesArray(SparseVolumeTextureParameterValues, Index);
 	RemoveLayerParameterIndicesArray(FontParameterValues, Index);
 	RemoveLayerParameterIndicesArray(StaticParametersRuntime.StaticSwitchParameters, Index);
-	RemoveLayerParameterIndicesArray(EnumerationParameterValues, Index);
 	if (EditorOnly)
 	{
 		RemoveLayerParameterIndicesArray(EditorOnly->StaticParameters.StaticComponentMaskParameters, Index);
@@ -817,9 +810,6 @@ bool UMaterialInstance::UpdateParameters()
 
 			// Static component mask parameters
 			bDirty = UpdateParameterSet<FStaticComponentMaskParameter, UMaterialExpressionStaticComponentMaskParameter>(EditorOnly->StaticParameters.StaticComponentMaskParameters, ParentMaterial) || bDirty;
-			
-			// Enumeration parameters
-			bDirty = UpdateParameterSet<FEnumerationParameterValue, UMaterialExpressionEnumerationParameter>(EnumerationParameterValues, ParentMaterial) || bDirty;
 		}
 
 		if (StaticParametersRuntime.bHasMaterialLayers && Parent)
@@ -845,7 +835,6 @@ bool UMaterialInstance::UpdateParameters()
 					RemapLayerParameterIndicesArray(FontParameterValues, RemapLayerIndices);
 					RemapLayerParameterIndicesArray(StaticParametersRuntime.StaticSwitchParameters, RemapLayerIndices);
 					RemapLayerParameterIndicesArray(EditorOnly->StaticParameters.StaticComponentMaskParameters, RemapLayerIndices);
-					RemapLayerParameterIndicesArray(EnumerationParameterValues, RemapLayerIndices);
 					bDirty = true;
 				}
 			}
@@ -980,14 +969,6 @@ void GameThread_InitMIParameters(const UMaterialInstance& Instance)
 		ParamRef.Value = FSparseVolumeTextureParameterValue::GetValue(Parameter);
 	}
 
-	// Enumeration parameters
-	for (const FEnumerationParameterValue& Parameter : Instance.EnumerationParameterValues)
-	{
-		auto& ParamRef = ParameterSet.EnumerationParameters.AddDefaulted_GetRef();
-		ParamRef.Info = Parameter.ParameterInfo;
-		ParamRef.Value = FEnumerationParameterValue::GetValue(Parameter);
-	}
-	
 	FStaticParameterSet StaticParamSet = Instance.GetStaticParameters();
 	ParameterSet.StaticSwitchParameters.Reserve(StaticParamSet.StaticSwitchParameters.Num());
 	for(const FStaticSwitchParameter& Param : StaticParamSet.StaticSwitchParameters)
@@ -1166,7 +1147,6 @@ bool UMaterialInstance::GetParameterOverrideValue(EMaterialParameterType Type, c
 		bResult = GameThread_GetParameterValue(GetEditorOnlyData()->StaticParameters.StaticComponentMaskParameters, ParameterInfo, OutResult);
 		break;
 #endif // WITH_EDITORONLY_DATA
-	case EMaterialParameterType::Enumeration: bResult = GameThread_GetParameterValue(EnumerationParameterValues, ParameterInfo, OutResult); break;
 	default: checkNoEntry(); break;
 	}
 	return bResult;
@@ -2311,7 +2291,6 @@ void UMaterialInstance::GetAllParametersOfType(EMaterialParameterType Type, TMap
 		case EMaterialParameterType::Font: GameThread_ApplyParameterOverrides(Instance->FontParameterValues, LayerIndexRemap, bSetOverride, OverridenParameters, OutParameters); break;
 		case EMaterialParameterType::StaticSwitch: GameThread_ApplyParameterOverrides(Instance->StaticParametersRuntime.StaticSwitchParameters, LayerIndexRemap, bSetOverride, OverridenParameters, OutParameters); break;
 		case EMaterialParameterType::ParameterCollection: GameThread_ApplyParameterOverrides(Instance->ParameterCollectionParameterValues, LayerIndexRemap, bSetOverride, OverridenParameters, OutParameters); break;
-		case EMaterialParameterType::Enumeration: GameThread_ApplyParameterOverrides(Instance->EnumerationParameterValues, LayerIndexRemap, bSetOverride, OverridenParameters, OutParameters); break;
 #if WITH_EDITORONLY_DATA
 		case EMaterialParameterType::StaticComponentMask:
 			GameThread_ApplyParameterOverrides(Instance->GetEditorOnlyData()->StaticParameters.StaticComponentMaskParameters, LayerIndexRemap, bSetOverride, OverridenParameters, OutParameters);
@@ -4015,7 +3994,6 @@ void UMaterialInstance::ReserveParameterValuesInternal(EMaterialParameterType Ty
 	case EMaterialParameterType::Font: FontParameterValues.Reserve(Capacity); break;
 	case EMaterialParameterType::RuntimeVirtualTexture: RuntimeVirtualTextureParameterValues.Reserve(Capacity); break;
 	case EMaterialParameterType::SparseVolumeTexture: SparseVolumeTextureParameterValues.Reserve(Capacity); break;
-	case EMaterialParameterType::Enumeration: EnumerationParameterValues.Reserve(Capacity); break;
 	default: checkNoEntry();
 	}
 }
@@ -4047,7 +4025,6 @@ void UMaterialInstance::AddParameterValueInternal(const FMaterialParameterInfo&
 	case EMaterialParameterType::RuntimeVirtualTexture: RuntimeVirtualTextureParameterValues.Emplace(ParameterInfo, Value.RuntimeVirtualTexture); break;
 	case EMaterialParameterType::SparseVolumeTexture: SparseVolumeTextureParameterValues.Emplace(ParameterInfo, Value.SparseVolumeTexture); break;
 	case EMaterialParameterType::StaticSwitch: break;
-	case EMaterialParameterType::Enumeration: EnumerationParameterValues.Emplace(ParameterInfo, Value.EnumEntry); break;
 	default: checkNoEntry();
 	}
 }
@@ -4078,7 +4055,6 @@ void UMaterialInstance::SetParameterValueInternal(const FMaterialParameterInfo&
 	case EMaterialParameterType::RuntimeVirtualTexture: SetRuntimeVirtualTextureParameterValueInternal(ParameterInfo, Value.RuntimeVirtualTexture); break;
 	case EMaterialParameterType::SparseVolumeTexture: SetSparseVolumeTextureParameterValueInternal(ParameterInfo, Value.SparseVolumeTexture); break;
 	case EMaterialParameterType::ParameterCollection: SetParameterCollectionParameterValueInternal(ParameterInfo, Value.ParameterCollection); break;
-	case EMaterialParameterType::Enumeration: SetEnumerationParameterValueInternal(ParameterInfo, Value.EnumEntry); break;
 	default: checkNoEntry();
 	}
 }
@@ -4458,32 +4434,6 @@ void UMaterialInstance::SetFontParameterValueInternal(const FMaterialParameterIn
 	}
 }
 
-void UMaterialInstance::SetEnumerationParameterValueInternal(const FMaterialParameterInfo& ParameterInfo, FName Value)
-{
-	LLM_SCOPE(ELLMTag::MaterialInstance);
-
-	FEnumerationParameterValue* ParameterValue = GameThread_FindParameterByName(EnumerationParameterValues, ParameterInfo);
-
-	bool bForceUpdate = false;
-	if (!ParameterValue)
-	{
-			// If there's no element for the named parameter in array yet, add one.
-			ParameterValue = new(EnumerationParameterValues) FEnumerationParameterValue;
-			ParameterValue->ParameterInfo = ParameterInfo;
-			ParameterValue->ExpressionGUID.Invalidate();
-			bForceUpdate = true;
-	}
-
-	// Don't enqueue an update if it isn't needed
-	if (bForceUpdate || ParameterValue->EnumEntry != Value)
-	{
-		ParameterValue->EnumEntry = Value;
-
-		// Update the material instance data in the rendering thread.
-		GameThread_UpdateMIParameter(this, *ParameterValue);
-	}
-}
-
 void UMaterialInstance::ClearParameterValuesInternal(EMaterialInstanceClearParameterFlag Flags)
 {
 	bool bUpdateResource = false;
@@ -4493,7 +4443,6 @@ void UMaterialInstance::ClearParameterValuesInternal(EMaterialInstanceClearParam
 		VectorParameterValues.Empty();
 		DoubleVectorParameterValues.Empty();
 		ParameterCollectionParameterValues.Empty();
-		EnumerationParameterValues.Empty();
 		bUpdateResource = true;
 	}
 
@@ -5448,10 +5397,6 @@ bool UMaterialInstance::Equivalent(const UMaterialInstance* CompareTo) const
 	{
 		return false;
 	}
-	if (!CompareValueArraysByExpressionGUID(EnumerationParameterValues, CompareTo->EnumerationParameterValues))
-	{
-		return false;
-	}
 	const FStaticParameterSet LocalStaticParameters = GetStaticParameters();
 	if (!LocalStaticParameters.Equivalent(CompareTo->GetStaticParameters()))
 	{
diff --git a/Engine/Source/Runtime/Engine/Private/Materials/MaterialInstanceConstant.cpp b/Engine/Source/Runtime/Engine/Private/Materials/MaterialInstanceConstant.cpp
--- a/Engine/Source/Runtime/Engine/Private/Materials/MaterialInstanceConstant.cpp
+++ b/Engine/Source/Runtime/Engine/Private/Materials/MaterialInstanceConstant.cpp
@@ -157,12 +157,6 @@ void UMaterialInstanceConstant::SetFontParameterValueEditorOnly(const FMaterialP
 	SetFontParameterValueInternal(ParameterInfo, FontValue, FontPage);
 }
 
-void UMaterialInstanceConstant::SetEnumerationParameterValueEditorOnly(const FMaterialParameterInfo& ParameterInfo, FName Value)
-{
-	check(GIsEditor || IsRunningCommandlet());
-	SetEnumerationParameterValueInternal(ParameterInfo, Value);
-}
-
 void UMaterialInstanceConstant::ClearParameterValuesEditorOnly()
 {
 	check(GIsEditor || IsRunningCommandlet());
diff --git a/Engine/Source/Runtime/Engine/Private/Materials/MaterialInstanceSupport.h b/Engine/Source/Runtime/Engine/Private/Materials/MaterialInstanceSupport.h
--- a/Engine/Source/Runtime/Engine/Private/Materials/MaterialInstanceSupport.h
+++ b/Engine/Source/Runtime/Engine/Private/Materials/MaterialInstanceSupport.h
@@ -250,7 +250,6 @@ public:
 		SparseVolumeTextureParameterArray.Empty();
 		ParameterCollectionParameterArray.Empty();
 		UserSceneTextureOverrides.Empty();
-		EnumerationParameterArray.Empty();
 	}
 
 	/**
@@ -419,8 +418,6 @@ private:
 	THashedMaterialParameterMap<const USparseVolumeTexture*> SparseVolumeTextureParameterArray;
 	/** Parameter collection parameters for this material instance. */
 	THashedMaterialParameterMap<const UMaterialParameterCollection*> ParameterCollectionParameterArray;
-	/** Enumeration parameters for this material instance. */
-	THashedMaterialParameterMap<FName> EnumerationParameterArray;
 	/** Remap layer indices for parent */
 	TArray<int32> ParentLayerIndexRemap;
 	/** User Scene Texture overrides for this material instance. */
@@ -438,7 +435,6 @@ template <> FORCEINLINE THashedMaterialParameterMap<const UTextureCollection*>&
 template <> FORCEINLINE THashedMaterialParameterMap<const URuntimeVirtualTexture*>& FMaterialInstanceResource::GetValueArray() { return RuntimeVirtualTextureParameterArray; }
 template <> FORCEINLINE THashedMaterialParameterMap<const USparseVolumeTexture*>& FMaterialInstanceResource::GetValueArray() { return SparseVolumeTextureParameterArray; }
 template <> FORCEINLINE THashedMaterialParameterMap<const UMaterialParameterCollection*>& FMaterialInstanceResource::GetValueArray() { return ParameterCollectionParameterArray; }
-template <> FORCEINLINE THashedMaterialParameterMap<FName>& FMaterialInstanceResource::GetValueArray() { return EnumerationParameterArray; }
 template <> FORCEINLINE const THashedMaterialParameterMap<bool>& FMaterialInstanceResource::GetValueArray() const { return StaticSwitchParameterArray; }
 template <> FORCEINLINE const THashedMaterialParameterMap<float>& FMaterialInstanceResource::GetValueArray() const { return ScalarParameterArray; }
 template <> FORCEINLINE const THashedMaterialParameterMap<FLinearColor>& FMaterialInstanceResource::GetValueArray() const { return VectorParameterArray; }
@@ -448,7 +444,6 @@ template <> FORCEINLINE const THashedMaterialParameterMap<const UTextureCollecti
 template <> FORCEINLINE const THashedMaterialParameterMap<const URuntimeVirtualTexture*>& FMaterialInstanceResource::GetValueArray() const { return RuntimeVirtualTextureParameterArray; }
 template <> FORCEINLINE const THashedMaterialParameterMap<const USparseVolumeTexture*>& FMaterialInstanceResource::GetValueArray() const { return SparseVolumeTextureParameterArray; }
 template <> FORCEINLINE const THashedMaterialParameterMap<const UMaterialParameterCollection*>& FMaterialInstanceResource::GetValueArray() const { return ParameterCollectionParameterArray; }
-template <> FORCEINLINE const THashedMaterialParameterMap<FName>& FMaterialInstanceResource::GetValueArray() const { return EnumerationParameterArray; }
 
 struct FMaterialInstanceParameterSet
 {
@@ -461,7 +456,6 @@ struct FMaterialInstanceParameterSet
 	TArray<THashedMaterialParameterMap<const URuntimeVirtualTexture*>::TNamedParameter>			RuntimeVirtualTextureParameters;
 	TArray<THashedMaterialParameterMap<const USparseVolumeTexture*>::TNamedParameter>			SparseVolumeTextureParameters;
 	TArray<THashedMaterialParameterMap<const UMaterialParameterCollection*>::TNamedParameter>	ParameterCollectionParameters;
-	TArray<THashedMaterialParameterMap<FName>::TNamedParameter>									EnumerationParameters;
 	TArray<FUserSceneTextureOverride>															UserSceneTextureOverrides;
 	FPostProcessBlendableOverrides																PostProcessBlendableOverrides;
 };
diff --git a/Engine/Source/Runtime/Engine/Private/Materials/MaterialInterface.cpp b/Engine/Source/Runtime/Engine/Private/Materials/MaterialInterface.cpp
--- a/Engine/Source/Runtime/Engine/Private/Materials/MaterialInterface.cpp
+++ b/Engine/Source/Runtime/Engine/Private/Materials/MaterialInterface.cpp
@@ -1061,17 +1061,6 @@ bool UMaterialInterface::GetFontParameterValue(const FHashedMaterialParameterInf
 	return false;
 }
 
-bool UMaterialInterface::GetEnumerationParameterValue(const FHashedMaterialParameterInfo& ParameterInfo, FName& OutValue, bool bOveriddenOnly) const
-{
-	FMaterialParameterMetadata Result;
-	if (GetParameterValue(EMaterialParameterType::Enumeration, ParameterInfo, Result, MakeParameterValueFlags(bOveriddenOnly)))
-	{
-		OutValue = Result.Value.EnumEntry;
-		return true;
-	}
-	return false;
-}
-
 bool UMaterialInterface::GetParameterDefaultValue(EMaterialParameterType Type, const FMemoryImageMaterialParameterInfo& ParameterInfo, FMaterialParameterMetadata& OutValue) const
 {
 	return GetParameterValue(Type, ParameterInfo, OutValue, EMaterialGetParameterValueFlags::CheckNonOverrides);
@@ -1178,17 +1167,6 @@ bool UMaterialInterface::GetFontParameterDefaultValue(const FHashedMaterialParam
 	return false;
 }
 
-bool UMaterialInterface::GetEnumerationParameterDefaultValue(const FHashedMaterialParameterInfo& ParameterInfo, FName& OutValue) const
-{
-	FMaterialParameterMetadata Result;
-	if (GetParameterDefaultValue(EMaterialParameterType::Enumeration, ParameterInfo, Result))
-	{
-		OutValue = Result.Value.EnumEntry;
-		return true;
-	}
-	return false;
-}
-
 #if WITH_EDITOR
 bool UMaterialInterface::GetStaticSwitchParameterDefaultValue(const FHashedMaterialParameterInfo& ParameterInfo, bool& OutValue, FGuid& OutExpressionGuid) const
 {
@@ -1277,11 +1255,6 @@ void UMaterialInterface::GetAllFontParameterInfo(TArray<FMaterialParameterInfo>&
 	GetAllParameterInfoOfType(EMaterialParameterType::Font, OutParameterInfo, OutParameterIds);
 }
 
-void UMaterialInterface::GetAllEnumerationParameterInfo(TArray<FMaterialParameterInfo>& OutParameterInfo, TArray<FGuid>& OutParameterIds) const
-{
-	GetAllParameterInfoOfType(EMaterialParameterType::Enumeration, OutParameterInfo, OutParameterIds);
-}
-
 #if WITH_EDITORONLY_DATA
 void UMaterialInterface::GetAllStaticSwitchParameterInfo(TArray<FMaterialParameterInfo>&OutParameterInfo, TArray<FGuid>&OutParameterIds) const
 {
diff --git a/Engine/Source/Runtime/Engine/Private/Materials/MaterialRenderProxy.cpp b/Engine/Source/Runtime/Engine/Private/Materials/MaterialRenderProxy.cpp
--- a/Engine/Source/Runtime/Engine/Private/Materials/MaterialRenderProxy.cpp
+++ b/Engine/Source/Runtime/Engine/Private/Materials/MaterialRenderProxy.cpp
@@ -1,7 +1,6 @@
 // Copyright Epic Games, Inc. All Rights Reserved.
 
 #include "Materials/MaterialRenderProxy.h"
-#include "Materials/MaterialEnumeration.h"
 #include "MaterialCachedData.h"
 #include "Engine/Texture.h"
 #include "EngineModule.h"
@@ -273,18 +272,7 @@ bool FMaterialRenderProxy::GetParameterShaderValue(EMaterialParameterType Type,
 	FMaterialParameterValue Value;
 	if (GetParameterValue(Type, ParameterInfo, Value, Context))
 	{
-		if (Type == EMaterialParameterType::Enumeration)
-		{
-			FMaterialParameterMetadata ParameterMetadata;
-			bool bResult = GetMaterialInterface()->GetCachedExpressionData().GetParameterValue(Type, ParameterInfo, ParameterMetadata);
-			check(bResult);
-
-			OutValue = Cast<IMaterialEnumerationProvider>(ParameterMetadata.Enumeration)->GetValueOrDefault(Value.EnumEntry);
-		}
-		else
-		{
-			OutValue = Value.AsShaderValue();
-		}
+		OutValue = Value.AsShaderValue();
 		return true;
 	}
 	return false;
diff --git a/Engine/Source/Runtime/Engine/Private/Materials/MaterialShared.cpp b/Engine/Source/Runtime/Engine/Private/Materials/MaterialShared.cpp
--- a/Engine/Source/Runtime/Engine/Private/Materials/MaterialShared.cpp
+++ b/Engine/Source/Runtime/Engine/Private/Materials/MaterialShared.cpp
@@ -5607,8 +5607,6 @@ UE::Shader::FType GetShaderValueType(EMaterialParameterType Type)
 		return FMaterialTextureValue::GetTypeName();
 	case EMaterialParameterType::SparseVolumeTexture:
 		return UE::Shader::EValueType::Void; // TODO
-	case EMaterialParameterType::Enumeration:
-		return UE::Shader::EValueType::Int1;
 	default:
 		checkNoEntry();
 		return UE::Shader::EValueType::Void;
diff --git a/Engine/Source/Runtime/Engine/Public/MaterialCachedData.h b/Engine/Source/Runtime/Engine/Public/MaterialCachedData.h
--- a/Engine/Source/Runtime/Engine/Public/MaterialCachedData.h
+++ b/Engine/Source/Runtime/Engine/Public/MaterialCachedData.h
@@ -327,12 +327,6 @@ struct FMaterialCachedExpressionData
 	UPROPERTY()
 	TArray<int32> FontPageValues;
 	
-	UPROPERTY()
-	TArray<TSoftObjectPtr<UObject>> Enumerations;
-
-	UPROPERTY()
-	TArray<FName> EnumerationEntryValues;
-
 	UPROPERTY()
 	TArray<TSoftObjectPtr<URuntimeVirtualTexture>> RuntimeVirtualTextureValues;
 
diff --git a/Engine/Source/Runtime/Engine/Public/Materials/MaterialEnumeration.h b/Engine/Source/Runtime/Engine/Public/Materials/MaterialEnumeration.h
--- a/Engine/Source/Runtime/Engine/Public/Materials/MaterialEnumeration.h
+++ b/Engine/Source/Runtime/Engine/Public/Materials/MaterialEnumeration.h
@@ -28,45 +28,3 @@ class UMaterialEnumerationProvider : public UInterface
 {
 	GENERATED_BODY()
 };
-
-USTRUCT(MinimalAPI)
-struct FMaterialEnumerationEntry
-{
-	GENERATED_BODY()
-
-	// Entry identifier.
-	UPROPERTY(EditAnywhere, Category=MaterialEnumerationEntry)
-	FName Name{};
-
-	// Associated integer value.
-	UPROPERTY(EditAnywhere, Category=MaterialEnumerationEntry)
-	int32 Value{};
-};
-
-UCLASS(MinimalAPI)
-class UMaterialEnumeration : public UDataAsset, public IMaterialEnumerationProvider
-{
-	GENERATED_BODY()
-
-public:
-	// Raw list of entries.
-	UPROPERTY(EditAnywhere, Category=MaterialEnumeration)
-	TArray<FMaterialEnumerationEntry> Entries;
-
-	// Cached and sorted version of Entries for faster lookup.
-	TArray<FMaterialEnumerationEntry> SortedEntries;
-
-	// IMaterialEnumerationProvider implementation:
-	virtual bool ResolveValue(FName EntryName, int32& OutValue, int32 DefaultValue = 0) const final;
-	virtual void ForEachEntry(TFunctionRef<void (FName Name, int32 Value)> Iterator) const final;
-
-private:
-	// Sorts entries by name for optimized access.
-	void SortEntries();
-
-	virtual void PostLoad() override;
-
-#if WITH_EDITOR
-	virtual void PostEditChangeProperty(struct FPropertyChangedEvent& PropertyChangedEvent) override;
-#endif
-};
diff --git a/Engine/Source/Runtime/Engine/Public/Materials/MaterialExpressionEnumerationParameter.h b/Engine/Source/Runtime/Engine/Public/Materials/MaterialExpressionEnumerationParameter.h
deleted file mode 100644
diff --git a/Engine/Source/Runtime/Engine/Public/Materials/MaterialFunctionInstance.h b/Engine/Source/Runtime/Engine/Public/Materials/MaterialFunctionInstance.h
--- a/Engine/Source/Runtime/Engine/Public/Materials/MaterialFunctionInstance.h
+++ b/Engine/Source/Runtime/Engine/Public/Materials/MaterialFunctionInstance.h
@@ -73,10 +73,6 @@ class UMaterialFunctionInstance : public UMaterialFunctionInterface
 	UPROPERTY(EditAnywhere, Category = MaterialFunctionInstance)
 	TArray<struct FSparseVolumeTextureParameterValue> SparseVolumeTextureParameterValues;
 
-	/** Enumeration parameters. */
-	UPROPERTY(EditAnywhere, BlueprintReadOnly, Category = MaterialInstance, meta = (EditFixedOrder))
-	TArray<struct FEnumerationParameterValue> EnumerationParameterValues;
-
 #if WITH_EDITOR
 	ENGINE_API void UpdateParameterSet();
 	ENGINE_API void OverrideMaterialInstanceParameterValues(class UMaterialInstance* Instance);
diff --git a/Engine/Source/Runtime/Engine/Public/Materials/MaterialInstance.h b/Engine/Source/Runtime/Engine/Public/Materials/MaterialInstance.h
--- a/Engine/Source/Runtime/Engine/Public/Materials/MaterialInstance.h
+++ b/Engine/Source/Runtime/Engine/Public/Materials/MaterialInstance.h
@@ -524,54 +524,6 @@ struct FFontParameterValue
 	}
 };
 
-/** Editable enumeration parameter. */
-USTRUCT(BlueprintType)
-struct FEnumerationParameterValue
-{
-	GENERATED_USTRUCT_BODY()
-
-	UPROPERTY(EditAnywhere, BlueprintReadWrite, Category=EnumerationParameterValue)
-	FMaterialParameterInfo ParameterInfo;
-
-	UPROPERTY(EditAnywhere, BlueprintReadWrite, Category=EnumerationParameterValue)
-	FName EnumEntry;
-
-	UPROPERTY()
-	FGuid ExpressionGUID;
-
-	explicit FEnumerationParameterValue(const FMaterialParameterInfo& InParameterInfo = FMaterialParameterInfo(), FName InEnumEntry = {})
-		: ParameterInfo(InParameterInfo), EnumEntry(InEnumEntry)
-	{
-	}
-
-	bool IsOverride() const { return true; }
-
-	bool IsValid() const { return GetValue(*this).IsValid(); }
-
-	bool operator==(const FEnumerationParameterValue& Other) const
-	{
-		return
-			ParameterInfo == Other.ParameterInfo &&
-			EnumEntry == Other.EnumEntry &&
-			ExpressionGUID == Other.ExpressionGUID;
-	}
-	bool operator!=(const FEnumerationParameterValue& Other) const
-	{
-		return !((*this) == Other);
-	}
-	
-	typedef FName ValueType;
-	static ValueType GetValue(const FEnumerationParameterValue& Parameter) { return Parameter.EnumEntry; }
-
-	void GetValue(FMaterialParameterMetadata& OutResult) const
-	{
-		OutResult.Value = FMaterialParameterValue(EnumEntry);
-#if WITH_EDITORONLY_DATA
-		OutResult.ExpressionGuid = ExpressionGUID;
-#endif
-	}
-};
-
 USTRUCT(BlueprintType)
 struct FUserSceneTextureOverride
 {
@@ -829,10 +781,6 @@ public:
 	UPROPERTY(EditAnywhere, BlueprintReadOnly, Category=MaterialInstance, meta = (EditFixedOrder))
 	TArray<struct FFontParameterValue> FontParameterValues;
 	
-	/** Enumeration parameters. */
-	UPROPERTY(EditAnywhere, BlueprintReadOnly, Category = MaterialInstance, meta = (EditFixedOrder))
-	TArray<FEnumerationParameterValue> EnumerationParameterValues;
-
 	/** User scene texture overrides.  Applies to post process domain materials only. */
 	UPROPERTY(EditAnywhere, Category = MaterialInstance)
 	TArray<struct FUserSceneTextureOverride> UserSceneTextureOverrides;
@@ -1264,7 +1212,6 @@ protected:
 	void SetRuntimeVirtualTextureParameterValueInternal(const FMaterialParameterInfo& ParameterInfo, class URuntimeVirtualTexture* Value);
 	void SetSparseVolumeTextureParameterValueInternal(const FMaterialParameterInfo& ParameterInfo, class USparseVolumeTexture* Value);
 	void SetFontParameterValueInternal(const FMaterialParameterInfo& ParameterInfo, class UFont* FontValue, int32 FontPage);
-	void SetEnumerationParameterValueInternal(const FMaterialParameterInfo& ParameterInfo, FName Value);
 	void ClearParameterValuesInternal(EMaterialInstanceClearParameterFlag Flags = EMaterialInstanceClearParameterFlag::AllNonStatic);
 
 	/** Initialize the material instance's resources. */
diff --git a/Engine/Source/Runtime/Engine/Public/Materials/MaterialInstanceConstant.h b/Engine/Source/Runtime/Engine/Public/Materials/MaterialInstanceConstant.h
--- a/Engine/Source/Runtime/Engine/Public/Materials/MaterialInstanceConstant.h
+++ b/Engine/Source/Runtime/Engine/Public/Materials/MaterialInstanceConstant.h
@@ -100,7 +100,6 @@ class UMaterialInstanceConstant : public UMaterialInstance
 	ENGINE_API void SetRuntimeVirtualTextureParameterValueEditorOnly(const FMaterialParameterInfo& ParameterInfo, class URuntimeVirtualTexture* Value);
 	ENGINE_API void SetSparseVolumeTextureParameterValueEditorOnly(const FMaterialParameterInfo& ParameterInfo, class USparseVolumeTexture* Value);
 	ENGINE_API void SetFontParameterValueEditorOnly(const FMaterialParameterInfo& ParameterInfo, class UFont* FontValue, int32 FontPage);
-	ENGINE_API void SetEnumerationParameterValueEditorOnly(const FMaterialParameterInfo& ParameterInfo, FName Value);
 
 	/**
 	 * Clear all parameter overrides on this material instance. This function
diff --git a/Engine/Source/Runtime/Engine/Public/Materials/MaterialInterface.h b/Engine/Source/Runtime/Engine/Public/Materials/MaterialInterface.h
--- a/Engine/Source/Runtime/Engine/Public/Materials/MaterialInterface.h
+++ b/Engine/Source/Runtime/Engine/Public/Materials/MaterialInterface.h
@@ -726,7 +726,6 @@ public:
 	ENGINE_API void GetAllRuntimeVirtualTextureParameterInfo(TArray<FMaterialParameterInfo>& OutParameterInfo, TArray<FGuid>& OutParameterIds) const;
 	ENGINE_API void GetAllSparseVolumeTextureParameterInfo(TArray<FMaterialParameterInfo>& OutParameterInfo, TArray<FGuid>& OutParameterIds) const;
 	ENGINE_API void GetAllFontParameterInfo(TArray<FMaterialParameterInfo>& OutParameterInfo, TArray<FGuid>& OutParameterIds) const;
-	ENGINE_API void GetAllEnumerationParameterInfo(TArray<FMaterialParameterInfo>& OutParameterInfo, TArray<FGuid>& OutParameterIds) const;
 
 #if WITH_EDITORONLY_DATA
 	ENGINE_API void GetAllStaticSwitchParameterInfo(TArray<FMaterialParameterInfo>& OutParameterInfo, TArray<FGuid>& OutParameterIds) const;
@@ -748,7 +747,6 @@ public:
 	ENGINE_API bool GetRuntimeVirtualTextureParameterDefaultValue(const FHashedMaterialParameterInfo& ParameterInfo, class URuntimeVirtualTexture*& OutValue) const;
 	ENGINE_API bool GetSparseVolumeTextureParameterDefaultValue(const FHashedMaterialParameterInfo& ParameterInfo, class USparseVolumeTexture*& OutValue) const;
 	ENGINE_API bool GetFontParameterDefaultValue(const FHashedMaterialParameterInfo& ParameterInfo, class UFont*& OutFontValue, int32& OutFontPage) const;
-	ENGINE_API bool GetEnumerationParameterDefaultValue(const FHashedMaterialParameterInfo& ParameterInfo, FName& OutValue) const;
 	
 #if WITH_EDITOR
 	ENGINE_API bool GetStaticSwitchParameterDefaultValue(const FHashedMaterialParameterInfo& ParameterInfo, bool& OutValue, FGuid& OutExpressionGuid) const;
@@ -990,7 +988,6 @@ public:
 	ENGINE_API bool GetTextureParameterChannelNames(const FHashedMaterialParameterInfo& ParameterInfo, FParameterChannelNames& OutValue) const;
 #endif
 	ENGINE_API bool GetFontParameterValue(const FHashedMaterialParameterInfo& ParameterInfo,class UFont*& OutFontValue, int32& OutFontPage, bool bOveriddenOnly = false) const;
-	ENGINE_API bool GetEnumerationParameterValue(const FHashedMaterialParameterInfo& ParameterInfo, FName& OutValue, bool bOveriddenOnly = false) const;
 
 	ENGINE_API virtual bool GetRefractionSettings(float& OutBiasValue) const;
 	ENGINE_API virtual bool GetUserSceneTextureOverride(FName& InOutValue) const;
diff --git a/Engine/Source/Runtime/Engine/Public/Materials/MaterialParameters.h b/Engine/Source/Runtime/Engine/Public/Materials/MaterialParameters.h
--- a/Engine/Source/Runtime/Engine/Public/Materials/MaterialParameters.h
+++ b/Engine/Source/Runtime/Engine/Public/Materials/MaterialParameters.h
@@ -195,7 +195,6 @@ enum class EMaterialParameterType : uint8
 	SparseVolumeTexture,
 	StaticSwitch,
 	ParameterCollection,
-	Enumeration,
 
 	NumRuntime, // Runtime parameter types must go above here, and editor-only ones below
 
@@ -218,7 +217,6 @@ inline bool IsNumericMaterialParameter(EMaterialParameterType InType)
 	case EMaterialParameterType::Scalar:
 	case EMaterialParameterType::Vector:
 	case EMaterialParameterType::DoubleVector:
-	case EMaterialParameterType::Enumeration:
 		return true;
 	default:
 		return false;
@@ -371,7 +369,6 @@ struct FMaterialParameterValue
 	FMaterialParameterValue(bool InValue) : Type(EMaterialParameterType::StaticSwitch) { Bool[0] = InValue; }
 	FMaterialParameterValue(const FStaticComponentMaskValue& InValue) : Type(EMaterialParameterType::StaticComponentMask) { Bool[0] = InValue.R; Bool[1] = InValue.G; Bool[2] = InValue.B; Bool[3] = InValue.A; }
 	FMaterialParameterValue(bool bMaskR, bool bMaskG, bool bMaskB, bool bMaskA) : Type(EMaterialParameterType::StaticComponentMask) { Bool[0] = bMaskR; Bool[1] = bMaskG; Bool[2] = bMaskB; Bool[3] = bMaskA; }
-	FMaterialParameterValue(FName InEnumEntry) : Type(EMaterialParameterType::Enumeration) { EnumEntry = InEnumEntry; }
 	ENGINE_API FMaterialParameterValue(EMaterialParameterType Type, const UE::Shader::FValue& InValue);
 
 	inline float AsScalar() const { check(Type == EMaterialParameterType::Scalar); return Float[0]; }
@@ -434,7 +431,6 @@ struct FMaterialParameterValue
 			Lhs.Bool[1] == Rhs.Bool[1] &&
 			Lhs.Bool[2] == Rhs.Bool[2] &&
 			Lhs.Bool[3] == Rhs.Bool[3];
-		case EMaterialParameterType::Enumeration: return Lhs.EnumEntry == Rhs.EnumEntry;
 		default: checkNoEntry(); return false;
 		}
 	}
@@ -498,10 +494,6 @@ struct FMaterialParameterMetadata
 	bool bOverride = false;
 #endif // WITH_EDITORONLY_DATA
 
-	/** The enumeration definition, valid when this is an enumeration parameter.
-	  * Note: this UObject must implement the IMaterialEnumerationProvider interface */
-	TObjectPtr<UObject> Enumeration;
-
 	/** Is a Static Switch parameter using a dynamic branch? */
 	bool bDynamicSwitchParameter = false;
 };
