Fix bug in operator Select handling in new material translator causing wrong order of rows and columns.

#rnx

diff --git a/Engine/Source/Runtime/Engine/Private/Materials/MaterialIREmitter.cpp b/Engine/Source/Runtime/Engine/Private/Materials/MaterialIREmitter.cpp
--- a/Engine/Source/Runtime/Engine/Private/Materials/MaterialIREmitter.cpp
+++ b/Engine/Source/Runtime/Engine/Private/Materials/MaterialIREmitter.cpp
@@ -2160,12 +2160,12 @@ static FType ValidateOperatorAndGetResultType(FEmitter& Emitter, EOperator Op, F
 		int32 MaxNumComponents = FMath::Max3(FirstArgumentType.GetPrimitive().NumComponents(), BPrimitiveType.NumComponents(), CPrimitiveType.NumComponents());
 
 		// Cast the first argument (the boolean condition) to a bool vector of the maximum number of components.
-		Arguments[0] = Emitter.Cast(Arguments[0], FType::MakePrimitive(EScalarKind::Bool, MaxNumComponents, 1));
+		Arguments[0] = Emitter.Cast(Arguments[0], FType::MakeVector(EScalarKind::Bool, MaxNumComponents));
 
 		// Compute the common type between the second and third argument types with a number of components equal to the max of all three.
 		FType CommonTypeBetweenSecondAndThirdArguments = Emitter.TryGetCommonType(
-			FType::MakePrimitive(BPrimitiveType.ScalarKind, MaxNumComponents, 1),
-			FType::MakePrimitive(CPrimitiveType.ScalarKind, MaxNumComponents, 1));
+			FType::MakeVector(BPrimitiveType.ScalarKind, MaxNumComponents),
+			FType::MakeVector(CPrimitiveType.ScalarKind, MaxNumComponents));
 		
 		// Getting the common type should always be possible.
 		check(CommonTypeBetweenSecondAndThirdArguments);
