[StaticMesh] Remove warning about static mesh render data bounds being significantly different from the cached mesh description bounds, as it can too often be triggered due to LOD (both manual and simplified)

diff --git a/Engine/Plugins/Experimental/MeshModelingToolsetExp/Source/GeometryProcessingAdapters/Private/ApproximateActorsImpl.cpp b/Engine/Plugins/Experimental/MeshModelingToolsetExp/Source/GeometryProcessingAdapters/Private/ApproximateActorsImpl.cpp
--- a/Engine/Plugins/Experimental/MeshModelingToolsetExp/Source/GeometryProcessingAdapters/Private/ApproximateActorsImpl.cpp
+++ b/Engine/Plugins/Experimental/MeshModelingToolsetExp/Source/GeometryProcessingAdapters/Private/ApproximateActorsImpl.cpp
@@ -1431,11 +1431,6 @@ UStaticMesh* FApproximateActorsImpl::EmitGeneratedMeshAsset(
 	BuildParameters.bInRebuildUVChannelData = true;
 	BuildParameters.bInEnforceLightmapRestrictions = true;
 
-	// Nanite meshes can sometime contain smaller triangles that could be discarded in the proxy mesh.
-	// If those are found on the outskirt of a merged mesh, the static mesh build will complain about
-	// the bounds delta being too high.
-	BuildParameters.bIgnoreBoundsDiff = Options.bGenerateNaniteEnabledMesh; 
-
 	MeshAssetOutputs.StaticMesh->Build(BuildParameters);
 
 	ResultsOut.NewMeshAssets.Add(MeshAssetOutputs.StaticMesh);
diff --git a/Engine/Source/Developer/MeshMergeUtilities/Private/MeshMergeUtilities.cpp b/Engine/Source/Developer/MeshMergeUtilities/Private/MeshMergeUtilities.cpp
--- a/Engine/Source/Developer/MeshMergeUtilities/Private/MeshMergeUtilities.cpp
+++ b/Engine/Source/Developer/MeshMergeUtilities/Private/MeshMergeUtilities.cpp
@@ -2742,11 +2742,6 @@ void FMeshMergeUtilities::MergeComponentsToStaticMesh(const TArray<UPrimitiveCom
 		BuildParameters.bInSilent = true;
 		BuildParameters.bInRebuildUVChannelData = true;
 		BuildParameters.bInEnforceLightmapRestrictions = true;
-		
-		// Nanite meshes can sometime contain smaller triangles that could be discarded in the proxy mesh.
-		// If those are found on the outskirt of a merged mesh, the static mesh build will complain about
-		// the bounds delta being too high.
-		BuildParameters.bIgnoreBoundsDiff = StaticMesh->GetNaniteSettings().bEnabled;
 
 		StaticMesh->Build(BuildParameters);
 
diff --git a/Engine/Source/Developer/MeshMergeUtilities/Private/ProxyGenerationProcessor.cpp b/Engine/Source/Developer/MeshMergeUtilities/Private/ProxyGenerationProcessor.cpp
--- a/Engine/Source/Developer/MeshMergeUtilities/Private/ProxyGenerationProcessor.cpp
+++ b/Engine/Source/Developer/MeshMergeUtilities/Private/ProxyGenerationProcessor.cpp
@@ -365,11 +365,6 @@ void FProxyGenerationProcessor::ProcessJob(const FGuid& JobGuid, FProxyGeneratio
 	BuildParameters.bInRebuildUVChannelData = true;
 	BuildParameters.bInEnforceLightmapRestrictions = true;
 
-	// Nanite meshes can sometime contain smaller triangles that could be discarded in the proxy mesh.
-	// If those are found on the outskirt of a merged mesh, the static mesh build will complain about
-	// the bounds delta being too high.
-	BuildParameters.bIgnoreBoundsDiff = StaticMesh->GetNaniteSettings().bEnabled;
-
 	StaticMesh->Build(BuildParameters);
 
 	OutAssetsToSync.Add(StaticMesh);
diff --git a/Engine/Source/Runtime/Engine/Classes/Engine/StaticMesh.h b/Engine/Source/Runtime/Engine/Classes/Engine/StaticMesh.h
--- a/Engine/Source/Runtime/Engine/Classes/Engine/StaticMesh.h
+++ b/Engine/Source/Runtime/Engine/Classes/Engine/StaticMesh.h
@@ -2014,11 +2014,27 @@ public:
 			, OutErrors(nullptr)
 			, bInRebuildUVChannelData(false)
 			, bInEnforceLightmapRestrictions(false)
-			, bIgnoreBoundsDiff(false)
 		{
+			PRAGMA_DISABLE_DEPRECATION_WARNINGS
+			bIgnoreBoundsDiff = false;
+			PRAGMA_ENABLE_DEPRECATION_WARNINGS
 		}
+#else
+		FBuildParameters() = default;
 #endif
 
+		FBuildParameters(const FBuildParameters& Other)
+		{
+			bInSilent = Other.bInSilent;
+			OutErrors = Other.OutErrors;
+			bInRebuildUVChannelData = Other.bInRebuildUVChannelData;
+			bInEnforceLightmapRestrictions = Other.bInEnforceLightmapRestrictions;
+
+			PRAGMA_DISABLE_DEPRECATION_WARNINGS
+			bIgnoreBoundsDiff = Other.bIgnoreBoundsDiff;
+			PRAGMA_ENABLE_DEPRECATION_WARNINGS
+		}
+
 		// If true will not popup a progress dialog.
 		bool bInSilent = false;
 
@@ -2032,6 +2048,7 @@ public:
 		bool bInEnforceLightmapRestrictions = false;
 
 		// Whether to avoid outputting a warning if the render data and MeshDescription bounds differ too much.
+		UE_DEPRECATED(5.7, "This option no longer has any effect and will be removed in a future release.")
 		bool bIgnoreBoundsDiff = false;
 	};
 
@@ -2401,7 +2418,7 @@ private:
 	/**
 	 * Complete the static mesh building process - Can't be done in parallel.
 	 */
-	void FinishBuildInternal(const TArray<IStaticMeshComponent*>& InAffectedComponents, bool bHasRenderDataChanged, bool bIgnoreBoundsDiff, bool bShouldComputeExtendedBounds = true);
+	void FinishBuildInternal(const TArray<IStaticMeshComponent*>& InAffectedComponents, bool bHasRenderDataChanged, bool bShouldComputeExtendedBounds = true);
 
 #if WITH_EDITORONLY_DATA
 	/**
diff --git a/Engine/Source/Runtime/Engine/Private/StaticMeshBuild.cpp b/Engine/Source/Runtime/Engine/Private/StaticMeshBuild.cpp
--- a/Engine/Source/Runtime/Engine/Private/StaticMeshBuild.cpp
+++ b/Engine/Source/Runtime/Engine/Private/StaticMeshBuild.cpp
@@ -345,7 +345,7 @@ void UStaticMesh::BatchBuild(const TArray<UStaticMesh*>& InStaticMeshes, const F
 				// Wait the result of the async task
 				const bool bHasRenderDataChanged = Task.Future.Get();
 
-				Task.StaticMesh->FinishBuildInternal(StaticMeshComponents.FindChecked(Task.StaticMesh), bHasRenderDataChanged, BuildParameters.bIgnoreBoundsDiff);
+				Task.StaticMesh->FinishBuildInternal(StaticMeshComponents.FindChecked(Task.StaticMesh), bHasRenderDataChanged);
 				FinalizeStaticMesh(Task.StaticMesh);
 			}
 		}
@@ -366,7 +366,7 @@ void UStaticMesh::BatchBuild(const TArray<UStaticMesh*>& InStaticMeshes, const F
 
 					const bool bHasRenderDataChanged = StaticMesh->ExecuteBuildInternal(BuildParameters);
 
-					StaticMesh->FinishBuildInternal(StaticMeshComponents.FindChecked(StaticMesh), bHasRenderDataChanged, BuildParameters.bIgnoreBoundsDiff);
+					StaticMesh->FinishBuildInternal(StaticMeshComponents.FindChecked(StaticMesh), bHasRenderDataChanged);
 					FinalizeStaticMesh(StaticMesh);
 				}
 			}
@@ -650,7 +650,7 @@ bool UStaticMesh::ExecuteBuildInternal(const FBuildParameters& BuildParameters)
 	return bHasRenderDataChanged;
 }
 
-void UStaticMesh::FinishBuildInternal(const TArray<IStaticMeshComponent*>& InAffectedComponents, bool bHasRenderDataChanged, bool bIgnoreBoundsDiff, bool bShouldComputeExtendedBounds)
+void UStaticMesh::FinishBuildInternal(const TArray<IStaticMeshComponent*>& InAffectedComponents, bool bHasRenderDataChanged, bool bShouldComputeExtendedBounds)
 {
 	TRACE_CPUPROFILER_EVENT_SCOPE(UStaticMesh::PostBuildInternal);
 
@@ -674,21 +674,7 @@ void UStaticMesh::FinishBuildInternal(const TArray<IStaticMeshComponent*>& InAff
 		// Calculate extended bounds
 		CalculateExtendedBounds();
 	}
-	else if (!IsNaniteLandscape() && !IsNaniteAssembly() && !bIgnoreBoundsDiff)
-	{
-		// We don't care about minor differences but might want to highlight if big differences are noted.
-		const FBoxSphereBounds RenderBounds = GetRenderData()->Bounds;
-		const FBoxSphereBounds CachedBounds = CachedMeshDescriptionBounds.GetValue();
-
-		const float SizeDifferencePercent   = 100.f * 2.0f * FMath::Abs(CachedBounds.BoxExtent.Size() - RenderBounds.BoxExtent.Size()) / (FMath::Abs(CachedBounds.BoxExtent.Size() + RenderBounds.BoxExtent.Size()) + UE_KINDA_SMALL_NUMBER);
-		const float OriginDifferencePercent = 100.f * 2.0f * FVector::Dist(CachedBounds.Origin, RenderBounds.Origin) / ((CachedBounds.Origin + RenderBounds.Origin).Size() + UE_KINDA_SMALL_NUMBER);
-		// Anything more than 5% is probably worth investigating
-		if (SizeDifferencePercent > 5.0f || OriginDifferencePercent > 5.0f)
-		{	
-			UE_LOG(LogStaticMesh, Warning, TEXT("The difference between RenderData Bounds (%s) and MeshDescription Bounds (%s) is significative for %s"), *RenderBounds.ToString(), *CachedBounds.ToString(), *GetFullName());
-		}
-	}
-
+	
 	// Async build is finished, we can unlock everything now.
 	ReleaseAsyncProperty();
 
diff --git a/Engine/Source/Runtime/Engine/Private/StaticMeshCompiler.cpp b/Engine/Source/Runtime/Engine/Private/StaticMeshCompiler.cpp
--- a/Engine/Source/Runtime/Engine/Private/StaticMeshCompiler.cpp
+++ b/Engine/Source/Runtime/Engine/Private/StaticMeshCompiler.cpp
@@ -402,7 +402,6 @@ void FStaticMeshCompilingManager::PostCompilation(UStaticMesh* StaticMesh)
 				StaticMesh->FinishBuildInternal(
 					ComponentsToUpdate,
 					LocalAsyncTask->GetTask().BuildContext->bHasRenderDataChanged,
-					LocalAsyncTask->GetTask().BuildContext->BuildParameters.bIgnoreBoundsDiff,
 					LocalAsyncTask->GetTask().BuildContext->bShouldComputeExtendedBounds
 				);
 
