Ensure valid skeletal mesh LOD index is used for rendering to fix invisible meshes

diff --git a/Engine/Source/Runtime/Engine/Private/SkeletalMeshSceneProxy.cpp b/Engine/Source/Runtime/Engine/Private/SkeletalMeshSceneProxy.cpp
--- a/Engine/Source/Runtime/Engine/Private/SkeletalMeshSceneProxy.cpp
+++ b/Engine/Source/Runtime/Engine/Private/SkeletalMeshSceneProxy.cpp
@@ -758,7 +758,7 @@ uint8 FSkeletalMeshSceneProxy::GetCurrentFirstLODIdx_Internal() const
 
 FDesiredLODLevel FSkeletalMeshSceneProxy::GetDesiredLODLevel_RenderThread(const FSceneView* View) const
 {
-	return FDesiredLODLevel::CreateFixed(MeshObject->GetLOD());
+	return FDesiredLODLevel::CreateFixed(FMath::Max(ClampedLODIndex, MeshObject->GetLOD()));
 }
 
 bool FSkeletalMeshSceneProxy::GetCachedGeometry(FRDGBuilder& GraphBuilder, FCachedGeometry& OutCachedGeometry) const 
