UE-186111 - Push-model blueprint array property not properly marked dirty if modified by different actor

The root cause of the issue is due to incorrect function signatures on array function thunks. Several of them mark their input array as 'const', but clearly change the array. When push-model support was added, this was problem was discovered and worked around by adding MARK_PROPERTY_DIRTY calls directly in the thunks. The problem with this solution is that it doesn't work when we change an array on a different actor. Additionally, the thunk doesn't have enough information to determine that easily.

The correct solution would be to leverage the existing code in FKCHandler_CallFunction that has sufficient context to inject MarkPropertyDirtyNode accurately. In order to get array functions working, we need to fix their signatures to use UPARAM(ref) instead of 'const'.

One issue with this fix is that there might be existing Blueprint code that incorrectly uses mutable array functions with read-only data. For now, we'll allow this, but with a warning to fix the code as it'll become an error in 5.9.

diff --git a/Engine/Config/BaseEditor.ini b/Engine/Config/BaseEditor.ini
--- a/Engine/Config/BaseEditor.ini
+++ b/Engine/Config/BaseEditor.ini
@@ -189,6 +189,26 @@ MenuPriority=10
 [/Script/BlueprintGraph.EdGraphSchema_K2]
 +EditoronlyBPFunctionRedirects=(ClassName="GameplayStatics", OldFunctionName="CreateSaveGameObjectFromBlueprint", NewFunctionName="CreateSaveGameObject", BlueprintParamName="SaveGameBlueprint", ClassParamName="SaveGameClass")
 
+; REMOVE IN UE 5.9
+; Many of the Blueprint container functions are incorrectly using 'const' for inputs that should actually be non-const.
+; However, changing these function signatures could result in Blueprint compilation errors if anyone
+; mistakenly connected a read-only container to one of these functions.
+; 'AllowedMutableContainerFunctions' allows us to continue using these read-only variables incorrectly,
+; but issues a warning instead of an error. Eventually, this should be removed as part of the deprecation process.
++AllowedMutableContainerFunctions=((Function="/Script/Engine.KismetArrayLibrary:Array_Add", ContainerParameterName="TargetArray"))
++AllowedMutableContainerFunctions=((Function="/Script/Engine.KismetArrayLibrary:Array_AddUnique", ContainerParameterName="TargetArray"))
++AllowedMutableContainerFunctions=((Function="/Script/Engine.KismetArrayLibrary:Array_Shuffle", ContainerParameterName="TargetArray"))
++AllowedMutableContainerFunctions=((Function="/Script/Engine.KismetArrayLibrary:Array_ShuffleFromStream", ContainerParameterName="TargetArray"))
++AllowedMutableContainerFunctions=((Function="/Script/Engine.KismetArrayLibrary:Array_Append", ContainerParameterName="TargetArray"))
++AllowedMutableContainerFunctions=((Function="/Script/Engine.KismetArrayLibrary:Array_Insert", ContainerParameterName="TargetArray"))
++AllowedMutableContainerFunctions=((Function="/Script/Engine.KismetArrayLibrary:Array_Remove", ContainerParameterName="TargetArray"))
++AllowedMutableContainerFunctions=((Function="/Script/Engine.KismetArrayLibrary:Array_RemoveItem", ContainerParameterName="TargetArray"))
++AllowedMutableContainerFunctions=((Function="/Script/Engine.KismetArrayLibrary:Array_Clear", ContainerParameterName="TargetArray"))
++AllowedMutableContainerFunctions=((Function="/Script/Engine.KismetArrayLibrary:Array_Resize", ContainerParameterName="TargetArray"))
++AllowedMutableContainerFunctions=((Function="/Script/Engine.KismetArrayLibrary:Array_Reverse", ContainerParameterName="TargetArray"))
++AllowedMutableContainerFunctions=((Function="/Script/Engine.KismetArrayLibrary:Array_Set", ContainerParameterName="TargetArray"))
++AllowedMutableContainerFunctions=((Function="/Script/Engine.KismetArrayLibrary:Array_Swap", ContainerParameterName="TargetArray"))
+
 [BlueprintEditor.Palette]
 bUseLegacyLayout=false
 
diff --git a/Engine/Source/Editor/BlueprintGraph/Classes/EdGraphSchema_K2.h b/Engine/Source/Editor/BlueprintGraph/Classes/EdGraphSchema_K2.h
--- a/Engine/Source/Editor/BlueprintGraph/Classes/EdGraphSchema_K2.h
+++ b/Engine/Source/Editor/BlueprintGraph/Classes/EdGraphSchema_K2.h
@@ -335,6 +335,28 @@ struct FBlueprintCallableFunctionRedirect
 	FString ClassParamName;
 };
 
+USTRUCT()
+struct FAllowedMutableContainerFunction
+{
+	GENERATED_BODY()
+
+	UPROPERTY()
+	TObjectPtr<UFunction> Function;
+
+	UPROPERTY()
+	FString ContainerParameterName;
+
+	friend uint32 GetTypeHash(const FAllowedMutableContainerFunction& Input)
+	{
+		return GetTypeHash(Input.Function);
+	}
+
+	bool operator==(const FAllowedMutableContainerFunction& RHS) const
+	{
+		return (Function == RHS.Function);
+	}
+};
+
 enum class EObjectReferenceType : uint8
 {
 	NotAnObject		= 0x00,
@@ -446,6 +468,9 @@ class UEdGraphSchema_K2 : public UEdGraphSchema
 	UPROPERTY(globalconfig)
 	TArray<FBlueprintCallableFunctionRedirect> EditoronlyBPFunctionRedirects;
 
+	UPROPERTY(globalconfig)
+	TSet<FAllowedMutableContainerFunction> AllowedMutableContainerFunctions;
+
 public:
 
 	//////////////////////////////////////////////////////////////////////////
diff --git a/Engine/Source/Editor/BlueprintGraph/Private/CallFunctionHandler.cpp b/Engine/Source/Editor/BlueprintGraph/Private/CallFunctionHandler.cpp
--- a/Engine/Source/Editor/BlueprintGraph/Private/CallFunctionHandler.cpp
+++ b/Engine/Source/Editor/BlueprintGraph/Private/CallFunctionHandler.cpp
@@ -38,13 +38,47 @@ namespace UE::K2NodeCallFunction::Private
 		return false;
 	}
 
+	// REMOVE IN 5.9
+	// 
+	// IsMutableContainerFunction and IsMutableContainerInput are used to permit connections between mutable array functions and read-only inputs.
+	// This will eventually become an error, but users might still have content that accidentally modifies read-only arrays.
+	// For now, we'll issue a warning and give them a chance to fix their content.
+
+	static bool IsMutableContainerFunction(UFunction* Function)
+	{
+		check(Function);
+
+		const UEdGraphSchema_K2* K2Schema = GetDefault<UEdGraphSchema_K2>();
+		check(K2Schema);
+
+		return K2Schema->AllowedMutableContainerFunctions.Contains(FAllowedMutableContainerFunction{Function, TEXT("")});
+	}
+
+	static bool IsMutableContainerInput(FName InputName, UFunction* Function)
+	{
+		const UEdGraphSchema_K2* K2Schema = GetDefault<UEdGraphSchema_K2>();
+		check(K2Schema);
+
+		if (const FAllowedMutableContainerFunction* Entry = K2Schema->AllowedMutableContainerFunctions.Find(FAllowedMutableContainerFunction{Function, TEXT("")}))
+		{
+			return (InputName == Entry->ContainerParameterName);
+		}
+
+		return false;
+	}
+
 	static bool FindMatchingReferencedNetOrFieldNotifyPropertyAndPin(TArray<UEdGraphPin*>& RemainingPins, FProperty* FunctionProperty, FProperty*& NetProperty, FProperty*& FieldNotifyProperty, UEdGraphPin*& PropertyObjectPin)
 	{
 		NetProperty = nullptr;
 		FieldNotifyProperty = nullptr;
 		PropertyObjectPin = nullptr;
 
-		if (UNLIKELY(FunctionProperty->HasAllPropertyFlags(CPF_OutParm | CPF_ReferenceParm) && !FunctionProperty->HasAnyPropertyFlags(CPF_ReturnParm | CPF_ConstParm)))
+		const bool bIsValidProperty =
+			(FunctionProperty->HasAllPropertyFlags(CPF_OutParm | CPF_ReferenceParm) && !FunctionProperty->HasAnyPropertyFlags(CPF_ReturnParm | CPF_ConstParm)) ||
+			IsMutableContainerInput(FunctionProperty->GetFName(), FunctionProperty->GetTypedOwner<UFunction>())
+		;
+
+		if (bIsValidProperty)
 		{
 			for (int32 i = 0; i < RemainingPins.Num(); ++i)
 			{
@@ -379,7 +413,14 @@ void FKCHandler_CallFunction::CreateFunctionCallStatement(FKismetFunctionContext
 								{
 									if (!Property->HasAnyPropertyFlags(CPF_ConstParm))
 									{
-										CompilerContext.MessageLog.Error(*LOCTEXT("PassReadOnlyReferenceParam_Error", "Cannot pass a read-only variable to a reference parameter @@").ToString(), PinMatch);
+										if (UE::K2NodeCallFunction::Private::IsMutableContainerFunction(Function))
+										{
+											CompilerContext.MessageLog.Warning(*LOCTEXT("MutatingReadOnlyData_Warning", "Array is changing read-only data @@. This will become an error in a future engine release.").ToString(), (*Term)->SourcePin);
+										}
+										else
+										{
+											CompilerContext.MessageLog.Error(*LOCTEXT("PassReadOnlyReferenceParam_Error", "Cannot pass a read-only variable to a reference parameter @@").ToString(), PinMatch);
+										}
 									}
 								}
 								else
@@ -876,9 +917,17 @@ void FKCHandler_CallFunction::Transform(FKismetFunctionContext& Context, UEdGrap
 			RemainingPins.Reserve(Node->Pins.Num());
 			for (UEdGraphPin* Pin : Node->Pins)
 			{
+				const bool bIsMutableContainerInput = 
+					UE::K2NodeCallFunction::Private::IsMutableContainerInput(Pin->GetFName(), Function)
+				;
+
+				const bool bIsNonConstReference =
+					(!Pin->PinType.bIsConst && Pin->PinType.bIsReference) ||
+					bIsMutableContainerInput
+				;
+
 				if (Pin->Direction == EGPD_Input &&
-					!Pin->PinType.bIsConst &&
-					Pin->PinType.bIsReference &&
+					bIsNonConstReference &&
 					Pin->LinkedTo.Num() > 0)
 				{
 					RemainingPins.Add(Pin);
diff --git a/Engine/Source/Runtime/Engine/Classes/Kismet/KismetArrayLibrary.h b/Engine/Source/Runtime/Engine/Classes/Kismet/KismetArrayLibrary.h
--- a/Engine/Source/Runtime/Engine/Classes/Kismet/KismetArrayLibrary.h
+++ b/Engine/Source/Runtime/Engine/Classes/Kismet/KismetArrayLibrary.h
@@ -34,7 +34,7 @@ class UKismetArrayLibrary : public UBlueprintFunctionLibrary
 	 *@return	The index of the newly added item
 	*/
 	UFUNCTION(BlueprintCallable, CustomThunk, meta=(DisplayName = "Add", CompactNodeTitle = "ADD", ArrayParm = "TargetArray", ArrayTypeDependentParams = "NewItem", AutoCreateRefTerm = "NewItem"), Category="Utilities|Array")
-	static ENGINE_API int32 Array_Add(const TArray<int32>& TargetArray, const int32& NewItem);
+	static ENGINE_API int32 Array_Add(UPARAM(ref) TArray<int32>& TargetArray, const int32& NewItem);
 
 	/**
 	*Add item to array (unique)
@@ -44,7 +44,7 @@ class UKismetArrayLibrary : public UBlueprintFunctionLibrary
 	*@return	The index of the newly added item, or INDEX_NONE if the item is already present in the array
 	*/
 	UFUNCTION(BlueprintCallable, CustomThunk, meta = (DisplayName = "Add Unique", CompactNodeTitle = "ADDUNIQUE", ArrayParm = "TargetArray", ArrayTypeDependentParams = "NewItem", AutoCreateRefTerm = "NewItem"), Category = "Utilities|Array")
-	static ENGINE_API int32 Array_AddUnique(const TArray<int32>& TargetArray, const int32& NewItem);
+	static ENGINE_API int32 Array_AddUnique(UPARAM(ref) TArray<int32>& TargetArray, const int32& NewItem);
 
 	/** 
 	 * Shuffle (randomize) the elements of an array
@@ -52,7 +52,7 @@ class UKismetArrayLibrary : public UBlueprintFunctionLibrary
 	 *@param	TargetArray		The array to shuffle
 	*/
 	UFUNCTION(BlueprintCallable, CustomThunk, meta=(DisplayName = "Shuffle", CompactNodeTitle = "SHUFFLE", ArrayParm = "TargetArray"), Category="Utilities|Array")
-	static ENGINE_API void Array_Shuffle(const TArray<int32>& TargetArray);
+	static ENGINE_API void Array_Shuffle(UPARAM(ref) TArray<int32>& TargetArray);
 	
 	/** 
 	 * Shuffle (randomize) the elements of an array from a specific stream of random data, useful for achieving determinism
@@ -61,7 +61,7 @@ class UKismetArrayLibrary : public UBlueprintFunctionLibrary
 	 * @param	RandomStream	The random stream
 	 */
 	UFUNCTION(BlueprintCallable, CustomThunk, meta=(DisplayName = "Shuffle from Stream", CompactNodeTitle = "SHUFFLE", ArrayParm = "TargetArray"), Category="Utilities|Array")
-	static ENGINE_API void Array_ShuffleFromStream(const TArray<int32>& TargetArray, UPARAM(Ref) FRandomStream& RandomStream);
+	static ENGINE_API void Array_ShuffleFromStream(UPARAM(ref) TArray<int32>& TargetArray, UPARAM(Ref) FRandomStream& RandomStream);
 
 	/** 
 	 * Checks if two arrays are memberwise identical
@@ -80,7 +80,7 @@ class UKismetArrayLibrary : public UBlueprintFunctionLibrary
 	 *@param	SourceArray		The array to add to the target array
 	*/
 	UFUNCTION(BlueprintCallable, CustomThunk, meta=(DisplayName = "Append Array", CompactNodeTitle = "APPEND", ArrayParm = "TargetArray,SourceArray", ArrayTypeDependentParams = "SourceArray"), Category="Utilities|Array")
-	static ENGINE_API void Array_Append(const TArray<int32>& TargetArray, const TArray<int32>& SourceArray);
+	static ENGINE_API void Array_Append(UPARAM(ref) TArray<int32>& TargetArray, const TArray<int32>& SourceArray);
 
 	/* 
 	 *Insert item at the given index into the array.
@@ -90,7 +90,7 @@ class UKismetArrayLibrary : public UBlueprintFunctionLibrary
 	 *@param	Index			The index at which to insert the item into the array
 	*/
 	UFUNCTION(BlueprintCallable, CustomThunk, meta=(DisplayName = "Insert", CompactNodeTitle = "INSERT", ArrayParm = "TargetArray", ArrayTypeDependentParams = "NewItem", AutoCreateRefTerm = "NewItem"), Category="Utilities|Array")
-	static ENGINE_API void Array_Insert(const TArray<int32>& TargetArray, const int32& NewItem, int32 Index);
+	static ENGINE_API void Array_Insert(UPARAM(ref) TArray<int32>& TargetArray, const int32& NewItem, int32 Index);
 
 
 	/* 
@@ -100,7 +100,7 @@ class UKismetArrayLibrary : public UBlueprintFunctionLibrary
 	 *@param	IndexToRemove	The index into the array to remove from
 	*/
 	UFUNCTION(BlueprintCallable, CustomThunk, meta=(DisplayName = "Remove Index", CompactNodeTitle = "REMOVE INDEX", ArrayParm = "TargetArray"), Category="Utilities|Array")
-	static ENGINE_API void Array_Remove(const TArray<int32>& TargetArray, int32 IndexToRemove);
+	static ENGINE_API void Array_Remove(UPARAM(ref) TArray<int32>& TargetArray, int32 IndexToRemove);
 
 	/* 
 	 *Remove all instances of item from array.
@@ -110,7 +110,7 @@ class UKismetArrayLibrary : public UBlueprintFunctionLibrary
 	 *@return	True if one or more items were removed
 	*/
 	UFUNCTION(BlueprintCallable, CustomThunk, meta=(DisplayName = "Remove Item", CompactNodeTitle = "REMOVE", ArrayParm = "TargetArray", ArrayTypeDependentParams = "Item", AutoCreateRefTerm = "Item"), Category="Utilities|Array")
-	static ENGINE_API bool Array_RemoveItem(const TArray<int32>& TargetArray, const int32 &Item);
+	static ENGINE_API bool Array_RemoveItem(UPARAM(ref) TArray<int32>& TargetArray, const int32 &Item);
 
 	/* 
 	 *Clear an array, removes all content
@@ -118,7 +118,7 @@ class UKismetArrayLibrary : public UBlueprintFunctionLibrary
 	 *@param	TargetArray		The array to clear
 	*/
 	UFUNCTION(BlueprintCallable, CustomThunk, meta=(DisplayName = "Clear", CompactNodeTitle = "CLEAR", Keywords = "empty", ArrayParm = "TargetArray"), Category="Utilities|Array")
-	static ENGINE_API void Array_Clear(const TArray<int32>& TargetArray);
+	static ENGINE_API void Array_Clear(UPARAM(ref) TArray<int32>& TargetArray);
 
 	/* 
 	 *Resize Array to specified size. 
@@ -127,7 +127,7 @@ class UKismetArrayLibrary : public UBlueprintFunctionLibrary
 	 *@param	Size			The new size of the array
 	*/
 	UFUNCTION(BlueprintCallable, CustomThunk, meta=(DisplayName = "Resize", CompactNodeTitle = "RESIZE", ArrayParm = "TargetArray"), Category="Utilities|Array")
-	static ENGINE_API void Array_Resize(const TArray<int32>& TargetArray, int32 Size);
+	static ENGINE_API void Array_Resize(UPARAM(ref) TArray<int32>& TargetArray, int32 Size);
 	
 	/**
 	 * Reverse the elements of an array
@@ -135,7 +135,7 @@ class UKismetArrayLibrary : public UBlueprintFunctionLibrary
 	 *@param	TargetArray		The array to reverse
 	*/
 	UFUNCTION(BlueprintCallable, CustomThunk, meta = (DisplayName = "Reverse", CompactNodeTitle = "REVERSE", ArrayParm = "TargetArray"), Category = "Utilities|Array")
-	static ENGINE_API void Array_Reverse(const TArray<int32>& TargetArray);
+	static ENGINE_API void Array_Reverse(UPARAM(ref) TArray<int32>& TargetArray);
 
 	/* 
 	 *Get the number of items in an array
@@ -193,7 +193,7 @@ class UKismetArrayLibrary : public UBlueprintFunctionLibrary
 	 *@param	bSizeToFit		If true, the array will expand if Index is greater than the current size of the array
 	*/
 	UFUNCTION(BlueprintCallable, CustomThunk, meta=(DisplayName = "Set Array Elem", ArrayParm = "TargetArray", ArrayTypeDependentParams = "Item", AutoCreateRefTerm = "Item"), Category="Utilities|Array")
-	static ENGINE_API void Array_Set(const TArray<int32>& TargetArray, int32 Index, const int32& Item, bool bSizeToFit);
+	static ENGINE_API void Array_Set(UPARAM(ref) TArray<int32>& TargetArray, int32 Index, const int32& Item, bool bSizeToFit);
 
 	/*
 	 *Swaps the elements at the specified positions in the specified array
@@ -204,7 +204,7 @@ class UKismetArrayLibrary : public UBlueprintFunctionLibrary
 	 *@param    SecondIndex     The index of the other element to be swapped
 	*/
 	UFUNCTION(BlueprintCallable, CustomThunk, meta=(DisplayName = "Swap Array Elements", CompactNodeTitle = "SWAP", ArrayParm = "TargetArray"), Category="Utilities|Array")
-	static ENGINE_API void Array_Swap(const TArray<int32>& TargetArray, int32 FirstIndex, int32 SecondIndex);
+	static ENGINE_API void Array_Swap(UPARAM(ref) TArray<int32>& TargetArray, int32 FirstIndex, int32 SecondIndex);
 
 	/*  
 	 *Finds the index of the first instance of the item within the array
@@ -400,7 +400,6 @@ public:
  
  		P_FINISH;
 		P_NATIVE_BEGIN;
-		MARK_PROPERTY_DIRTY(Stack.Object, ArrayProperty);
 		*(int32*)RESULT_PARAM = GenericArray_Add(ArrayAddr, ArrayProperty, StorageSpace);
 		P_NATIVE_END;
 		InnerProp->DestroyValue(StorageSpace);
@@ -430,7 +429,6 @@ public:
 
 		P_FINISH;
 		P_NATIVE_BEGIN;
-		MARK_PROPERTY_DIRTY(Stack.Object, ArrayProperty);
 		*(int32*)RESULT_PARAM = GenericArray_AddUnique(ArrayAddr, ArrayProperty, StorageSpace);
 		P_NATIVE_END;
 		InnerProp->DestroyValue(StorageSpace);
@@ -450,7 +448,6 @@ public:
 
 		P_FINISH;
 		P_NATIVE_BEGIN;
-		MARK_PROPERTY_DIRTY(Stack.Object, ArrayProperty);
 		GenericArray_Shuffle(ArrayAddr, ArrayProperty);
 		P_NATIVE_END;
 	}
@@ -473,7 +470,6 @@ public:
 
 		P_FINISH;
 		P_NATIVE_BEGIN;
-		MARK_PROPERTY_DIRTY(Stack.Object, ArrayProperty);
 		GenericArray_ShuffleFromStream(ArrayAddr, ArrayProperty, RandomStream);
 		P_NATIVE_END;
 	}
@@ -532,7 +528,6 @@ public:
 
 		P_FINISH;
 		P_NATIVE_BEGIN;
-		MARK_PROPERTY_DIRTY(Stack.Object, TargetArrayProperty);
 		GenericArray_Append(TargetArrayAddr, TargetArrayProperty, SourceArrayAddr, SourceArrayProperty);
 		P_NATIVE_END;
 	}
@@ -562,7 +557,6 @@ public:
 		P_GET_PROPERTY(FIntProperty, Index);
 		P_FINISH;
 		P_NATIVE_BEGIN;
-		MARK_PROPERTY_DIRTY(Stack.Object, ArrayProperty);
 		GenericArray_Insert(ArrayAddr, ArrayProperty, StorageSpace, Index);
 		P_NATIVE_END;
 		InnerProp->DestroyValue(StorageSpace);
@@ -583,7 +577,6 @@ public:
 		P_GET_PROPERTY(FIntProperty, Index);
 		P_FINISH;
 		P_NATIVE_BEGIN;
-		MARK_PROPERTY_DIRTY(Stack.Object, ArrayProperty);
 		GenericArray_Remove(ArrayAddr, ArrayProperty, Index);
 		P_NATIVE_END;
 	}
@@ -615,7 +608,6 @@ public:
 		// Bools need to be processed internally by the property so that C++ bool value is properly set.
 		GenericArray_HandleBool(InnerProp, ItemPtr);
 		P_NATIVE_BEGIN;
-		MARK_PROPERTY_DIRTY(Stack.Object, ArrayProperty);
 		*(bool*)RESULT_PARAM = GenericArray_RemoveItem(ArrayAddr, ArrayProperty, ItemPtr);
 		P_NATIVE_END;
 
@@ -635,7 +627,6 @@ public:
 		}
 		P_FINISH;
 		P_NATIVE_BEGIN;
-		MARK_PROPERTY_DIRTY(Stack.Object, ArrayProperty);
 		GenericArray_Clear(ArrayAddr, ArrayProperty);
 		P_NATIVE_END;
 	}
@@ -654,7 +645,6 @@ public:
 		P_GET_PROPERTY(FIntProperty, Size);
 		P_FINISH;
 		P_NATIVE_BEGIN;
-		MARK_PROPERTY_DIRTY(Stack.Object, ArrayProperty);
 		GenericArray_Resize(ArrayAddr, ArrayProperty, Size);
 		P_NATIVE_END;
 	}
@@ -815,7 +805,6 @@ public:
 		P_FINISH;
 
 		P_NATIVE_BEGIN;
-		MARK_PROPERTY_DIRTY(Stack.Object, ArrayProperty);
 		GenericArray_Set(ArrayAddr, ArrayProperty, Index, StorageSpace, bSizeToFit);
 		P_NATIVE_END;
 		InnerProp->DestroyValue(StorageSpace);
@@ -839,7 +828,6 @@ public:
 
 		P_FINISH;
 		P_NATIVE_BEGIN;
-		MARK_PROPERTY_DIRTY(Stack.Object, ArrayProperty);
 		GenericArray_Swap(ArrayAddr, ArrayProperty, First, Second);
 		P_NATIVE_END;
 	}
diff --git a/Engine/Source/Runtime/Engine/Private/KismetArrayLibrary.cpp b/Engine/Source/Runtime/Engine/Private/KismetArrayLibrary.cpp
--- a/Engine/Source/Runtime/Engine/Private/KismetArrayLibrary.cpp
+++ b/Engine/Source/Runtime/Engine/Private/KismetArrayLibrary.cpp
@@ -659,39 +659,39 @@ void UKismetArrayLibrary::GenericArray_RandomFromStream(void* TargetArray, const
 //////////////////////////////////////////////////////////////////////////
 // Stubs for the UFunctions declared as kismet callable.  These are never actually called...the CustomThunk code calls the appropriate native function with a void* reference to the array
 
-int32 UKismetArrayLibrary::Array_Add(const TArray<int32>& TargetArray, const int32& NewItem)
+int32 UKismetArrayLibrary::Array_Add(TArray<int32>& TargetArray, const int32& NewItem)
 {
 	// We should never hit these!  They're stubs to avoid NoExport on the class.  Call the Generic* equivalent instead
 	check(0);
 	return 0;
 }
 
-void UKismetArrayLibrary::Array_Shuffle(const TArray<int32>& TargetArray)
+void UKismetArrayLibrary::Array_Shuffle(TArray<int32>& TargetArray)
 {
 	// We should never hit these!  They're stubs to avoid NoExport on the class.  Call the Generic* equivalent instead
 	check(0);
 }
 
-void UKismetArrayLibrary::Array_Insert(const TArray<int32>& TargetArray, const int32& NewItem, int32 Index)
+void UKismetArrayLibrary::Array_Insert(TArray<int32>& TargetArray, const int32& NewItem, int32 Index)
 {
 	// We should never hit these!  They're stubs to avoid NoExport on the class.  Call the Generic* equivalent instead
 	check(0);
 }
 
-void UKismetArrayLibrary::Array_Remove(const TArray<int32>& TargetArray, int32 IndexToRemove)
+void UKismetArrayLibrary::Array_Remove(TArray<int32>& TargetArray, int32 IndexToRemove)
 {
 	// We should never hit these!  They're stubs to avoid NoExport on the class.  Call the Generic* equivalent instead
 	check(0);
 }
 
-bool UKismetArrayLibrary::Array_RemoveItem(const TArray<int32>& TargetArray, const int32& IndexToRemove)
+bool UKismetArrayLibrary::Array_RemoveItem(TArray<int32>& TargetArray, const int32& IndexToRemove)
 {
 	// We should never hit these!  They're stubs to avoid NoExport on the class.  Call the Generic* equivalent instead
 	check(0);
 	return true;
 }
 
-void UKismetArrayLibrary::Array_Clear(const TArray<int32>& TargetArray)
+void UKismetArrayLibrary::Array_Clear(TArray<int32>& TargetArray)
 {
 	// We should never hit these!  They're stubs to avoid NoExport on the class.  Call the Generic* equivalent instead
 	check(0);
@@ -737,13 +737,13 @@ void UKismetArrayLibrary::Array_Get(const TArray<int32>& TargetArray, int32 Inde
 	check(0);
 }
 
-void UKismetArrayLibrary::Array_Set(const TArray<int32>& TargetArray, int32 Index, const int32& NewItem, bool bSizeToFit)
+void UKismetArrayLibrary::Array_Set(TArray<int32>& TargetArray, int32 Index, const int32& NewItem, bool bSizeToFit)
 {
 	// We should never hit these!  They're stubs to avoid NoExport on the class.  Call the Generic* equivalent instead
 	check(0);
 }
 
-void UKismetArrayLibrary::Array_Swap(const TArray<int32>& TargetArray, int32 First, int32 Second)
+void UKismetArrayLibrary::Array_Swap(TArray<int32>& TargetArray, int32 First, int32 Second)
 {
 	// We should never hit these!  They're stubs to avoid NoExport on the class.  Call the Generic* equivalent instead
 	check(0);
