Fix bad virtual texture thumbnails.
Each virtual texture, when loaded, reserves and fills a small texture copy of itself that is used for filling thumbnails amongst other things.
The allocation was being done through RDG (and exported from RDG graph) but when multiple virtual textures were initialized on the same frame this memory was aliased.
This led to virtual textures effectively stomping on each others thumbnails.
The fix is to the texture allocation through RHI and import to the RDG graph.

diff --git a/Engine/Source/Runtime/Engine/Private/Texture2D.cpp b/Engine/Source/Runtime/Engine/Private/Texture2D.cpp
--- a/Engine/Source/Runtime/Engine/Private/Texture2D.cpp
+++ b/Engine/Source/Runtime/Engine/Private/Texture2D.cpp
@@ -1737,8 +1737,6 @@ void FVirtualTexture2DResource::InitializeEditorResources(FRHICommandListBase& R
 			const uint32 LayerMask = 1u;
 
 			// Get or create the target texture to composite tiles to.
-			TRefCountPtr<IPooledRenderTarget> PooledRenderTarget;
-			
 			// First try reusing the current referenced texture.
 			// We do that because swapping the referenced texture can be expensive if we are doing a recreate of all textures.
 			if (FRHITexture* ReferencedTexture = TextureReferenceRHI->GetReferencedTexture())
@@ -1746,7 +1744,6 @@ void FVirtualTexture2DResource::InitializeEditorResources(FRHICommandListBase& R
 				FRHITextureDesc const& ExistingDesc = ReferencedTexture->GetDesc();
 				if (ExistingDesc.Format == PixelFormat && ExistingDesc.GetSize() == FIntVector(MipWidthInTiles * TileSizeInPixels, MipHeightInTiles * TileSizeInPixels, 1))
 				{
-					PooledRenderTarget = CreateRenderTarget(ReferencedTexture, TEXT("VTEditorResource"));
 					Texture2DRHI = ReferencedTexture;
 				}
 			}
@@ -1754,13 +1751,14 @@ void FVirtualTexture2DResource::InitializeEditorResources(FRHICommandListBase& R
 			// If we can't reuse the referenced texture then create it.
 			if (Texture2DRHI == nullptr)
 			{
-				FRDGTextureRef Texture = GraphBuilder.CreateTexture(FRDGTextureDesc::Create2D(
-					FIntPoint(MipWidthInTiles * TileSizeInPixels, MipHeightInTiles * TileSizeInPixels),
-					PixelFormat, FClearValueBinding::Black,	TexCreateFlags), TEXT("VTEditorResource"));
+				const FRHITextureCreateDesc Desc =
+					FRHITextureCreateDesc::Create2D(TEXT("VTEditorResource"), MipWidthInTiles * TileSizeInPixels, MipHeightInTiles * TileSizeInPixels, PixelFormat)
+					.SetFlags(ETextureCreateFlags::ShaderResource);
 
-				PooledRenderTarget = GraphBuilder.ConvertToExternalTexture(Texture);
-				Texture2DRHI = PooledRenderTarget->GetRHI();
+				Texture2DRHI = RHICreateTexture(Desc);
 			}
+			
+			TRefCountPtr<IPooledRenderTarget> PooledRenderTarget = CreateRenderTarget(Texture2DRHI, TEXT("VTEditorResource"));
 
 			// Produce all the tiles.
 			TArray<IVirtualTextureFinalizer*> Finalizers;
