New NetworkPhysicsSettings: EventSchedulingMinDelaySeconds to delay an event's application so server and peers can execute on the same frame/at the same time.

diff --git a/Engine/Source/Runtime/Engine/Public/Physics/NetworkPhysicsSettingsComponent.h b/Engine/Source/Runtime/Engine/Public/Physics/NetworkPhysicsSettingsComponent.h
--- a/Engine/Source/Runtime/Engine/Public/Physics/NetworkPhysicsSettingsComponent.h
+++ b/Engine/Source/Runtime/Engine/Public/Physics/NetworkPhysicsSettingsComponent.h
@@ -89,9 +89,14 @@ struct FNetworkPhysicsSettings
 	UPROPERTY(config, EditDefaultsOnly, Category = "Overrides", Meta = (EditCondition = "bOverrideSimProxyRepMode"))
 	EPhysicsReplicationMode SimProxyRepMode = EPhysicsReplicationMode::PredictiveInterpolation;
 	
-	// Register this actors Autonomous Proxy as a focal point / focal particle in Physics Replication LOD.
+	// Register this actor's Autonomous Proxy as a focal point / focal particle in Physics Replication LOD.
 	UPROPERTY(config, EditDefaultsOnly, Category = "Overrides")
 	bool bFocalParticleInPhysicsReplicationLOD = false;
+
+	// Minimum delay added, in seconds, when scheduling an event far enough in the future to ensure server and all peers that event at the same time
+	// This value should be large enough to accommodate the largest network Round Trip Time, beyond which it is understood that the simulation will suffer from corrections
+	UPROPERTY(config, EditDefaultsOnly, Category = "General")
+	float EventSchedulingMinDelaySeconds = 0.3f;
 };
 
 USTRUCT()
@@ -668,7 +673,7 @@ public:
 
 private:
 	FNetworkPhysicsSettingsComponentAsync* NetworkPhysicsSettings_Internal;
-	static const FNetworkPhysicsSettingsData SettingsNetworkPhysicsData_Default;
+	ENGINE_API static const FNetworkPhysicsSettingsData SettingsNetworkPhysicsData_Default;
 
 	/* --- Static API --- */
 private:
