Adding GetCurrentServerFrame_External, a util function used for networked physics that returns the server frame number. This is LocalFrame + NetworkTickOffset in the client case, or just the local frame number on the server.

diff --git a/Engine/Source/Runtime/Engine/Private/PhysicsEngine/NetworkPhysicsComponent.cpp b/Engine/Source/Runtime/Engine/Private/PhysicsEngine/NetworkPhysicsComponent.cpp
--- a/Engine/Source/Runtime/Engine/Private/PhysicsEngine/NetworkPhysicsComponent.cpp
+++ b/Engine/Source/Runtime/Engine/Private/PhysicsEngine/NetworkPhysicsComponent.cpp
@@ -2645,3 +2645,28 @@ const int32 FAsyncNetworkPhysicsComponent::SetupRewindData()
 
 	return NumFrames;
 }
+
+namespace UE::NetworkPhysicsUtils
+{
+	int32 GetUpcomingServerFrame_External(UWorld* World)
+	{
+		if (World)
+		{
+			if (FPhysScene* Scene = World->GetPhysicsScene())
+			{
+				if (Chaos::FPhysicsSolver* Solver = Scene->GetSolver())
+				{
+					int32 ServerFrame = Solver->GetMarshallingManager().GetInternalStep_External();
+					if (APlayerController* PlayerController = World->GetFirstPlayerController())
+					{
+						int32 NetworkPhysicsTickOffset = PlayerController->GetNetworkPhysicsTickOffset();
+						ServerFrame += NetworkPhysicsTickOffset;
+					}
+					return ServerFrame;
+				}
+			}
+		}
+
+		return INDEX_NONE;
+	}
+}
diff --git a/Engine/Source/Runtime/Engine/Public/Physics/NetworkPhysicsComponent.h b/Engine/Source/Runtime/Engine/Public/Physics/NetworkPhysicsComponent.h
--- a/Engine/Source/Runtime/Engine/Public/Physics/NetworkPhysicsComponent.h
+++ b/Engine/Source/Runtime/Engine/Public/Physics/NetworkPhysicsComponent.h
@@ -1890,3 +1890,9 @@ private:
 	FDelegateHandle DelegateOnPreProcessInputs_Internal;
 	FDelegateHandle DelegateOnPostProcessInputs_Internal;
 };
+
+namespace UE::NetworkPhysicsUtils
+{
+	// Returns the index of the frame about to be simulated, in the server timeline
+	ENGINE_API int32 GetUpcomingServerFrame_External(UWorld* World);
+}
\ No newline at end of file
