Changed UMaterialFunctionInterface to include the BlueprintType annotation

diff --git a/Engine/Source/Runtime/Engine/Public/Materials/MaterialFunctionInterface.h b/Engine/Source/Runtime/Engine/Public/Materials/MaterialFunctionInterface.h
--- a/Engine/Source/Runtime/Engine/Public/Materials/MaterialFunctionInterface.h
+++ b/Engine/Source/Runtime/Engine/Public/Materials/MaterialFunctionInterface.h
@@ -54,7 +54,7 @@ public:
 /**
  * A Material Function is a collection of material expressions that can be reused in different materials
  */
-UCLASS(abstract, hidecategories=object, MinimalAPI)
+UCLASS(BlueprintType, abstract, hidecategories=object, MinimalAPI)
 class UMaterialFunctionInterface : public UObject
 {
 	GENERATED_BODY()
