Follow-up fix for compile error in CL 44246053 when WITH_DYNAMIC_RESOLUTION is false

diff --git a/Engine/Source/Runtime/Engine/Private/GameViewportClient.cpp b/Engine/Source/Runtime/Engine/Private/GameViewportClient.cpp
--- a/Engine/Source/Runtime/Engine/Private/GameViewportClient.cpp
+++ b/Engine/Source/Runtime/Engine/Private/GameViewportClient.cpp
@@ -1583,9 +1583,8 @@ void UGameViewportClient::Draw(FViewport* InViewport, FCanvas* SceneCanvas)
 			check(ViewFamily.SecondaryViewFraction > 0.0f);
 		}
 
+#if WITH_DYNAMIC_RESOLUTION
 		// Setup main view family with screen percentage interface by dynamic resolution if screen percentage is enabled.
-		#if WITH_DYNAMIC_RESOLUTION
-
 		// Force screen percentage support if we're using stereo rendering and redirecting primary scaling to the HMD runtime.
 		bool bDynamicResolutionScalingViaHMDRuntime = false;
 		bool bAllResolutionScalingViaHMDRuntime = false;
@@ -1656,9 +1655,10 @@ void UGameViewportClient::Draw(FViewport* InViewport, FCanvas* SceneCanvas)
 			}
 #endif // CSV_PROFILER
 		}
-		#endif
 
 		ViewFamily.EngineShowFlags.ScreenPercentage = ViewFamily.EngineShowFlags.ScreenPercentage || bUsesDynamicResolution || bAllResolutionScalingViaHMDRuntime;
+#endif // WITH_DYNAMIC_RESOLUTION
+		
 		bFinalScreenPercentageShowFlag = ViewFamily.EngineShowFlags.ScreenPercentage;
 	}
 
