Fixed issue where screen-space widget components wouldn't render for players other than player 1 if splitscreen is disabled. There is now an option (ShowAllPlayerWidgetsWhenSplitscreenDisabled) to overlay the player layers, instead of cullling out all player layers except for player 1.

diff --git a/Engine/Source/Runtime/Engine/Private/Slate/SGameLayerManager.cpp b/Engine/Source/Runtime/Engine/Private/Slate/SGameLayerManager.cpp
--- a/Engine/Source/Runtime/Engine/Private/Slate/SGameLayerManager.cpp
+++ b/Engine/Source/Runtime/Engine/Private/Slate/SGameLayerManager.cpp
@@ -10,6 +10,7 @@
 #include "Engine/GameEngine.h"
 #include "Engine/UserInterfaceSettings.h"
 #include "GeneralProjectSettings.h"
+#include "GameMapsSettings.h"
 #include "Input/HittestGrid.h"
 #include "Widgets/LayerManager/STooltipPresenter.h"
 #include "Widgets/Layout/SBox.h"
@@ -638,16 +639,26 @@ void SGameLayerManager::AddOrUpdatePlayerLayers(const FGeometry& AllottedGeometr
 	{
 		ULocalPlayer* Player = GamePlayers[PlayerIndex];
 
-		if ( SplitType < SplitInfo.Num() && PlayerIndex < SplitInfo[SplitType].PlayerData.Num() )
+		if ( SplitType < SplitInfo.Num() )
 		{
-			TSharedPtr<FPlayerLayer> PlayerLayer = FindOrCreatePlayerLayer(Player);
-			FPerPlayerSplitscreenData& SplitData = SplitInfo[SplitType].PlayerData[PlayerIndex];
-
 			// Viewport Sizes
 			FVector2D Position(0.0, 0.0);
-			FVector2D Size(SplitData.SizeX, SplitData.SizeY);
-			GetNormalizeRect(Player, Position, Size);
+			FVector2D Size(0.0, 0.0);
 
+			TSharedPtr<FPlayerLayer> PlayerLayer = FindOrCreatePlayerLayer(Player);
+			if (PlayerIndex < SplitInfo[SplitType].PlayerData.Num())
+			{
+				FPerPlayerSplitscreenData& SplitData = SplitInfo[SplitType].PlayerData[PlayerIndex];
+				Size = FVector2D(SplitData.SizeX, SplitData.SizeY);
+			}
+			else if (SplitType == ESplitScreenType::None && GetDefault<UGameMapsSettings>()->bShowAllPlayerWidgetsWhenSplitscreenDisabled )
+			{
+				// If splitscreen is disabled, we should overlay all of the player layers in the viewport
+				FPerPlayerSplitscreenData& SplitData = SplitInfo[SplitType].PlayerData[0];
+				Size = FVector2D(SplitData.SizeX, SplitData.SizeY);
+			}		
+
+			GetNormalizeRect(Player, Position, Size);
 			Size = Size * AllottedGeometry.GetLocalSize() * InverseDPIScale;
 			Position = Position * AllottedGeometry.GetLocalSize() * InverseDPIScale;
 
diff --git a/Engine/Source/Runtime/EngineSettings/Classes/GameMapsSettings.h b/Engine/Source/Runtime/EngineSettings/Classes/GameMapsSettings.h
--- a/Engine/Source/Runtime/EngineSettings/Classes/GameMapsSettings.h
+++ b/Engine/Source/Runtime/EngineSettings/Classes/GameMapsSettings.h
@@ -187,6 +187,10 @@ public:
 	UPROPERTY(config, EditAnywhere, Category=LocalMultiplayer, meta=(editcondition="bUseSplitScreen"))
 	EFourPlayerSplitScreenType FourPlayerSplitscreenLayout;
 
+	/** Determines whether all player widget layers should be displayed in the main viewport when splitscreen is disabled. If unchecked, only Player 1 widget will be visible */
+	UPROPERTY(config, EditAnywhere, Category=LocalMultiplayer)
+	bool bShowAllPlayerWidgetsWhenSplitscreenDisabled;
+
 	/**
 	* If enabled, this will make so that gamepads start being assigned to the second controller ID in local multiplayer games.
 	* In PIE sessions with multiple windows, this has the same effect as enabling "Route 1st Gamepad to 2nd Client"
diff --git a/Engine/Source/Runtime/EngineSettings/Private/EngineSettingsModule.cpp b/Engine/Source/Runtime/EngineSettings/Private/EngineSettingsModule.cpp
--- a/Engine/Source/Runtime/EngineSettings/Private/EngineSettingsModule.cpp
+++ b/Engine/Source/Runtime/EngineSettings/Private/EngineSettingsModule.cpp
@@ -45,6 +45,7 @@ UGameMapsSettings::UGameMapsSettings( const FObjectInitializer& ObjectInitialize
 	, bUseSplitscreen(true)
 	, TwoPlayerSplitscreenLayout(ETwoPlayerSplitScreenType::Horizontal)
 	, ThreePlayerSplitscreenLayout(EThreePlayerSplitScreenType::FavorTop)
+	, bShowAllPlayerWidgetsWhenSplitscreenDisabled(false)
 	, bOffsetPlayerGamepadIds(false)
 { }
 
