Make sound Pitch and Volume tooltips more accurate to their functions.

diff --git a/Engine/Source/Runtime/Engine/Classes/Sound/SoundWave.h b/Engine/Source/Runtime/Engine/Classes/Sound/SoundWave.h
--- a/Engine/Source/Runtime/Engine/Classes/Sound/SoundWave.h
+++ b/Engine/Source/Runtime/Engine/Classes/Sound/SoundWave.h
@@ -758,11 +758,11 @@ public:
 	UPROPERTY(EditAnywhere, BlueprintReadOnly, Category = Subtitles)
 	float SubtitlePriority;
 
-	/** Playback volume of sound 0 to 1 - Default is 1.0. */
+	/** Linear volume multiplier for sound playback, default is 1.0 */
 	UPROPERTY(Category = Sound, meta = (ClampMin = "0.0"), EditAnywhere)
 	float Volume;
 
-	/** Playback pitch for sound. */
+	/** Playback rate multiplier for sound, range is 0.125 - 4.0, default is 1.0  */
 	UPROPERTY(Category = Sound, meta = (ClampMin = "0.125", ClampMax = "4.0"), EditAnywhere)
 	float Pitch;
 
