Removing deprecated single-precision float server time replication, and updated comments for clarity on how often the server updates its replicated world time to clients.

diff --git a/Engine/Source/Runtime/Engine/Classes/GameFramework/GameStateBase.h b/Engine/Source/Runtime/Engine/Classes/GameFramework/GameStateBase.h
--- a/Engine/Source/Runtime/Engine/Classes/GameFramework/GameStateBase.h
+++ b/Engine/Source/Runtime/Engine/Classes/GameFramework/GameStateBase.h
@@ -142,20 +142,11 @@ protected:
 	/** Called periodically to update ReplicatedWorldTimeSecondsDouble */
 	ENGINE_API virtual void UpdateServerTimeSeconds();
 
-	/** Allows clients to calculate ServerWorldTimeSecondsDelta */
-	UE_DEPRECATED(5.2, "OnRep_ReplicatedWorldTimeSeconds() is deprecated. Use OnRep_ReplicatedWorldTimeSecondsDouble().")
-	UFUNCTION()
-	ENGINE_API virtual void OnRep_ReplicatedWorldTimeSeconds() final;
-
 	/** Allows clients to calculate ServerWorldTimeSecondsDelta */
 	UFUNCTION()
 	ENGINE_API virtual void OnRep_ReplicatedWorldTimeSecondsDouble();
 
 	/** Server TimeSeconds. Useful for syncing up animation and gameplay. */
-	UE_DEPRECATED(5.2, "ReplicatedWorldTimeSeconds is deprecated. Use ReplicatedWorldTimeSecondsDouble.")
-	UPROPERTY(Transient, ReplicatedUsing=OnRep_ReplicatedWorldTimeSeconds)
-	float ReplicatedWorldTimeSeconds;
-
 	UPROPERTY(Transient, ReplicatedUsing = OnRep_ReplicatedWorldTimeSecondsDouble)
 	double ReplicatedWorldTimeSecondsDouble;
 
@@ -163,7 +154,7 @@ protected:
 	UPROPERTY(Transient)
 	float ServerWorldTimeSecondsDelta;
 
-	/** Frequency that the server updates the replicated TimeSeconds from the world. Set to zero to disable periodic updates. */
+	/** How often the server updates the replicated TimeSeconds from the world. Values mean "update every X seconds", so 0.5 means update every half second. Set to zero to disable periodic updates. */
 	UPROPERTY(EditDefaultsOnly, Category=GameState)
 	float ServerWorldTimeSecondsUpdateFrequency;
 
diff --git a/Engine/Source/Runtime/Engine/Private/GameStateBase.cpp b/Engine/Source/Runtime/Engine/Private/GameStateBase.cpp
--- a/Engine/Source/Runtime/Engine/Private/GameStateBase.cpp
+++ b/Engine/Source/Runtime/Engine/Private/GameStateBase.cpp
@@ -32,7 +32,7 @@ AGameStateBase::AGameStateBase(const FObjectInitializer& ObjectInitializer)
 	// may end up deleting the new PlayerStates they had just spawned.
 	bNetLoadOnClient = false;
 
-	// Default to every few seconds.
+	// Default to every 100 ms
 	ServerWorldTimeSecondsUpdateFrequency = 0.1f;
 
 	SumServerWorldTimeSecondsDelta = 0.0;
@@ -161,24 +161,6 @@ void AGameStateBase::UpdateServerTimeSeconds()
 	}
 }
 
-void AGameStateBase::OnRep_ReplicatedWorldTimeSeconds()
-{
-	UWorld* World = GetWorld();
-	if (World)
-	{
-		const UDemoNetDriver* DemoNetDriver = World ? World->GetDemoNetDriver() : nullptr;
-
-		// Support old replays.
-		if (DemoNetDriver && DemoNetDriver->IsPlaying() && (DemoNetDriver->GetPlaybackEngineNetworkProtocolVersion() < FEngineNetworkCustomVersion::GameStateReplicatedTimeAsDouble))
-		{
-			PRAGMA_DISABLE_DEPRECATION_WARNINGS
-			ReplicatedWorldTimeSecondsDouble = ReplicatedWorldTimeSeconds;
-			PRAGMA_ENABLE_DEPRECATION_WARNINGS
-			OnRep_ReplicatedWorldTimeSecondsDouble();
-		}
-	}
-}
-
 void AGameStateBase::OnRep_ReplicatedWorldTimeSecondsDouble()
 {
 	UWorld* World = GetWorld();
@@ -276,10 +258,6 @@ void AGameStateBase::GetLifetimeReplicatedProps(TArray<FLifetimeProperty>& OutLi
 
 	FDoRepLifetimeParams Params;
 	Params.bIsPushBased = true;
-	PRAGMA_DISABLE_DEPRECATION_WARNINGS
-	// ReplicatedWorldTimeSeconds is still 'replicated' to support old replays, however this variable does not change value on the server now and is never dirtyed.
-	DOREPLIFETIME_WITH_PARAMS_FAST(AGameStateBase, ReplicatedWorldTimeSeconds, Params);
-	PRAGMA_ENABLE_DEPRECATION_WARNINGS
 
 	DOREPLIFETIME( AGameStateBase, bReplicatedHasBegunPlay );
 }
