Migrate fix for handling deformers not being used on all lods from USkinnedMeshComponent::IsSkinCacheAllowed() over to FSkinnedMeshSceneProxyDesc::IsSkinCacheAllowed()
#tests Played with deformers on a multi lod mesh

diff --git a/Engine/Source/Runtime/Engine/Private/SkinnedMeshSceneProxyDesc.cpp b/Engine/Source/Runtime/Engine/Private/SkinnedMeshSceneProxyDesc.cpp
--- a/Engine/Source/Runtime/Engine/Private/SkinnedMeshSceneProxyDesc.cpp
+++ b/Engine/Source/Runtime/Engine/Private/SkinnedMeshSceneProxyDesc.cpp
@@ -312,7 +312,7 @@ bool FSkinnedMeshSceneProxyDesc::IsSkinCacheAllowed(int32 LodIdx) const
 		return false;
 	}
 
-	if (GetMeshDeformerInstance() != nullptr)
+	if (GetMeshDeformerInstanceForLOD(LodIdx) != nullptr)
 	{
 		// Disable skin cache if a mesh deformer is in use.
 		// Any animation buffers are expected to be owned by the MeshDeformer.
