[SkeletalMesh HalfEdge] Make sure the deformer asset is added to the asset registry before querying info about it

diff --git a/Engine/Source/Runtime/Engine/Private/Rendering/SkeletalMeshHalfEdgeBufferAccessor.cpp b/Engine/Source/Runtime/Engine/Private/Rendering/SkeletalMeshHalfEdgeBufferAccessor.cpp
--- a/Engine/Source/Runtime/Engine/Private/Rendering/SkeletalMeshHalfEdgeBufferAccessor.cpp
+++ b/Engine/Source/Runtime/Engine/Private/Rendering/SkeletalMeshHalfEdgeBufferAccessor.cpp
@@ -22,10 +22,13 @@ bool SkeletalMeshHalfEdgeBufferAccessor::IsHalfEdgeRequired(TSoftObjectPtr<UObje
 		return false;
 	}
 	
-	const FAssetRegistryModule& AssetRegistryModule = FModuleManager::LoadModuleChecked<FAssetRegistryModule>("AssetRegistry");	
-
+	IAssetRegistry& AssetRegistry = FModuleManager::LoadModuleChecked<FAssetRegistryModule>("AssetRegistry").Get();
+	
+	// Always ask for scan in case no initial scan was done (Commandlets) or if AssetRegistry isn't done scanning.
+	AssetRegistry.ScanSynchronous({ }, { InAssetSoftPtr.GetLongPackageName()});
+		
 	FAssetData AssetData;
-	AssetRegistryModule.TryGetAssetByObjectPath(InAssetSoftPtr.ToSoftObjectPath(), AssetData);
+	AssetRegistry.TryGetAssetByObjectPath(InAssetSoftPtr.ToSoftObjectPath(), AssetData);
 
 	// Asset types that implement USkeletalMeshHalfEdgeBufferAccessor should have this tag,
 	// if not, it is likely old and requires a resave to have this information saved as asset tag
@@ -35,7 +38,7 @@ bool SkeletalMeshHalfEdgeBufferAccessor::IsHalfEdgeRequired(TSoftObjectPtr<UObje
 		return bIsHalfEdgeRequired;
 	}
 	
-	// Fallback path
+	// Fallback path if input is an old deformer graph asset
 	//
 	// If an asset do not have this tag serialized is likely old and haven't been re-saved recently.
 	// Therefore, we have to simply assume that the deformer needs half-edge buffer here to make sure it still works,
