PR #13545: Fixing UE-209483 a Stream Cache Handle Crash

diff --git a/Engine/Source/Runtime/Engine/Private/ContentStreaming.cpp b/Engine/Source/Runtime/Engine/Private/ContentStreaming.cpp
--- a/Engine/Source/Runtime/Engine/Private/ContentStreaming.cpp
+++ b/Engine/Source/Runtime/Engine/Private/ContentStreaming.cpp
@@ -1488,6 +1488,11 @@ FAudioChunkHandle::FAudioChunkHandle(const uint8* InData, uint32 NumBytes, const
 			CorrespondingWaveGuid = SoundWaveData->GetGUID();
 		}
 	}
+
+	if (IsValid())
+	{
+		IStreamingManager::Get().GetAudioStreamingManager().AddReferenceToChunk(*this);
+	}
 }
 
 FAudioChunkHandle::FAudioChunkHandle(const FAudioChunkHandle& Other)
@@ -1615,4 +1620,4 @@ FAudioChunkHandle IAudioStreamingManager::BuildChunkHandle(const uint8* InData,
 void IAudioStreamingManager::LogWarning()
 {
 	UE_LOG(LogContentStreaming, Warning, TEXT("This function is now deprecated on IAudioStreamingManager and no-ops."));
-}
+}
\ No newline at end of file
