PR #13589: Remove check preventing FSoundWaveProxy from using procedural USoundWaves

[RN] Removed invalid check on sound wave proxy not being a procedural sound wave

diff --git a/Engine/Source/Runtime/Engine/Private/SoundWave.cpp b/Engine/Source/Runtime/Engine/Private/SoundWave.cpp
--- a/Engine/Source/Runtime/Engine/Private/SoundWave.cpp
+++ b/Engine/Source/Runtime/Engine/Private/SoundWave.cpp
@@ -5160,8 +5160,6 @@ FSoundWaveProxy::FSoundWaveProxy(USoundWave* InWave)
 	// this should have been allocated by the USoundWave and should always be valid
 	check(SoundWaveDataPtr);
 
-	checkf(!InWave->bProcedural, TEXT("FSoundWaveProxy can not  be made from a procedural sound wave."));
-
 	bool bIsStreaming = InWave->IsStreaming(nullptr);
 
 	// non-streaming sources need resource data initialized before the FSoundWaveProxy
