[NetPhysics] Fix issue where inputs can get dropped when server input-buffer runs empty because only the latest received input gets injected at the front of the buffer.

Fix: When receiving inputs on the server, if any of them have not been received before but the server has produced an input for that frame already then merge the received input to the next received input and perform the same for that input. The last input with potential merged data from previous inputs gets injected at the front of the buffer if it's empty, to be used on the current frame.
No input will even be received without being either used or merged into another input that will be used.

diff --git a/Engine/Source/Runtime/Engine/Private/PhysicsEngine/NetworkPhysicsComponent.cpp b/Engine/Source/Runtime/Engine/Private/PhysicsEngine/NetworkPhysicsComponent.cpp
--- a/Engine/Source/Runtime/Engine/Private/PhysicsEngine/NetworkPhysicsComponent.cpp
+++ b/Engine/Source/Runtime/Engine/Private/PhysicsEngine/NetworkPhysicsComponent.cpp
@@ -1866,7 +1866,8 @@ void FAsyncNetworkPhysicsComponent::ConsumeAsyncInput(const int32 PhysicsStep)
 			auto ReceiveHelper = [&](Chaos::FBaseRewindHistory* History, Chaos::FBaseRewindHistory* ReceiveData, const bool bImportant, const bool bCompareData)
 			{
 				const bool bCompareDataForRewind = bCompareData && !bIsServer;
-				const int32 ResimFrame = History->ReceiveNewData(*ReceiveData, (bIsServer ? 0 : NetworkPhysicsTickOffset), bCompareDataForRewind, bImportant);
+				const int32 TryInjectAtFrame = bIsServer ? PhysicsStep : 0;
+				const int32 ResimFrame = History->ReceiveNewData(*ReceiveData, (bIsServer ? 0 : NetworkPhysicsTickOffset), bCompareDataForRewind, bImportant, TryInjectAtFrame);
 				if (bCompareDataForRewind)
 				{
 					TriggerResimulation(ResimFrame);
@@ -1889,24 +1890,6 @@ void FAsyncNetworkPhysicsComponent::ConsumeAsyncInput(const int32 PhysicsStep)
 			// Receive Inputs
 			if (AsyncInput->InputData && AsyncInput->InputData->HasDataInHistory())
 			{
-				/* Extract latest received input from client on the server, to be used if input buffer runs empty
-				* TODO, improve flow to not require this before ReceiveHelper */
-				FNetworkPhysicsBase* PhysicsData = nullptr;
-				if (bIsServer && LatestInputReceiveData)
-				{
-					PhysicsData = LatestInputReceiveData.Get();
-					if (AsyncInput->InputData->ExtractData(AsyncInput->InputData->GetLatestFrame(), false, PhysicsData, true) == false)
-					{
-						// Extraction failed
-						ensureMsgf(false, TEXT("Failed to extract latest input data from received inputs"));
-						PhysicsData = nullptr;
-
-#if DEBUG_NETWORK_PHYSICS
-						AsyncInput->InputData->DebugData(FString::Printf(TEXT("SERVER | PT | Failed to extract LatestInputReceiveData | LatestFrame: %d | Name: %s"), AsyncInput->InputData->GetLatestFrame(), *GetActorName()));
-#endif
-					}
-				}
-
 				// Validate data in the received inputs on the server
 				if (!ComponentSettings.GetValidateDataOnGameThread() && bIsServer)
 				{
@@ -1926,22 +1909,6 @@ void FAsyncNetworkPhysicsComponent::ConsumeAsyncInput(const int32 PhysicsStep)
 					}
 				}
 				ReceiveHelper(InputHistory.Get(), AsyncInput->InputData.Get(), /*bImportant*/false, bCompareInput);
-
-				/* If the server-side input history doesn't have any entries ahead of the current physics tick, the input buffer is empty, inject the latest received input as the input for the current tick.
-				* This happens during a desync where the client is far behind the server */
-				if (bIsServer && InputHistory->GetLatestFrame() < PhysicsStep && PhysicsData)
-				{
-#if DEBUG_NETWORK_PHYSICS
-					UE_LOG(LogChaos, Log, TEXT("SERVER | PT | Input Buffer Empty, Injecting Received Input at frame %d || LocalFrame = %d || ServerFrame = %d || bDataAltered = %d || Data: %s || Actor: %s")
-						, PhysicsStep, PhysicsData->LocalFrame, PhysicsData->ServerFrame, PhysicsData->bDataAltered, *PhysicsData->DebugData(), *GetActorName());
-#endif
-					
-					// Record data in InputHistory
-					PhysicsData->LocalFrame = PhysicsStep;
-					PhysicsData->ServerFrame = PhysicsStep;
-					InputHistory->RecordData(PhysicsStep, PhysicsData);
-				}
-
 			}
 
 			// Receive States
@@ -2104,8 +2071,7 @@ void FAsyncNetworkPhysicsComponent::OnPreProcessInputs_Internal(const int32 Phys
 				// Check if the extracted data was altered and if we have a hole in the buffer
 				if (bIsServer && PhysicsData->bDataAltered)
 				{
-					// Check if the buffer is empty
-					if (PhysicsStep <= InputHistory->GetLatestFrame())
+					if (PhysicsStep < InputHistory->GetLatestFrame())
 					{
 						// A missing input was detected and buffer is not empty, inform the owning client to send more inputs in each RPC to not get a gaps in the buffer
 						// NOTE: We don't send more extra inputs when the buffer runs empty since that case is corrected via time dilation, not sending extra inputs
@@ -2173,10 +2139,11 @@ void FAsyncNetworkPhysicsComponent::OnPreProcessInputs_Internal(const int32 Phys
 					}
 				}
 
-				// If the extracted input data was altered (merged, extrapolated, interpolated) on the server, record it into the history for it to get replicated to clients
+				// If the extracted input data was altered (merged, extrapolated, interpolated, injected) on the server, record it into the history for it to get replicated to clients
 				if (bIsServer && PhysicsData->bDataAltered)
 				{
-					PhysicsData->bReceivedData = true; // Mark the input data as received so that it doesn't get overwritten by incoming client inputs
+					// Explicitly say this input was not received, since it was altered by the server and when receiving the input for this frame it should get processed as altered but not received when calling ReceiveNewData()
+					PhysicsData->bReceivedData = false;
 					PhysicsData->bImportant = false;
 					InputHistory->RecordData(PhysicsStep, PhysicsData);
 				}
diff --git a/Engine/Source/Runtime/Engine/Public/Physics/NetworkPhysicsComponent.h b/Engine/Source/Runtime/Engine/Public/Physics/NetworkPhysicsComponent.h
--- a/Engine/Source/Runtime/Engine/Public/Physics/NetworkPhysicsComponent.h
+++ b/Engine/Source/Runtime/Engine/Public/Physics/NetworkPhysicsComponent.h
@@ -232,48 +232,103 @@ struct TNetRewindHistory : public Chaos::TDataRewindHistory<DataType>
 		return Copy;
 	}
 
-	virtual int32 ReceiveNewData(Chaos::FBaseRewindHistory& NewData, const int32 FrameOffset, bool CompareDataForRewind = false, const bool bImportant = false) override
+	virtual int32 ReceiveNewData(Chaos::FBaseRewindHistory& NewData, const int32 FrameOffset, bool CompareDataForRewind = false, const bool bImportant = false, int32 TryInjectAtFrame = INDEX_NONE) override
 	{
 		TNetRewindHistory& NetNewData = static_cast<TNetRewindHistory&>(NewData);
 
 		int32 RewindFrame = INDEX_NONE;
-		if (NetNewData.NumFrames > 0)
+		DataType& ReceiveData = NetNewData.GetAndLoadEarliestData();
+		
+		for (int32 Itr = 0; Itr < NetNewData.GetHistorySize(); Itr++)
 		{
-			for (int32 FrameIndex = 0; FrameIndex < NetNewData.NumFrames; ++FrameIndex)
+			DataType* NextReceiveData = NetNewData.GetAndLoadNextIncrementalData();
+
+			ReceiveData.bImportant = bImportant;
+			ReceiveData.bReceivedData = true; // Received data is marked to differentiate from locally predicted data
+
+			ReceiveData.LocalFrame = ReceiveData.ServerFrame - FrameOffset;
+
+			if (ShouldRecordReceivedDataOnFrame(ReceiveData, NextReceiveData))
 			{
-				DataType& FrameData = NetNewData.DataHistory[FrameIndex];
-				FrameData.bImportant = bImportant;
-				FrameData.bReceivedData = true; // Received data is marked to differentiate from locally predicted data
+				if (CompareDataForRewind && ReceiveData.LocalFrame > RewindFrame && TriggerRewindFromNewData(ReceiveData))
+				{
+					RewindFrame = ReceiveData.LocalFrame;
+				}
 
-				FrameData.LocalFrame = FrameData.ServerFrame - FrameOffset;
+				Super::RecordData(ReceiveData.LocalFrame, &ReceiveData);
+			}
 
-				if (ShouldRecordReceivedDataOnFrame(FrameData))
+			if (NextReceiveData)
+			{
+				ReceiveData = *NextReceiveData;
+			}
+			else
+			{
+				// Record a copy of the last data at specified frame if the history doesn't have data for that frame yet
+				if (TryInjectAtFrame > Super::GetLatestFrame())
 				{
-					if (CompareDataForRewind && FrameData.LocalFrame > RewindFrame && TriggerRewindFromNewData(FrameData))
-					{
-						RewindFrame = FrameData.LocalFrame;
-					}
+#if DEBUG_NETWORK_PHYSICS
+					UE_LOG(LogChaos, Log, TEXT("SERVER | PT | Input Buffer Empty, Injecting Received Input at frame %d || LocalFrame = %d || ServerFrame = %d || bDataAltered = %d || Data: %s")
+						, TryInjectAtFrame, ReceiveData.LocalFrame, ReceiveData.ServerFrame, ReceiveData.bDataAltered, *ReceiveData.DebugData());
+#endif
 
-					Super::RecordData(FrameData.LocalFrame, &FrameData);
+					DataType InjectData = ReceiveData;
+					InjectData.bReceivedData = false;
+					InjectData.bDataAltered = true;
+					InjectData.LocalFrame = TryInjectAtFrame;
+					InjectData.ServerFrame = TryInjectAtFrame + FrameOffset;
+
+					Super::RecordData(TryInjectAtFrame, &InjectData);
 				}
+
+				break;
 			}
 		}
+		
 		return RewindFrame;
 	}
 
 	/** Check if we should record received data into history.
 	* Can for example block received data from client from overriding server authoritative data */
-	virtual bool ShouldRecordReceivedDataOnFrame(const DataType& ReceivedData)
+	virtual bool ShouldRecordReceivedDataOnFrame(const DataType& ReceivedData, DataType* NextReceivedData = nullptr)
 	{
 		if (ReceivedData.LocalFrame < 0)
 		{
 			return false;
 		}
 
-		// Allow received data to get recorded if it's for a newer frame than already stored at the history index or if the data stored is not marked as received (i.e. it's locally predicted)
+		// Get the cached data at the index slot for the received data 
 		Super::LoadData(ReceivedData.LocalFrame);
 		DataType& Data = Super::DataHistory[Super::CurrentIndex];
-		return (Data.LocalFrame < ReceivedData.LocalFrame || (!Data.bReceivedData && Data.LocalFrame == ReceivedData.LocalFrame));
+
+		if (Data.LocalFrame < ReceivedData.LocalFrame)
+		{
+			// Allow recording the received data if it's newer than the already recorded data on this index
+			return true;
+		}
+		else if (!Data.bReceivedData && Data.LocalFrame == ReceivedData.LocalFrame)
+		{
+			// If the data exist but is not marked as received and marked as altered, the server has produced this input via extrapolation or interpolation, don't overwrite it with data from the client.
+			if (Data.bDataAltered)
+			{
+				// If we have a newer received data (since we receive multiple data at the same time) merge this data into the newer which can either get recorded or injected at the front of the input buffer on the server.
+				if (NextReceivedData != nullptr)
+				{
+					NextReceivedData->MergeData(ReceivedData);
+					NextReceivedData->bDataAltered = true;
+				}
+
+				// Mark the already cached data as received so that we don't perform this merge again when receiving redundant inputs / states
+				Data.bReceivedData = true; 
+				return false;
+			}
+			else
+			{
+				// Allow recording the received data if we already have data for the same frame cached, not marked as altered, meaning it was predicted data on the client
+				return true;
+			}
+		}
+		return false;
 	}
 
 	/** Compares new received data with local predicted data and returns true if they differ enough to trigger a resimulation  */
@@ -1065,10 +1120,13 @@ struct TNetworkPhysicsDataHelper : FNetworkPhysicsDataHelper
 	{
 		Chaos::TDataRewindHistory<DataType>* NetHistory = static_cast<Chaos::TDataRewindHistory<DataType>*>(To);
 
+		// Record at same index as taken from
+		int32 Idx = 0;
 		for (const TInstancedStruct<FNetworkPhysicsBase>& DataInstance : From.DataArray)
 		{
 			const DataType& Data = DataInstance.Get<DataType>();
-			NetHistory->RecordData(Data.LocalFrame, &Data);
+			NetHistory->RecordData(Idx, &Data);
+			Idx++;
 		}
 	}
 
@@ -1148,9 +1206,9 @@ struct TNetworkPhysicsDataHelper : FNetworkPhysicsDataHelper
 		Chaos::TDataRewindHistory<DataType>* NetHistory = static_cast<Chaos::TDataRewindHistory<DataType>*>(History);
 		TArray<DataType>& HistoryArray = NetHistory->GetDataHistory();
 
-		for (int32 FromIdx = 0; FromIdx < NetHistory->GetHistorySize(); FromIdx++)
+		for (int32 Idx = 0; Idx < NetHistory->GetHistorySize(); Idx++)
 		{
-			FNetworkPhysicsBase& Data = HistoryArray[FromIdx];
+			FNetworkPhysicsBase& Data = HistoryArray[Idx];
 			Interface.ValidateInput(Data);
 		}
 	}
diff --git a/Engine/Source/Runtime/Experimental/Chaos/Public/RewindData.h b/Engine/Source/Runtime/Experimental/Chaos/Public/RewindData.h
--- a/Engine/Source/Runtime/Experimental/Chaos/Public/RewindData.h
+++ b/Engine/Source/Runtime/Experimental/Chaos/Public/RewindData.h
@@ -98,8 +98,16 @@ struct FBaseRewindHistory
 	* @param EndFrame = Excluded */
 	virtual TUniquePtr<FBaseRewindHistory> CopyFramesWithOffset(const uint32 StartFrame, const uint32 EndFrame, const int32 FrameOffset) = 0;
 
-	/** Copy new data (received from the network) into this history, returns frame to resimulate from if @param CompareDataForRewind is set to true and compared data differ enough */
-	virtual int32 ReceiveNewData(FBaseRewindHistory& NewData, const int32 FrameOffset, const bool CompareDataForRewind = false, const bool bImportant = false) { return INDEX_NONE; }
+	/** Copy new data (received from the network) into this history
+	* Returns frame to resimulate from if @param CompareDataForRewind is set to true and compared data differ enough 
+	* @param TryInjectAtFrame puts the latest of the received data entries at the specified frame (if there isn't data for that frame yet), used to inject the latest data at the head of the buffer if it's empty */
+	virtual int32 ReceiveNewData(FBaseRewindHistory& NewData, const int32 FrameOffset, const bool CompareDataForRewind = false, const bool bImportant = false, int32 TryInjectAtFrame = INDEX_NONE) { return INDEX_NONE; }
+	
+	UE_DEPRECATED(5.7, "Deprecated, use the ReceiveNewData call with parameter InjectAtFrameIfEmpty, pass in nullptr to opt out of implementing a function.")
+	virtual int32 ReceiveNewData(FBaseRewindHistory& NewData, const int32 FrameOffset, const bool CompareDataForRewind = false, const bool bImportant = false) 
+	{
+		return ReceiveNewData(NewData, FrameOffset, CompareDataForRewind, bImportant, INDEX_NONE);
+	}
 
 	/** Serialize the data to or from a network archive */
 	UE_DEPRECATED(5.6, "Deprecated, use the NetSerialize call with parameter that takes a function DataSetupFunction, pass in nullptr to opt out of implementing a function.")
@@ -257,7 +265,7 @@ public :
 				const int32 DeltaFrame = FMath::Abs(ExtractFrame - DataHistory[MinFrameIndex].LocalFrame);
 
 				ExtractedData.LocalFrame = ExtractFrame;
-				ExtractedData.ServerFrame = DataHistory[MinFrameIndex].ServerFrame + DeltaFrame;
+				ExtractedData.ServerFrame += DeltaFrame;
 				ExtractedData.bDataAltered = true; // This history entry is now altered and doesn't correspond to the source entry
 
 #if DEBUG_NETWORK_PHYSICS
@@ -340,6 +348,39 @@ public :
 
 	const DataType& GetCurrentData() const { return DataHistory[CurrentIndex]; }
 
+	/** Returns the earliest data entry and loads it as the current data of the history */
+	DataType& GetAndLoadEarliestData()
+	{
+		int32 EarliestFrame = INT_MAX;
+		for (int32 FrameIndex = 0; FrameIndex < NumFrames; ++FrameIndex)
+		{
+			DataType& Data = DataHistory[FrameIndex];
+			if (Data.LocalFrame > INDEX_NONE && Data.LocalFrame < EarliestFrame)
+			{
+				CurrentIndex = FrameIndex;
+				CurrentFrame = Data.LocalFrame;
+				EarliestFrame = CurrentFrame;
+			}
+		}
+
+		return GetCurrentData();
+	}
+
+	/** Returns the next incremental data entry after the CurrentIndex (and loads it as the current data of the history), if there isn't a data entry with a higher LocalFrame value, nullptr is returned */
+	DataType* GetAndLoadNextIncrementalData()
+	{
+		const int32 NextIndex = (CurrentIndex + 1) % NumFrames;
+		if (DataHistory[NextIndex].LocalFrame > CurrentFrame)
+		{
+			CurrentIndex = NextIndex;
+			CurrentFrame = DataHistory[NextIndex].LocalFrame;
+
+			return &DataHistory[NextIndex];
+		}
+
+		return nullptr;
+	}
+
 	/** Get the number of valid data in the buffer index range */
 	FORCEINLINE uint32 NumValidData(const uint32 StartFrame, const uint32 EndFrame) const
 	{
