Fix for Material Aggregates compiling through the Material Attributes output pin of a Material in the new translator by retrieving compiled sub-properties from the aggregate value.
- Migrated attribute behaviour to the IR's material aggregate type and added default type setting to match old translator behaviour.
- Added default values to be assigned whenever an aggregate value does not have a compilation path/pin connected etc.
- Changed aggregate type to use a dynamic array which allows us to map Substrate and non-Substrate types, as well as fixing compilation order (i.e. Normal first).
- Ideally the aggregate should become ground truth and act as the interface for all attribute type mapping and the rest of the new translator should never see the AttributeDefintionMap, but this is outside the scope of this fix as it needs a considerable refactor.
- Altered shader stage logic to be mapped to all their properties correctly (previously WPO was in Vertex and everything else in Pixel).
- CustomisedUVs disabled in the new translator but previously were not functional anyway. Vertex stage components in general need revising but this is also outside the scope of this fix.
- Similarly SubsurfaceColor needs revising, but is also outside the scope of this change for now, ongoing discussion on slack. For now, this change creates parity between the old and new translator, unfortunately using similar hacks but it's the healthiest option currently, but will rip them out asap.
- Substrate entries only compile with Substrate enabled now, not written out to the shader at all otherwise.
- Fixed up some OutputValueType checks which are required for the various aggregate nodes to work correctly with the new translator.
- Integrated Massimo's additional fixes for mapping MCT_MaterialAttributes to appropriate the Aggregate type.

diff --git a/Engine/Source/Runtime/Engine/Private/Materials/MaterialAggregate.cpp b/Engine/Source/Runtime/Engine/Private/Materials/MaterialAggregate.cpp
--- a/Engine/Source/Runtime/Engine/Private/Materials/MaterialAggregate.cpp
+++ b/Engine/Source/Runtime/Engine/Private/Materials/MaterialAggregate.cpp
@@ -3,7 +3,10 @@
 #include "Materials/MaterialAggregate.h"
 #include "Framework/Notifications/NotificationManager.h"
 #include "Widgets/Notifications/SNotificationList.h"
+#include "Materials/MaterialIREmitter.h"
+#include "Materials/MaterialIRInternal.h"
 #include "Materials/MaterialAttributeDefinitionMap.h"
+#include "RenderUtils.h"
 #include "SceneTypes.h"
 #include "UObject/Package.h"
 
@@ -24,16 +27,17 @@ EMaterialAggregateAttributeType MaterialValueTypeToMaterialAggregateAttributeTyp
 		case MCT_Float2: return EMaterialAggregateAttributeType::Float2;
 		case MCT_Float3: return EMaterialAggregateAttributeType::Float3;
 		case MCT_Float4: return EMaterialAggregateAttributeType::Float4;
-		
-		case MCT_ShadingModel: return EMaterialAggregateAttributeType::ShadingModel;
-
-		case MCT_Bool: return EMaterialAggregateAttributeType::Bool1;
 	
+		case MCT_UInt:
 		case MCT_UInt1: return EMaterialAggregateAttributeType::UInt1;
 		case MCT_UInt2: return EMaterialAggregateAttributeType::UInt2;
 		case MCT_UInt3: return EMaterialAggregateAttributeType::UInt3;
 		case MCT_UInt4: return EMaterialAggregateAttributeType::UInt4;
-		case MCT_UInt: return EMaterialAggregateAttributeType::UInt4;
+
+		case MCT_Bool: return EMaterialAggregateAttributeType::Bool1;
+
+		case MCT_ShadingModel: return EMaterialAggregateAttributeType::ShadingModel;
+		case MCT_MaterialAttributes: return EMaterialAggregateAttributeType::MaterialAttributes;
 	
 		default:
 			checkf(false, TEXT("This material value type '%d' cannot be expressed to a material aggregate attribute type."), (int32)Type);
@@ -61,7 +65,7 @@ EMaterialValueType FMaterialAggregateAttribute::ToMaterialValueType() const
 		case EMaterialAggregateAttributeType::Float4: return MCT_Float4;
 		
 		case EMaterialAggregateAttributeType::ShadingModel: return MCT_ShadingModel;
-		
+		case EMaterialAggregateAttributeType::MaterialAttributes: return MCT_MaterialAttributes;
 		case EMaterialAggregateAttributeType::Aggregate:
 			return Aggregate.Get() == UMaterialAggregate::GetMaterialAttributes() ? MCT_MaterialAttributes : MCT_Unknown;
 		
@@ -128,46 +132,92 @@ static void ClearCyclicDependencies(UMaterialAggregate* Aggregate, bool bDisplay
 	}
 }
 
-static FMaterialAggregateAttribute MaterialPropertyToAttribute(EMaterialProperty Property)
+static FMaterialAggregateAttribute MakeAttributeFromMaterialProperty(EMaterialProperty Property)
 { 
 	FMaterialAggregateAttribute Attribute;
-	Attribute.Name = *FMaterialAttributeDefinitionMap::GetAttributeName(Property);
-	Attribute.Type = MaterialValueTypeToMaterialAggregateAttributeType(FMaterialAttributeDefinitionMap::GetValueType(Property));
-	Attribute.DefaultValue = FMaterialAttributeDefinitionMap::GetDefaultValue(Property);
+	Attribute.Name = (Property == MP_SubsurfaceColor) ? TEXT("Subsurface") : *FMaterialAttributeDefinitionMap::GetAttributeName(Property);
+	Attribute.Type = MaterialValueTypeToMaterialAggregateAttributeType((Property == MP_SubsurfaceColor) ? MCT_Float4 : FMaterialAttributeDefinitionMap::GetValueType(Property));
+	
+	//ShadingModel defaults to 0, i.e. unlit. This is overriden for the old material system via editor only data, but we can't access that here so we force the value to Lit(1.0)
+	Attribute.DefaultValue = (Property == MP_ShadingModel) ? FVector4f(1.0f, 0.0f, 0.0f, 0.0f) : FMaterialAttributeDefinitionMap::GetDefaultValue(Property);
+	
 	return Attribute;
 }
 
-// This array declares the material properties that are referenced by each attribute in the MaterialAttributes aggregate.
-static constexpr EMaterialProperty MaterialAttributesProperties[] = {
-	MP_BaseColor,
-	MP_Metallic,
-	MP_Specular,
-	MP_Roughness,
-	MP_Anisotropy,
-	MP_EmissiveColor,
-	MP_Opacity,
-	MP_OpacityMask,
-	MP_Normal,
-	MP_Tangent,
-	MP_WorldPositionOffset,
-	MP_SubsurfaceColor,
-	MP_CustomData0,
-	MP_CustomData1,
-	MP_AmbientOcclusion,
-	MP_Refraction,
-	MP_CustomizedUVs0,
-	MP_CustomizedUVs1,
-	MP_CustomizedUVs2,
-	MP_CustomizedUVs3,
-	MP_CustomizedUVs4,
-	MP_CustomizedUVs5,
-	MP_CustomizedUVs6,
-	MP_CustomizedUVs7,
-	MP_PixelDepthOffset,
-	MP_ShadingModel,
-	MP_Displacement,
+struct FMaterialAttributePropertyIndexMap
+{
+	// Constant time lookups for Property/Index mapping
+	EMaterialProperty IndexToProperty[MP_MAX];
+	int PropertyToIndex[MP_MAX];
+
+	int NumAttributes = 0;
+
+	static const FMaterialAttributePropertyIndexMap& Get()
+	{
+		static const FMaterialAttributePropertyIndexMap MaterialAttributeProperties;
+		return MaterialAttributeProperties;
+	}
+
+	// This map is the ground truth for which material properties we compile in the new compiler.
+	// Dynamically extendable whilst retaining constant time referencing.
+	FMaterialAttributePropertyIndexMap()
+	{
+		//Initialise all empty values to defaults so we aren't referencing incorrect data after creating the maps
+		for (int32 ArrayIndex = 0; ArrayIndex < MP_MAX; ArrayIndex++)
+		{
+			IndexToProperty[ArrayIndex] = MP_MAX;
+			PropertyToIndex[ArrayIndex] = INDEX_NONE;
+		}
+
+		// The normal input is read back from the value set in the material attribute.
+		// For this reason, the normal attribute is evaluated and set first, ensuring that
+		// other inputs can read its value.
+		PushAttribute(MP_Normal);
+
+		PushAttribute(MP_BaseColor);
+		PushAttribute(MP_Metallic);
+		PushAttribute(MP_Specular);
+		PushAttribute(MP_Roughness);
+		PushAttribute(MP_Anisotropy);
+		PushAttribute(MP_EmissiveColor);
+		PushAttribute(MP_Opacity);
+		PushAttribute(MP_OpacityMask);
+		PushAttribute(MP_Tangent);
+		PushAttribute(MP_WorldPositionOffset);
+		PushAttribute(MP_SubsurfaceColor);
+		PushAttribute(MP_CustomData0);
+		PushAttribute(MP_CustomData1);
+		PushAttribute(MP_AmbientOcclusion);
+		PushAttribute(MP_Refraction);
+		PushAttribute(MP_PixelDepthOffset);
+		PushAttribute(MP_ShadingModel);
+		PushAttribute(MP_Displacement);
+
+		// Bridge legacy material attributes to Substrate's inputs last (when required)
+		if (Substrate::IsSubstrateEnabled())
+		{			
+			PushAttribute(MP_FrontMaterial);
+			PushAttribute(MP_SurfaceThickness);
+		}
+	}
+
+	void PushAttribute(EMaterialProperty Property)
+	{
+		IndexToProperty[NumAttributes] = Property;
+		PropertyToIndex[Property] = NumAttributes++;
+	}
+
+	bool IsValidIndex(int32 Index) const
+	{
+		return Index >= 0 && Index < NumAttributes;
+	}
 };
 
+const TConstArrayView<EMaterialProperty> UMaterialAggregate::GetMaterialAttributesProperties()
+{
+	return { FMaterialAttributePropertyIndexMap::Get().IndexToProperty, FMaterialAttributePropertyIndexMap::Get().NumAttributes};
+}
+
 const UMaterialAggregate* UMaterialAggregate::GetMaterialAttributes()
 {
 	static TObjectPtr<UMaterialAggregate> Instance = []
@@ -175,55 +225,35 @@ const UMaterialAggregate* UMaterialAggregate::GetMaterialAttributes()
 		TObjectPtr<UMaterialAggregate> MA = NewObject<UMaterialAggregate>(GetTransientPackage());
 		MA->Rename(TEXT("MaterialAttributes"));
 		MA->AddToRoot();
-		for (EMaterialProperty Property : MaterialAttributesProperties)
+		for (EMaterialProperty Property : GetMaterialAttributesProperties())
 		{
-			MA->Attributes.Add(MaterialPropertyToAttribute(Property));
+			MA->Attributes.Add(MakeAttributeFromMaterialProperty(Property));
 		}
 		return MA;
 	}();
 	return Instance.Get();
 }
 
-int32 UMaterialAggregate::MaterialPropertyToAttributeIndex(EMaterialProperty Property)
+const FMaterialAggregateAttribute* UMaterialAggregate::GetMaterialAttribute(EMaterialProperty Property)
 {
-	switch (Property)
+	int32 PropertyIndex = MaterialPropertyToAttributeIndex(Property);
+	if (FMaterialAttributePropertyIndexMap::Get().IsValidIndex(PropertyIndex))
 	{
-		case MP_BaseColor: return 0;
-		case MP_Metallic: return 1;
-		case MP_Specular: return 2;
-		case MP_Roughness: return 3;
-		case MP_Anisotropy: return 4;
-		case MP_EmissiveColor: return 5;
-		case MP_Opacity: return 6;
-		case MP_OpacityMask: return 7;
-		case MP_Normal: return 8;
-		case MP_Tangent: return 9;
-		case MP_WorldPositionOffset: return 10;
-		case MP_SubsurfaceColor: return 11;
-		case MP_CustomData0: return 12;
-		case MP_CustomData1: return 13;
-		case MP_AmbientOcclusion: return 14;
-		case MP_Refraction: return 15;
-		case MP_CustomizedUVs0: return 16;
-		case MP_CustomizedUVs1: return 17;
-		case MP_CustomizedUVs2: return 18;
-		case MP_CustomizedUVs3: return 19;
-		case MP_CustomizedUVs4: return 20;
-		case MP_CustomizedUVs5: return 21;
-		case MP_CustomizedUVs6: return 22;
-		case MP_CustomizedUVs7: return 23;
-		case MP_PixelDepthOffset: return 24;
-		case MP_ShadingModel: return 25;
-		case MP_Displacement: return 26;
-		default: return INDEX_NONE;
+		return &GetMaterialAttributes()->Attributes[PropertyIndex];
 	}
+	return nullptr;
+}
+
+int32 UMaterialAggregate::MaterialPropertyToAttributeIndex(EMaterialProperty Property)
+{
+	check(Property < MP_MAX);
+	return (Property < MP_MAX) ? FMaterialAttributePropertyIndexMap::Get().PropertyToIndex[Property] : INDEX_NONE;
 }
 
 EMaterialProperty UMaterialAggregate::AttributeIndexToMaterialProperty(int32 MaterialAttributeIndex)
 {
-	return MaterialAttributeIndex < UE_ARRAY_COUNT(MaterialAttributesProperties)
-		? MaterialAttributesProperties[MaterialAttributeIndex]
-		: MP_MAX;
+	ensure(FMaterialAttributePropertyIndexMap::Get().IsValidIndex(MaterialAttributeIndex));
+	return FMaterialAttributePropertyIndexMap::Get().IndexToProperty[MaterialAttributeIndex];
 }
 
 int32 UMaterialAggregate::FindAttributeIndexByName(FName Name) const
diff --git a/Engine/Source/Runtime/Engine/Private/Materials/MaterialExpressions.cpp b/Engine/Source/Runtime/Engine/Private/Materials/MaterialExpressions.cpp
--- a/Engine/Source/Runtime/Engine/Private/Materials/MaterialExpressions.cpp
+++ b/Engine/Source/Runtime/Engine/Private/Materials/MaterialExpressions.cpp
@@ -17940,6 +17940,16 @@ EMaterialValueType UMaterialExpressionMaterialFunctionCall::GetInputValueType(in
 	}
 	return MCT_Unknown;
 }
+
+EMaterialValueType UMaterialExpressionMaterialFunctionCall::GetOutputValueType(int32 OutputIndex)
+{
+	if (OutputIndex >= 0 && OutputIndex < FunctionOutputs.Num() && FunctionOutputs[OutputIndex].ExpressionOutput)
+	{
+		return FunctionOutputs[OutputIndex].ExpressionOutput->GetOutputValueType(0);
+	}
+	return Super::GetOutputValueType(OutputIndex);
+}
+
 #endif // WITH_EDITOR
 
 ///////////////////////////////////////////////////////////////////////////////
@@ -18559,6 +18569,16 @@ EMaterialValueType UMaterialExpressionFunctionOutput::GetInputValueType(int32 In
 	return static_cast<EMaterialValueType>(MCT_Float | MCT_MaterialAttributes | MCT_Substrate);
 }
 
+EMaterialValueType UMaterialExpressionFunctionOutput::GetOutputValueType(int32 OutputIndex)
+{
+	const FExpressionInput& ExpressionInputRef = A.GetTracedInput();
+	if (ExpressionInputRef.Expression)
+	{
+		return ExpressionInputRef.Expression->GetOutputValueType(ExpressionInputRef.OutputIndex);
+	}
+	return Super::GetOutputValueType(OutputIndex);
+}
+
 int32 UMaterialExpressionFunctionOutput::Compile(class FMaterialCompiler* Compiler, int32 OutputIndex)
 {
 	if (!A.GetTracedInput().Expression)
@@ -24517,6 +24537,11 @@ EMaterialValueType UMaterialExpressionAggregate::GetOutputValueType(int32 Output
 	return GetPinValueType(this, OutputIndex);
 }
 
+bool UMaterialExpressionAggregate::IsResultMaterialAttributes(int32 OutputIndex)
+{
+	return GetOutputValueType(OutputIndex) == MCT_MaterialAttributes;
+}
+
 void UMaterialExpressionAggregate::PreEditChange(FProperty* PropertyAboutToChange)
 {
 	PreEditData.Aggregate = GetAggregate();
diff --git a/Engine/Source/Runtime/Engine/Private/Materials/MaterialIR.cpp b/Engine/Source/Runtime/Engine/Private/Materials/MaterialIR.cpp
--- a/Engine/Source/Runtime/Engine/Private/Materials/MaterialIR.cpp
+++ b/Engine/Source/Runtime/Engine/Private/Materials/MaterialIR.cpp
@@ -3,6 +3,7 @@
 #include "Materials/MaterialIR.h"
 #include "Materials/MaterialIRTypes.h"
 #include "Materials/MaterialAggregate.h"
+#include "Materials/MaterialAttributeDefinitionMap.h"
 #include "Misc/LargeWorldRenderPosition.h"
 #include "Engine/Texture.h"
 
@@ -24,8 +25,8 @@ const TCHAR* LexToString(EStage Stage)
 }
 
 bool MaterialOutputEvaluatesInStage(EMaterialProperty Property, EStage Stage)
-{
-	if (Property == MP_WorldPositionOffset)
+{	
+	if (FMaterialAttributeDefinitionMap::GetShaderFrequency(Property) == EShaderFrequency::SF_Vertex)
 	{
 		return Stage == Stage_Vertex;
 	}
diff --git a/Engine/Source/Runtime/Engine/Private/Materials/MaterialIREmitter.cpp b/Engine/Source/Runtime/Engine/Private/Materials/MaterialIREmitter.cpp
--- a/Engine/Source/Runtime/Engine/Private/Materials/MaterialIREmitter.cpp
+++ b/Engine/Source/Runtime/Engine/Private/Materials/MaterialIREmitter.cpp
@@ -397,6 +397,7 @@ FType FEmitter::GetMaterialAggregateAttributeType(const UMaterialAggregate* Aggr
 		case EMaterialAggregateAttributeType::Float2: return FType::MakeVector(EScalarKind::Float, 2);
 		case EMaterialAggregateAttributeType::Float3: return FType::MakeVector(EScalarKind::Float, 3);
 		case EMaterialAggregateAttributeType::Float4: return FType::MakeVector(EScalarKind::Float, 4);
+		case EMaterialAggregateAttributeType::MaterialAttributes: return FType::MakeAggregate(UMaterialAggregate::GetMaterialAttributes());
 		case EMaterialAggregateAttributeType::Aggregate: return FType::MakeAggregate(Aggregate->Attributes[AttributeIndex].Aggregate.Get());
 		default: UE_MIR_UNREACHABLE();
 	}
@@ -1246,21 +1247,22 @@ static FValueRef EmitAttributeDefaultValue(FEmitter& Emitter, const FMaterialAgg
 {
 	switch (Attribute.Type)
 	{
-		case EMaterialAggregateAttributeType::Bool1: return Emitter.ConstantBool((bool)Attribute.DefaultValue.X);
-		case EMaterialAggregateAttributeType::Bool2: return Emitter.ConstantBool2((bool)Attribute.DefaultValue.X, (bool)Attribute.DefaultValue.Y);
-		case EMaterialAggregateAttributeType::Bool3: return Emitter.ConstantBool3((bool)Attribute.DefaultValue.X, (bool)Attribute.DefaultValue.Y, (bool)Attribute.DefaultValue.Z); 
-		case EMaterialAggregateAttributeType::Bool4: return Emitter.ConstantBool4((bool)Attribute.DefaultValue.X, (bool)Attribute.DefaultValue.Y, (bool)Attribute.DefaultValue.Z, (bool)Attribute.DefaultValue.W);
-		case EMaterialAggregateAttributeType::ShadingModel:
-		case EMaterialAggregateAttributeType::UInt1: return Emitter.ConstantInt((TInteger)Attribute.DefaultValue.X);
-		case EMaterialAggregateAttributeType::UInt2: return Emitter.ConstantInt2({ (TInteger)Attribute.DefaultValue.X, (TInteger)Attribute.DefaultValue.Y });
-		case EMaterialAggregateAttributeType::UInt3: return Emitter.ConstantInt3({ (TInteger)Attribute.DefaultValue.X, (TInteger)Attribute.DefaultValue.Y, (TInteger)Attribute.DefaultValue.Z });
-		case EMaterialAggregateAttributeType::UInt4: return Emitter.ConstantInt4({ (TInteger)Attribute.DefaultValue.X, (TInteger)Attribute.DefaultValue.Y, (TInteger)Attribute.DefaultValue.Z, (TInteger)Attribute.DefaultValue.W });
-		case EMaterialAggregateAttributeType::Float1: return Emitter.ConstantFloat((TFloat)Attribute.DefaultValue.X);
-		case EMaterialAggregateAttributeType::Float2: return Emitter.ConstantFloat2({ (TFloat)Attribute.DefaultValue.X, (TFloat)Attribute.DefaultValue.Y });
-		case EMaterialAggregateAttributeType::Float3: return Emitter.ConstantFloat3({ (TFloat)Attribute.DefaultValue.X, (TFloat)Attribute.DefaultValue.Y, (TFloat)Attribute.DefaultValue.Z });
-		case EMaterialAggregateAttributeType::Float4: return Emitter.ConstantFloat4(UE::Math::TVector4<TFloat>{ Attribute.DefaultValue.X, Attribute.DefaultValue.Y, Attribute.DefaultValue.Z, Attribute.DefaultValue.W });
-		case EMaterialAggregateAttributeType::Aggregate: return Emitter.Aggregate(Attribute.Aggregate);
-		default: UE_MIR_UNREACHABLE();
+	case EMaterialAggregateAttributeType::Bool1: return Emitter.ConstantBool((bool)Attribute.DefaultValue.X);
+	case EMaterialAggregateAttributeType::Bool2: return Emitter.ConstantBool2((bool)Attribute.DefaultValue.X, (bool)Attribute.DefaultValue.Y);
+	case EMaterialAggregateAttributeType::Bool3: return Emitter.ConstantBool3((bool)Attribute.DefaultValue.X, (bool)Attribute.DefaultValue.Y, (bool)Attribute.DefaultValue.Z);
+	case EMaterialAggregateAttributeType::Bool4: return Emitter.ConstantBool4((bool)Attribute.DefaultValue.X, (bool)Attribute.DefaultValue.Y, (bool)Attribute.DefaultValue.Z, (bool)Attribute.DefaultValue.W);
+	case EMaterialAggregateAttributeType::ShadingModel:
+	case EMaterialAggregateAttributeType::UInt1: return Emitter.ConstantInt((TInteger)Attribute.DefaultValue.X);
+	case EMaterialAggregateAttributeType::UInt2: return Emitter.ConstantInt2({ (TInteger)Attribute.DefaultValue.X, (TInteger)Attribute.DefaultValue.Y });
+	case EMaterialAggregateAttributeType::UInt3: return Emitter.ConstantInt3({ (TInteger)Attribute.DefaultValue.X, (TInteger)Attribute.DefaultValue.Y, (TInteger)Attribute.DefaultValue.Z });
+	case EMaterialAggregateAttributeType::UInt4: return Emitter.ConstantInt4({ (TInteger)Attribute.DefaultValue.X, (TInteger)Attribute.DefaultValue.Y, (TInteger)Attribute.DefaultValue.Z, (TInteger)Attribute.DefaultValue.W });
+	case EMaterialAggregateAttributeType::Float1: return Emitter.ConstantFloat((TFloat)Attribute.DefaultValue.X);
+	case EMaterialAggregateAttributeType::Float2: return Emitter.ConstantFloat2({ (TFloat)Attribute.DefaultValue.X, (TFloat)Attribute.DefaultValue.Y });
+	case EMaterialAggregateAttributeType::Float3: return Emitter.ConstantFloat3({ (TFloat)Attribute.DefaultValue.X, (TFloat)Attribute.DefaultValue.Y, (TFloat)Attribute.DefaultValue.Z });
+	case EMaterialAggregateAttributeType::Float4: return Emitter.ConstantFloat4(UE::Math::TVector4<TFloat>{ Attribute.DefaultValue.X, Attribute.DefaultValue.Y, Attribute.DefaultValue.Z, Attribute.DefaultValue.W });
+	case EMaterialAggregateAttributeType::MaterialAttributes: return Emitter.Aggregate(UMaterialAggregate::GetMaterialAttributes());
+	case EMaterialAggregateAttributeType::Aggregate: return Emitter.Aggregate(Attribute.Aggregate);
+	default: UE_MIR_UNREACHABLE();
 	}
 }
 
diff --git a/Engine/Source/Runtime/Engine/Private/Materials/MaterialIRInternal.cpp b/Engine/Source/Runtime/Engine/Private/Materials/MaterialIRInternal.cpp
--- a/Engine/Source/Runtime/Engine/Private/Materials/MaterialIRInternal.cpp
+++ b/Engine/Source/Runtime/Engine/Private/Materials/MaterialIRInternal.cpp
@@ -10,80 +10,6 @@
 
 namespace MIR::Internal {
 
-bool IsMaterialPropertyEnabled(EMaterialProperty InProperty)
-{
-	switch (InProperty)
-	{
-		case MP_Normal:
-		case MP_Tangent:
-		case MP_EmissiveColor:
-		case MP_Opacity:
-		case MP_OpacityMask:
-		case MP_BaseColor:
-		case MP_Metallic:
-		case MP_Specular:
-		case MP_Roughness:
-		case MP_Anisotropy:
-		case MP_AmbientOcclusion:
-		case MP_Refraction:
-		case MP_PixelDepthOffset:
-		case MP_SubsurfaceColor:
-		case MP_ShadingModel:
-		case MP_SurfaceThickness:
-		case MP_FrontMaterial:
-		case MP_Displacement:
-		case MP_WorldPositionOffset:
-		case MP_CustomData0:
-		case MP_CustomData1:
-			return true;
-		default:
-			return false;
-	}
-}
-
-bool NextMaterialAttributeInput(UMaterial* BaseMaterial, int32& PropertyIndex)
-{
-	for (; PropertyIndex < MP_MAX; ++PropertyIndex)
-	{
-		EMaterialProperty Property = (EMaterialProperty)PropertyIndex;
-		if (MIR::Internal::IsMaterialPropertyEnabled(Property) 
-			&& Property != MP_SubsurfaceColor
-			&& Property != MP_FrontMaterial)
-		{
-			return true;
-		}
-	}
-
-	return false;
-}
-
-MIR::FValue* CreateMaterialAttributeDefaultValue(FEmitter& Emitter, const UMaterial* Material, EMaterialProperty Property)
-{
-	EMaterialValueType Type = FMaterialAttributeDefinitionMap::GetValueType(Property);
-	FVector4f DefaultValue = FMaterialAttributeDefinitionMap::GetDefaultValue(Property);
-
-	switch (Type)
-	{
-		case MCT_ShadingModel: return Emitter.ConstantInt(Material->GetShadingModels().GetFirstShadingModel());
-
-		case MCT_Float1: return Emitter.ConstantFloat(  DefaultValue.X );
-		case MCT_Float2: return Emitter.ConstantFloat2({ DefaultValue.X, DefaultValue.Y });
-		case MCT_Float3: return Emitter.ConstantFloat3({ DefaultValue.X, DefaultValue.Y, DefaultValue.Z });
-		case MCT_Float:
-		case MCT_Float4: return Emitter.ConstantFloat4(DefaultValue);
-
-		case MCT_UInt1: return Emitter.ConstantInt((TInteger)DefaultValue.X);
-		case MCT_UInt2: return Emitter.ConstantInt2({ (TInteger)DefaultValue.X, (TInteger)DefaultValue.Y });
-		case MCT_UInt3: return Emitter.ConstantInt3({ (TInteger)DefaultValue.X, (TInteger)DefaultValue.Y, (TInteger)DefaultValue.Z });
-		case MCT_UInt:
-		case MCT_UInt4: return Emitter.ConstantInt4({ (TInteger)DefaultValue.X, (TInteger)DefaultValue.Y, (TInteger)DefaultValue.Z, (TInteger)DefaultValue.W });
-
-		case MCT_Substrate: return Emitter.PromoteSubstrateParameter();
-
-		default: UE_MIR_UNREACHABLE();
-	}
-}
-
 EMaterialTextureParameterType TextureMaterialValueTypeToParameterType(EMaterialValueType Type)
 {
 	switch (Type)
diff --git a/Engine/Source/Runtime/Engine/Private/Materials/MaterialIRInternal.h b/Engine/Source/Runtime/Engine/Private/Materials/MaterialIRInternal.h
--- a/Engine/Source/Runtime/Engine/Private/Materials/MaterialIRInternal.h
+++ b/Engine/Source/Runtime/Engine/Private/Materials/MaterialIRInternal.h
@@ -7,15 +7,6 @@
 #if WITH_EDITOR
 
 namespace MIR::Internal {
-
-//
-bool IsMaterialPropertyEnabled(EMaterialProperty InProperty);
-
-//
-bool NextMaterialAttributeInput(UMaterial* BaseMaterial, int32& PropertyIndex);
-
-//
-FValue* CreateMaterialAttributeDefaultValue(FEmitter& Emitter, const UMaterial* Material, EMaterialProperty Property);
  
 //
 EMaterialTextureParameterType TextureMaterialValueTypeToParameterType(EMaterialValueType Type);
diff --git a/Engine/Source/Runtime/Engine/Private/Materials/MaterialIRModuleBuilder.cpp b/Engine/Source/Runtime/Engine/Private/Materials/MaterialIRModuleBuilder.cpp
--- a/Engine/Source/Runtime/Engine/Private/Materials/MaterialIRModuleBuilder.cpp
+++ b/Engine/Source/Runtime/Engine/Private/Materials/MaterialIRModuleBuilder.cpp
@@ -11,6 +11,7 @@
 #include "Materials/MaterialIRTypes.h"
 #include "Materials/MaterialIREmitter.h"
 #include "Materials/MaterialIRInternal.h"
+#include "Materials/MaterialAggregate.h"
 #include "Materials/MaterialAttributeDefinitionMap.h"
 #include "Materials/MaterialExpressionFunctionInput.h"
 #include "Materials/MaterialExpressionFunctionOutput.h"
@@ -92,6 +93,7 @@ struct FMaterialIRModuleBuilderImpl
 	MIR::FEmitter* Emitter;
 	TArray<FAnalysisContext> AnalysisContextStack;
 	FMaterialIRValueAnalyzer ValueAnalyzer;
+	MIR::FValue* DefaultMaterialAggregate;
 	FColorMaterialInput PreviewInput;
 
 	void Step_Initialize()
@@ -104,43 +106,45 @@ struct FMaterialIRModuleBuilderImpl
 		Emitter->Initialize();
 		AnalysisContextStack.Emplace();
 
+		DefaultMaterialAggregate = Emitter->Aggregate(UMaterialAggregate::GetMaterialAttributes());
+
 		PreviewInput.Expression = Builder->PreviewExpression;
 	}
 
 	void Step_GenerateOutputInstructions()
 	{
-		// Helper local function to process a single material attribute output.
-		auto PrepareSingleMaterialAttribute = [this](EMaterialProperty Property)
+		//If we are processing from the Material Attributes output node, push the dependency here.
+		const bool bUseMaterialAttributes = Builder->Material->bUseMaterialAttributes;
+		if (bUseMaterialAttributes)
 		{
-			// Read the material input associated to this property
 			FMaterialInputDescription Input;
-			GetExpressionInputDescription(Property, Input);
-
-			// Push the connected expression to this material attribute input as a dependency, if any.
+			GetExpressionInputDescription(MP_MaterialAttributes, Input);
 			PushDependency(Input.Input->Expression);
+		}
 
-			// Create the set material output instruction (leave it unset for now; it is set in Step_FlowValuesIntoMaterialOutputs)
-			Emitter->SetMaterialOutput(Property, nullptr);
-		};
-		
-		// The normal input is read back from the value set in the material attribute.
-		// For this reason, the normal attribute is evaluated and set first, ensuring that
-		// other inputs can read its value.
-		PrepareSingleMaterialAttribute(MP_Normal);
-
-		// Then prepare all the other material attributes.
-		for (int32 Index = 0; MIR::Internal::NextMaterialAttributeInput(Builder->Material, Index); ++Index)
+		// Helper local function to process a single material attribute output.
+		auto PrepareSingleMaterialAttribute = [&](EMaterialProperty Property)
 		{
-			if (Index != MP_Normal)
+			//If we are not using the Material Attributes output node, then process the properties as normal
+			if (!bUseMaterialAttributes)
 			{
-				PrepareSingleMaterialAttribute((EMaterialProperty)Index);
+				// Read the material input associated to this property
+				FMaterialInputDescription Input;
+				GetExpressionInputDescription(Property, Input);
+
+				// Push the connected expression to this material attribute input as a dependency, if any.
+				PushDependency(Input.Input->Expression);
 			}
-		}
 
-		if (Substrate::IsSubstrateEnabled())
+			// Create the set material output instruction (leave it unset for now; it is set in Step_FlowValuesIntoMaterialOutputs)
+			// Note: Material Attributes need to still map the output correctly at the end so this is always set.
+			Emitter->SetMaterialOutput(Property, nullptr);
+		};
+
+		// Prepare the material attributes.
+		for (EMaterialProperty Property : UMaterialAggregate::GetMaterialAttributesProperties())
 		{
-			// Bridge legacy material attributes to Substrate's "FrontMaterial" input last
-			PrepareSingleMaterialAttribute(MP_FrontMaterial);
+			PrepareSingleMaterialAttribute(Property);
 		}
 	}
 
@@ -342,12 +346,6 @@ struct FMaterialIRModuleBuilderImpl
 			// Get the value flowing into the function output inside the function in the current context.
 			MIR::FValue* Value = Emitter->Input(FunctionOutput.ExpressionOutput->GetInput(0));
 
-			// Get the function output type.
-			MIR::FType OutputType = MIR::FType::FromMaterialValueType(FunctionOutput.ExpressionOutput->GetOutputValueType(0));
-
-			// Cast the value to the expected output type. This may fail (value will be poison). 
-			Value = Emitter->Cast(Value, OutputType);
-
 			// And flow it to the relative function *call* output in the parent context.
 			ParentContext.SetOutputValue(Call->GetOutput(i), Value);
 		}
@@ -361,23 +359,47 @@ struct FMaterialIRModuleBuilderImpl
 
 	void Step_FlowValuesIntoMaterialOutputs()
 	{
+		// Get the material attribute input info.
+		const bool bUseMaterialAttributes = Builder->Material->bUseMaterialAttributes;
+		FMaterialInputDescription MaterialAttributesInput;
+		MIR::FValue* MaterialAttributesValue{};
+		if (bUseMaterialAttributes)
+		{
+			ensure(GetExpressionInputDescription(MP_MaterialAttributes, MaterialAttributesInput));
+
+			FlowValueThroughConnection(MaterialAttributesInput.Input);
+
+			// Fetch the value from the material attributes input.
+			MaterialAttributesValue = MIR::Internal::FetchValueFromExpressionInput(this, MaterialAttributesInput.Input);
+			if (!MaterialAttributesValue)
+			{
+				// If none there, use the default aggregate
+				MaterialAttributesValue = DefaultMaterialAggregate;
+			}
+
+			// Make sure a valid value is present and it is of the correct type.
+			check(!MaterialAttributesValue->IsPoison() && MaterialAttributesValue->Type.AsAggregate() == UMaterialAggregate::GetMaterialAttributes());
+		}
+
 		for (uint32 StageIndex = 0; StageIndex < MIR::NumStages; ++StageIndex)
 		{
 			for (MIR::FSetMaterialOutput* Output : Module->Outputs[StageIndex])
 			{
-				// Get the material attribute input info.
 				FMaterialInputDescription Input;
 				ensure(GetExpressionInputDescription(Output->Property, Input));
 
+				if (bUseMaterialAttributes)
+				{
+					// If material is in material attributes mode, extract this property's attribute from the material attributes aggregate value
+					// and manually flow it into the this property's material expression input pin.
+					MIR::FValue* PropertyValue = Emitter->Subscript(MaterialAttributesValue, UMaterialAggregate::MaterialPropertyToAttributeIndex(Output->Property));
+					MIR::Internal::BindValueToExpressionInput(this, Input.Input, PropertyValue);
+				}
+
 				// Flow, if any, the value coming in through the connection to the material attribute pin.
 				FlowValueThroughConnection(Input.Input);
 
-				if (Input.bUseConstant)
-				{
-					// If input is marked to use constant, assign this output to the specified constant value.
-					Output->Arg = Emitter->ConstantFromShaderValue(Input.ConstantValue);
-				}
-				else if (MIR::FValue* Value = MIR::Internal::FetchValueFromExpressionInput(this, Input.Input))
+				if (MIR::FValue* Value = MIR::Internal::FetchValueFromExpressionInput(this, Input.Input))
 				{
 					// If this property is the emissive color and we're previewing the material, apply gamma correction to the previewed value.
 					if (Output->Property == MP_EmissiveColor && PreviewInput.IsConnected())
@@ -402,10 +424,15 @@ struct FMaterialIRModuleBuilderImpl
 						Output->Arg = Emitter->Cast(Value, OutputArgType);
 					}
 				}
+				else if (Input.bUseConstant)
+				{
+					// If input is marked to use constant, assign this output to the specified constant value.
+					Output->Arg = Emitter->ConstantFromShaderValue(Input.ConstantValue);
+				}
 				else
 				{
 					// Otherwise, fallback to assigning this material output to its default value.
-					Output->Arg = MIR::Internal::CreateMaterialAttributeDefaultValue(*Emitter, Builder->Material, Output->Property);
+					Output->Arg = Emitter->Subscript(DefaultMaterialAggregate, UMaterialAggregate::MaterialPropertyToAttributeIndex(Output->Property));
 				}
 
 				// The value being set to this material output is now valid.
@@ -453,7 +480,22 @@ struct FMaterialIRModuleBuilderImpl
 		}
 		else
 		{
-			return Builder->Material->GetExpressionInputDescription(Property, Input);
+			bool bResult = Builder->Material->GetExpressionInputDescription(Property, Input);
+			/**
+			* MP_SubsurfaceColor is currently hacked in the old translator to float4, 
+			* but we rely on default types (i.e.float3) for default values in the material editor output.
+			* 
+			* This hack resolves the default value to use float4 rather than float3 until we can 
+			* implement a permanent float4 alternative method in the new translator 
+			* (i.e. work towards deprecating the MP_SubsurfaceColor hacks scattered throughout UE).
+			*/
+			if (Property == MP_SubsurfaceColor)
+			{
+				Input.Type = UE::Shader::EValueType::Float4;
+				Input.ConstantValue = UE::Shader::FValue(Input.ConstantValue.AsLinearColor());
+			}
+
+			return bResult;
 		}
 	}
 
diff --git a/Engine/Source/Runtime/Engine/Private/Materials/MaterialIRToHLSLTranslator.cpp b/Engine/Source/Runtime/Engine/Private/Materials/MaterialIRToHLSLTranslator.cpp
--- a/Engine/Source/Runtime/Engine/Private/Materials/MaterialIRToHLSLTranslator.cpp
+++ b/Engine/Source/Runtime/Engine/Private/Materials/MaterialIRToHLSLTranslator.cpp
@@ -6,6 +6,7 @@
 #include "MaterialDomain.h"
 #include "MaterialIRInternal.h"
 #include "Materials/Material.h"
+#include "Materials/MaterialAggregate.h"
 #include "Materials/MaterialAttributeDefinitionMap.h"
 #include "Materials/MaterialExpressionVolumetricAdvancedMaterialOutput.h"
 #include "Materials/MaterialExternalCodeRegistry.h"
@@ -265,22 +266,18 @@ struct FPrivate : FMaterialIRToHLSLTranslation
 
 	void GeneratePixelAttributesHLSL()
 	{
-		for (int32 PropertyIndex = 0; PropertyIndex < MP_MAX; ++PropertyIndex)
+		for (EMaterialProperty Property : UMaterialAggregate::GetMaterialAttributesProperties())
 		{
-			EMaterialProperty Property = (EMaterialProperty)PropertyIndex;
-			if (!MIR::Internal::IsMaterialPropertyEnabled(Property) || !MIR::MaterialOutputEvaluatesInStage(Property, MIR::Stage_Pixel))
+			if (!MIR::MaterialOutputEvaluatesInStage(Property, MIR::Stage_Pixel))
 			{
 				continue;
 			}
-		
-			check(FMaterialAttributeDefinitionMap::GetShaderFrequency(Property) == SF_Pixel);
 			
-			// Special case MP_SubsurfaceColor as the actual property is a combination of the color and the profile but we don't want to expose the profile
-			FStringView PropertyName = (Property == MP_SubsurfaceColor) ? TEXTVIEW("Subsurface") : FStringView{ FMaterialAttributeDefinitionMap::GetAttributeName(Property) };
-			EMaterialValueType Type = (Property == MP_SubsurfaceColor) ? MCT_Float4 : FMaterialAttributeDefinitionMap::GetValueType(Property);
-			check(PropertyName.Len() > 0);
-
-			PixelAttributesHLSL.Appendf(TEXT(TAB "%s %s;\n"), GetHLSLTypeString(Type), PropertyName.GetData());
+			if (const FMaterialAggregateAttribute* PropertyAggregate = UMaterialAggregate::GetMaterialAttribute(Property))
+			{
+				check(PropertyAggregate->Name.IsValid());
+				PixelAttributesHLSL.Appendf(TEXT(TAB "%s %s;\n"), GetHLSLTypeString(PropertyAggregate->ToMaterialValueType()), *(PropertyAggregate->Name.ToString()));
+			}
 		}
 	}
 	
@@ -316,10 +313,6 @@ struct FPrivate : FMaterialIRToHLSLTranslation
 				Printer << TEXTVIEW("#endif") << NewLine;
 				Printer << TEXTVIEW("Parameters.SharedLocalBases.Count = 1") << EndOfStatement;
 			}
-			else
-			{
-				Printer << TEXTVIEW("PixelMaterialInputs.FrontMaterial = GetInitialisedSubstrateData()") << EndOfStatement;
-			}
 		}
 
 		EvaluateOtherMaterialAttributesHLSL[CurrentStage] = MoveTemp(Printer.Buffer);
@@ -1172,16 +1165,20 @@ struct FPrivate : FMaterialIRToHLSLTranslation
 
 	void LowerSetMaterialOutput(const MIR::FSetMaterialOutput* Output)
 	{
-		// Special case MP_SubsurfaceColor as the actual property is a combination of the color and the profile but we don't want to expose the profile
-		FStringView PropertyName = (Output->Property == MP_SubsurfaceColor) ? TEXTVIEW("Subsurface") : FStringView{ FMaterialAttributeDefinitionMap::GetAttributeName(Output->Property) };
-
 		if (Output->Property == MP_WorldPositionOffset)
 		{
 			Printer << TEXTVIEW("return ");
 		}
 		else
 		{
-			Printer << TEXTVIEW("PixelMaterialInputs.") << PropertyName << TEXTVIEW(" = ");
+			if (const FMaterialAggregateAttribute* PropertyAggregate = UMaterialAggregate::GetMaterialAttribute(Output->Property))
+			{
+				Printer << TEXTVIEW("PixelMaterialInputs.") << PropertyAggregate->Name.ToString() << TEXTVIEW(" = ");
+			}
+			else
+			{
+				UE_MIR_UNREACHABLE();
+			}			
 		}
 
 		if (Substrate::IsSubstrateEnabled())
@@ -1845,8 +1842,12 @@ struct FPrivate : FMaterialIRToHLSLTranslation
 		FString GetMaterialCustomizedUVS;
 		for (int32 CustomUVIndex = 0; CustomUVIndex < ModuleStatistics.NumPixelTexCoords; CustomUVIndex++)
 		{
-			const FString AttributeName = FMaterialAttributeDefinitionMap::GetAttributeName((EMaterialProperty)(MP_CustomizedUVs0 + CustomUVIndex));
-			GetMaterialCustomizedUVS.Appendf(TEXT(TAB "OutTexCoords[%u] = Parameters.MaterialAttributes.%s;\n"), CustomUVIndex, *AttributeName);
+			const FMaterialAggregateAttribute* PropertyAggregate = UMaterialAggregate::GetMaterialAttribute((EMaterialProperty)(MP_CustomizedUVs0 + CustomUVIndex));
+			if (!PropertyAggregate)
+			{
+				continue;
+			}
+			GetMaterialCustomizedUVS.Appendf(TEXT(TAB "OutTexCoords[%u] = Parameters.MaterialAttributes.%s;\n"), CustomUVIndex, *(PropertyAggregate->Name.ToString()));
 		}
 		Params.Add(TEXT("get_material_customized_u_vs"), MoveTemp(GetMaterialCustomizedUVS));
 		
diff --git a/Engine/Source/Runtime/Engine/Private/Materials/MaterialIRTypes.cpp b/Engine/Source/Runtime/Engine/Private/Materials/MaterialIRTypes.cpp
--- a/Engine/Source/Runtime/Engine/Private/Materials/MaterialIRTypes.cpp
+++ b/Engine/Source/Runtime/Engine/Private/Materials/MaterialIRTypes.cpp
@@ -109,6 +109,7 @@ FType FType::FromMaterialValueType(EMaterialValueType Type)
 		case MCT_Float3x3: return FType::MakeFloat(3, 3);
 		case MCT_Float4x4: return FType::MakeFloat(4, 4);
 		case MCT_LWCMatrix: return FType::MakeDouble(4, 4);
+		case MCT_MaterialAttributes: return FType::MakeAggregate(UMaterialAggregate::GetMaterialAttributes());
 		default: UE_MIR_UNREACHABLE();
 	}
 }
diff --git a/Engine/Source/Runtime/Engine/Public/Materials/MaterialAggregate.h b/Engine/Source/Runtime/Engine/Public/Materials/MaterialAggregate.h
--- a/Engine/Source/Runtime/Engine/Public/Materials/MaterialAggregate.h
+++ b/Engine/Source/Runtime/Engine/Public/Materials/MaterialAggregate.h
@@ -29,8 +29,13 @@ enum class EMaterialAggregateAttributeType
 	Float3,
 	Float4,
 
+	// Attribute is a shading mode. Used in by MaterialAttributes aggregate.
 	ShadingModel,
 
+	// Attribute is of type MaterialAttributes aggregate.
+	MaterialAttributes,
+
+	// Attribute is of specified user aggregate type.
 	Aggregate,
 };
 
@@ -85,15 +90,25 @@ public:
 public:
 #if WITH_EDITOR
 
+	static const TConstArrayView<EMaterialProperty> GetMaterialAttributesProperties();
+
 	// Returns a reference to the global material attributes aggregate used by the engine.
 	static const UMaterialAggregate* GetMaterialAttributes();
 
+	// Returns a single cached aggregate for the specified material property.
+	static const FMaterialAggregateAttribute* GetMaterialAttribute(EMaterialProperty Property);
+
 	// Maps a material property to the corresponding attribute index in the aggregate.
 	static int32 MaterialPropertyToAttributeIndex(EMaterialProperty Property);
 
 	// Maps an attribute index back to the corresponding material property.
 	static EMaterialProperty AttributeIndexToMaterialProperty(int32 MaterialAttributeIndex);
 
+	static inline bool IsMaterialPropertyEnabled(EMaterialProperty Property)
+	{
+		return MaterialPropertyToAttributeIndex(Property) != INDEX_NONE;
+	}
+
 	// Finds the index of an attribute in the aggregate by name.
 	int32 FindAttributeIndexByName(FName Name) const;
 
diff --git a/Engine/Source/Runtime/Engine/Public/Materials/MaterialExpressionAggregate.h b/Engine/Source/Runtime/Engine/Public/Materials/MaterialExpressionAggregate.h
--- a/Engine/Source/Runtime/Engine/Public/Materials/MaterialExpressionAggregate.h
+++ b/Engine/Source/Runtime/Engine/Public/Materials/MaterialExpressionAggregate.h
@@ -68,6 +68,7 @@ public:
 	virtual FName 				GetInputName(int32 InputIndex) const override;
 	virtual EMaterialValueType 	GetInputValueType(int32 InputIndex) override;
 	virtual EMaterialValueType 	GetOutputValueType(int32 OutputIndex) override;
+	virtual bool				IsResultMaterialAttributes(int32 OutputIndex) override;
 	virtual void 				PreEditChange(FProperty* PropertyAboutToChange) override;
 	virtual void 				PostEditChangeProperty(FPropertyChangedEvent& PropertyChangedEvent) override;
 	virtual void 				Build(MIR::FEmitter& Emitter) override;
diff --git a/Engine/Source/Runtime/Engine/Public/Materials/MaterialExpressionFunctionOutput.h b/Engine/Source/Runtime/Engine/Public/Materials/MaterialExpressionFunctionOutput.h
--- a/Engine/Source/Runtime/Engine/Public/Materials/MaterialExpressionFunctionOutput.h
+++ b/Engine/Source/Runtime/Engine/Public/Materials/MaterialExpressionFunctionOutput.h
@@ -64,6 +64,7 @@ class UMaterialExpressionFunctionOutput : public UMaterialExpression
 	}
 	virtual void GetExpressionToolTip(TArray<FString>& OutToolTip) override;
 	virtual EMaterialValueType GetInputValueType(int32 InputIndex) override;
+	virtual EMaterialValueType GetOutputValueType(int32 OutputIndex) override;
 	virtual bool IsResultMaterialAttributes(int32 OutputIndex) override;
 	virtual bool IsResultSubstrateMaterial(int32 OutputIndex) override;
 	virtual void GatherSubstrateMaterialInfo(FSubstrateMaterialInfo& SubstrateMaterialInfo, int32 OutputIndex) override;
diff --git a/Engine/Source/Runtime/Engine/Public/Materials/MaterialExpressionMaterialFunctionCall.h b/Engine/Source/Runtime/Engine/Public/Materials/MaterialExpressionMaterialFunctionCall.h
--- a/Engine/Source/Runtime/Engine/Public/Materials/MaterialExpressionMaterialFunctionCall.h
+++ b/Engine/Source/Runtime/Engine/Public/Materials/MaterialExpressionMaterialFunctionCall.h
@@ -136,6 +136,7 @@ class UMaterialExpressionMaterialFunctionCall : public UMaterialExpression
 	ENGINE_API virtual void GatherSubstrateMaterialInfo(FSubstrateMaterialInfo& SubstrateMaterialInfo, int32 OutputIndex) override;
 	ENGINE_API virtual FSubstrateOperator* SubstrateGenerateMaterialTopologyTree(class FMaterialCompiler* Compiler, class UMaterialExpression* Parent, int32 OutputIndex) override;
 	ENGINE_API virtual EMaterialValueType GetInputValueType(int32 InputIndex) override;
+	ENGINE_API virtual EMaterialValueType GetOutputValueType(int32 OutputIndex) override;
 	//~ End UMaterialExpression Interface
 
 	/** Util to get name of a particular type, optionally with type appended in parenthesis */
diff --git a/Engine/Source/Runtime/Engine/Public/Materials/MaterialIRCommon.h b/Engine/Source/Runtime/Engine/Public/Materials/MaterialIRCommon.h
--- a/Engine/Source/Runtime/Engine/Public/Materials/MaterialIRCommon.h
+++ b/Engine/Source/Runtime/Engine/Public/Materials/MaterialIRCommon.h
@@ -22,6 +22,7 @@ class UTexture;
 enum EMaterialProperty : int;
 struct FExpressionInput;
 struct FExpressionOutput;
+struct FMaterialAggregateAttribute;
 struct FMaterialInputDescription;
 struct FMaterialInsights;
 struct FMaterialIRModuleBuilder;
