Tighten name matching logic slightly for debug mute/solo.

diff --git a/Engine/Source/Runtime/Engine/Private/AudioDebug.cpp b/Engine/Source/Runtime/Engine/Private/AudioDebug.cpp
--- a/Engine/Source/Runtime/Engine/Private/AudioDebug.cpp
+++ b/Engine/Source/Runtime/Engine/Private/AudioDebug.cpp
@@ -1006,9 +1006,14 @@ namespace Audio
 		FString& OutReason) const
 	{
 		check(IsInAudioThread());
+
+		const FString PlainNameString = FName(Name).GetPlainNameString();
 	
 		// Allow for partial matches of the name.
-		auto MatchesPartOfName = [&Name](FName i) -> bool { return Name.Contains(i.GetPlainNameString()); };
+		auto MatchesPartOfName = [&Name, &PlainNameString](FName i) -> bool
+			{
+				return Name.Contains(i.ToString()) || (PlainNameString == i.GetPlainNameString());
+			};
 
 		// Solo active?
 		if (Solos.Num() > 0)
