Fixed an incorrect flag for populating previous instance data in instance update change set causing crashes on mismatch.

diff --git a/Engine/Source/Runtime/Engine/Public/InstanceData/InstanceUpdateChangeSet.h b/Engine/Source/Runtime/Engine/Public/InstanceData/InstanceUpdateChangeSet.h
--- a/Engine/Source/Runtime/Engine/Public/InstanceData/InstanceUpdateChangeSet.h
+++ b/Engine/Source/Runtime/Engine/Public/InstanceData/InstanceUpdateChangeSet.h
@@ -318,7 +318,7 @@ public:
 	// These setups define the mapping from a delta attribute bit in the tracker to the data array. 
 	// This is not a simple 1:1 mapping as we only track a few bits, and there is also special overrides to take into account.
 	auto GetTransformSetup()		{ return GetSetup<FInstanceAttributeTracker::EFlag::TransformChanged>(true, bUpdateAllInstanceTransforms, Transforms); }
-	auto GetPrevTransformSetup()	{ return GetSetup<FInstanceAttributeTracker::EFlag::TransformChanged>(Flags.bHasPerInstanceDynamicData, ForceFullFlags.bHasPerInstanceDynamicData, PrevTransforms); }
+	auto GetPrevTransformSetup()	{ return GetSetup<FInstanceAttributeTracker::EFlag::TransformChanged>(Flags.bHasPerInstanceDynamicData, bUpdateAllInstanceTransforms || ForceFullFlags.bHasPerInstanceDynamicData, PrevTransforms); }
 	auto GetCustomDataSetup()		{ return GetSetup<FInstanceAttributeTracker::EFlag::CustomDataChanged>(Flags.bHasPerInstanceCustomData, ForceFullFlags.bHasPerInstanceCustomData, PerInstanceCustomData, NumCustomDataFloats); }
 	// These use an identity delta, which means they send all or nothing.
 	auto GetLocalBoundsSetup()		{ return GetSetup(Flags.bHasPerInstanceLocalBounds, InstanceLocalBounds); }
