Skeletal Mesh: Fix FMorphTargetDeltaIterator crash when there is uncompressed MorphTarget data in cooked builds.

#rnx
#tests Frontend in cooked build

diff --git a/Engine/Source/Runtime/Engine/Classes/Animation/MorphTarget.h b/Engine/Source/Runtime/Engine/Classes/Animation/MorphTarget.h
--- a/Engine/Source/Runtime/Engine/Classes/Animation/MorphTarget.h
+++ b/Engine/Source/Runtime/Engine/Classes/Animation/MorphTarget.h
@@ -132,7 +132,7 @@ private:
 			}
 		}
 #if !WITH_EDITOR
-		if (!IterativeDecode(Token, CompressedDelta->PackedDeltaHeaders, CompressedDelta->PackedDeltaData, CompressedDelta->PositionPrecision, CompressedDelta->TangentPrecision, CurrentDelta))
+		else if (!CompressedDelta || !IterativeDecode(Token, CompressedDelta->PackedDeltaHeaders, CompressedDelta->PackedDeltaData, CompressedDelta->PositionPrecision, CompressedDelta->TangentPrecision, CurrentDelta))
 		{
 			// If the decode failed, we're either at the end, or the data is invalid. Stop any further decoding.
 			Clear();
