Re-merged the implementation of ID update in FISMCInstanceDataSceneProxy to use the common code path (the difference was identified and fixed in CL 40898307).

#rnx

diff --git a/Engine/Source/Runtime/Engine/Private/InstancedStaticMesh/ISMInstanceDataSceneProxy.cpp b/Engine/Source/Runtime/Engine/Private/InstancedStaticMesh/ISMInstanceDataSceneProxy.cpp
--- a/Engine/Source/Runtime/Engine/Private/InstancedStaticMesh/ISMInstanceDataSceneProxy.cpp
+++ b/Engine/Source/Runtime/Engine/Private/InstancedStaticMesh/ISMInstanceDataSceneProxy.cpp
@@ -326,40 +326,7 @@ void FISMCInstanceDataSceneProxy::Build(FISMInstanceUpdateChangeSet&& ChangeSet)
 template <typename IndexRemapType>
 void FISMCInstanceDataSceneProxy::UpdateIdMapping(FISMInstanceUpdateChangeSet& ChangeSet, const IndexRemapType &IndexRemap)
 {
-	// update mapping, create explicit mapping if needed
-	if (ChangeSet.bIdentityIdMap && IndexRemap.IsIdentity())
-	{
-		// Reset to identity mapping with the new number of instances
-		InstanceIdIndexMap.Reset(ChangeSet.PostUpdateNumInstances);
-	}
-	else
-	{
-		InstanceIdIndexMap.ResizeExplicit(ChangeSet.PostUpdateNumInstances, ChangeSet.MaxInstanceId);
-
-		// If any were removed, we need to clear the associated IDs, before updating (since they may have been added again)
-		for (TConstSetBitIterator<> It(ChangeSet.InstanceAttributeTracker.GetRemovedIterator()); It; ++It)
-		{
-			// There may be more bits set as things that are marked as removed may no longer be in the map
-			if (It.GetIndex() >= InstanceIdIndexMap.GetMaxInstanceId())
-			{
-				break;
-			}
-			InstanceIdIndexMap.SetInvalid(FPrimitiveInstanceId{It.GetIndex()});
-		}
-
-		// Update index mappings (if not identity)
-		auto IndexDelta = ChangeSet.GetIndexChangedDelta();
-		for (auto It = IndexDelta.GetIterator(); It; ++It)
-		{
-			int32 NewInstanceIndex = It.GetIndex();
-			int32 ItemIndex = It.GetItemIndex();
-			
-			IndexRemap.Remap(ItemIndex, NewInstanceIndex);
-
-			FPrimitiveInstanceId InstanceId = ChangeSet.bIdentityIdMap ? FPrimitiveInstanceId{ItemIndex} : ChangeSet.IndexToIdMapDeltaData[ItemIndex];
-			InstanceIdIndexMap.Update(InstanceId, NewInstanceIndex);
-		}
-	}
+	::UpdateIdMapping(ChangeSet, IndexRemap, InstanceIdIndexMap);
 }
 
 void FISMCInstanceDataSceneProxy::Update(FISMInstanceUpdateChangeSet&& ChangeSet)
