[Shaders] Fix path limit for distributed shader compile jobs and improve readability of shortened filenames with descriptive task names.

The maximum path length in the distributed shader compiler did not account for the NULL-character.
This is implemented via the 'kMaxFilePathSuffixLenWithNullChar' constant.

Also improve the readability of shortened filenames when a descriptive task name is constructed:
Instead of just trimming the name at the end, the task name is constructed with removed characters in the middle,
e.g. "2288.uba.DefaultPostProcessMate---alTextureVSBaseColor-0".

Additionally, report an error if the base path exceeds the limit.
Otherwise, any shortened filename will be invalid nonetheless.

#rnx

diff --git a/Engine/Source/Runtime/Engine/Private/ShaderCompiler/ShaderCompilerDistributed.cpp b/Engine/Source/Runtime/Engine/Private/ShaderCompiler/ShaderCompilerDistributed.cpp
--- a/Engine/Source/Runtime/Engine/Private/ShaderCompiler/ShaderCompilerDistributed.cpp
+++ b/Engine/Source/Runtime/Engine/Private/ShaderCompiler/ShaderCompilerDistributed.cpp
@@ -143,38 +143,64 @@ static void ConvertDebugNameToFilename(TStringBuilder<1024>& StringBuilder, cons
 	}
 }
 
-static void BuildDescriptiveTaskFilename(TStringBuilder<1024>& StringBuilder, const FShaderCompileJob& SingleJob)
+static void BuildDescriptiveTaskFilename(TStringBuilder<1024>& StringBuilder, const TArray<FShaderCommonCompileJobPtr>& JobsToSerialize, int32 BasePathLen, int32 PathLimit)
 {
-	ConvertDebugNameToFilename(StringBuilder, SingleJob.Input.DebugGroupName);
-}
+	StringBuilder << TEXT(".");
 
-void FShaderCompileDistributedThreadRunnable_Interface::DispatchShaderCompileJobsBatch(TArray<FShaderCommonCompileJobPtr>& JobsToSerialize)
-{
-	// Generate unique filename for shader compiler I/O files
-	FString BaseFilePath = CachedController.CreateUniqueFilePath();
-
-	TStringBuilder<1024> BaseFilePathBuilder;
 	if (JobsToSerialize.Num() == 1)
 	{
 		// Decorate filename with single job description
 		if (const FShaderCompileJob* FirstSingleJob = JobsToSerialize[0]->GetSingleShaderJob())
 		{
-			BaseFilePathBuilder << TEXT(".");
-			BuildDescriptiveTaskFilename(BaseFilePathBuilder, *FirstSingleJob);
+			const int32 BaseFilenameLen = BasePathLen + StringBuilder.Len();
+			const int32 TotalFilenameLen = BaseFilenameLen + FirstSingleJob->Input.DebugGroupName.Len();
+			if (TotalFilenameLen > PathLimit)
+			{
+				const int32 RemainingLen = PathLimit - BaseFilenameLen;
+
+				// If even a shortened task filename with just 5 characters were to exceed the path limit,
+				// don't bother trying to construct a descriptive task name, just fall back to default.
+				constexpr int32 kMinFilenameSuffixLen = 5;
+				if (RemainingLen >= kMinFilenameSuffixLen)
+				{
+					const int32 RemainingLenHalfPart = (RemainingLen - 3) / 2;
+					FString ShortenedDebugGroupName = FString::Printf(TEXT("%s---%s"), *FirstSingleJob->Input.DebugGroupName.Left(RemainingLenHalfPart), *FirstSingleJob->Input.DebugGroupName.Right(RemainingLenHalfPart));
+					ConvertDebugNameToFilename(StringBuilder, ShortenedDebugGroupName);
+					return;
+				}
+			}
+			else
+			{
+				ConvertDebugNameToFilename(StringBuilder, FirstSingleJob->Input.DebugGroupName);
+				return;
+			}
 		}
 	}
-	else
+
+	// Decorate filename with number of jobs
+	StringBuilder << TEXT("j-") << JobsToSerialize.Num();
+}
+
+void FShaderCompileDistributedThreadRunnable_Interface::DispatchShaderCompileJobsBatch(TArray<FShaderCommonCompileJobPtr>& JobsToSerialize)
+{
+	// Generate unique filename for shader compiler I/O files
+	FString BaseFilePath = CachedController.CreateUniqueFilePath();
+
+	constexpr int32 kMaxFilePathSuffixLenWithNullChar = 5; // for ".out\0"
+	const int32 PathLimit = FPlatformMisc::GetExternalAppMaxPathLength() - kMaxFilePathSuffixLenWithNullChar;
+	if (BaseFilePath.Len() > PathLimit)
 	{
-		// Decorate filename with number of jobs
-		BaseFilePathBuilder << TEXT(".j-") << JobsToSerialize.Num();
+		UE_LOG(LogShaderCompilers, Error, TEXT("Distributed job batch exceeded path limit: Length=%d, Limit=%d, Path=%s"), BaseFilePath.Len(), PathLimit, *BaseFilePath);
 	}
+
+	TStringBuilder<1024> BaseFilePathBuilder;
+	BuildDescriptiveTaskFilename(BaseFilePathBuilder, JobsToSerialize, BaseFilePath.Len(), PathLimit);
 	BaseFilePath += BaseFilePathBuilder.ToView();
 
-	// Ensure input and output files don't exceed platform path limit
-	const int32 MaxBaseFilePathLen = FPlatformMisc::GetMaxPathLength() - 4;
-	if (BaseFilePath.Len() > MaxBaseFilePathLen)
+	// Ensure the constructed input and output filenames don't exceed platform path limit
+	if (BaseFilePath.Len() > PathLimit)
 	{
-		BaseFilePath.LeftChopInline(BaseFilePath.Len() - MaxBaseFilePathLen);
+		BaseFilePath.LeftChopInline(BaseFilePath.Len() - PathLimit);
 	}
 
 	FString InputFilePath = BaseFilePath + TEXT(".in");
