Add per poly collision setter function to SkeletalMeshComponent

diff --git a/Engine/Source/Runtime/Engine/Classes/Components/SkeletalMeshComponent.h b/Engine/Source/Runtime/Engine/Classes/Components/SkeletalMeshComponent.h
--- a/Engine/Source/Runtime/Engine/Classes/Components/SkeletalMeshComponent.h
+++ b/Engine/Source/Runtime/Engine/Classes/Components/SkeletalMeshComponent.h
@@ -1876,6 +1876,14 @@ public:
 	ENGINE_API virtual void SetPhysMaterialOverride(UPhysicalMaterial* NewPhysMaterial) override;
 	ENGINE_API virtual bool GetSquaredDistanceToCollision(const FVector& Point, float& OutSquaredDistance, FVector& OutClosestPointOnCollision) const override;
 
+	/**
+	 *	Enables or disables per poly collision. This will recreate the physics state.
+	 *
+	 *	@param bEnablePerPolyCollision	Whether per poly collision should be enabled or disabled.
+	 */
+	UFUNCTION(BlueprintCallable, Category="Physics")
+	ENGINE_API void SetEnablePerPolyCollision(bool bInEnablePerPolyCollision);
+
 	/**
 	 *	Enables or disables gravity for the given bone.
 	 *	NAME_None indicates the root body will be edited.
diff --git a/Engine/Source/Runtime/Engine/Private/SkeletalMeshComponentPhysics.cpp b/Engine/Source/Runtime/Engine/Private/SkeletalMeshComponentPhysics.cpp
--- a/Engine/Source/Runtime/Engine/Private/SkeletalMeshComponentPhysics.cpp
+++ b/Engine/Source/Runtime/Engine/Private/SkeletalMeshComponentPhysics.cpp
@@ -2467,6 +2467,21 @@ void USkeletalMeshComponent::ComputeSkinnedTangentBasis(USkeletalMeshComponent*
 	}
 }
 
+void USkeletalMeshComponent::SetEnablePerPolyCollision(bool bInEnablePerPolyCollision)
+{
+	if(bInEnablePerPolyCollision != bEnablePerPolyCollision)
+	{
+		DestroyPhysicsState();
+
+		bEnablePerPolyCollision = bInEnablePerPolyCollision;
+
+		if (IsRegistered())
+		{
+			CreatePhysicsState();
+		}
+	}
+}
+
 void USkeletalMeshComponent::SetEnableBodyGravity(bool bEnableGravity, FName BoneName)
 {
 	if (FBodyInstance* BI = GetBodyInstance(BoneName))
