Initialize both experimental and legacy subtitle fields for DialogueWaves to prevent errors when switching between them.

diff --git a/Engine/Source/Runtime/Engine/Private/DialogueWave.cpp b/Engine/Source/Runtime/Engine/Private/DialogueWave.cpp
--- a/Engine/Source/Runtime/Engine/Private/DialogueWave.cpp
+++ b/Engine/Source/Runtime/Engine/Private/DialogueWave.cpp
@@ -852,20 +852,16 @@ void UDialogueWave::UpdateMappingProxy(FDialogueContextMapping& ContextMapping)
 		ContextMapping.Proxy->SoundWave = ContextMapping.SoundWave;
 		CopySoundBasePropertiesToProxy(ContextMapping.SoundWave, ContextMapping.Proxy);
 
-		const bool UseNewSubtitles = (CVarUseNewSubtitles.GetValueOnGameThread() == 1);
-		if (UseNewSubtitles)
-		{
-			ContextMapping.Proxy->SubtitleText = GetLocalizedSubtitle(ContextMapping);
-		}
-		else
-		{
-			FSubtitleCue NewSubtitleCue;
-			NewSubtitleCue.Text = GetLocalizedSubtitle(ContextMapping);
-			NewSubtitleCue.Time = 0.0f;
-
-			ContextMapping.Proxy->Subtitles.Empty();
-			ContextMapping.Proxy->Subtitles.Add(NewSubtitleCue);
-		}
+		// Set up new subtitles...
+		ContextMapping.Proxy->SubtitleText = GetLocalizedSubtitle(ContextMapping);
+		// ...and legacy subtitles.
+		// This is to make sure both types are properly initialized if the migration cvar au.UseNewSubtitles is swapped while the game is running.
+		FSubtitleCue NewSubtitleCue;
+		NewSubtitleCue.Text = GetLocalizedSubtitle(ContextMapping);
+		NewSubtitleCue.Time = 0.0f;
+
+		ContextMapping.Proxy->Subtitles.Empty();
+		ContextMapping.Proxy->Subtitles.Add(NewSubtitleCue);
 	}
 }
 
