Dstm: Add support for physics constraint component migration. Only supports constraint compponents such that the constrained primitive components are within the same actor.

diff --git a/Engine/Source/Runtime/Engine/Classes/PhysicsEngine/PhysicsConstraintComponent.h b/Engine/Source/Runtime/Engine/Classes/PhysicsEngine/PhysicsConstraintComponent.h
--- a/Engine/Source/Runtime/Engine/Classes/PhysicsEngine/PhysicsConstraintComponent.h
+++ b/Engine/Source/Runtime/Engine/Classes/PhysicsEngine/PhysicsConstraintComponent.h
@@ -8,6 +8,10 @@
 #include "PhysicsEngine/ConstraintInstance.h"
 #include "PhysicsConstraintComponent.generated.h"
 
+#if UE_WITH_REMOTE_OBJECT_HANDLE
+#include "UObject/UObjectMigrationContext.h"
+#endif
+
 class AActor;
 class UPrimitiveComponent;
 struct FBodyInstance;
@@ -74,7 +78,7 @@ public:
 	ENGINE_API virtual void PostEditChangeChainProperty(FPropertyChangedChainEvent& PropertyChangedEvent) override;
 	ENGINE_API virtual void PostEditChangeProperty(FPropertyChangedEvent& PropertyChangedEvent) override;
 #endif // WITH_EDITOR
-#if WITH_EDITORONLY_DATA
+#if WITH_EDITORONLY_DATA || UE_WITH_REMOTE_OBJECT_HANDLE
 	ENGINE_API virtual void Serialize(FArchive& Ar) override;
 #endif
 	//End UObject interface
@@ -89,6 +93,15 @@ public:
 	ENGINE_API virtual void InitializeComponent() override;
 	//End ActorComponent interface
 
+#if UE_WITH_REMOTE_OBJECT_HANDLE
+	ENGINE_API virtual void PostMigrate(const struct FUObjectMigrationContext& MigrationContext) override;
+	ENGINE_API void ApplyMigratedPhysicsState(Chaos::FSerializedDataBufferPtr MigratedStateData = nullptr);
+	ENGINE_API bool HasValidPhysicsState() const;
+#endif
+
+	UFUNCTION()
+	ENGINE_API void ReInitConstraintAndApplyPhysicsState(UPrimitiveComponent* ChangedComponent, EComponentPhysicsStateChange StateChange);
+
 	//~ Begin SceneComponent Interface
 #if WITH_EDITOR
 	ENGINE_API virtual void PostEditComponentMove(bool bFinished) override;
@@ -423,5 +436,14 @@ private:
 
 	/** Wrapper that calls our constraint plasticity delegate */
 	ENGINE_API void OnPlasticDeformationWrapper(int32 ConstraintIndex);
+#if UE_WITH_REMOTE_OBJECT_HANDLE
+	bool CanReInitConstraintAndApplyPhysicsState() const;
+
+	void MigratePhysicsState(FArchive& Ar);
+	
+	ENGINE_API bool CanSendMigratedPhysicsState(const FUObjectMigrationContext* MigrationContext, EObjectMigrationSendType ObjMigrationSendType) const;
+
+	void CheckAndUnbindPhysicsStateDelegates();
+#endif
 };
 
diff --git a/Engine/Source/Runtime/Engine/Private/PhysicsEngine/Experimental/PhysInterface_Chaos.cpp b/Engine/Source/Runtime/Engine/Private/PhysicsEngine/Experimental/PhysInterface_Chaos.cpp
--- a/Engine/Source/Runtime/Engine/Private/PhysicsEngine/Experimental/PhysInterface_Chaos.cpp
+++ b/Engine/Source/Runtime/Engine/Private/PhysicsEngine/Experimental/PhysInterface_Chaos.cpp
@@ -25,6 +25,7 @@
 #include "PhysicalMaterials/PhysicalMaterialMask.h"
 #include "PhysicalMaterials/PhysicalMaterial.h"
 #include "PhysicsEngine/ConstraintDrives.h"
+#include "PhysicsEngine/PhysicsConstraintComponent.h"
 #include "PhysicsProxy/SingleParticlePhysicsProxy.h"
 #include "PhysicsEngine/ConstraintTypes.h"
 
@@ -198,17 +199,20 @@ void FPhysInterface_Chaos::UpdateMaterialMask(FPhysicsMaterialMaskHandle& InHand
 	Chaos::FPhysicalMaterialManager::Get().UpdateMaterialMask(InHandle);
 }
 
-FPhysScene* FPhysInterface_Chaos::GetCurrentSceneForComponent(UPrimitiveComponent* InComponent)
+FPhysScene* FPhysInterface_Chaos::GetCurrentSceneForComponent(USceneComponent* InComponent)
 {
 	if (!InComponent)
 	{
 		return nullptr;
 	}
 
-	FBodyInstance* BodyInstance = InComponent->GetBodyInstance();
-	if (BodyInstance && BodyInstance->GetPhysicsActorHandle())
+	if (UPrimitiveComponent* PrimComp = Cast<UPrimitiveComponent>(InComponent))
 	{
-		return static_cast<FPhysScene*>(GetCurrentScene(BodyInstance->GetPhysicsActorHandle()));
+		FBodyInstance* BodyInstance = PrimComp->GetBodyInstance();
+		if (BodyInstance && BodyInstance->GetPhysicsActorHandle())
+		{
+			return static_cast<FPhysScene*>(GetCurrentScene(BodyInstance->GetPhysicsActorHandle()));
+		}
 	}
 
 	if (UWorld* ComponentWorld = InComponent->GetWorld())
diff --git a/Engine/Source/Runtime/Engine/Private/PhysicsEngine/PhysicsConstraintComponent.cpp b/Engine/Source/Runtime/Engine/Private/PhysicsEngine/PhysicsConstraintComponent.cpp
--- a/Engine/Source/Runtime/Engine/Private/PhysicsEngine/PhysicsConstraintComponent.cpp
+++ b/Engine/Source/Runtime/Engine/Private/PhysicsEngine/PhysicsConstraintComponent.cpp
@@ -15,10 +15,19 @@
 #include "UObject/ICookInfo.h"
 #include "UObject/SoftObjectPath.h"
 
+#if UE_WITH_REMOTE_OBJECT_HANDLE
+#include "PhysicsInterfaceDeclaresCore.h"
+#include "Physics/Experimental/PhysScene_Chaos.h"
+#include "Chaos/Serialization/SerializedMultiPhysicsState.h"
+#endif
+
 #include UE_INLINE_GENERATED_CPP_BY_NAME(PhysicsConstraintComponent)
 
 #define LOCTEXT_NAMESPACE "ConstraintComponent"
 
+DECLARE_LOG_CATEGORY_EXTERN(LogPhysicsConstraints, Log, All);
+DEFINE_LOG_CATEGORY(LogPhysicsConstraints);
+
 UPhysicsConstraintComponent::UPhysicsConstraintComponent(const FObjectInitializer& ObjectInitializer)
 	: Super(ObjectInitializer)
 {
@@ -397,6 +406,9 @@ void UPhysicsConstraintComponent::OnUnregister()
 {
 	Super::OnUnregister();
 	TermComponentConstraint();
+#if UE_WITH_REMOTE_OBJECT_HANDLE
+	CheckAndUnbindPhysicsStateDelegates();
+#endif
 }
 
 void UPhysicsConstraintComponent::BeginDestroy()
@@ -659,7 +671,9 @@ void UPhysicsConstraintComponent::UpdateSpriteTexture()
 	}
 }
 #endif // WITH_EDITOR
-#if WITH_EDITORONLY_DATA
+
+#if WITH_EDITORONLY_DATA || UE_WITH_REMOTE_OBJECT_HANDLE
+
 void UPhysicsConstraintComponent::Serialize(FArchive& Ar)
 {
 	Super::Serialize(Ar);
@@ -679,9 +693,233 @@ void UPhysicsConstraintComponent::Serialize(FArchive& Ar)
 		}
 	}
 #endif
+
+#if UE_WITH_REMOTE_OBJECT_HANDLE
+	// For remote objects we need non-UProperties to be transferred
+	if (Ar.IsMigratingRemoteObjects())
+	{
+		Ar << OverrideComponent1;
+		Ar << OverrideComponent2;
+
+		MigratePhysicsState(Ar);
+	}
+
+	if (Ar.IsSaving())
+	{
+		CheckAndUnbindPhysicsStateDelegates();
+	}
+#endif
+}
+
+#endif
+
+#if UE_WITH_REMOTE_OBJECT_HANDLE
+
+void UPhysicsConstraintComponent::MigratePhysicsState(FArchive& Ar)
+{
+	const FUObjectMigrationContext* MigrationContext = Ar.GetMigrationContext();
+	check(MigrationContext);
+
+	if (Ar.IsLoading())
+	{
+		Chaos::FSerializedDataBufferPtr MigratedStateData = MakeUnique<Chaos::FSerializedDataBuffer>();
+
+		Ar << (*MigratedStateData);
+
+		const EObjectMigrationRecvType ObjMigrationRecvType = MigrationContext->GetObjectMigrationRecvType(this);
+
+		if (MigratedStateData->GetDataAsByteArrayRef().IsEmpty())
+		{
+			UE_LOG(LogPhysicsConstraints, Verbose, TEXT("Migrated an empty physics state (this can be expected) | [%s] will use its default physics state or the one already had | Migration Context [%s]."), *GetPathNameSafe(this), ToString(ObjMigrationRecvType));
+			return;
+		}
+
+		if (const bool bFinishedDestroy = HasAnyFlags(EObjectFlags::RF_MirroredGarbage))
+		{
+			return;
+		}
+
+		FPhysScene* PhysicsScene = FPhysicsInterface::GetCurrentSceneForComponent(this);
+		if (ensure(PhysicsScene))
+		{
+			PhysicsScene->PushPendingSerializedStateForObject(this, MoveTemp(MigratedStateData));
+		}
+	}
+	else if (Ar.IsSaving())
+	{
+		const EObjectMigrationSendType ObjMigrationSendType = MigrationContext->GetObjectMigrationSendType(this);
+
+		FPhysScene* PhysicsScene = FPhysicsInterface::GetCurrentSceneForComponent(this);
+
+		if (CanSendMigratedPhysicsState(MigrationContext, ObjMigrationSendType))
+		{
+			Chaos::FSerializedMultiPhysicsState MigratedState(ObjMigrationSendType == EObjectMigrationSendType::ReassignOwnership ? Chaos::ESerializedDataContext::Both : Chaos::ESerializedDataContext::External);
+			MigratedState.SetNumOfConstraints(1);
+
+			MigratedState.AppendDataFromElement(ConstraintInstance);
+
+			// We need re-serialize the state as a chaos data buffer, so we can use it as is on de-serialization
+			Chaos::FSerializedDataBuffer MigratedStateAsBytes;
+			FMemoryWriter MigratedStateWriter(MigratedStateAsBytes.GetDataAsByteArrayRef());
+			MigratedState.Serialize(MigratedStateWriter);
+
+			Ar << MigratedStateAsBytes;
+		}
+		else
+		{
+			// This object might be being migrated before we had a chance to process a previous migration
+			Chaos::FSerializedDataBufferPtr RawMigratedData = PhysicsScene ? PhysicsScene->PopPendingSerializedStateForObject(this) : nullptr;
+
+			if (RawMigratedData)
+			{
+				UE_LOG(LogPhysicsConstraints, Warning, TEXT("Sending a pending physics state | [%s]. It is sent away before processed. | Migration Context [%s]."), *GetPathNameSafe(this), ToString(ObjMigrationSendType));
+				Ar << *RawMigratedData;
+			}
+			else
+			{
+				// We need to serialize an empty buffer as array, so we can handle this case during the load phase post migration
+				Chaos::FSerializedDataBuffer EmptyPhysicsState;
+				Ar << EmptyPhysicsState;
+			}
+		}
+	}
+}
+
+bool UPhysicsConstraintComponent::CanReInitConstraintAndApplyPhysicsState() const
+{
+	FBodyInstance* Body1 = GetBodyInstance(EConstraintFrame::Frame1);
+	FBodyInstance* Body2 = GetBodyInstance(EConstraintFrame::Frame2);
+
+	Chaos::FPhysicsObject* Object1 = (Body1 && Body1->IsValidBodyInstance()) ? Body1->GetPhysicsActor()->GetPhysicsObject() : GetPhysicsObject(EConstraintFrame::Frame1);
+	Chaos::FPhysicsObject* Object2 = (Body2 && Body2->IsValidBodyInstance()) ? Body2->GetPhysicsActor()->GetPhysicsObject() : GetPhysicsObject(EConstraintFrame::Frame2);
+
+	return (((!Body1) == (!Object1)) && ((!Body2) == (!Object2)));
+}
+
+void UPhysicsConstraintComponent::PostMigrate(const struct FUObjectMigrationContext& MigrationContext)
+{
+	Super::PostMigrate(MigrationContext);
+
+	if (MigrationContext.MigrationSide == EObjectMigrationSide::Receive)
+	{
+		EObjectMigrationRecvType RecvType = MigrationContext.GetObjectMigrationRecvType(this);
+
+		if (HasValidPhysicsState())
+		{
+			ApplyMigratedPhysicsState();
+		}
+		else if (IsValid(this))
+		{
+			if (CanReInitConstraintAndApplyPhysicsState())
+			{
+				InitComponentConstraint();
+				ApplyMigratedPhysicsState();
+			}
+			else
+			{
+				if (UPrimitiveComponent* Comp1 = GetComponentInternal(EConstraintFrame::Frame1))
+				{
+					Comp1->OnComponentPhysicsStateChanged.AddUniqueDynamic(this, &UPhysicsConstraintComponent::ReInitConstraintAndApplyPhysicsState);
+				}
+				if (UPrimitiveComponent* Comp2 = GetComponentInternal(EConstraintFrame::Frame2))
+				{
+					Comp2->OnComponentPhysicsStateChanged.AddUniqueDynamic(this, &UPhysicsConstraintComponent::ReInitConstraintAndApplyPhysicsState);
+				}
+			}
+		}
+	}
 }
 
+void UPhysicsConstraintComponent::ApplyMigratedPhysicsState(Chaos::FSerializedDataBufferPtr MigratedStateData)
+{
+	Chaos::FSerializedDataBufferPtr RawMigratedData;
+
+	if (MigratedStateData)
+	{
+		RawMigratedData = MoveTemp(MigratedStateData);
+	}
+	else if (FPhysScene* PhysicsScene = FPhysicsInterface::GetCurrentSceneForComponent(this))
+	{
+		RawMigratedData = PhysicsScene->PopPendingSerializedStateForObject(this);
+	}
+
+	if (RawMigratedData)
+	{
+		FMemoryReader MultiBodyStateReader(RawMigratedData->GetDataAsByteArrayRef());
+
+		Chaos::FSerializedMultiPhysicsState MigratedState;
+		MigratedState.Serialize(MultiBodyStateReader);
+
+		const Chaos::FSerializedMultiPhysicsState::FHeader& MigratedStateHeader = MigratedState.GetHeader();
+		const bool bIsValidMigratedState = MigratedStateHeader.IsValid() && MigratedStateHeader.NumConstraints == 1;
+		if (!ensure(bIsValidMigratedState))
+		{
+			return;
+		}
+
+		UE_LOG(LogPhysicsConstraints, Verbose, TEXT("[%hs] Applying migrated physics state | Data context [%s]"), __func__, Chaos::LexToString(MigratedState.GetHeader().DataContext));
+
+		if (ensure(ConstraintInstance.IsValidConstraintInstance()))
+		{
+			MigratedState.ApplyDataToElement(ConstraintInstance);
+		}
+	}
+}
+
+bool UPhysicsConstraintComponent::CanSendMigratedPhysicsState(const FUObjectMigrationContext* MigrationContext, EObjectMigrationSendType ObjMigrationSendType) const
+{
+	if (!HasValidPhysicsState())
+	{
+		return false;
+	}
+
+	if (ObjMigrationSendType == EObjectMigrationSendType::ReturnBorrowed)
+	{
+		return MigrationContext->MultiServerCommitRequestId.IsValid();
+	}
+
+	return true;
+}
+
+bool UPhysicsConstraintComponent::HasValidPhysicsState() const
+{
+	return ConstraintInstance.IsValidConstraintInstance();
+}
+
+void UPhysicsConstraintComponent::CheckAndUnbindPhysicsStateDelegates()
+{
+	TFunction<void(UPrimitiveComponent*, const FString&)> CheckAndUnbind =
+		[this](UPrimitiveComponent* Comp, const FString& FrameLabel)
+		{
+			if (Comp)
+			{
+				if (Comp->OnComponentPhysicsStateChanged.IsAlreadyBound(this, &UPhysicsConstraintComponent::ReInitConstraintAndApplyPhysicsState))
+				{
+					UE_LOG(LogPhysics, Warning,
+						TEXT("UPhysicsConstraintComponent::CheckAndUnbindPhysicsStateDelegates: PhysicsConstraintComponent [%s] still had physics state delegate bound for %s. Cleaning up."),
+						*GetNameSafe(this), *FrameLabel);
+
+					Comp->OnComponentPhysicsStateChanged.RemoveDynamic(this, &UPhysicsConstraintComponent::ReInitConstraintAndApplyPhysicsState);
+				}
+			}
+		};
+
+	CheckAndUnbind(GetComponentInternal(EConstraintFrame::Frame1), TEXT("Frame1"));
+	CheckAndUnbind(GetComponentInternal(EConstraintFrame::Frame2), TEXT("Frame2"));
+}
+#endif
+
+void UPhysicsConstraintComponent::ReInitConstraintAndApplyPhysicsState(UPrimitiveComponent* ChangedComponent, EComponentPhysicsStateChange StateChange)
+{
+#if UE_WITH_REMOTE_OBJECT_HANDLE
+	if (StateChange == EComponentPhysicsStateChange::Created && CanReInitConstraintAndApplyPhysicsState())
+	{
+		InitComponentConstraint();
+		ApplyMigratedPhysicsState();
+		CheckAndUnbindPhysicsStateDelegates();
+	}
 #endif
+}
 
 void UPhysicsConstraintComponent::SetLinearPositionDrive( bool bEnableDriveX, bool bEnableDriveY, bool bEnableDriveZ )
 {
diff --git a/Engine/Source/Runtime/Engine/Public/Physics/Experimental/PhysInterface_Chaos.h b/Engine/Source/Runtime/Engine/Public/Physics/Experimental/PhysInterface_Chaos.h
--- a/Engine/Source/Runtime/Engine/Public/Physics/Experimental/PhysInterface_Chaos.h
+++ b/Engine/Source/Runtime/Engine/Public/Physics/Experimental/PhysInterface_Chaos.h
@@ -29,6 +29,7 @@ struct FAngularDriveConstraint;
 
 
 class AWorldSettings;
+class UPhysicsConstraintComponent;
 
 class FPhysInterface_Chaos : public FGenericPhysicsInterface, public FChaosEngineInterface
 {
@@ -47,7 +48,7 @@ public:
 
 	// Actor interface functions
 
-	static ENGINE_API FPhysScene* GetCurrentSceneForComponent(UPrimitiveComponent* InComponent);
+	static ENGINE_API FPhysScene* GetCurrentSceneForComponent(USceneComponent* InComponent);
 
 	static ENGINE_API void FlushScene(FPhysScene* InScene);
 	static ENGINE_API bool IsInScene(const FPhysicsActorHandle& InActorReference);
