[Camera Manager] Remove some old deprecated 5.0 functions:

- GetCameraCachePOV
- GetLastFrameCameraCachePOV

[RN] Remove the deprecated functions on the camera manager, GetCameraCachePOV and GetLastFrameCameraCachePOV

diff --git a/Engine/Source/Runtime/Engine/Classes/Camera/PlayerCameraManager.h b/Engine/Source/Runtime/Engine/Classes/Camera/PlayerCameraManager.h
--- a/Engine/Source/Runtime/Engine/Classes/Camera/PlayerCameraManager.h
+++ b/Engine/Source/Runtime/Engine/Classes/Camera/PlayerCameraManager.h
@@ -189,9 +189,6 @@ class APlayerCameraManager : public AActor
 {
 	GENERATED_UCLASS_BODY()
 
-	// destructor for handling property deprecation, please remove after all deprecated properties are gone
-	ENGINE_API virtual ~APlayerCameraManager();
-
 	/** PlayerController that owns this Camera actor */
 	UPROPERTY(transient)
 	TObjectPtr<class APlayerController> PCOwner;
@@ -336,14 +333,6 @@ public:
 	/** Gets value of LastFrameCameraCachePrivate.POV */
 	ENGINE_API virtual const FMinimalViewInfo& GetLastFrameCameraCacheView() const;
 
-	/** Gets value of CameraCachePrivate.POV */
-	UE_DEPRECATED(5.0, "APlayerCameraManager::GetCameraCacheView is favored now, and this function forwards to that one.")
-	ENGINE_API virtual FMinimalViewInfo GetCameraCachePOV() const;
-
-	/** Gets value of LastFrameCameraCachePrivate.POV */
-	UE_DEPRECATED(5.0, "APlayerCameraManager::GetLastFrameCameraCacheView is favored now, and this function forwards to that one.")
-	ENGINE_API virtual FMinimalViewInfo GetLastFrameCameraCachePOV() const;
-
 	/** Get value of CameraCachePrivate.Time  */
 	float GetCameraCacheTime() const { return CameraCachePrivate.TimeStamp; }
 
diff --git a/Engine/Source/Runtime/Engine/Private/PlayerCameraManager.cpp b/Engine/Source/Runtime/Engine/Private/PlayerCameraManager.cpp
--- a/Engine/Source/Runtime/Engine/Private/PlayerCameraManager.cpp
+++ b/Engine/Source/Runtime/Engine/Private/PlayerCameraManager.cpp
@@ -1421,22 +1421,6 @@ const FMinimalViewInfo& APlayerCameraManager::GetLastFrameCameraCacheView() cons
 	return LastFrameCameraCachePrivate.POV;
 }
 
-FMinimalViewInfo APlayerCameraManager::GetCameraCachePOV() const
-{
-	return GetCameraCacheView();
-}
-
-FMinimalViewInfo APlayerCameraManager::GetLastFrameCameraCachePOV() const
-{
-	return GetLastFrameCameraCacheView();
-}
-
-PRAGMA_DISABLE_DEPRECATION_WARNINGS
-APlayerCameraManager::~APlayerCameraManager()
-{
-}
-PRAGMA_ENABLE_DEPRECATION_WARNINGS
-
 //////////////////////////////////////////////////////////////////////////
 // FTViewTarget
 
