ActiveSound: set closest listener before pushing interface parameters

#rnx

diff --git a/Engine/Source/Runtime/Engine/Private/ActiveSound.cpp b/Engine/Source/Runtime/Engine/Private/ActiveSound.cpp
--- a/Engine/Source/Runtime/Engine/Private/ActiveSound.cpp
+++ b/Engine/Source/Runtime/Engine/Private/ActiveSound.cpp
@@ -14,7 +14,6 @@
 #include "IAudioParameterTransmitter.h"
 #include "SubtitleManager.h"
 
-
 static int32 AudioOcclusionDisabledCvar = 0;
 FAutoConsoleVariableRef CVarAudioOcclusionEnabled(
 	TEXT("au.DisableOcclusion"),
@@ -934,15 +933,26 @@ void FActiveSound::UpdateInterfaceParameters(const TArray<FListener>& InListener
 
 int32 FActiveSound::GetClosestListenerIndex() const
 {
-	// If we haven't cached the cloest listener index, just directly query it from the audio device.
-	// This is cached in UpdateWaveInstances.
+	// If we haven't cached the closest listener index, just directly query it from the audio device.
 	if (ClosestListenerIndex == INDEX_NONE && AudioDevice)
 	{
-		return AudioDevice->FindClosestListenerIndex(Transform);
+		CacheClosestListener();
 	}
 	return ClosestListenerIndex;
 }
 
+void FActiveSound::CacheClosestListener() const
+{
+	// Early outs.
+	if (!AudioDevice || Sound == nullptr || !Sound->IsPlayable())
+	{
+		return;
+	}
+
+	SCOPE_CYCLE_COUNTER(STAT_AudioFindNearestLocation);
+	ClosestListenerIndex = AudioDevice->FindClosestListenerIndex(Transform);
+}
+
 void FActiveSound::UpdateWaveInstances(TArray<FWaveInstance*> &InWaveInstances, const float DeltaTime)
 {
 	// Reset whether or not the active sound is playing audio.
@@ -957,11 +967,6 @@ void FActiveSound::UpdateWaveInstances(TArray<FWaveInstance*> &InWaveInstances,
 		return;
 	}
 
-	{
-		SCOPE_CYCLE_COUNTER(STAT_AudioFindNearestLocation);
-		ClosestListenerIndex = AudioDevice->FindClosestListenerIndex(Transform);
-	}
-
 	FocusData.PriorityHighest = 1.0f;
 
 	FSoundParseParameters ParseParams;
diff --git a/Engine/Source/Runtime/Engine/Private/AudioDevice.cpp b/Engine/Source/Runtime/Engine/Private/AudioDevice.cpp
--- a/Engine/Source/Runtime/Engine/Private/AudioDevice.cpp
+++ b/Engine/Source/Runtime/Engine/Private/AudioDevice.cpp
@@ -4279,6 +4279,7 @@ int32 FAudioDevice::GetSortedActiveWaveInstances(TArray<FWaveInstance*>& WaveIns
 						UsedDeltaTime = 0.0f;
 					}
 
+					ActiveSound->CacheClosestListener();
 					ActiveSound->UpdateInterfaceParameters(Listeners);
 					ActiveSound->UpdateWaveInstances(WaveInstances, UsedDeltaTime);
 				}
diff --git a/Engine/Source/Runtime/Engine/Public/ActiveSound.h b/Engine/Source/Runtime/Engine/Public/ActiveSound.h
--- a/Engine/Source/Runtime/Engine/Public/ActiveSound.h
+++ b/Engine/Source/Runtime/Engine/Public/ActiveSound.h
@@ -831,7 +831,9 @@ private:
 	static ENGINE_API FTraceDelegate ActiveSoundTraceDelegate;
 
 	/** Cached index to the closest listener. So we don't have to do the work to find it twice. */
-	int32 ClosestListenerIndex = INDEX_NONE;
+	mutable int32 ClosestListenerIndex = INDEX_NONE;
+
+	void CacheClosestListener() const;
 
 	/** This is a friend so the audio device can call Stop() on the active sound. */
 	friend class FAudioDevice;
