Material Translation: Rearrange ValidCast logic.
ValidCast was only checking for trivial condition DestType == MCT_MaterialAttributes at the end of the if/else logic.
This meant that we didn't take the simple path when dealing with non-constant uniform expressions.
This always forced immediate resolution of uniform expressions preventing further preshader combination.
This was always seen in material attribute material function inputs which are used for material layers.
That meant that material layers did not analyze preshader logic across layer boundaries.
Fixing this might give a small improvement to performance for affected materials.

diff --git a/Engine/Source/Runtime/Core/Private/UObject/FortniteMainBranchObjectVersion.cpp b/Engine/Source/Runtime/Core/Private/UObject/FortniteMainBranchObjectVersion.cpp
--- a/Engine/Source/Runtime/Core/Private/UObject/FortniteMainBranchObjectVersion.cpp
+++ b/Engine/Source/Runtime/Core/Private/UObject/FortniteMainBranchObjectVersion.cpp
@@ -14,6 +14,6 @@ TMap<FGuid, FGuid> FFortniteMainBranchObjectVersion::GetSystemGuids()
 	SystemGuids.Add(DevGuids.Niagara_LatestScriptCompileVersion, FGuid("97F0FD7149D90D42B65D6557AF477D1A"));
 	SystemGuids.Add(DevGuids.SkeletalMeshDerivedDataVersion, FGuid("C998F2F3B2884A02B7E290B8C25835F0"));
 	SystemGuids.Add(DevGuids.STATICMESH_DERIVEDDATA_VER, FGuid("3ABF32149F9D4D838750368AB95573BE"));
-	SystemGuids.Add(DevGuids.MaterialTranslationDDCVersion, FGuid("F2C6CB87A94449AE8C34DC29FBA1D02B"));
+	SystemGuids.Add(DevGuids.MaterialTranslationDDCVersion, FGuid("5F49E591FCFD4E21AA8D37F8D15A5856"));
 	return SystemGuids;
 }
diff --git a/Engine/Source/Runtime/Engine/Private/Materials/HLSLMaterialTranslator.cpp b/Engine/Source/Runtime/Engine/Private/Materials/HLSLMaterialTranslator.cpp
--- a/Engine/Source/Runtime/Engine/Private/Materials/HLSLMaterialTranslator.cpp
+++ b/Engine/Source/Runtime/Engine/Private/Materials/HLSLMaterialTranslator.cpp
@@ -5067,28 +5067,25 @@ int32 FHLSLMaterialTranslator::ValidCast(int32 Code, EMaterialValueType DestType
 
 	const EMaterialCastFlags Flags = EMaterialCastFlags::ValidCast;
 	const EMaterialValueType SourceType = GetParameterType(Code);
-	int32 CompiledResult = INDEX_NONE;
 
 	if (SourceType & DestType)
 	{
-		CompiledResult = Code;
+		return Code;
 	}
-	else if(GetParameterUniformExpression(Code) && !GetParameterUniformExpression(Code)->IsConstant())
+	else if ((SourceType & (MCT_TextureVirtual | MCT_TextureMeshPaint | MCT_TextureMaterialCache)) && (DestType & MCT_Texture2D))
 	{
-		if ((SourceType & MCT_TextureVirtual) && (DestType & MCT_Texture2D))
-		{
-			return Code;
-		}
-		else
-		{
-			return ValidCast(AccessUniformExpression(Code), DestType);
-		}
+		return Code;
 	}
-	else if ((SourceType & (MCT_TextureMeshPaint | MCT_TextureMaterialCache)) && (DestType & MCT_Texture2D))
+	else if (DestType == MCT_MaterialAttributes)
 	{
+		// We can feed any type into a material attributes socket as we're really just passing them through.
 		return Code;
 	}
-	else if(IsFloatNumericType(SourceType) && IsFloatNumericType(DestType))
+	else if (GetParameterUniformExpression(Code) && !GetParameterUniformExpression(Code)->IsConstant())
+	{
+		return ValidCast(AccessUniformExpression(Code), DestType);
+	}
+	else if (IsFloatNumericType(SourceType) && IsFloatNumericType(DestType))
 	{
 		const FDerivInfo CodeDerivInfo = GetDerivInfo(Code);
 
@@ -5113,20 +5110,8 @@ int32 FHLSLMaterialTranslator::ValidCast(int32 Code, EMaterialValueType DestType
 			return AddInlinedCodeChunk(DestType, *FiniteCode);
 		}
 	}
-	else
-	{
-		//We can feed any type into a material attributes socket as we're really just passing them through.
-		if( DestType == MCT_MaterialAttributes )
-		{
-			CompiledResult = Code;
-		}
-		else
-		{
-			CompiledResult = Errorf(TEXT("Cannot cast from %s to %s."), DescribeType(SourceType), DescribeType(DestType));
-		}
-	}
-
-	return CompiledResult;
+	
+	return Errorf(TEXT("Cannot cast from %s to %s."), DescribeType(SourceType), DescribeType(DestType));
 }
 
 EMaterialCastFlags GetForceCastFlags(uint32 ForceCastFlags)
