Fixing custom attenuation curve tool tips.

diff --git a/Engine/Source/Runtime/Engine/Classes/Engine/Attenuation.h b/Engine/Source/Runtime/Engine/Classes/Engine/Attenuation.h
--- a/Engine/Source/Runtime/Engine/Classes/Engine/Attenuation.h
+++ b/Engine/Source/Runtime/Engine/Classes/Engine/Attenuation.h
@@ -102,7 +102,7 @@ struct FBaseAttenuationSettings
 	UPROPERTY(EditAnywhere, BlueprintReadWrite, Category = AttenuationDistance, meta = (ClampMin = "0"))
 	float ConeSphereFalloffDistance;
 
-	/* The custom volume attenuation curve to use. */
+	/* The custom volume attenuation curve to use. Value clamped at 1.0. */
 	UPROPERTY(EditAnywhere, BlueprintReadWrite, Category= AttenuationDistance)
 	FRuntimeFloatCurve CustomAttenuationCurve;
 
diff --git a/Engine/Source/Runtime/Engine/Classes/Sound/SoundAttenuation.h b/Engine/Source/Runtime/Engine/Classes/Sound/SoundAttenuation.h
--- a/Engine/Source/Runtime/Engine/Classes/Sound/SoundAttenuation.h
+++ b/Engine/Source/Runtime/Engine/Classes/Sound/SoundAttenuation.h
@@ -390,7 +390,7 @@ struct FSoundAttenuationSettings : public FBaseAttenuationSettings
 	UPROPERTY(EditAnywhere, BlueprintReadWrite, Category = AttenuationPriority, meta = (ClampMin = "0.0", UIMin = "0.0", UIMax = "1.0", DisplayName = "Attenuation Priority"))
 	float ManualPriorityAttenuation;
 
-	/* The custom reverb send curve to use for distance-based send level. */
+	/* The custom reverb send curve to use for distance-based send level. Values clamped between 0.0 and 1.0. */
 	UPROPERTY(EditAnywhere, BlueprintReadWrite, Category = AttenuationReverbSend)
 	FRuntimeFloatCurve CustomReverbSendCurve;
 
