[Audio Insights] Add new event types to the Event Log - Start and Stop triggers, virtualize/realize, stop all, flush audio device, pause and resume

diff --git a/Engine/Plugins/AudioGameplay/Source/AudioGameplay/Private/SoundHandleSubsystem.cpp b/Engine/Plugins/AudioGameplay/Source/AudioGameplay/Private/SoundHandleSubsystem.cpp
--- a/Engine/Plugins/AudioGameplay/Source/AudioGameplay/Private/SoundHandleSubsystem.cpp
+++ b/Engine/Plugins/AudioGameplay/Source/AudioGameplay/Private/SoundHandleSubsystem.cpp
@@ -4,6 +4,7 @@
 
 #include "ActiveSound.h"
 #include "AudioDevice.h"
+#include "Audio/AudioTraceUtil.h"
 #include "Audio/ISoundHandleOwner.h"
 
 #include UE_INLINE_GENERATED_CPP_BY_NAME(SoundHandleSubsystem)
@@ -99,8 +100,14 @@ Audio::EResult USoundHandleSubsystem::Play(const Audio::FSoundHandleID ID)
 
 			Handle->ActiveSound.MaxDistance = MaxDistance;
 		}
+
+#if UE_AUDIO_PROFILERTRACE_ENABLED
+		const FString* EventLogID = &Audio::Trace::EventLog::ID::PlayRequestSoundHandle;
+#else
+		const FString* EventLogID = nullptr;
+#endif // UE_AUDIO_PROFILERTRACE_ENABLED
 		
-		AudioDevice->AddNewActiveSound(Handle->ActiveSound);		
+		AudioDevice->AddNewActiveSound(Handle->ActiveSound, nullptr, EventLogID);
 		return Audio::EResult::Success;
 	}
 	
@@ -112,6 +119,13 @@ void USoundHandleSubsystem::Stop(const Audio::FSoundHandleID ID)
 	if(FSoundHandle* Handle = ActiveHandles.Find(ID))
 	{
 		FAudioDevice* AudioDevice = GetAudioDeviceHandle().GetAudioDevice();
-		AudioDevice->StopActiveSound(&Handle->ActiveSound);
+
+#if UE_AUDIO_PROFILERTRACE_ENABLED
+		const FString* EventLogID = &Audio::Trace::EventLog::ID::StopRequestedSoundHandle;
+#else
+		const FString* EventLogID = nullptr;
+#endif // UE_AUDIO_PROFILERTRACE_ENABLED
+
+		AudioDevice->StopActiveSound(&Handle->ActiveSound, EventLogID);
 	}
 }
diff --git a/Engine/Plugins/AudioInsights/Source/AudioInsights/Private/Views/AudioEventLogDashboardViewFactory.cpp b/Engine/Plugins/AudioInsights/Source/AudioInsights/Private/Views/AudioEventLogDashboardViewFactory.cpp
--- a/Engine/Plugins/AudioInsights/Source/AudioInsights/Private/Views/AudioEventLogDashboardViewFactory.cpp
+++ b/Engine/Plugins/AudioInsights/Source/AudioInsights/Private/Views/AudioEventLogDashboardViewFactory.cpp
@@ -156,7 +156,16 @@ namespace UE::Audio::Insights
 					PlayOrderColumnID,
 					{
 						LOCTEXT("AudioEventLog_PlayOrderColumnDisplayName", "Play Order"),
-						[](const IDashboardDataViewEntry& InData) { return FText::AsNumber(AudioEventLogPrivate::CastEntry(InData).PlayOrder); },
+						[](const IDashboardDataViewEntry& InData)
+						{ 
+							const uint32 PlayOrder = AudioEventLogPrivate::CastEntry(InData).PlayOrder;
+							if (PlayOrder == INDEX_NONE)
+							{
+								return FText::GetEmpty();
+							}
+
+							return FText::AsNumber(AudioEventLogPrivate::CastEntry(InData).PlayOrder);
+						},
 						nullptr /* GetIconName */,
 						false /* bDefaultHidden */,
 						0.15f /* FillWidth */
@@ -281,15 +290,21 @@ namespace UE::Audio::Insights
 		}
 
 		const FAudioEventLogDashboardEntry& SelectedEventLogEntry = CastEntry(*SelectedItem);
-
-		for (const TSharedPtr<IDashboardDataViewEntry>& Entry : DataViewEntries)
+		if (SelectedEventLogEntry.PlayOrder == INDEX_NONE)
 		{
-			if (!Entry.IsValid())
+			FilteredEntriesListView->ClearHighlightedItems();
+		}
+		else
+		{
+			for (const TSharedPtr<IDashboardDataViewEntry>& Entry : DataViewEntries)
 			{
-				continue;
-			}
+				if (!Entry.IsValid())
+				{
+					continue;
+				}
 
-			FilteredEntriesListView->SetItemHighlighted(Entry, CastEntry(*Entry).PlayOrder == SelectedEventLogEntry.PlayOrder);
+				FilteredEntriesListView->SetItemHighlighted(Entry, CastEntry(*Entry).PlayOrder == SelectedEventLogEntry.PlayOrder);
+			}
 		}
 
 		if (SelectInfo != ESelectInfo::Type::Direct)
diff --git a/Engine/Source/Editor/UnrealEd/Private/Commandlets/AudioMixerCommandlet.cpp b/Engine/Source/Editor/UnrealEd/Private/Commandlets/AudioMixerCommandlet.cpp
--- a/Engine/Source/Editor/UnrealEd/Private/Commandlets/AudioMixerCommandlet.cpp
+++ b/Engine/Source/Editor/UnrealEd/Private/Commandlets/AudioMixerCommandlet.cpp
@@ -8,6 +8,7 @@
 #include "Sound/SoundAttenuation.h"
 #include "Sound/SoundWave.h"
 #include "Audio.h"
+#include "Audio/AudioTraceUtil.h"
 #include "UObject/UObjectHash.h"
 #include "UObject/Package.h"
 #include "ActiveSound.h"
@@ -173,8 +174,14 @@ static void PlayOneShotSound(FAudioDevice* InAudioDevice, const TArray<USoundWav
 	NewActiveSound.bAllowSpatialization = false;
 	NewActiveSound.Priority = 1.0f;
 
+#if UE_AUDIO_PROFILERTRACE_ENABLED
+	const FString* EventLogID = &Audio::Trace::EventLog::ID::PlayRequestOneShot;
+#else
+	const FString* EventLogID = nullptr;
+#endif // UE_AUDIO_PROFILERTRACE_ENABLED
+
 	// Add it to the audio device
-	InAudioDevice->AddNewActiveSound(NewActiveSound);
+	InAudioDevice->AddNewActiveSound(NewActiveSound, nullptr, EventLogID);
 }
 
 static UAudioComponent* SpawnLoopingSound(FAudioDevice* InAudioDevice, UWorld* World, const TArray<USoundWave*>& InSounds, bool bAllowSpatialization = false, USoundAttenuation* SoundAttenuation = nullptr, float StartTime = 0.0f)
diff --git a/Engine/Source/Runtime/Engine/Private/ActiveSound.cpp b/Engine/Source/Runtime/Engine/Private/ActiveSound.cpp
--- a/Engine/Source/Runtime/Engine/Private/ActiveSound.cpp
+++ b/Engine/Source/Runtime/Engine/Private/ActiveSound.cpp
@@ -2,6 +2,7 @@
 #include "ActiveSound.h"
 
 #include "Audio/AudioDebug.h"
+#include "Audio/AudioTraceUtil.h"
 #include "Misc/App.h"
 #include "AudioDevice.h"
 #include "AudioLinkSettingsAbstract.h"
@@ -688,6 +689,10 @@ void FActiveSound::Stop()
 {
 	if (AudioDevice)
 	{
+#if UE_AUDIO_PROFILERTRACE_ENABLED
+		Audio::Trace::EventLog::SendActiveSoundEvent(*AudioDevice, *this, Audio::Trace::EventLog::ID::StopRequestActiveSound);
+#endif // UE_AUDIO_PROFILERTRACE_ENABLED
+
 		AudioDevice->AddSoundToStop(this);
 	}
 }
diff --git a/Engine/Source/Runtime/Engine/Private/Audio/AudioTraceUtil.cpp b/Engine/Source/Runtime/Engine/Private/Audio/AudioTraceUtil.cpp
--- a/Engine/Source/Runtime/Engine/Private/Audio/AudioTraceUtil.cpp
+++ b/Engine/Source/Runtime/Engine/Private/Audio/AudioTraceUtil.cpp
@@ -6,8 +6,6 @@
 #include "Styling/SlateIconFinder.h"
 #include "Trace/Trace.h"
 
-#define LOCTEXT_NAMESPACE "AudioTraceUtil"
-
 #if UE_AUDIO_PROFILERTRACE_ENABLED
 
 UE_TRACE_EVENT_BEGIN(Audio, EventLog)
@@ -31,9 +29,6 @@ namespace Audio::Trace
 		const FString SoundWave = TEXT("SoundWave");
 		const FString SoundCueTemplate = TEXT("SoundCueTemplate");
 
-		const FString SoundStartLogMessage = LOCTEXT("EventLogTraceMessage_SoundStart", "Start").ToString();
-		const FString SoundStopLogMessage = LOCTEXT("EventLogTraceMessage_SoundStop", "Stop").ToString();
-
 		bool NameMatchesBaseAudioAsset(const FString& SoundClassName)
 		{
 			return SoundClassName == MetaSoundSource
@@ -120,7 +115,7 @@ namespace Audio::Trace
 
 	namespace EventLog
 	{
-		void SendSoundAssetEventLog(const FActiveSound& InActiveSound, const FString& Event)
+		void SendActiveSoundEvent(const FAudioDevice& AudioDevice, const FActiveSound& InActiveSound, const FString& Event)
 		{
 			const bool bChannelEnabled = UE_TRACE_CHANNELEXPR_IS_ENABLED(AudioMixerChannel);
 
@@ -132,33 +127,35 @@ namespace Audio::Trace
 					return;
 				}
 
-				const FString SoundBaseAssetName = Util::GetSoundBaseAssetName(SoundObjectPtr->GetClass());
-
-				UE_TRACE_LOG(Audio, EventLog, AudioMixerChannel)
-					<< EventLog.DeviceId(InActiveSound.AudioDevice->DeviceID)
-					<< EventLog.Timestamp(FPlatformTime::Cycles64())
-					<< EventLog.PlayOrder(InActiveSound.GetPlayOrder())
-					<< EventLog.AssetPath(*SoundObjectPtr->GetPathName())
-					<< EventLog.EventLogName(*Event)
-					<< EventLog.ActorName(*Util::GetOwnerActorName(InActiveSound))
-					<< EventLog.ActorIconName(*Util::GetOwnerActorIconName(InActiveSound).ToString())
-					<< EventLog.SoundClassName(*SoundBaseAssetName);
+				SendEvent(AudioDevice.DeviceID
+						, Event
+						, InActiveSound.GetPlayOrder()
+						, SoundObjectPtr->GetPathName()
+						, Util::GetOwnerActorName(InActiveSound)
+						, Util::GetOwnerActorIconName(InActiveSound).ToString()
+						, Util::GetSoundBaseAssetName(SoundObjectPtr->GetClass()));
 			}
 		}
 
-		const FString& GetStartEventName()
+		void SendEvent(const Audio::FDeviceId AudioDeviceID, const FString& Event, const uint32 AudioObjectID, const FString& AssetPath, const FString& ActorName, const FString& ActorIconName, const FString& SoundClassName)
 		{
-			return Private::SoundStartLogMessage;
-		}
+			const bool bChannelEnabled = UE_TRACE_CHANNELEXPR_IS_ENABLED(AudioMixerChannel);
 
-		const FString& GetStopEventName()
-		{
-			return Private::SoundStopLogMessage;
+			if (bChannelEnabled)
+			{
+				UE_TRACE_LOG(Audio, EventLog, AudioMixerChannel)
+					<< EventLog.DeviceId(AudioDeviceID)
+					<< EventLog.Timestamp(FPlatformTime::Cycles64())
+					<< EventLog.PlayOrder(AudioObjectID)
+					<< EventLog.AssetPath(*AssetPath)
+					<< EventLog.EventLogName(*Event)
+					<< EventLog.ActorName(*ActorName)
+					<< EventLog.ActorIconName(*ActorIconName)
+					<< EventLog.SoundClassName(*SoundClassName);
+			}
 		}
 	} // namespace EventLog
 
 } // namespace Audio::AudioTraceUtil
 
-#endif // UE_AUDIO_PROFILERTRACE_ENABLED
-
-#undef LOCTEXT_NAMESPACE // AudioTraceUtil
\ No newline at end of file
+#endif // UE_AUDIO_PROFILERTRACE_ENABLED
\ No newline at end of file
diff --git a/Engine/Source/Runtime/Engine/Private/AudioDevice.cpp b/Engine/Source/Runtime/Engine/Private/AudioDevice.cpp
--- a/Engine/Source/Runtime/Engine/Private/AudioDevice.cpp
+++ b/Engine/Source/Runtime/Engine/Private/AudioDevice.cpp
@@ -5105,7 +5105,7 @@ void FAudioDevice::SendSoundStartedTraceMessage(const FActiveSound& InActiveSoun
 			<< SoundStart.Timestamp(FPlatformTime::Cycles64())
 			<< SoundStart.PlayOrder(InActiveSound.PlayOrder)
 			<< SoundStart.AssetPath(*SoundObjectPtr->GetPathName())
-			<< SoundStart.EventLogName(*Audio::Trace::EventLog::GetStartEventName())
+			<< SoundStart.EventLogName(*Audio::Trace::EventLog::ID::SoundStart)
 			<< SoundStart.SoundClassName(*Audio::Trace::Util::GetSoundBaseAssetName(SoundObjectPtr->GetClass()))
 			<< SoundStart.ActorName(*Audio::Trace::Util::GetOwnerActorName(InActiveSound))
 			<< SoundStart.ActorLabel(*Audio::Trace::Util::GetOwnerActorLabel(InActiveSound))
@@ -5215,6 +5215,10 @@ void FAudioDevice::StopAllSounds(bool bShouldStopUISounds)
 		return;
 	}
 
+#if UE_AUDIO_PROFILERTRACE_ENABLED
+	Audio::Trace::EventLog::SendEvent(DeviceID, Audio::Trace::EventLog::ID::StopAllRequested);
+#endif // UE_AUDIO_PROFILERTRACE_ENABLED
+
 	for (int32 SoundIndex=ActiveSounds.Num() - 1; SoundIndex >= 0; --SoundIndex)
 	{
 		FActiveSound* ActiveSound = ActiveSounds[SoundIndex];
@@ -5275,7 +5279,7 @@ void FAudioDevice::NotifyPluginListenersWorldChanged(UWorld* World)
 	}
 }
 
-void FAudioDevice::AddNewActiveSound(const FActiveSound& ActiveSound, const TArray<FAudioParameter>* InDefaultParams)
+void FAudioDevice::AddNewActiveSound(const FActiveSound& ActiveSound, const TArray<FAudioParameter>* InDefaultParams, const FString* EventLogID)
 {
 	TArray<FAudioParameter> Params;
 	if (InDefaultParams)
@@ -5285,11 +5289,18 @@ void FAudioDevice::AddNewActiveSound(const FActiveSound& ActiveSound, const TArr
 	//Previously the copy was done in AddNewActiveSoundInternal when creating the Lamba function. 
 	TSharedRef<FActiveSound> NewUniqueActiveSound = MakeShared<FActiveSound>(ActiveSound);
 
-	AddNewActiveSound(NewUniqueActiveSound, MoveTemp(Params));
+	AddNewActiveSound(NewUniqueActiveSound, MoveTemp(Params), EventLogID);
 }
 
-void FAudioDevice::AddNewActiveSound(TSharedRef<FActiveSound> ActiveSound, TArray<FAudioParameter>&& InDefaultParams)
+void FAudioDevice::AddNewActiveSound(TSharedRef<FActiveSound> ActiveSound, TArray<FAudioParameter>&& InDefaultParams, const FString* EventLogID)
 {
+#if UE_AUDIO_PROFILERTRACE_ENABLED
+	if (EventLogID)
+	{
+		Audio::Trace::EventLog::SendActiveSoundEvent(*this, ActiveSound.Get(), *EventLogID);
+	}
+#endif // UE_AUDIO_PROFILERTRACE_ENABLED
+
 	if (USoundBase* Sound = ActiveSound.Get().GetSound())
 	{
 		Sound->InitResources();
@@ -5483,6 +5494,10 @@ void FAudioDevice::AddNewActiveSoundInternal(const FActiveSound& InNewActiveSoun
 		FActiveSound& VirtualActiveSound = InVirtualLoopToRetrigger->GetActiveSound();
 		AddSoundToStop(&VirtualActiveSound);
 
+#if UE_AUDIO_PROFILERTRACE_ENABLED
+		Audio::Trace::EventLog::SendActiveSoundEvent(*this, VirtualActiveSound, Audio::Trace::EventLog::ID::SoundRealized);
+#endif // UE_AUDIO_PROFILERTRACE_ENABLED
+
 		// Clear must be called after AddSoundToStop to ensure AudioComponent is properly removed from AudioComponentIDToActiveSoundMap
 		VirtualActiveSound.ClearAudioComponent();
 
@@ -5657,6 +5672,10 @@ void FAudioDevice::AddVirtualLoop(const FAudioVirtualLoop& InVirtualLoop)
 
 	VirtualLoops.Add(&ActiveSound, MoveTemp(VirtualLoop));
 
+#if UE_AUDIO_PROFILERTRACE_ENABLED
+	Audio::Trace::EventLog::SendActiveSoundEvent(*this, ActiveSound, Audio::Trace::EventLog::ID::SoundVirtualized);
+#endif // UE_AUDIO_PROFILERTRACE_ENABLED
+
 	NotifySubsystemsVirtualizedSoundCreated(ActiveSound);
 }
 
@@ -5880,16 +5899,16 @@ bool FAudioDevice::IsPendingStop(FActiveSound* ActiveSound)
 	return PendingSoundsToStop.Contains(ActiveSound) || PendingSoundsToDelete.Contains(ActiveSound);
 }
 
-void FAudioDevice::StopActiveSound(const uint64 AudioComponentID)
+void FAudioDevice::StopActiveSound(const uint64 AudioComponentID, const FString* EventLogID)
 {
 	if (!IsInAudioThread())
 	{
 		DECLARE_CYCLE_STAT(TEXT("FAudioThreadTask.StopActiveSound"), STAT_AudioStopActiveSound, STATGROUP_AudioThreadCommands);
 
 		FAudioDevice* AudioDevice = this;
-		FAudioThread::RunCommandOnAudioThread([AudioDevice, AudioComponentID]()
+		FAudioThread::RunCommandOnAudioThread([AudioDevice, AudioComponentID, EventLogID]()
 		{
-			AudioDevice->StopActiveSound(AudioComponentID);
+			AudioDevice->StopActiveSound(AudioComponentID, EventLogID);
 		}, GET_STATID(STAT_AudioStopActiveSound));
 
 		return;
@@ -5897,7 +5916,7 @@ void FAudioDevice::StopActiveSound(const uint64 AudioComponentID)
 
 	const Audio::FDeviceId AudioDeviceID = DeviceID;
 
-	SendCommandToActiveSounds(AudioComponentID, [AudioDeviceID](FActiveSound& ActiveSound)
+	SendCommandToActiveSounds(AudioComponentID, [AudioDeviceID, EventLogID](FActiveSound& ActiveSound)
 	{
 		FAudioDeviceManager* AudioDeviceManager = GEngine->GetAudioDeviceManager();
 			
@@ -5906,15 +5925,30 @@ void FAudioDevice::StopActiveSound(const uint64 AudioComponentID)
 			FAudioDeviceHandle AudioDevice = AudioDeviceManager->GetAudioDevice(AudioDeviceID);
 			if (AudioDevice.IsValid())
 			{
+#if UE_AUDIO_PROFILERTRACE_ENABLED
+				if (EventLogID)
+				{
+					Audio::Trace::EventLog::SendActiveSoundEvent(*AudioDevice, ActiveSound, *EventLogID);
+				}
+#endif // UE_AUDIO_PROFILERTRACE_ENABLED
+
 				AudioDevice->AddSoundToStop(&ActiveSound);
 			}
 		}
 	});
 }
 
-void FAudioDevice::StopActiveSound(FActiveSound* ActiveSound)
+void FAudioDevice::StopActiveSound(FActiveSound* ActiveSound, const FString* EventLogID)
 {
 	check(IsInAudioThread());
+
+#if UE_AUDIO_PROFILERTRACE_ENABLED
+	if (ActiveSound && EventLogID)
+	{
+		Audio::Trace::EventLog::SendActiveSoundEvent(*this, *ActiveSound, *EventLogID);
+	}
+#endif // UE_AUDIO_PROFILERTRACE_ENABLED
+
 	AddSoundToStop(ActiveSound);
 }
 
@@ -5922,9 +5956,27 @@ void FAudioDevice::PauseActiveSound(const uint64 AudioComponentID, const bool bI
 {
 	check(IsInAudioThread());
 	
-	SendCommandToActiveSounds(AudioComponentID, [bInIsPaused](FActiveSound& ActiveSound)
+	const Audio::FDeviceId AudioDeviceID = DeviceID;
+
+	SendCommandToActiveSounds(AudioComponentID, [bInIsPaused, AudioDeviceID](FActiveSound& ActiveSound)
 	{
 		ActiveSound.bIsPaused = bInIsPaused;
+
+#if UE_AUDIO_PROFILERTRACE_ENABLED
+		if (GEngine)
+		{
+			FAudioDeviceManager* AudioDeviceManager = GEngine->GetAudioDeviceManager();
+
+			if (AudioDeviceManager)
+			{
+				FAudioDeviceHandle AudioDevice = AudioDeviceManager->GetAudioDevice(AudioDeviceID);
+				if (AudioDevice.IsValid())
+				{
+					Audio::Trace::EventLog::SendActiveSoundEvent(*AudioDevice, ActiveSound, bInIsPaused ? Audio::Trace::EventLog::ID::PauseSoundRequested : Audio::Trace::EventLog::ID::ResumeSoundRequested);
+				}
+			}
+		}
+#endif // UE_AUDIO_PROFILERTRACE_ENABLED
 	});
 }
 
@@ -6029,7 +6081,7 @@ void FAudioDevice::RemoveActiveSound(FActiveSound* ActiveSound)
 				<< SoundStop.Timestamp(FPlatformTime::Cycles64())
 				<< SoundStop.PlayOrder(ActiveSound->PlayOrder)
 				<< SoundStop.AssetPath(*SoundObjectPtr->GetPathName())
-				<< SoundStop.EventLogName(*Audio::Trace::EventLog::GetStopEventName())
+				<< SoundStop.EventLogName(*Audio::Trace::EventLog::ID::SoundStop)
 				<< SoundStop.SoundClassName(*Audio::Trace::Util::GetSoundBaseAssetName(SoundObjectPtr->GetClass()))
 				<< SoundStop.ActorName(*Audio::Trace::Util::GetOwnerActorName(*ActiveSound))
 				<< SoundStop.ActorIconName(*Audio::Trace::Util::GetOwnerActorIconName(*ActiveSound).ToString());
@@ -6747,7 +6799,14 @@ void FAudioDevice::PlaySoundAtLocation(USoundBase* Sound, UWorld* World, float V
 		{
 			NotifySubsystemsSetOwner(NewActiveSound, OwningActor);
 		}
-		AddNewActiveSound(MoveTemp(NewSharedActiveSound), MoveTemp(Params));
+
+#if UE_AUDIO_PROFILERTRACE_ENABLED
+		const FString* EventLogID = &Audio::Trace::EventLog::ID::PlayRequestSoundAtLocation;
+#else
+		const FString* EventLogID = nullptr;
+#endif // UE_AUDIO_PROFILERTRACE_ENABLED
+
+		AddNewActiveSound(MoveTemp(NewSharedActiveSound), MoveTemp(Params), EventLogID);
 	}
 	else
 	{
@@ -6808,6 +6867,10 @@ void FAudioDevice::Flush(UWorld* WorldToFlush, bool bClearActivatedReverb)
 		FadeOut();
 	}
 
+#if UE_AUDIO_PROFILERTRACE_ENABLED
+	Audio::Trace::EventLog::SendEvent(DeviceID, Audio::Trace::EventLog::ID::FlushAudioDeviceRequested);
+#endif // UE_AUDIO_PROFILERTRACE_ENABLED
+
 	// Stop all audio components attached to the scene
 	bool bFoundIgnoredComponent = false;
 	for (int32 Index = ActiveSounds.Num() - 1; Index >= 0; --Index)
@@ -7274,6 +7337,11 @@ void FAudioDevice::StopSoundsUsingResource(USoundWave* SoundWave, TArray<UAudioC
 						StoppedComponents->Add(AudioComponent);
 					}
 				}
+
+#if UE_AUDIO_PROFILERTRACE_ENABLED
+				Audio::Trace::EventLog::SendActiveSoundEvent(*this, *ActiveSound, Audio::Trace::EventLog::ID::StopRequestSoundsUsingResource);
+#endif // UE_AUDIO_PROFILERTRACE_ENABLED
+
 				AddSoundToStop(ActiveSound);
 				bStoppedSounds = true;
 				break;
diff --git a/Engine/Source/Runtime/Engine/Private/Components/AudioComponent.cpp b/Engine/Source/Runtime/Engine/Private/Components/AudioComponent.cpp
--- a/Engine/Source/Runtime/Engine/Private/Components/AudioComponent.cpp
+++ b/Engine/Source/Runtime/Engine/Private/Components/AudioComponent.cpp
@@ -2,6 +2,7 @@
 
 #include "Components/AudioComponent.h"
 #include "Audio/ActorSoundParameterInterface.h"
+#include "Audio/AudioTraceUtil.h"
 #include "AudioDevice.h"
 #include "Components/BillboardComponent.h"
 #include "Engine/Engine.h"
@@ -874,7 +875,13 @@ void UAudioComponent::PlayInternal(const PlayInternalRequestData& InPlayRequestD
 	TArray<FAudioParameter> InstanceParamsCopy = InstanceParameters;
 	FAudioParameter::Merge(MoveTemp(InstanceParamsCopy), SoundParams);
 
-	AudioDevice->AddNewActiveSound(NewSharedActiveSound, MoveTemp(SoundParams));
+#if UE_AUDIO_PROFILERTRACE_ENABLED
+	const FString* EventLogID = &Audio::Trace::EventLog::ID::PlayRequestAudioComponent;
+#else
+	const FString* EventLogID = nullptr;
+#endif // UE_AUDIO_PROFILERTRACE_ENABLED
+
+	AudioDevice->AddNewActiveSound(NewSharedActiveSound, MoveTemp(SoundParams), EventLogID);
 
 	LastSoundPlayOrder = NewActiveSound.GetPlayOrder();
 
@@ -1047,7 +1054,13 @@ void UAudioComponent::Stop()
 
 	UE_LOG(LogAudio, Verbose, TEXT("%g: Stopping AudioComponent : '%s' with Sound: '%s'"), GetAudioTimeSeconds(), *GetFullName(), Sound ? *Sound->GetName() : TEXT("nullptr"));
 
-	AudioDevice->StopActiveSound(AudioComponentID);
+#if UE_AUDIO_PROFILERTRACE_ENABLED
+	const FString* EventLogID = &Audio::Trace::EventLog::ID::StopRequestAudioComponent;
+#else
+	const FString* EventLogID = nullptr;
+#endif // UE_AUDIO_PROFILERTRACE_ENABLED
+
+	AudioDevice->StopActiveSound(AudioComponentID, EventLogID);
 
 	BroadcastPlayState();
 }
diff --git a/Engine/Source/Runtime/Engine/Private/GameplayStatics.cpp b/Engine/Source/Runtime/Engine/Private/GameplayStatics.cpp
--- a/Engine/Source/Runtime/Engine/Private/GameplayStatics.cpp
+++ b/Engine/Source/Runtime/Engine/Private/GameplayStatics.cpp
@@ -51,6 +51,7 @@
 #include "Components/SceneCaptureComponent2D.h"
 #include "Sound/SoundCue.h"
 #include "Audio/ActorSoundParameterInterface.h"
+#include "Audio/AudioTraceUtil.h"
 #include "Engine/DamageEvents.h"
 
 #include UE_INLINE_GENERATED_CPP_BY_NAME(GameplayStatics)
@@ -1635,7 +1636,13 @@ void UGameplayStatics::PlaySound2D(const UObject* WorldContextObject, USoundBase
 		TArray<FAudioParameter> Params;
 		UActorSoundParameterInterface::Fill(ActiveSoundOwner, Params);
 
-		AudioDevice->AddNewActiveSound(NewActiveSound, &Params);
+#if UE_AUDIO_PROFILERTRACE_ENABLED
+		const FString* EventLogID = &Audio::Trace::EventLog::ID::PlayRequestSound2D;
+#else
+		const FString* EventLogID = nullptr;
+#endif // UE_AUDIO_PROFILERTRACE_ENABLED
+
+		AudioDevice->AddNewActiveSound(NewActiveSound, &Params, EventLogID);
 	}
 }
 
diff --git a/Engine/Source/Runtime/Engine/Private/Slate/SlateSoundDevice.cpp b/Engine/Source/Runtime/Engine/Private/Slate/SlateSoundDevice.cpp
--- a/Engine/Source/Runtime/Engine/Private/Slate/SlateSoundDevice.cpp
+++ b/Engine/Source/Runtime/Engine/Private/Slate/SlateSoundDevice.cpp
@@ -3,6 +3,7 @@
 #include "Slate/SlateSoundDevice.h"
 #include "Sound/SlateSound.h"
 #include "AudioDevice.h"
+#include "Audio/AudioTraceUtil.h"
 #include "Engine/Engine.h"
 
 DEFINE_LOG_CATEGORY_STATIC(LogSlateSoundDevice, Log, All);
@@ -24,7 +25,13 @@ void FSlateSoundDevice::PlaySound(const FSlateSound& Sound, int32 UserIndex) con
 				NewActiveSound.UserIndex = UserIndex;
 				NewActiveSound.Priority = SoundResource->Priority;
 
-				AudioDevice->AddNewActiveSound(NewActiveSound);
+#if UE_AUDIO_PROFILERTRACE_ENABLED
+				const FString* EventLogID = &Audio::Trace::EventLog::ID::PlayRequestSlateSound;
+#else
+				const FString* EventLogID = nullptr;
+#endif // UE_AUDIO_PROFILERTRACE_ENABLED
+
+				AudioDevice->AddNewActiveSound(NewActiveSound, nullptr, EventLogID);
 			}
 			else if(Object)
 			{
diff --git a/Engine/Source/Runtime/Engine/Private/SoundConcurrency.cpp b/Engine/Source/Runtime/Engine/Private/SoundConcurrency.cpp
--- a/Engine/Source/Runtime/Engine/Private/SoundConcurrency.cpp
+++ b/Engine/Source/Runtime/Engine/Private/SoundConcurrency.cpp
@@ -4,6 +4,7 @@
 
 #include "Algo/AllOf.h"
 #include "AudioDevice.h"
+#include "Audio/AudioTraceUtil.h"
 
 #include UE_INLINE_GENERATED_CPP_BY_NAME(SoundConcurrency)
 
@@ -1074,6 +1075,10 @@ void FSoundConcurrencyManager::StopDueToVoiceStealing(FActiveSound& ActiveSound)
 	}
 	else
 	{
+#if UE_AUDIO_PROFILERTRACE_ENABLED
+		Audio::Trace::EventLog::SendActiveSoundEvent(*ActiveSound.AudioDevice, ActiveSound, Audio::Trace::EventLog::ID::StopRequestConcurrency);
+#endif // UE_AUDIO_PROFILERTRACE_ENABLED
+
 		ActiveSound.AudioDevice->AddSoundToStop(&ActiveSound);
 	}
 
diff --git a/Engine/Source/Runtime/Engine/Public/Audio/AudioTraceUtil.h b/Engine/Source/Runtime/Engine/Public/Audio/AudioTraceUtil.h
--- a/Engine/Source/Runtime/Engine/Public/Audio/AudioTraceUtil.h
+++ b/Engine/Source/Runtime/Engine/Public/Audio/AudioTraceUtil.h
@@ -1,12 +1,16 @@
 // Copyright Epic Games, Inc. All Rights Reserved.
 #pragma once
 
+#include "AudioDevice.h"
 #include "AudioMixerTrace.h"
 #include "CoreMinimal.h"
 
+#define LOCTEXT_NAMESPACE "AudioTraceUtil"
+
 #if UE_AUDIO_PROFILERTRACE_ENABLED
 
 struct FActiveSound;
+
 namespace Audio::Trace
 {
 	namespace Util
@@ -20,12 +24,57 @@ namespace Audio::Trace
 
 	namespace EventLog
 	{
-		ENGINE_API void SendSoundAssetEventLog(const FActiveSound& InActiveSound, const FString& Event);
+		// Event Log in-built event types
+		namespace ID
+		{
+			// Active Sound : Playing Audio
+			const FString SoundStart = LOCTEXT("EventLogTraceMessage_Playing", "Playing").ToString();
+			const FString SoundStop = LOCTEXT("EventLogTraceMessage_Stopped", "Stopped").ToString();
+
+			// Virtualization
+			const FString SoundVirtualized = LOCTEXT("EventLogTraceMessage_Virtualize", "Virtualized").ToString();
+			const FString SoundRealized = LOCTEXT("EventLogTraceMessage_Realize", "Realized").ToString();
+
+			// Triggers
+			const FString PlayRequestSoundHandle = LOCTEXT("EventLogTraceMessage_PlayRequestedSoundHandle", "Play Request : Sound Handle").ToString();
+			const FString StopRequestedSoundHandle = LOCTEXT("EventLogTraceMessage_StopRequestedSoundHandle", "Stop Request : Sound Handle").ToString();
+
+			const FString PlayRequestAudioComponent = LOCTEXT("EventLogTraceMessage_PlayRequestedAudioComponent", "Play Request : Audio Component").ToString();
+			const FString StopRequestAudioComponent = LOCTEXT("EventLogTraceMessage_StopRequestedAudioComponent", "Stop Request : Audio Component").ToString();
+
+			const FString PlayRequestOneShot = LOCTEXT("EventLogTraceMessage_PlayRequestedOneShot", "Play Request : One shot").ToString();
+			const FString PlayRequestSoundAtLocation = LOCTEXT("EventLogTraceMessage_PlayRequestedSoundAtLocation", "Play Request : Sound at location").ToString();
+			const FString PlayRequestSound2D = LOCTEXT("EventLogTraceMessage_PlayRequested2DSound", "Play Request : Play Sound 2D").ToString();
+			const FString PlayRequestSlateSound = LOCTEXT("EventLogTraceMessage_PlayRequestedSlateSound", "Play Request : Slate Sound").ToString();
+
+			const FString StopRequestActiveSound = LOCTEXT("EventLogTraceMessage_StopRequestedActiveSound", "Stop Request : Active Sound").ToString();
+			const FString StopRequestSoundsUsingResource = LOCTEXT("EventLogTraceMessage_StopRequestedSoundsUsingResource", "Stop Request : Sounds using resource").ToString();
+			const FString StopRequestConcurrency = LOCTEXT("EventLogTraceMessage_StopRequestedConcurrency", "Stop Request : Concurrency").ToString();
+
+			const FString PauseSoundRequested = LOCTEXT("EventLogTraceMessage_PauseSoundRequested", "Paused").ToString();
+			const FString ResumeSoundRequested = LOCTEXT("EventLogTraceMessage_ResumeSoundRequested", "Resumed").ToString();
+
+			const FString StopAllRequested = LOCTEXT("EventLogTraceMessage_StopAllRequested", "Stop All").ToString();
+			const FString FlushAudioDeviceRequested = LOCTEXT("EventLogTraceMessage_FlushAudioDeviceRequested", "Flush Audio Device").ToString();
+		}// namespace ID
+
+		// Add an entry to the Audio Insights Event Log that related to an Active Sound
+		// Event may be customized to post custom event types to the Event Log
+		ENGINE_API void SendActiveSoundEvent(const FAudioDevice& AudioDevice, const FActiveSound& InActiveSound, const FString& Event);
+
+		// Add an entry to the Audio Insights Event Log
+		ENGINE_API void SendEvent(const Audio::FDeviceId AudioDeviceID
+								, const FString& Event
+								, const uint32 AudioObjectID = INDEX_NONE
+								, const FString& AssetPath = FString()
+								, const FString& ActorName = FString()
+								, const FString& ActorIconName = FString()
+								, const FString& SoundClassName = FString());
 
-		ENGINE_API const FString& GetStartEventName();
-		ENGINE_API const FString& GetStopEventName();
 	} // namespace EventLog
 
 } // namespace Audio::Trace
 
-#endif // #if UE_AUDIO_PROFILERTRACE_ENABLED
\ No newline at end of file
+#endif // #if UE_AUDIO_PROFILERTRACE_ENABLED
+
+#undef LOCTEXT_NAMESPACE // AudioTraceUtil
\ No newline at end of file
diff --git a/Engine/Source/Runtime/Engine/Public/AudioDevice.h b/Engine/Source/Runtime/Engine/Public/AudioDevice.h
--- a/Engine/Source/Runtime/Engine/Public/AudioDevice.h
+++ b/Engine/Source/Runtime/Engine/Public/AudioDevice.h
@@ -769,12 +769,12 @@ public:
 	/**
 	 * Adds an active sound to the audio device
 	 */
-	ENGINE_API void AddNewActiveSound(const FActiveSound& ActiveSound, const TArray<FAudioParameter>* InDefaultParams = nullptr);
+	ENGINE_API void AddNewActiveSound(const FActiveSound& ActiveSound, const TArray<FAudioParameter>* InDefaultParams = nullptr, const FString* EventLogID = nullptr);
 
 	/**
 	 * Adds an active sound to the audio device
 	*/
-	ENGINE_API void AddNewActiveSound(TSharedRef<FActiveSound> ActiveSound, TArray<FAudioParameter>&& InDefaultParams);
+	ENGINE_API void AddNewActiveSound(TSharedRef<FActiveSound> ActiveSound, TArray<FAudioParameter>&& InDefaultParams, const FString* EventLogID = nullptr);
 
 	/**
 	 * Adds an active sound to the audio device
@@ -790,12 +790,12 @@ public:
 	/**
 	 * Removes the active sound for the specified audio component
 	 */
-	ENGINE_API void StopActiveSound(uint64 AudioComponentID);
+	ENGINE_API void StopActiveSound(uint64 AudioComponentID, const FString* EventLogID = nullptr);
 
 	/**
 	* (Deprecated in favor of AddSoundToStop). Stops the active sound
 	*/
-	ENGINE_API void StopActiveSound(FActiveSound* ActiveSound);
+	ENGINE_API void StopActiveSound(FActiveSound* ActiveSound, const FString* EventLogID = nullptr);
 
 	/**
 	* Pauses the active sound for the specified audio component
