[Audio Insights] Moving EventLogIDs to AudioMixerCore so we can reference them in both AudioInsights and Engine

diff --git a/Engine/Source/Runtime/AudioMixerCore/Public/AudioMixerTrace.h b/Engine/Source/Runtime/AudioMixerCore/Public/AudioMixerTrace.h
--- a/Engine/Source/Runtime/AudioMixerCore/Public/AudioMixerTrace.h
+++ b/Engine/Source/Runtime/AudioMixerCore/Public/AudioMixerTrace.h
@@ -2,6 +2,7 @@
 
 #pragma once
 
+#include "CoreMinimal.h"
 #include "ProfilingDebugging/CpuProfilerTrace.h"
 #include "Trace/Trace.h"
 
@@ -20,3 +21,41 @@
 		#define AUDIO_MIXER_TRACE_CPUPROFILER_EVENT_SCOPE(Name)
 	#endif
 #endif
+
+#define LOCTEXT_NAMESPACE "AudioInsightsEventLog"
+
+// Audio Insights Event Log in-built event types
+namespace Audio::Trace::EventLog::ID
+{
+	// Active Sound : Playing Audio
+	const FString SoundStart = LOCTEXT("EventLogTraceMessage_Playing", "Playing").ToString();
+	const FString SoundStop = LOCTEXT("EventLogTraceMessage_Stopped", "Stopped").ToString();
+
+	// Virtualization
+	const FString SoundVirtualized = LOCTEXT("EventLogTraceMessage_Virtualize", "Virtualized").ToString();
+	const FString SoundRealized = LOCTEXT("EventLogTraceMessage_Realize", "Realized").ToString();
+
+	// Triggers
+	const FString PlayRequestSoundHandle = LOCTEXT("EventLogTraceMessage_PlayRequestedSoundHandle", "Play Request : Sound Handle").ToString();
+	const FString StopRequestedSoundHandle = LOCTEXT("EventLogTraceMessage_StopRequestedSoundHandle", "Stop Request : Sound Handle").ToString();
+
+	const FString PlayRequestAudioComponent = LOCTEXT("EventLogTraceMessage_PlayRequestedAudioComponent", "Play Request : Audio Component").ToString();
+	const FString StopRequestAudioComponent = LOCTEXT("EventLogTraceMessage_StopRequestedAudioComponent", "Stop Request : Audio Component").ToString();
+
+	const FString PlayRequestOneShot = LOCTEXT("EventLogTraceMessage_PlayRequestedOneShot", "Play Request : One shot").ToString();
+	const FString PlayRequestSoundAtLocation = LOCTEXT("EventLogTraceMessage_PlayRequestedSoundAtLocation", "Play Request : Sound at location").ToString();
+	const FString PlayRequestSound2D = LOCTEXT("EventLogTraceMessage_PlayRequested2DSound", "Play Request : Play Sound 2D").ToString();
+	const FString PlayRequestSlateSound = LOCTEXT("EventLogTraceMessage_PlayRequestedSlateSound", "Play Request : Slate Sound").ToString();
+
+	const FString StopRequestActiveSound = LOCTEXT("EventLogTraceMessage_StopRequestedActiveSound", "Stop Request : Active Sound").ToString();
+	const FString StopRequestSoundsUsingResource = LOCTEXT("EventLogTraceMessage_StopRequestedSoundsUsingResource", "Stop Request : Sounds using resource").ToString();
+	const FString StopRequestConcurrency = LOCTEXT("EventLogTraceMessage_StopRequestedConcurrency", "Stop Request : Concurrency").ToString();
+
+	const FString PauseSoundRequested = LOCTEXT("EventLogTraceMessage_PauseSoundRequested", "Paused").ToString();
+	const FString ResumeSoundRequested = LOCTEXT("EventLogTraceMessage_ResumeSoundRequested", "Resumed").ToString();
+
+	const FString StopAllRequested = LOCTEXT("EventLogTraceMessage_StopAllRequested", "Stop All").ToString();
+	const FString FlushAudioDeviceRequested = LOCTEXT("EventLogTraceMessage_FlushAudioDeviceRequested", "Flush Audio Device").ToString();
+}
+
+#undef LOCTEXT_NAMESPACE // AudioInsightsEventLog
\ No newline at end of file
diff --git a/Engine/Source/Runtime/Engine/Public/Audio/AudioTraceUtil.h b/Engine/Source/Runtime/Engine/Public/Audio/AudioTraceUtil.h
--- a/Engine/Source/Runtime/Engine/Public/Audio/AudioTraceUtil.h
+++ b/Engine/Source/Runtime/Engine/Public/Audio/AudioTraceUtil.h
@@ -5,8 +5,6 @@
 #include "AudioMixerTrace.h"
 #include "CoreMinimal.h"
 
-#define LOCTEXT_NAMESPACE "AudioTraceUtil"
-
 #if UE_AUDIO_PROFILERTRACE_ENABLED
 
 struct FActiveSound;
@@ -24,40 +22,6 @@ namespace Audio::Trace
 
 	namespace EventLog
 	{
-		// Event Log in-built event types
-		namespace ID
-		{
-			// Active Sound : Playing Audio
-			const FString SoundStart = LOCTEXT("EventLogTraceMessage_Playing", "Playing").ToString();
-			const FString SoundStop = LOCTEXT("EventLogTraceMessage_Stopped", "Stopped").ToString();
-
-			// Virtualization
-			const FString SoundVirtualized = LOCTEXT("EventLogTraceMessage_Virtualize", "Virtualized").ToString();
-			const FString SoundRealized = LOCTEXT("EventLogTraceMessage_Realize", "Realized").ToString();
-
-			// Triggers
-			const FString PlayRequestSoundHandle = LOCTEXT("EventLogTraceMessage_PlayRequestedSoundHandle", "Play Request : Sound Handle").ToString();
-			const FString StopRequestedSoundHandle = LOCTEXT("EventLogTraceMessage_StopRequestedSoundHandle", "Stop Request : Sound Handle").ToString();
-
-			const FString PlayRequestAudioComponent = LOCTEXT("EventLogTraceMessage_PlayRequestedAudioComponent", "Play Request : Audio Component").ToString();
-			const FString StopRequestAudioComponent = LOCTEXT("EventLogTraceMessage_StopRequestedAudioComponent", "Stop Request : Audio Component").ToString();
-
-			const FString PlayRequestOneShot = LOCTEXT("EventLogTraceMessage_PlayRequestedOneShot", "Play Request : One shot").ToString();
-			const FString PlayRequestSoundAtLocation = LOCTEXT("EventLogTraceMessage_PlayRequestedSoundAtLocation", "Play Request : Sound at location").ToString();
-			const FString PlayRequestSound2D = LOCTEXT("EventLogTraceMessage_PlayRequested2DSound", "Play Request : Play Sound 2D").ToString();
-			const FString PlayRequestSlateSound = LOCTEXT("EventLogTraceMessage_PlayRequestedSlateSound", "Play Request : Slate Sound").ToString();
-
-			const FString StopRequestActiveSound = LOCTEXT("EventLogTraceMessage_StopRequestedActiveSound", "Stop Request : Active Sound").ToString();
-			const FString StopRequestSoundsUsingResource = LOCTEXT("EventLogTraceMessage_StopRequestedSoundsUsingResource", "Stop Request : Sounds using resource").ToString();
-			const FString StopRequestConcurrency = LOCTEXT("EventLogTraceMessage_StopRequestedConcurrency", "Stop Request : Concurrency").ToString();
-
-			const FString PauseSoundRequested = LOCTEXT("EventLogTraceMessage_PauseSoundRequested", "Paused").ToString();
-			const FString ResumeSoundRequested = LOCTEXT("EventLogTraceMessage_ResumeSoundRequested", "Resumed").ToString();
-
-			const FString StopAllRequested = LOCTEXT("EventLogTraceMessage_StopAllRequested", "Stop All").ToString();
-			const FString FlushAudioDeviceRequested = LOCTEXT("EventLogTraceMessage_FlushAudioDeviceRequested", "Flush Audio Device").ToString();
-		}// namespace ID
-
 		// Add an entry to the Audio Insights Event Log that related to an Active Sound
 		// Event may be customized to post custom event types to the Event Log
 		ENGINE_API void SendActiveSoundEvent(const FAudioDevice& AudioDevice, const FActiveSound& InActiveSound, const FString& Event);
@@ -75,6 +39,4 @@ namespace Audio::Trace
 
 } // namespace Audio::Trace
 
-#endif // #if UE_AUDIO_PROFILERTRACE_ENABLED
-
-#undef LOCTEXT_NAMESPACE // AudioTraceUtil
\ No newline at end of file
+#endif // #if UE_AUDIO_PROFILERTRACE_ENABLED
\ No newline at end of file
