Fixed an issue where an SRV wouldn't be created for index buffers when expected.

Index buffers can start out with a null buffer RHI handle, which becomes valid later when the geometry is streamed in. The previous code skipped creating the SRV on startup in this case, but the desired behavior is to create the SRV if the buffer will ever be valid.

diff --git a/Engine/Source/Runtime/Engine/Public/RawIndexBuffer.h b/Engine/Source/Runtime/Engine/Public/RawIndexBuffer.h
--- a/Engine/Source/Runtime/Engine/Public/RawIndexBuffer.h
+++ b/Engine/Source/Runtime/Engine/Public/RawIndexBuffer.h
@@ -429,13 +429,7 @@ public:
 			return {};
 		}
 		
-		const bool bHadIndexData = Num() > 0;
-		if (!bHadIndexData)
-		{
-			return {};
-		}
-		
-		if ((IndexBufferRHI->GetUsage() & EBufferUsageFlags::ShaderResource) != EBufferUsageFlags::ShaderResource)
+		if (!IsShaderResource(Indices.GetAllowCPUAccess()))
 		{
 			return {};
 		}
