Remove the declaration of method that was removed in CL40093672

diff --git a/Engine/Source/Runtime/Engine/Classes/Engine/Level.h b/Engine/Source/Runtime/Engine/Classes/Engine/Level.h
--- a/Engine/Source/Runtime/Engine/Classes/Engine/Level.h
+++ b/Engine/Source/Runtime/Engine/Classes/Engine/Level.h
@@ -995,11 +995,6 @@ public:
 	 */
 	ENGINE_API void CleanupLevel(bool bCleanupResources = true, bool bUnloadFromEditor = true);
 
-	/**
-	 * Cleans all references from and to this level that may be preventing it from being Garbage Collected
-	 */
-	ENGINE_API void CleanupReferences();
-
 	/**
 	 * Clears all components of actors associated with this level (aka in Actors array) and 
 	 * also the BSP model components.
