Updated GetKeysForAction and GetKeysForAxis in UPlayerInput to be virtual.

diff --git a/Engine/Source/Runtime/Engine/Classes/GameFramework/PlayerInput.h b/Engine/Source/Runtime/Engine/Classes/GameFramework/PlayerInput.h
--- a/Engine/Source/Runtime/Engine/Classes/GameFramework/PlayerInput.h
+++ b/Engine/Source/Runtime/Engine/Classes/GameFramework/PlayerInput.h
@@ -742,10 +742,10 @@ public:
 #endif
 
 	/** Returns the list of keys mapped to the specified Action Name */
-	ENGINE_API const TArray<FInputActionKeyMapping>& GetKeysForAction(const FName ActionName) const;
+	ENGINE_API virtual const TArray<FInputActionKeyMapping>& GetKeysForAction(const FName ActionName) const;
 
 	/** Returns the list of keys mapped to the specified Axis Name */
-	ENGINE_API const TArray<FInputAxisKeyMapping>& GetKeysForAxis(const FName AxisName) const;
+	ENGINE_API virtual const TArray<FInputAxisKeyMapping>& GetKeysForAxis(const FName AxisName) const;
 
 	static const TArray<FInputActionKeyMapping>& GetEngineDefinedActionMappings() { return EngineDefinedActionMappings; }
 	static const TArray<FInputAxisKeyMapping>& GetEngineDefinedAxisMappings() { return EngineDefinedAxisMappings; }
