[StandaloneHLOD] Add AWorldPartitionHLODOnlyLevelInstance Actor to allow users to place Standalone HLODs in the world without having to add actual Level Instance to the world
* After adding the actor to the world users can select the level for which Standalone HLODs should be loaded
* The actual level is neither loaded nor cooked. It's only used by the Standalone HLOD Subsystem to pull in the corresponding Standalone HLOD levels

diff --git a/Engine/Source/Runtime/Engine/Private/LevelInstance/LevelInstanceInterface.cpp b/Engine/Source/Runtime/Engine/Private/LevelInstance/LevelInstanceInterface.cpp
--- a/Engine/Source/Runtime/Engine/Private/LevelInstance/LevelInstanceInterface.cpp
+++ b/Engine/Source/Runtime/Engine/Private/LevelInstance/LevelInstanceInterface.cpp
@@ -193,6 +193,18 @@ bool ILevelInstanceInterface::IsLevelInstanceHLODRelevant() const
 	return false;
 }
 
+bool ILevelInstanceInterface::HasStandaloneHLOD() const
+{
+	if (ULevel* LoadedLevel = GetLoadedLevel())
+	{
+		if (UWorldPartition* WorldPartition = LoadedLevel->GetWorldPartition())
+		{
+			return WorldPartition->HasStandaloneHLOD();
+		}
+	}
+	return false;
+}
+
 bool ILevelInstanceInterface::IsEditing() const
 {
 	// No need to check LevelInstanceID here since Editing Level Instance is referenced directly:
diff --git a/Engine/Source/Runtime/Engine/Private/WorldPartition/HLOD/HLODOnlyLevelInstanceActor.cpp b/Engine/Source/Runtime/Engine/Private/WorldPartition/HLOD/HLODOnlyLevelInstanceActor.cpp
new file mode 100644
--- /dev/null
+++ b/Engine/Source/Runtime/Engine/Private/WorldPartition/HLOD/HLODOnlyLevelInstanceActor.cpp
@@ -0,0 +1,110 @@
+// Copyright Epic Games, Inc. All Rights Reserved.
+
+#include "WorldPartition/HLOD/HLODOnlyLevelInstanceActor.h"
+
+#include "Engine/World.h"
+#include "LevelInstance/LevelInstanceComponent.h"
+#include "WorldPartition/LevelInstance/LevelInstanceActorDesc.h"
+
+AWorldPartitionHLODOnlyLevelInstance::AWorldPartitionHLODOnlyLevelInstance(const FObjectInitializer& ObjectInitializer)
+	: Super(ObjectInitializer)
+	, LevelInstanceActorGuid(this)
+	, LevelInstanceActorImpl(this)
+{
+	RootComponent = CreateDefaultSubobject<ULevelInstanceComponent>(TEXT("Root"));
+	RootComponent->Mobility = EComponentMobility::Static;
+
+	bIsEditorOnlyActor = true;
+}
+
+const FLevelInstanceID& AWorldPartitionHLODOnlyLevelInstance::GetLevelInstanceID() const
+{
+	return LevelInstanceActorImpl.GetLevelInstanceID();
+}
+
+bool AWorldPartitionHLODOnlyLevelInstance::HasValidLevelInstanceID() const
+{
+	return LevelInstanceActorImpl.HasValidLevelInstanceID();
+}
+
+const FGuid& AWorldPartitionHLODOnlyLevelInstance::GetLevelInstanceGuid() const
+{
+	return LevelInstanceActorGuid.GetGuid();
+}
+
+const TSoftObjectPtr<UWorld>& AWorldPartitionHLODOnlyLevelInstance::GetWorldAsset() const
+{
+	return WorldAsset;
+}
+
+bool AWorldPartitionHLODOnlyLevelInstance::IsLoadingEnabled() const
+{
+	return false;
+}
+
+bool AWorldPartitionHLODOnlyLevelInstance::SetWorldAsset(TSoftObjectPtr<UWorld> InWorldAsset)
+{
+	WorldAsset = InWorldAsset;
+	return true;
+}
+
+#if WITH_EDITOR
+ULevelInstanceComponent* AWorldPartitionHLODOnlyLevelInstance::GetLevelInstanceComponent() const
+{
+	return Cast<ULevelInstanceComponent>(RootComponent);
+}
+
+ELevelInstanceRuntimeBehavior AWorldPartitionHLODOnlyLevelInstance::GetDesiredRuntimeBehavior() const
+{
+	return ELevelInstanceRuntimeBehavior::LevelStreaming;
+}
+
+ELevelInstanceRuntimeBehavior AWorldPartitionHLODOnlyLevelInstance::GetDefaultRuntimeBehavior() const
+{
+	return ELevelInstanceRuntimeBehavior::LevelStreaming;
+}
+
+bool AWorldPartitionHLODOnlyLevelInstance::HasStandaloneHLOD() const
+{
+	return true;
+}
+
+bool AWorldPartitionHLODOnlyLevelInstance::CanEditChange(const FProperty* InProperty) const
+{
+	if (!Super::CanEditChange(InProperty))
+	{
+		return false;
+	}
+
+	if (InProperty && InProperty->GetFName() == GET_MEMBER_NAME_CHECKED(AActor, bIsEditorOnlyActor))
+	{
+		return false;
+	}
+
+	return true;
+}
+#endif // WITH_EDITOR
+
+void AWorldPartitionHLODOnlyLevelInstance::PostRegisterAllComponents()
+{
+	Super::PostRegisterAllComponents();
+
+	if (LevelInstanceActorGuid.IsValid())
+	{
+		LevelInstanceActorImpl.RegisterLevelInstance();
+	}
+}
+
+void AWorldPartitionHLODOnlyLevelInstance::PostUnregisterAllComponents()
+{
+	Super::PostUnregisterAllComponents();
+
+	LevelInstanceActorImpl.UnregisterLevelInstance();
+}
+
+#if WITH_EDITOR
+TUniquePtr<FWorldPartitionActorDesc> AWorldPartitionHLODOnlyLevelInstance::CreateClassActorDesc() const
+{
+	return TUniquePtr<FWorldPartitionActorDesc>(new FLevelInstanceActorDesc());
+}
+#endif
\ No newline at end of file
diff --git a/Engine/Source/Runtime/Engine/Private/WorldPartition/LevelInstance/LevelInstanceActorDesc.cpp b/Engine/Source/Runtime/Engine/Private/WorldPartition/LevelInstance/LevelInstanceActorDesc.cpp
--- a/Engine/Source/Runtime/Engine/Private/WorldPartition/LevelInstance/LevelInstanceActorDesc.cpp
+++ b/Engine/Source/Runtime/Engine/Private/WorldPartition/LevelInstance/LevelInstanceActorDesc.cpp
@@ -61,13 +61,7 @@ void FLevelInstanceActorDesc::Init(const AActor* InActor)
 	
 	bIsChildContainerInstance = IsChildContainerInstanceInternal();
 
-	if (ULevel* LoadedLevel = LevelInstance->GetLoadedLevel())
-	{
-		if (UWorldPartition* WorldPartition = LoadedLevel->GetWorldPartition())
-		{
-			bHasStandaloneHLOD = WorldPartition->HasStandaloneHLOD();
-		}
-	}
+	bHasStandaloneHLOD = LevelInstance->HasStandaloneHLOD();
 
 	if (LevelInstance->SupportsPropertyOverrides())
 	{
diff --git a/Engine/Source/Runtime/Engine/Public/LevelInstance/LevelInstanceInterface.h b/Engine/Source/Runtime/Engine/Public/LevelInstance/LevelInstanceInterface.h
--- a/Engine/Source/Runtime/Engine/Public/LevelInstance/LevelInstanceInterface.h
+++ b/Engine/Source/Runtime/Engine/Public/LevelInstance/LevelInstanceInterface.h
@@ -155,6 +155,8 @@ class ILevelInstanceInterface
 
 	// Whether the level instance is HLOD relevant. If not, all it's actors will be excluded from HLOD.
 	ENGINE_API virtual bool IsLevelInstanceHLODRelevant() const;
+	
+	ENGINE_API virtual bool HasStandaloneHLOD() const;
 #endif
 
 	/** Sets the UWorld asset reference when loading a LevelInstance */
diff --git a/Engine/Source/Runtime/Engine/Public/WorldPartition/HLOD/HLODOnlyLevelInstanceActor.h b/Engine/Source/Runtime/Engine/Public/WorldPartition/HLOD/HLODOnlyLevelInstanceActor.h
new file mode 100644
--- /dev/null
+++ b/Engine/Source/Runtime/Engine/Public/WorldPartition/HLOD/HLODOnlyLevelInstanceActor.h
@@ -0,0 +1,59 @@
+// Copyright Epic Games, Inc. All Rights Reserved.
+
+#pragma once
+
+#include "LevelInstance/LevelInstanceActorGuid.h"
+#include "LevelInstance/LevelInstanceActorImpl.h"
+#include "LevelInstance/LevelInstanceInterface.h"
+#include "HLODOnlyLevelInstanceActor.generated.h"
+
+UCLASS(MinimalAPI, HideCategories=(Rendering, Replication, Collision, Physics, Navigation, Networking, Input, Actor, LevelInstance, HLOD, WorldPartition, Cooking, DataLayers))
+class AWorldPartitionHLODOnlyLevelInstance : public AActor, public ILevelInstanceInterface
+{
+	GENERATED_BODY()
+
+public:
+	ENGINE_API AWorldPartitionHLODOnlyLevelInstance(const FObjectInitializer& ObjectInitializer);
+
+	//~ Begin ILevelInstanceInterface
+	ENGINE_API virtual const FLevelInstanceID& GetLevelInstanceID() const override;
+	ENGINE_API virtual bool HasValidLevelInstanceID() const override;
+	ENGINE_API virtual const FGuid& GetLevelInstanceGuid() const override;
+	ENGINE_API virtual const TSoftObjectPtr<UWorld>& GetWorldAsset() const override;
+	ENGINE_API virtual bool IsLoadingEnabled() const override;
+	ENGINE_API virtual bool SetWorldAsset(TSoftObjectPtr<UWorld> InWorldAsset) override;
+
+#if WITH_EDITOR
+	ENGINE_API virtual ULevelInstanceComponent* GetLevelInstanceComponent() const override;
+	ENGINE_API virtual ELevelInstanceRuntimeBehavior GetDesiredRuntimeBehavior() const override;
+	ENGINE_API virtual ELevelInstanceRuntimeBehavior GetDefaultRuntimeBehavior() const override;
+	ENGINE_API virtual bool HasStandaloneHLOD() const override;
+#endif // WITH_EDITOR
+	//~ End ILevelInstanceInterface
+
+	//~ Begin UObject Interface
+#if WITH_EDITOR
+	ENGINE_API virtual bool CanEditChange(const FProperty* InProperty) const override;
+#endif
+	//~ End UObject Interface.
+
+	//~ Begin AActor Interface.
+	ENGINE_API void PostRegisterAllComponents() override;
+	ENGINE_API void PostUnregisterAllComponents() override;
+	
+#if WITH_EDITOR
+	ENGINE_API virtual TUniquePtr<class FWorldPartitionActorDesc> CreateClassActorDesc() const override;
+#endif
+	//~ End AActor Interface.
+
+protected:
+	UPROPERTY(EditAnywhere, Category = Level, meta = (NoCreate, DisplayName = "Level"))
+	TSoftObjectPtr<UWorld> WorldAsset;
+	
+	UPROPERTY(Transient)
+	FGuid LevelInstanceSpawnGuid;
+
+private:
+	FLevelInstanceActorGuid LevelInstanceActorGuid;
+	FLevelInstanceActorImpl LevelInstanceActorImpl;
+};
\ No newline at end of file
