Material Translator: Fix (new) implementation of UMaterialExpressionFunctionOutput::GetOutputValueType() not returning MCT_MaterialAttributes when it should.
It relies on traversing the input to find the type, but at least some expressions can't be relied on to give the correct result here.
Those expressions do return true for IsResultMaterialAttributes() so now check that first. This is what we do in the base FMaterialExpression implementation of GetOutputValueType() which we are now overriding.
At the same time fix a couple of material expressions that were giving the bad return value for GetOutputValueType() and causing cook failures.
Both should just use the default implementation from FMaterialExpression.
UMaterialExpressionStaticSwitch had a bad override to MCT_Unknown.
UMaterialExpressionStaticSwitchParameter derives from UMaterialExpressionStaticBoolParameter but shouldn't use its implementation.
#rnx

diff --git a/Engine/Source/Runtime/Core/Private/UObject/FortniteMainBranchObjectVersion.cpp b/Engine/Source/Runtime/Core/Private/UObject/FortniteMainBranchObjectVersion.cpp
--- a/Engine/Source/Runtime/Core/Private/UObject/FortniteMainBranchObjectVersion.cpp
+++ b/Engine/Source/Runtime/Core/Private/UObject/FortniteMainBranchObjectVersion.cpp
@@ -14,6 +14,6 @@ TMap<FGuid, FGuid> FFortniteMainBranchObjectVersion::GetSystemGuids()
 	SystemGuids.Add(DevGuids.Niagara_LatestScriptCompileVersion, FGuid("97F0FD7149D90D42B65D6557AF477D1A"));
 	SystemGuids.Add(DevGuids.SkeletalMeshDerivedDataVersion, FGuid("C998F2F3B2884A02B7E290B8C25835F0"));
 	SystemGuids.Add(DevGuids.STATICMESH_DERIVEDDATA_VER, FGuid("3ABF32149F9D4D838750368AB95573BE"));
-	SystemGuids.Add(DevGuids.MaterialTranslationDDCVersion, FGuid("5F49E591FCFD4E21AA8D37F8D15A5856"));
+	SystemGuids.Add(DevGuids.MaterialTranslationDDCVersion, FGuid("7F010615E10C4835B08B2395E7274119"));
 	return SystemGuids;
 }
diff --git a/Engine/Source/Runtime/Engine/Private/Materials/MaterialExpressions.cpp b/Engine/Source/Runtime/Engine/Private/Materials/MaterialExpressions.cpp
--- a/Engine/Source/Runtime/Engine/Private/Materials/MaterialExpressions.cpp
+++ b/Engine/Source/Runtime/Engine/Private/Materials/MaterialExpressions.cpp
@@ -18574,6 +18574,14 @@ EMaterialValueType UMaterialExpressionFunctionOutput::GetOutputValueType(int32 O
 	const FExpressionInput& ExpressionInputRef = A.GetTracedInput();
 	if (ExpressionInputRef.Expression)
 	{
+		if (ExpressionInputRef.Expression->IsResultMaterialAttributes(OutputIndex))
+		{
+			return MCT_MaterialAttributes;
+		}
+		if (ExpressionInputRef.Expression->IsResultSubstrateMaterial(OutputIndex))
+		{
+			return MCT_Substrate;
+		}
 		return ExpressionInputRef.Expression->GetOutputValueType(ExpressionInputRef.OutputIndex);
 	}
 	return Super::GetOutputValueType(OutputIndex);
diff --git a/Engine/Source/Runtime/Engine/Public/Materials/MaterialExpressionStaticSwitch.h b/Engine/Source/Runtime/Engine/Public/Materials/MaterialExpressionStaticSwitch.h
--- a/Engine/Source/Runtime/Engine/Public/Materials/MaterialExpressionStaticSwitch.h
+++ b/Engine/Source/Runtime/Engine/Public/Materials/MaterialExpressionStaticSwitch.h
@@ -38,7 +38,6 @@ public:
 	virtual FName GetInputName(int32 InputIndex) const override;
 	virtual bool IsResultMaterialAttributes(int32 OutputIndex) override;
 	virtual EMaterialValueType GetInputValueType(int32 InputIndex) override;
-	virtual EMaterialValueType GetOutputValueType(int32 OutputIndex) override {return MCT_Unknown;}
 	virtual bool IsStaticExpression() const override { return true; }
 
 	virtual bool IsResultSubstrateMaterial(int32 OutputIndex) override;
diff --git a/Engine/Source/Runtime/Engine/Public/Materials/MaterialExpressionStaticSwitchParameter.h b/Engine/Source/Runtime/Engine/Public/Materials/MaterialExpressionStaticSwitchParameter.h
--- a/Engine/Source/Runtime/Engine/Public/Materials/MaterialExpressionStaticSwitchParameter.h
+++ b/Engine/Source/Runtime/Engine/Public/Materials/MaterialExpressionStaticSwitchParameter.h
@@ -31,7 +31,7 @@ public:
 	virtual void GetCaption(TArray<FString>& OutCaptions) const override;
 	virtual bool IsResultMaterialAttributes(int32 OutputIndex) override;
 	virtual EMaterialValueType GetInputValueType(int32 InputIndex) override {return MCT_Unknown;}
-	virtual EMaterialValueType GetOutputValueType(int32 OutputIndex) override {return MCT_Unknown;}
+	virtual EMaterialValueType GetOutputValueType(int32 OutputIndex) override { return UMaterialExpression::GetOutputValueType(OutputIndex); }
 	virtual bool IsStaticExpression() const override { return true; }
 
 	virtual bool IsResultSubstrateMaterial(int32 OutputIndex) override;
