Add additional project-specific texture groups

#tests local editor

diff --git a/Engine/Source/Runtime/Engine/Classes/Engine/TextureDefines.h b/Engine/Source/Runtime/Engine/Classes/Engine/TextureDefines.h
--- a/Engine/Source/Runtime/Engine/Classes/Engine/TextureDefines.h
+++ b/Engine/Source/Runtime/Engine/Classes/Engine/TextureDefines.h
@@ -105,6 +105,22 @@ enum TextureGroup : int
 	TEXTUREGROUP_Project30 UMETA(DisplayName="ini:Project Group 30"),
 	TEXTUREGROUP_Project31 UMETA(DisplayName="ini:Project Group 31"),
 	TEXTUREGROUP_Project32 UMETA(DisplayName="ini:Project Group 32"),
+	TEXTUREGROUP_Project33 UMETA(DisplayName="ini:Project Group 33"),
+	TEXTUREGROUP_Project34 UMETA(DisplayName="ini:Project Group 34"),
+	TEXTUREGROUP_Project35 UMETA(DisplayName="ini:Project Group 35"),
+	TEXTUREGROUP_Project36 UMETA(DisplayName="ini:Project Group 36"),
+	TEXTUREGROUP_Project37 UMETA(DisplayName="ini:Project Group 37"),
+	TEXTUREGROUP_Project38 UMETA(DisplayName="ini:Project Group 38"),
+	TEXTUREGROUP_Project39 UMETA(DisplayName="ini:Project Group 39"),
+	TEXTUREGROUP_Project40 UMETA(DisplayName="ini:Project Group 40"),
+	TEXTUREGROUP_Project41 UMETA(DisplayName="ini:Project Group 41"),
+	TEXTUREGROUP_Project42 UMETA(DisplayName="ini:Project Group 42"),
+	TEXTUREGROUP_Project43 UMETA(DisplayName="ini:Project Group 43"),
+	TEXTUREGROUP_Project44 UMETA(DisplayName="ini:Project Group 44"),
+	TEXTUREGROUP_Project45 UMETA(DisplayName="ini:Project Group 45"),
+	TEXTUREGROUP_Project46 UMETA(DisplayName="ini:Project Group 46"),
+	TEXTUREGROUP_Project47 UMETA(DisplayName="ini:Project Group 47"),
+	TEXTUREGROUP_Project48 UMETA(DisplayName="ini:Project Group 48"),
 	TEXTUREGROUP_MAX,
 };
 
