Moved bHIddenEdTemporary flag from UPrimitiveComponent to USceneComponent. Added a missing check for bHiddenEdTemporary to ShouldRender to fix some code paths which do not check UPrimitiveComponent::ShouldComponentAddToScene.

diff --git a/Engine/Source/Runtime/Engine/Classes/Components/PrimitiveComponent.h b/Engine/Source/Runtime/Engine/Classes/Components/PrimitiveComponent.h
--- a/Engine/Source/Runtime/Engine/Classes/Components/PrimitiveComponent.h
+++ b/Engine/Source/Runtime/Engine/Classes/Components/PrimitiveComponent.h
@@ -383,12 +383,6 @@ private:
 	UPROPERTY(EditAnywhere, BlueprintGetter = GetGenerateOverlapEvents, BlueprintSetter = SetGenerateOverlapEvents, Category = Collision)
 	uint8 bGenerateOverlapEvents : 1;
 
-#if WITH_EDITORONLY_DATA
-	/** Whether this component is temporarily hidden within the editor. */
-	UPROPERTY(Transient)
-	uint8 bHiddenEdTemporary:1;
-#endif // WITH_EDITORONLY_DATA
-
 public:
 	/**
 	 * If true, this component will generate individual overlaps for each overlapping physics body if it is a multi-body component. When false, this component will
@@ -2639,18 +2633,6 @@ public:
 	/** Returns mask that represents in which views this primitive is hidden */
 	ENGINE_API virtual uint64 GetHiddenEditorViews() const;
 
-	/**
-	 * Returns whether or not this component was explicitly hidden in the editor for the duration of the current editor session
-	 * @param bIncludeParent - Whether to recurse up actor hierarchy or not
-	 */
-	ENGINE_API bool IsTemporarilyHiddenInEditor(const bool bIncludeParent) const;
-
-	/**
- 	 * Explicitly sets whether or not this component is hidden in the editor for the duration of the current editor session
- 	 * @param bIsHidden	True if the component is hidden
- 	 */
-	ENGINE_API void SetIsTemporarilyHiddenInEditor(const bool bInIsHidden);
-
 	/** Sets whether this component is being moved by the editor so the renderer can render velocities for it, even when Static. */
 	ENGINE_API void SetIsBeingMovedByEditor(bool bNewIsBeingMoved);
 
diff --git a/Engine/Source/Runtime/Engine/Classes/Components/SceneComponent.h b/Engine/Source/Runtime/Engine/Classes/Components/SceneComponent.h
--- a/Engine/Source/Runtime/Engine/Classes/Components/SceneComponent.h
+++ b/Engine/Source/Runtime/Engine/Classes/Components/SceneComponent.h
@@ -187,6 +187,12 @@ private:
 	UPROPERTY(Transient, Replicated)
 	uint8 bShouldSnapScaleWhenAttached : 1;
 
+#if WITH_EDITORONLY_DATA
+	/** Whether this component is temporarily hidden within the editor. */
+	UPROPERTY(Transient)
+	uint8 bHiddenEdTemporary : 1;
+#endif // WITH_EDITORONLY_DATA
+
 	/** Sets bShouldBeAttached, push model aware. */
 	ENGINE_API void SetShouldBeAttached(bool bNewShouldBeAttached);
 
@@ -917,6 +923,19 @@ public:
 		SetHiddenInGame(NewHidden, bPropagateToChildren ? EVisibilityPropagation::Propagate : EVisibilityPropagation::DirtyOnly);
 	}
 
+#if WITH_EDITOR
+	UE_DEPRECATED(5.7, "Please call GetIsTemporarilyHiddenInEditorFlag()")
+	ENGINE_API bool IsTemporarilyHiddenInEditor(const bool bIncludeParent) const
+	{
+		return bHiddenEdTemporary;
+	}
+	/**
+ 	 * Explicitly sets whether or not this component is hidden in the editor for the duration of the current editor session
+ 	 * @param bIsHidden	True if the component is hidden
+ 	 */
+	ENGINE_API void SetIsTemporarilyHiddenInEditor(const bool bInIsHidden);
+#endif // WITH_EDITOR
+
 	//~ Begin ActorComponent Interface
 	ENGINE_API virtual void GetLifetimeReplicatedProps(TArray<FLifetimeProperty>& OutLifetimeProps) const override;
 	ENGINE_API virtual void OnPreRegister() override;
diff --git a/Engine/Source/Runtime/Engine/Private/Components/PrimitiveComponent.cpp b/Engine/Source/Runtime/Engine/Private/Components/PrimitiveComponent.cpp
--- a/Engine/Source/Runtime/Engine/Private/Components/PrimitiveComponent.cpp
+++ b/Engine/Source/Runtime/Engine/Private/Components/PrimitiveComponent.cpp
@@ -383,9 +383,6 @@ UPrimitiveComponent::UPrimitiveComponent(const FObjectInitializer& ObjectInitial
 	bIsBeingMovedByEditor = false;
 
 	SetGenerateOverlapEvents(true);
-#if WITH_EDITORONLY_DATA
-	bHiddenEdTemporary = false;
-#endif // WITH_EDITORONLY_DATA 
 	bMultiBodyOverlap = false;
 	bReturnMaterialOnMove = false;
 	bCanEverAffectNavigation = false;
@@ -1991,15 +1988,12 @@ bool UPrimitiveComponent::ShouldComponentAddToScene() const
 #if WITH_EDITOR
 	AActor* Owner = GetOwner();
 	const bool bIsHiddenInEditor = GIsEditor && Owner && Owner->IsHiddenEd();
-
-	constexpr bool bIncludeParent = false;
-	const bool bIsTemporarilyHiddenInEditor = IsTemporarilyHiddenInEditor(bIncludeParent);
 #else
 	const bool bIsHiddenInEditor = false;
 	const bool bIsTemporarilyHiddenInEditor = false;
 #endif
 
-	return bSceneAdd && !bIsTemporarilyHiddenInEditor && (ShouldRender() || (CastShadow && bCastHiddenShadow && !bIsHiddenInEditor) || bAffectIndirectLightingWhileHidden || bRayTracingFarField);
+	return bSceneAdd && (ShouldRender() || (CastShadow && bCastHiddenShadow && !bIsHiddenInEditor) || bAffectIndirectLightingWhileHidden || bRayTracingFarField);
 }
 
 bool UPrimitiveComponent::ShouldCreatePhysicsState() const
@@ -2274,30 +2268,6 @@ uint64 UPrimitiveComponent::GetHiddenEditorViews() const
 	return OwnerActor ? OwnerActor->HiddenEditorViews : 0;
 }
 
-bool UPrimitiveComponent::IsTemporarilyHiddenInEditor(const bool bIncludeParent) const
-{
-	if (bHiddenEdTemporary)
-	{
-		return true;
-	}
-
-	if (bIncludeParent)
-	{
-		return GetOwner() && GetOwner()->IsTemporarilyHiddenInEditor(true);
-	}
-
-	return false;
-}
-
-void UPrimitiveComponent::SetIsTemporarilyHiddenInEditor(const bool bInIsHidden)
-{
-	if (bHiddenEdTemporary != bInIsHidden)
-	{
-		bHiddenEdTemporary = bInIsHidden;
-		MarkRenderStateDirty();
-	}
-}
-
 void UPrimitiveComponent::SetIsBeingMovedByEditor(bool bIsBeingMoved)
 {
 	bIsBeingMovedByEditor = bIsBeingMoved;
diff --git a/Engine/Source/Runtime/Engine/Private/Components/SceneComponent.cpp b/Engine/Source/Runtime/Engine/Private/Components/SceneComponent.cpp
--- a/Engine/Source/Runtime/Engine/Private/Components/SceneComponent.cpp
+++ b/Engine/Source/Runtime/Engine/Private/Components/SceneComponent.cpp
@@ -117,6 +117,10 @@ USceneComponent::USceneComponent(const FObjectInitializer& ObjectInitializer /*=
 	SetVisibleFlag(true);
 	bAutoActivate = false;
 	SetShouldBeAttached(AttachParent != nullptr);
+
+#if WITH_EDITORONLY_DATA
+	bHiddenEdTemporary = false;
+#endif // WITH_EDITORONLY_DATA
 }
 
 #if WITH_EDITORONLY_DATA
@@ -3477,7 +3481,7 @@ bool USceneComponent::ShouldRender() const
 	
 	const bool bShowInEditor = 
 #if WITH_EDITOR
-		GIsEditor ? (!Owner || !Owner->IsHiddenEd()) : false;
+		GIsEditor ? !bHiddenEdTemporary && (!Owner || !Owner->IsHiddenEd()) : false;
 #else
 		false;
 #endif
@@ -3504,7 +3508,7 @@ bool USceneComponent::CanEverRender() const
 
 	const bool bShowInEditor =
 #if WITH_EDITOR
-		GIsEditor ? (!Owner || !Owner->IsHiddenEd()) : false;
+		GIsEditor ? !bHiddenEdTemporary && (!Owner || !Owner->IsHiddenEd()) : false;
 #else
 		false;
 #endif
@@ -3649,6 +3653,19 @@ void USceneComponent::SetHiddenInGame(const bool bNewHiddenGame, const USceneCom
 	}
 }
 
+
+#if WITH_EDITOR
+void USceneComponent::SetIsTemporarilyHiddenInEditor(const bool bInIsHidden)
+{
+	if (bHiddenEdTemporary != bInIsHidden)
+	{
+		bHiddenEdTemporary = bInIsHidden;
+		MarkRenderStateDirty();
+	}
+}
+#endif // WITH_EDITOR
+
+
 void USceneComponent::ApplyWorldOffset(const FVector& InOffset, bool bWorldShift)
 {
 	Super::ApplyWorldOffset(InOffset, bWorldShift);
