Fixing bug where arrays of structs with custom identical operators could fail to copy from Blueprint generated classes correctly.
Adding check to custom identical operator when generating custom property lists for Blueprint instance post-construct initialization fast path.

diff --git a/Engine/Source/Runtime/Engine/Private/BlueprintGeneratedClass.cpp b/Engine/Source/Runtime/Engine/Private/BlueprintGeneratedClass.cpp
--- a/Engine/Source/Runtime/Engine/Private/BlueprintGeneratedClass.cpp
+++ b/Engine/Source/Runtime/Engine/Private/BlueprintGeneratedClass.cpp
@@ -1136,8 +1136,22 @@ bool UBlueprintGeneratedClass::BuildCustomArrayPropertyListForPostConstruction(F
 				*CurrentArrayNodePtr = new FCustomPropertyListNode(ArrayProperty, ArrayValueIndex);
 				CustomPropertyListForPostConstruction.Add(*CurrentArrayNodePtr);
 
-				// Recursively gather up all struct fields that differ and assign to the array item value node's sub property list.
-				if (BuildCustomPropertyListForPostConstruction((*CurrentArrayNodePtr)->SubPropertyList, InnerStructProperty->Struct, ArrayPropertyValue, DefaultArrayPropertyValue))
+				UScriptStruct::ICppStructOps* CppStructOps = nullptr;
+				if (InnerStructProperty->Struct)
+				{
+					CppStructOps = InnerStructProperty->Struct->GetCppStructOps();
+				}
+
+				// Check if we should initialize using the full value (e.g. a USTRUCT with one or more non-reflected fields).
+				bool bIsIdentical = false;
+				const uint32 PortFlags = 0;
+				if (!CppStructOps || !CppStructOps->HasIdentical() || !CppStructOps->Identical(ArrayPropertyValue, DefaultArrayPropertyValue, PortFlags, bIsIdentical))
+				{
+					// Recursively gather up all struct fields that differ and assign to the array item value node's sub property list.
+					bIsIdentical = !BuildCustomPropertyListForPostConstruction((*CurrentArrayNodePtr)->SubPropertyList, InnerStructProperty->Struct, ArrayPropertyValue, DefaultArrayPropertyValue);
+				}
+
+				if (!bIsIdentical)
 				{
 					// Advance to the next node in the list.
 					CurrentArrayNodePtr = &(*CurrentArrayNodePtr)->PropertyListNext;
