PR #13553: fix typo in BoxElem.h

diff --git a/Engine/Source/Runtime/Engine/Classes/PhysicsEngine/BoxElem.h b/Engine/Source/Runtime/Engine/Classes/PhysicsEngine/BoxElem.h
--- a/Engine/Source/Runtime/Engine/Classes/PhysicsEngine/BoxElem.h
+++ b/Engine/Source/Runtime/Engine/Classes/PhysicsEngine/BoxElem.h
@@ -31,7 +31,7 @@ struct FKBoxElem : public FKShapeElem
 	UPROPERTY(Category=Box, EditAnywhere, meta = (ClampMin = "-360", ClampMax = "360"))
 	FRotator Rotation;
 
-	/** Extent of the box along the y-axis */
+	/** Extent of the box along the x-axis */
 	UPROPERTY(Category= Box, EditAnywhere, meta =(DisplayName = "X Extent"))
 	float X;
 
