Fix shader compilation issue with certain materials.

Anisotropy input tracking was too conservative, causing a mismatch between the requested shader compilation and the actual need at runtime. When using the SetMaterialAttribute node, if an input is added to the list, it will be marked as connected. But if no connection is attached to this input, the default value (=0) will be used, i.e. no anisotropy.

Also adds Substrate as possible ShadingModel type in case of explicit conversion.

The mismatch was between:
* MaterialParameters.bHasAnisotropyConnected
* GetConnectedMaterialAttributesBitmask()

diff --git a/Engine/Source/Runtime/Engine/Private/Materials/MaterialExpressions.cpp b/Engine/Source/Runtime/Engine/Private/Materials/MaterialExpressions.cpp
--- a/Engine/Source/Runtime/Engine/Private/Materials/MaterialExpressions.cpp
+++ b/Engine/Source/Runtime/Engine/Private/Materials/MaterialExpressions.cpp
@@ -6532,10 +6532,19 @@ uint64 UMaterialExpressionSetMaterialAttributes::GetConnectedInputs() const
 {
 	uint64 Out = 0ull;
 	const int32 NumInputPins = AttributeSetTypes.Num();
-	for (int32 i = 0; i < NumInputPins; ++i)
+	for (int32 PinIndex = 0; PinIndex < NumInputPins; ++PinIndex)
 	{
-		const uint64 Bitmask = FMaterialAttributeDefinitionMap::GetBitmask(AttributeSetTypes[i]);
-		Out |= Bitmask;
+		const uint64 Bitmask = FMaterialAttributeDefinitionMap::GetBitmask(AttributeSetTypes[PinIndex]);
+		const EMaterialProperty Property = FMaterialAttributeDefinitionMap::GetProperty(AttributeSetTypes[PinIndex]);
+		const FExpressionInput& AttributeInput = Inputs[PinIndex + 1];
+
+		// If Anisotropy has been added to the SetMaterialAttributes list of attributes, but not connected, its input will be set to 0. 
+		// This means that the material won't have any anisotropy effect. For this reason we by pass input marking.
+		const bool bMarkInputAsConnected = Property != EMaterialProperty::MP_Anisotropy || (Property == EMaterialProperty::MP_Anisotropy && AttributeInput.GetTracedInput().Expression);
+		if (bMarkInputAsConnected)
+		{
+			Out |= Bitmask;
+		}
 	}
 	return Out;
 }
diff --git a/Engine/Source/Runtime/Renderer/Private/AnisotropyRendering.cpp b/Engine/Source/Runtime/Renderer/Private/AnisotropyRendering.cpp
--- a/Engine/Source/Runtime/Renderer/Private/AnisotropyRendering.cpp
+++ b/Engine/Source/Runtime/Renderer/Private/AnisotropyRendering.cpp
@@ -34,7 +34,7 @@ static bool IsAnisotropyPassCompatible(const EShaderPlatform Platform, FMaterial
 		FDataDrivenShaderPlatformInfo::GetSupportsAnisotropicMaterials(Platform) &&
 		MaterialParameters.bHasAnisotropyConnected &&
 		!IsTranslucentBlendMode(MaterialParameters) &&
-		MaterialParameters.ShadingModels.HasAnyShadingModel({ MSM_DefaultLit, MSM_ClearCoat }) &&
+		MaterialParameters.ShadingModels.HasAnyShadingModel({ MSM_DefaultLit, MSM_ClearCoat, MSM_Strata }) &&
 		(!bSubstrateEnabled || (bSubstrateEnabled && Substrate::IsSubstrateBlendableGBufferEnabled(Platform))); // Substrate renders anisotropy surface natively, without extra pass. Unless blendable gbuffer is used.
 }
 
