Fixing material shader map not releasing compilation id if shaders failed to compile

diff --git a/Engine/Source/Runtime/Engine/Private/ShaderCompiler/ShaderCompiler.cpp b/Engine/Source/Runtime/Engine/Private/ShaderCompiler/ShaderCompiler.cpp
--- a/Engine/Source/Runtime/Engine/Private/ShaderCompiler/ShaderCompiler.cpp
+++ b/Engine/Source/Runtime/Engine/Private/ShaderCompiler/ShaderCompiler.cpp
@@ -2203,6 +2203,10 @@ void FShaderCompilingManager::ProcessCompiledShaderMaps(
 					Material->GameThreadCompilingShaderMapId = 0u;
 					Material->GameThreadPendingCompilerEnvironment.SafeRelease();
 					MaterialDependencies.RemoveAt(MaterialIndex);
+					if (MaterialDependencies.Num() == 0 && CompilingShaderMap->CompilingMaterialNumExternalDependencies == 0)
+					{
+						CompilingShaderMap->ReleaseCompilingId();
+					}
 					MaterialsToReleaseCompilingId.Add(Material);
 				}
 				else
