Don't store shader output for preview/transient materials in DDC

diff --git a/Engine/Source/Runtime/Engine/Private/Materials/MaterialShader.cpp b/Engine/Source/Runtime/Engine/Private/Materials/MaterialShader.cpp
--- a/Engine/Source/Runtime/Engine/Private/Materials/MaterialShader.cpp
+++ b/Engine/Source/Runtime/Engine/Private/Materials/MaterialShader.cpp
@@ -1826,6 +1826,7 @@ static void PrepareMaterialShaderCompileJob(EShaderPlatform Platform,
 	const FShaderCompileJobKey& Key = NewJob->Key;
 	const FMaterialShaderType* ShaderType = Key.ShaderType->AsMaterialShaderType();
 
+	NewJob->bBypassCache = Material->IsPreview() || !Material->IsPersistent();
 	NewJob->Input.SharedEnvironment = MaterialEnvironment;
 	FShaderCompilerEnvironment& ShaderEnvironment = NewJob->Input.Environment;
 
diff --git a/Engine/Source/Runtime/Engine/Private/Materials/MeshMaterialShader.cpp b/Engine/Source/Runtime/Engine/Private/Materials/MeshMaterialShader.cpp
--- a/Engine/Source/Runtime/Engine/Private/Materials/MeshMaterialShader.cpp
+++ b/Engine/Source/Runtime/Engine/Private/Materials/MeshMaterialShader.cpp
@@ -30,6 +30,7 @@ static void PrepareMeshMaterialShaderCompileJob(EShaderPlatform Platform,
 	ShaderEnvironment.TargetPlatform = MaterialEnvironment->TargetPlatform;
 	NewJob->bIsDefaultMaterial = Material->IsDefaultMaterial();
 	NewJob->bIsGlobalShader = false;
+	NewJob->bBypassCache = Material->IsPreview() || !Material->IsPersistent();
 
 	if (IsShaderDevelopmentModeEnabled())
 	{
diff --git a/Engine/Source/Runtime/Engine/Private/ShaderCompiler/ShaderCompilerJobCache.cpp b/Engine/Source/Runtime/Engine/Private/ShaderCompiler/ShaderCompilerJobCache.cpp
--- a/Engine/Source/Runtime/Engine/Private/ShaderCompiler/ShaderCompilerJobCache.cpp
+++ b/Engine/Source/Runtime/Engine/Private/ShaderCompiler/ShaderCompilerJobCache.cpp
@@ -1647,7 +1647,7 @@ void FShaderJobCache::AddToCacheAndProcessPending(FShaderCommonCompileJob* Finis
 
 		if (FinishedJob->bSucceeded)
 		{
-			const bool bAddToDDC = GShaderCompilerPerShaderDDCGlobal || !(FinishedJob->bIsDefaultMaterial || FinishedJob->bIsGlobalShader);
+			const bool bAddToDDC = !FinishedJob->bBypassCache && (GShaderCompilerPerShaderDDCGlobal || !(FinishedJob->bIsDefaultMaterial || FinishedJob->bIsGlobalShader));
 			// we only cache jobs that succeded
 			AddJobOutput(JobData, FinishedJob, InputHash, SaveContext, FinishedDuplicateJobs.Num(), bAddToDDC);
 		}
