Fix Set/Get Material Attributes nodes not compiling in the new translator when the input pin is not attached.

diff --git a/Engine/Source/Runtime/Engine/Private/Materials/MaterialExpressionsIR.cpp b/Engine/Source/Runtime/Engine/Private/Materials/MaterialExpressionsIR.cpp
--- a/Engine/Source/Runtime/Engine/Private/Materials/MaterialExpressionsIR.cpp
+++ b/Engine/Source/Runtime/Engine/Private/Materials/MaterialExpressionsIR.cpp
@@ -2411,7 +2411,7 @@ void UMaterialExpressionGetMaterialAttributes::Build(MIR::FEmitter& Em)
 {
 	CheckMaterialAttributesExpression(Em, AttributeGetTypes);
 
-	FValueRef Prototype = Em.CheckIsAggregate(Em.Input(&MaterialAttributes), UMaterialAggregate::GetMaterialAttributes());
+	FValueRef Prototype = Em.CheckIsAggregate(Em.TryInput(&MaterialAttributes), UMaterialAggregate::GetMaterialAttributes());
 	
 	UE_MIR_CHECKPOINT(Em);
 	
@@ -2433,7 +2433,7 @@ void UMaterialExpressionSetMaterialAttributes::Build(MIR::FEmitter& Em)
 {
 	CheckMaterialAttributesExpression(Em, AttributeSetTypes);
 	
-	FValueRef Prototype = Em.CheckIsAggregate(Em.Input(&Inputs[0]), UMaterialAggregate::GetMaterialAttributes());
+	FValueRef Prototype = Em.CheckIsAggregate(Em.TryInput(&Inputs[0]), UMaterialAggregate::GetMaterialAttributes());
 	
 	UE_MIR_CHECKPOINT(Em);
 
