Fixed deprecation message referencing a non-existing function.

diff --git a/Engine/Source/Runtime/Engine/Classes/Components/SceneComponent.h b/Engine/Source/Runtime/Engine/Classes/Components/SceneComponent.h
--- a/Engine/Source/Runtime/Engine/Classes/Components/SceneComponent.h
+++ b/Engine/Source/Runtime/Engine/Classes/Components/SceneComponent.h
@@ -924,7 +924,7 @@ public:
 	}
 
 #if WITH_EDITOR
-	UE_DEPRECATED(5.7, "Please call GetIsTemporarilyHiddenInEditorFlag()")
+	UE_DEPRECATED(5.7, "IsTemporarilyHiddenInEditor is now deprecated.")
 	ENGINE_API bool IsTemporarilyHiddenInEditor(const bool bIncludeParent) const
 	{
 		return bHiddenEdTemporary;
