[Gameplay] Add an accessor to on the Local Player to get a local player subsytem from a provided player controller. This will remove the need to constantly get the owning controller when you just want access to the Local Player subsytem. Good example uses cases where you often have the player controller but want the LP subsystem would be enhanced input, common input, and common ui subsystems.

Also remove the previously deprecated GetSubsystemArray function, which was deprecated in 5.4.

#rnx

diff --git a/Engine/Source/Runtime/Engine/Classes/Engine/LocalPlayer.h b/Engine/Source/Runtime/Engine/Classes/Engine/LocalPlayer.h
--- a/Engine/Source/Runtime/Engine/Classes/Engine/LocalPlayer.h
+++ b/Engine/Source/Runtime/Engine/Classes/Engine/LocalPlayer.h
@@ -346,17 +346,20 @@ public:
 	}
 
 	/**
-	 * Get all Subsystem of specified type, this is only necessary for interfaces that can have multiple implementations instanced at a time.
-	 *
-	 * Do not hold onto this Array reference unless you are sure the lifetime is less than that of ULocalPlayer
+	 * Get a subsystem of the specificied type from the provided Player Controller's owning ULocalPlayer.
+	 * Returns nullptr if the Subsystem cannot be found or the given controller or its local player are null.
 	 */
 	template <typename TSubsystemClass>
-	UE_DEPRECATED(5.4, "This function is unsafe for re-entrancy and has been deprecated. Use GetSubsystemArrayCopy or ForEachSubsystem instead")
-	const TArray<TSubsystemClass*>& GetSubsystemArray() const
+	static FORCEINLINE TSubsystemClass* GetSubsystemFromController(const APlayerController* PC)
 	{
-		return SubsystemCollection.GetSubsystemArray<TSubsystemClass>(TSubsystemClass::StaticClass());
+		if (PC)
+		{
+			const ULocalPlayer* LocalPlayer = PC->GetLocalPlayer();
+			return LocalPlayer ? LocalPlayer->GetSubsystem<TSubsystemClass>() : nullptr;
+		}
+		return nullptr;
 	}
-
+	
 	/**
 	 * Get all Subsystem of specified type, this is only necessary for interfaces that can have multiple implementations instanced at a time.
 	 *
@@ -369,7 +372,7 @@ public:
 	}
 
 	/**
-	 * Performs an operation on all all Subsystem of specified type, this is only necessary for interfaces that can have multiple implementations instanced at a time.
+	 * Performs an operation on all Subsystem of specified type, this is only necessary for interfaces that can have multiple implementations instanced at a time.
 	 */
 	template <typename TSubsystemClass>
 	void ForEachSubsystem(TFunctionRef<void(TSubsystemClass*)> Operation) const
