Fixed imprecision in new material translator handling allocation of uniform material parameters.

[RN] none
#rnx

diff --git a/Engine/Source/Runtime/Engine/Private/Materials/MaterialIRValueAnalyzer.cpp b/Engine/Source/Runtime/Engine/Private/Materials/MaterialIRValueAnalyzer.cpp
--- a/Engine/Source/Runtime/Engine/Private/Materials/MaterialIRValueAnalyzer.cpp
+++ b/Engine/Source/Runtime/Engine/Private/Materials/MaterialIRValueAnalyzer.cpp
@@ -563,8 +563,8 @@ static void AnalyzePrimitiveUniformParameter(FMaterialIRValueAnalyzer& Analyzer,
 	if (UsedNumComponents >= 4)
 	{
 		check(UsedNumComponents == 4); 
-		const uint32 SizeInBytes = PrimitiveType.IsDouble() ? 32 : 16;
-		GlobalComponentOffset = UniformExpressionSet.AllocateFromUniformBuffer(SizeInBytes) / 4;
+		const uint32 SizeInBytes = PrimitiveType.IsDouble() ? 2 : 1;
+		GlobalComponentOffset = UniformExpressionSet.AllocateFromUniformBuffer(SizeInBytes) * 4;
 	}
 
 	// Finally, check whether in the used slot there's any slack left, and if so, record it in the appropriate free-list.
diff --git a/Engine/Source/Runtime/Engine/Private/Materials/MaterialUniformExpressions.cpp b/Engine/Source/Runtime/Engine/Private/Materials/MaterialUniformExpressions.cpp
--- a/Engine/Source/Runtime/Engine/Private/Materials/MaterialUniformExpressions.cpp
+++ b/Engine/Source/Runtime/Engine/Private/Materials/MaterialUniformExpressions.cpp
@@ -719,10 +719,10 @@ uint32 FUniformExpressionSet::GetNumericParameterEvaluationOffset(uint32 Paramet
 	return UniformParameterEvaluations[ParameterIndex].BufferOffset;
 }
 
-uint32 FUniformExpressionSet::AllocateFromUniformBuffer(uint32 SizeInBytes)
+uint32 FUniformExpressionSet::AllocateFromUniformBuffer(uint32 NumOfFloat4s)
 {
 	uint32 Temp = UniformPreshaderBufferSize;
-	UniformPreshaderBufferSize += SizeInBytes;
+	UniformPreshaderBufferSize += NumOfFloat4s;
 	return Temp;
 }
 
diff --git a/Engine/Source/Runtime/Engine/Public/MaterialShared.h b/Engine/Source/Runtime/Engine/Public/MaterialShared.h
--- a/Engine/Source/Runtime/Engine/Public/MaterialShared.h
+++ b/Engine/Source/Runtime/Engine/Public/MaterialShared.h
@@ -708,7 +708,10 @@ public:
 	/// Returns the offset of specified parameter given with AddNumericParameterEvaluation.
 	uint32 GetNumericParameterEvaluationOffset(uint32 ParameterIndex) const;
 	
-	uint32 AllocateFromUniformBuffer(uint32 SizeInBytes);
+	/// Allocates the specified number of float4s in the uniform buffer and returns the index
+	/// of the first newly allocated float4 in the buffer.
+	/// Note: multiply returned value by sizeof(float4) to know the byte offset in the buffer.
+	uint32 AllocateFromUniformBuffer(uint32 NumOfFloat4s);
 	
 	TConstArrayView<FMaterialVirtualTextureStack> GetVTStacks() const { return VTStacks; }
 	const FMaterialVirtualTextureStack& GetVTStack(int32 Index) const { return VTStacks[Index]; }
