Fixes primitive type validation crash in UMaterialExpressionOneMinus::Build.

#rnx

diff --git a/Engine/Source/Runtime/Engine/Private/Materials/MaterialExpressionsIR.cpp b/Engine/Source/Runtime/Engine/Private/Materials/MaterialExpressionsIR.cpp
--- a/Engine/Source/Runtime/Engine/Private/Materials/MaterialExpressionsIR.cpp
+++ b/Engine/Source/Runtime/Engine/Private/Materials/MaterialExpressionsIR.cpp
@@ -1041,6 +1041,8 @@ void UMaterialExpressionOneMinus::Build(MIR::FEmitter& Em)
 	// Default input to zero if not connected, then get it as a primitive.
 	FValueRef Value = Em.InputDefaultFloat(&Input, 0.0f);
 
+	UE_MIR_CHECKPOINT(Em); // verify the value is valid
+
 	// Make a "One" value of the same type and dimension as input's.
 	FValueRef One = Em.ConstantOne(Value->Type.AsPrimitive()->ScalarKind);
 
