Fixed incorrect reinterpret_cast destination type in GetTransformFromArray.

Caused by CL 44140129.

diff --git a/Engine/Source/Runtime/Engine/Private/Animation/AnimationRuntime.cpp b/Engine/Source/Runtime/Engine/Private/Animation/AnimationRuntime.cpp
--- a/Engine/Source/Runtime/Engine/Private/Animation/AnimationRuntime.cpp
+++ b/Engine/Source/Runtime/Engine/Private/Animation/AnimationRuntime.cpp
@@ -1833,7 +1833,7 @@ struct FBlendPosesPerBoneFilterScratchArea : public TThreadSingleton<FBlendPoses
 // Helper function to get FTransform from a PoseIndex and BoneIndex
 extern "C" const uint8* GetTransformFromArray(const uint8 *BlendPoseBase, const int32 PoseIndex, const int32 BoneIndex)
 {
-	const TArray<struct FCompactPose>& BlendPoses = *reinterpret_cast<const TArray<struct FCompactPose>*>(BlendPoseBase);
+	const TConstArrayView<struct FCompactPose>& BlendPoses = *reinterpret_cast<const TConstArrayView<struct FCompactPose>*>(BlendPoseBase);
 	const FTransform* BlendPose = &BlendPoses[PoseIndex][FCompactPoseBoneIndex(BoneIndex)];
 	return reinterpret_cast<const uint8*>(BlendPose);
 }
