Fix CIS error

diff --git a/Engine/Source/Runtime/Engine/Public/WorldPartition/HLOD/HLODLayer.h b/Engine/Source/Runtime/Engine/Public/WorldPartition/HLOD/HLODLayer.h
--- a/Engine/Source/Runtime/Engine/Public/WorldPartition/HLOD/HLODLayer.h
+++ b/Engine/Source/Runtime/Engine/Public/WorldPartition/HLOD/HLODLayer.h
@@ -87,10 +87,12 @@ private:
 #endif
 
 private:
+#if WITH_EDITORONLY_DATA
 	UPROPERTY(VisibleAnywhere, Instanced, Category = General, meta = (AllowEditInlineCustomization, EditInline, ShowInnerProperties))
 	TObjectPtr<UHLODLayerEditorOnlyData> EditorOnlyData;
+#endif
 
-	// Deprecated - Moved to optinal editor data through UHLODLayerEditorOnlyData
+	// Deprecated - Moved to optional editor data through UHLODLayerEditorOnlyData
 #if WITH_EDITORONLY_DATA
 	UPROPERTY()
 	EHLODLayerType LayerType_DEPRECATED;
