Reuse value returned from GetMaterialShaderMapKey

#rnx

diff --git a/Engine/Source/Runtime/Engine/Private/Materials/MaterialShader.cpp b/Engine/Source/Runtime/Engine/Private/Materials/MaterialShader.cpp
--- a/Engine/Source/Runtime/Engine/Private/Materials/MaterialShader.cpp
+++ b/Engine/Source/Runtime/Engine/Private/Materials/MaterialShader.cpp
@@ -2140,7 +2140,7 @@ TSharedRef<FMaterialShaderMap::FAsyncLoadContext> FMaterialShaderMap::BeginLoadF
 			{
 				FCacheGetRequest Request;
 				Request.Name = GetMaterialShaderMapName(Material->GetFullPath(), ShaderMapId, InPlatform);
-				Request.Key = GetMaterialShaderMapKey(Result->DataKey);
+				Request.Key = CacheKey;
 				Result->Material = Material;
 				Result->Platform = InPlatform;
 
