Increase reserved memory for ShaderMaterialMapKey

#rnx

diff --git a/Engine/Source/Runtime/Engine/Private/Materials/MaterialShader.cpp b/Engine/Source/Runtime/Engine/Private/Materials/MaterialShader.cpp
--- a/Engine/Source/Runtime/Engine/Private/Materials/MaterialShader.cpp
+++ b/Engine/Source/Runtime/Engine/Private/Materials/MaterialShader.cpp
@@ -213,7 +213,7 @@ FString GetMaterialShaderMapKeyString(const FMaterialShaderMapId& ShaderMapId,
 	bool bIncludeKeyStringShaderDependencies)
 {
 	FString ShaderMapKeyString;
-	ShaderMapKeyString.Reserve(16384);
+	ShaderMapKeyString.Reserve(1024 * 1024 * 2);
 	FMaterialKeyGeneratorContext Context(ShaderMapKeyString, Platform);
 	Context.SetFlags(EMaterialKeyInclude::ShaderDependencies, bIncludeKeyStringShaderDependencies);
 	RecordOrEmitMaterialShaderMapKey(Context, ShaderMapId, ShaderParameters);
