A blueprint actor's bindings to another actor's multicast delegates are now cleaned up when the actor is destroyed and reinstanced. Fixes a bug that bindings to other actors' multicast delegates (including BP event dispatchers) accumulate when recompiling this actor blueprint. Previously this resulted in event nodes firing multiple times, increasingly when the blueprint is recompiled.

diff --git a/Engine/Source/Runtime/Engine/Private/ActorConstruction.cpp b/Engine/Source/Runtime/Engine/Private/ActorConstruction.cpp
--- a/Engine/Source/Runtime/Engine/Private/ActorConstruction.cpp
+++ b/Engine/Source/Runtime/Engine/Private/ActorConstruction.cpp
@@ -573,6 +573,14 @@ void AActor::RerunConstructionScripts()
 		// Reset random streams
 		ResetPropertiesForConstruction();
 
+		// Cleanup any bindings from this actor to other objects, like bindings to referenced level actors.
+		// This is intended to only cleanup bindings by engine code that correspond to user nodes, like
+		// Event nodes in blueprint graphs for another actor's events dispatchers. If those nodes still exist
+		// the bindings will be rebound on ExecuteConstruction.
+#if WITH_EDITOR
+		UBlueprintGeneratedClass::UnbindDynamicDelegates(GetClass(), this);
+#endif
+
 		// Exchange net roles before running construction scripts
 		UWorld *OwningWorld = GetWorld();
 		if (OwningWorld && OwningWorld->IsNetMode(NM_Client))
