NetTest
* Added AsyncLoading stall test (Iris only)

#rnx

diff --git a/Engine/Source/Runtime/Engine/Public/Tests/NetTestHelpers.h b/Engine/Source/Runtime/Engine/Public/Tests/NetTestHelpers.h
--- a/Engine/Source/Runtime/Engine/Public/Tests/NetTestHelpers.h
+++ b/Engine/Source/Runtime/Engine/Public/Tests/NetTestHelpers.h
@@ -148,6 +148,15 @@ public:
 		UReplicationSystem* IrisRepSystem = nullptr;
 		UObjectReplicationBridge* IrisRepBridge = nullptr;
 #endif
+
+		bool IsValid() const
+		{
+			return World && NetDriver
+#if UE_WITH_IRIS
+				&& IrisRepSystem && IrisRepBridge
+#endif
+				;
+		}
 	};
 
 public:
