[UE] Instanciate RBAN cape collider bodies in ref pose
- Added new InstantiatePhysicsAssetBodiesRefPose function to USkeletalMeshComponent
- Call new function in FAnimNode_RigidBody::CollectClothColliderObjects to create physics bodies in ref pose instead of world space, mirroring the way the node instanciates its main physics asset.

diff --git a/Engine/Source/Runtime/AnimGraphRuntime/Private/BoneControllers/AnimNode_RigidBody.cpp b/Engine/Source/Runtime/AnimGraphRuntime/Private/BoneControllers/AnimNode_RigidBody.cpp
--- a/Engine/Source/Runtime/AnimGraphRuntime/Private/BoneControllers/AnimNode_RigidBody.cpp
+++ b/Engine/Source/Runtime/AnimGraphRuntime/Private/BoneControllers/AnimNode_RigidBody.cpp
@@ -1627,7 +1627,7 @@ void FAnimNode_RigidBody::CollectClothColliderObjects(const USkeletalMeshCompone
 			if (SourceComponent && PhysicsAsset)
 			{
 				TArray<FBodyInstance*> BodyInstances;
-				SourceComponent->InstantiatePhysicsAssetBodies(*PhysicsAsset, BodyInstances);
+				SourceComponent->InstantiatePhysicsAssetBodiesRefPose(*PhysicsAsset, BodyInstances);
 
 				for (uint32 BodyInstanceIndex = 0, BodyInstanceMax = BodyInstances.Num(); BodyInstanceIndex < BodyInstanceMax; ++BodyInstanceIndex)
 				{
diff --git a/Engine/Source/Runtime/Engine/Classes/Components/SkeletalMeshComponent.h b/Engine/Source/Runtime/Engine/Classes/Components/SkeletalMeshComponent.h
--- a/Engine/Source/Runtime/Engine/Classes/Components/SkeletalMeshComponent.h
+++ b/Engine/Source/Runtime/Engine/Classes/Components/SkeletalMeshComponent.h
@@ -2151,6 +2151,9 @@ public:
 	/** Instantiates only the bodies given a physics asset, not the constraints. The Created bodies are owned by the calling code and must be freed when necessary.*/
 	ENGINE_API void InstantiatePhysicsAssetBodies(const UPhysicsAsset& PhysAsset, TArray<FBodyInstance*>& OutBodies, FPhysScene* PhysScene = nullptr, USkeletalMeshComponent* OwningComponent = nullptr, int32 UseRootBodyIndex = INDEX_NONE, const FPhysicsAggregateHandle& UseAggregate = FPhysicsAggregateHandle()) const;
 
+	/** Instantiates only the bodies given a physics asset, not the constraints. Instead of reading the current component state, this reads the ref-pose from the reference skeleton of the mesh. Useful if trying to create bodies to be used during any evaluation work. The Created bodies are owned by the calling code and must be freed when necessary.*/
+	ENGINE_API void InstantiatePhysicsAssetBodiesRefPose(const UPhysicsAsset& PhysAsset, TArray<FBodyInstance*>& OutBodies, FPhysScene* PhysScene = nullptr, USkeletalMeshComponent* OwningComponent = nullptr, int32 UseRootBodyIndex = INDEX_NONE, const FPhysicsAggregateHandle& UseAggregate = FPhysicsAggregateHandle()) const;
+
 	/** Instantiates bodies given a physics asset like InstantiatePhysicsAsset but instead of reading the current component state, this reads the ref-pose from the reference skeleton of the mesh. Useful if trying to create bodies to be used during any evaluation work */
 	ENGINE_API void InstantiatePhysicsAssetRefPose(const UPhysicsAsset& PhysAsset, const FVector& Scale3D, TArray<FBodyInstance*>& OutBodies, TArray<FConstraintInstance*>& OutConstraints, FPhysScene* PhysScene = nullptr, USkeletalMeshComponent* OwningComponent = nullptr, int32 UseRootBodyIndex = INDEX_NONE, const FPhysicsAggregateHandle& UseAggregate = FPhysicsAggregateHandle(), bool bCreateBodiesInRefPose = false) const;
 
diff --git a/Engine/Source/Runtime/Engine/Private/SkeletalMeshComponentPhysics.cpp b/Engine/Source/Runtime/Engine/Private/SkeletalMeshComponentPhysics.cpp
--- a/Engine/Source/Runtime/Engine/Private/SkeletalMeshComponentPhysics.cpp
+++ b/Engine/Source/Runtime/Engine/Private/SkeletalMeshComponentPhysics.cpp
@@ -781,6 +781,24 @@ void USkeletalMeshComponent::InstantiatePhysicsAssetBodies(const UPhysicsAsset&
 	InstantiatePhysicsAssetBodies_Internal(PhysAsset, OutBodies, BoneTMCallable, nullptr, PhysScene, OwningComponent, UseRootBodyIndex, UseAggregate);
 }
 
+void USkeletalMeshComponent::InstantiatePhysicsAssetBodiesRefPose(const UPhysicsAsset& PhysAsset, TArray<FBodyInstance*>& OutBodies, FPhysScene* PhysScene /*= nullptr*/, USkeletalMeshComponent* OwningComponent /*= nullptr*/, int32 UseRootBodyIndex /*= INDEX_NONE*/, const FPhysicsAggregateHandle& UseAggregate /*= FPhysicsAggregateHandle()*/) const
+{
+	const FReferenceSkeleton& RefSkeleton = GetSkeletalMeshAsset()->GetRefSkeleton();
+
+	// Create the bodies posed in component-space with the ref pose
+	TArray<FTransform> CachedTransforms;
+	TArray<bool> CachedTransformReady;
+	CachedTransforms.SetNumUninitialized(RefSkeleton.GetRefBonePose().Num());
+	CachedTransformReady.SetNumZeroed(RefSkeleton.GetRefBonePose().Num());
+
+	auto BoneTMCallable = [&](int32 BoneIndex)
+	{
+		return FAnimationRuntime::GetComponentSpaceTransformWithCache(RefSkeleton, RefSkeleton.GetRefBonePose(), BoneIndex, CachedTransforms, CachedTransformReady);
+	};
+
+	InstantiatePhysicsAssetBodies_Internal(PhysAsset, OutBodies, BoneTMCallable, nullptr, PhysScene, OwningComponent, UseRootBodyIndex, UseAggregate);
+}
+
 void USkeletalMeshComponent::InstantiatePhysicsAssetRefPose(const UPhysicsAsset& PhysAsset, const FVector& Scale3D, TArray<FBodyInstance*>& OutBodies, TArray<FConstraintInstance*>& OutConstraints, FPhysScene* PhysScene /*= nullptr*/, USkeletalMeshComponent* OwningComponent /*= nullptr*/, int32 UseRootBodyIndex /*= INDEX_NONE*/, const FPhysicsAggregateHandle& UseAggregate, bool bCreateBodiesInRefPose) const
 {
 	if(GetSkeletalMeshAsset())
