[UE] Add nullptr check to USkeletalMeshComponent::InstantiatePhysicsAssetBodiesRefPose

#rnx

diff --git a/Engine/Source/Runtime/Engine/Private/SkeletalMeshComponentPhysics.cpp b/Engine/Source/Runtime/Engine/Private/SkeletalMeshComponentPhysics.cpp
--- a/Engine/Source/Runtime/Engine/Private/SkeletalMeshComponentPhysics.cpp
+++ b/Engine/Source/Runtime/Engine/Private/SkeletalMeshComponentPhysics.cpp
@@ -783,20 +783,23 @@ void USkeletalMeshComponent::InstantiatePhysicsAssetBodies(const UPhysicsAsset&
 
 void USkeletalMeshComponent::InstantiatePhysicsAssetBodiesRefPose(const UPhysicsAsset& PhysAsset, TArray<FBodyInstance*>& OutBodies, FPhysScene* PhysScene /*= nullptr*/, USkeletalMeshComponent* OwningComponent /*= nullptr*/, int32 UseRootBodyIndex /*= INDEX_NONE*/, const FPhysicsAggregateHandle& UseAggregate /*= FPhysicsAggregateHandle()*/) const
 {
-	const FReferenceSkeleton& RefSkeleton = GetSkeletalMeshAsset()->GetRefSkeleton();
+	if (GetSkeletalMeshAsset())
+	{
+		const FReferenceSkeleton& RefSkeleton = GetSkeletalMeshAsset()->GetRefSkeleton();
 
-	// Create the bodies posed in component-space with the ref pose
-	TArray<FTransform> CachedTransforms;
-	TArray<bool> CachedTransformReady;
-	CachedTransforms.SetNumUninitialized(RefSkeleton.GetRefBonePose().Num());
-	CachedTransformReady.SetNumZeroed(RefSkeleton.GetRefBonePose().Num());
+		// Create the bodies posed in component-space with the ref pose
+		TArray<FTransform> CachedTransforms;
+		TArray<bool> CachedTransformReady;
+		CachedTransforms.SetNumUninitialized(RefSkeleton.GetRefBonePose().Num());
+		CachedTransformReady.SetNumZeroed(RefSkeleton.GetRefBonePose().Num());
 
-	auto BoneTMCallable = [&](int32 BoneIndex)
-	{
-		return FAnimationRuntime::GetComponentSpaceTransformWithCache(RefSkeleton, RefSkeleton.GetRefBonePose(), BoneIndex, CachedTransforms, CachedTransformReady);
-	};
+		auto BoneTMCallable = [&](int32 BoneIndex)
+			{
+				return FAnimationRuntime::GetComponentSpaceTransformWithCache(RefSkeleton, RefSkeleton.GetRefBonePose(), BoneIndex, CachedTransforms, CachedTransformReady);
+			};
 
-	InstantiatePhysicsAssetBodies_Internal(PhysAsset, OutBodies, BoneTMCallable, nullptr, PhysScene, OwningComponent, UseRootBodyIndex, UseAggregate);
+		InstantiatePhysicsAssetBodies_Internal(PhysAsset, OutBodies, BoneTMCallable, nullptr, PhysScene, OwningComponent, UseRootBodyIndex, UseAggregate);
+	}
 }
 
 void USkeletalMeshComponent::InstantiatePhysicsAssetRefPose(const UPhysicsAsset& PhysAsset, const FVector& Scale3D, TArray<FBodyInstance*>& OutBodies, TArray<FConstraintInstance*>& OutConstraints, FPhysScene* PhysScene /*= nullptr*/, USkeletalMeshComponent* OwningComponent /*= nullptr*/, int32 UseRootBodyIndex /*= INDEX_NONE*/, const FPhysicsAggregateHandle& UseAggregate, bool bCreateBodiesInRefPose) const
