Substrate - fixed preview tile on expressions nodes.

#rnx
#tests PC editor

diff --git a/Engine/Source/Runtime/Engine/Private/Materials/HLSLMaterialTranslator.cpp b/Engine/Source/Runtime/Engine/Private/Materials/HLSLMaterialTranslator.cpp
--- a/Engine/Source/Runtime/Engine/Private/Materials/HLSLMaterialTranslator.cpp
+++ b/Engine/Source/Runtime/Engine/Private/Materials/HLSLMaterialTranslator.cpp
@@ -1748,7 +1748,8 @@ void FHLSLMaterialTranslator::TranslateMaterial()
 		NormalCodeChunkEnd = SharedPropertyCodeChunks[NormalShaderFrequency].Num();
 	}
 
-	if (!bSubstrateEnabled || !bUseRootNodeToSubstrateHiddenConversion)
+	if (!bSubstrateEnabled || !bUseRootNodeToSubstrateHiddenConversion 
+		|| GetSubstrateMaterialExportType() != ESubstrateMaterialExport::SME_None)	// Needed for material preview on node or any other export.
 	{
 		// Rest of properties
 		Chunk[MP_EmissiveColor]	= Material->CompilePropertyAndSetMaterialProperty(MP_EmissiveColor,	this);
