[PR #9426] Add error message to PaintVerticesLerpAlongAxis if the Static Mesh does not have bAllowCPUAccess set to true when built, a warning message if in editor

- The current implementation will no longer result in a crash like the original PR says, but an error message is needed since the function will not execute as expected in a build

diff --git a/Engine/Source/Runtime/Engine/Private/MeshVertexPainter/MeshVertexPainter.cpp b/Engine/Source/Runtime/Engine/Private/MeshVertexPainter/MeshVertexPainter.cpp
--- a/Engine/Source/Runtime/Engine/Private/MeshVertexPainter/MeshVertexPainter.cpp
+++ b/Engine/Source/Runtime/Engine/Private/MeshVertexPainter/MeshVertexPainter.cpp
@@ -49,18 +49,33 @@ void FMeshVertexPainter::PaintVerticesSingleColor(UStaticMeshComponent* StaticMe
 
 void FMeshVertexPainter::PaintVerticesLerpAlongAxis(UStaticMeshComponent* StaticMeshComponent, const FLinearColor& StartColor, const FLinearColor& EndColor, EVertexPaintAxis Axis, bool bConvertToSRGB)
 {
-	if (!StaticMeshComponent || !StaticMeshComponent->GetStaticMesh())
+	TObjectPtr<UStaticMesh> StaticMesh;
+	if (!StaticMeshComponent || !(StaticMesh = StaticMeshComponent->GetStaticMesh()))
 	{
 		return;
 	}
+	if (!StaticMesh->bAllowCPUAccess)
+	{
+#if WITH_EDITOR
+		UE_LOG(LogStaticMesh,
+			Warning,
+			TEXT("bAllowCPUAccess should be set to true for StaticMesh '%s' before calling PaintVerticesLerpAlongAxis(), otherwise the mesh geometry data isn't accessible when packed."),
+			*StaticMesh->GetName());
+#else
+		UE_LOG(LogStaticMesh,
+			Error,
+			TEXT("bAllowCPUAccess must be set to true for StaticMesh '%s' before calling PaintVerticesLerpAlongAxis(), otherwise the mesh geometry data isn't accessible when packed."),
+			*StaticMesh->GetName());
+#endif
+	}
 
-	const FBoxSphereBounds Bounds = StaticMeshComponent->GetStaticMesh()->GetBounds();
+	const FBoxSphereBounds Bounds = StaticMesh->GetBounds();
 	const FBox Box = Bounds.GetBox();
 	static_assert(static_cast<int32>(EVertexPaintAxis::X) == 0, "EVertexPaintAxis not correctly defined");
 	const float AxisMin = Box.Min.Component(static_cast<int32>(Axis));
 	const float AxisMax = Box.Max.Component(static_cast<int32>(Axis));
 
-	const int32 NumMeshLODs = StaticMeshComponent->GetStaticMesh()->GetNumLODs();
+	const int32 NumMeshLODs = StaticMesh->GetNumLODs();
 	StaticMeshComponent->SetLODDataCount(NumMeshLODs, NumMeshLODs);
 
 	uint32 LODIndex = 0;
@@ -69,7 +84,7 @@ void FMeshVertexPainter::PaintVerticesLerpAlongAxis(UStaticMeshComponent* Static
 		StaticMeshComponent->RemoveInstanceVertexColorsFromLOD(LODIndex);
 		check(LODInfo.OverrideVertexColors == nullptr);
 
-		const FStaticMeshLODResources& LODModel = StaticMeshComponent->GetStaticMesh()->GetRenderData()->LODResources[LODIndex];
+		const FStaticMeshLODResources& LODModel = StaticMesh->GetRenderData()->LODResources[LODIndex];
 		const FPositionVertexBuffer& PositionVertexBuffer = LODModel.VertexBuffers.PositionVertexBuffer;
 		const uint32 NumVertices = PositionVertexBuffer.GetNumVertices();
 
