[PR #9426] Add missed return after error message to PaintVerticesLerpAlongAxis

diff --git a/Engine/Source/Runtime/Engine/Private/MeshVertexPainter/MeshVertexPainter.cpp b/Engine/Source/Runtime/Engine/Private/MeshVertexPainter/MeshVertexPainter.cpp
--- a/Engine/Source/Runtime/Engine/Private/MeshVertexPainter/MeshVertexPainter.cpp
+++ b/Engine/Source/Runtime/Engine/Private/MeshVertexPainter/MeshVertexPainter.cpp
@@ -66,6 +66,7 @@ void FMeshVertexPainter::PaintVerticesLerpAlongAxis(UStaticMeshComponent* Static
 			Error,
 			TEXT("bAllowCPUAccess must be set to true for StaticMesh '%s' before calling PaintVerticesLerpAlongAxis(), otherwise the mesh geometry data isn't accessible when packed."),
 			*StaticMesh->GetName());
+		return;
 #endif
 	}
 
