Fixed a bug in APlayerController::IsLocalController() where it could incorrectly return true for remote clients during teardown.

Previously, when the NetDriver was destroyed (e.g. after a listen server transitioned to a standalone map), GetNetMode() would fall back to NM_Standalone. This caused IsLocalController() to incorrectly return true for remote controllers that happened to have a ULocalPlayer, leading to incorrect behavior during controller destruction.

The fix avoids mutating bIsLocalPlayerController during teardown and adds fallback logic that conservatively checks for both a ULocalPlayer and ROLE_Authority before returning true when the NetDriver is null.

diff --git a/Engine/Source/Runtime/Engine/Private/PlayerController.cpp b/Engine/Source/Runtime/Engine/Private/PlayerController.cpp
--- a/Engine/Source/Runtime/Engine/Private/PlayerController.cpp
+++ b/Engine/Source/Runtime/Engine/Private/PlayerController.cpp
@@ -301,7 +301,24 @@ bool APlayerController::IsLocalController() const
 		checkSlow(!bIsLocalPlayerController);
 		return false;
 	}
-	
+	// Fallback logic: during teardown the NetDriver may have already been destroyed.
+	// At this point, GetNetMode() will likely return NM_Standalone, so we can no longer trust it alone.
+	// Instead, we do a conservative check:
+	// - If we have a ULocalPlayer, we might be local.
+	// - But only treat it as local if we're also the authority (i.e., on the server).
+	if (GetNetDriver() == nullptr)
+	{
+		const bool bHasLocalPlayer = Cast<ULocalPlayer>(Player) != nullptr;
+		const bool bIsServer = GetLocalRole() == ROLE_Authority;
+
+		UE_LOG(LogTemp, Warning, TEXT("IsLocalController fallback: %s | Player: %s | Role: %d → %s"),
+			*GetName(),
+			Player ? *Player->GetName() : TEXT("NULL"),
+			(int32)GetLocalRole(),
+			(bHasLocalPlayer && bIsServer) ? TEXT("TRUE") : TEXT("FALSE"));
+
+		return bHasLocalPlayer && bIsServer;
+	}
 	// Fast path if we have this bool set.
 	if (bIsLocalPlayerController)
 	{
@@ -327,6 +344,8 @@ bool APlayerController::IsLocalController() const
 }
 
 
+
+
 void APlayerController::FailedToSpawnPawn()
 {
 	Super::FailedToSpawnPawn();
@@ -1579,6 +1598,7 @@ void APlayerController::BeginPlay()
 void APlayerController::EndPlay(const EEndPlayReason::Type EndPlayReason)
 {
 	ULocalPlayer* LocalPlayer = Cast<ULocalPlayer>(Player);
+
 	if (LocalPlayer)
 	{
 		if (VirtualJoystick.IsValid())
@@ -6680,4 +6700,3 @@ void ANoPawnPlayerController::ServerSetViewTargetPosition_Implementation(FVector
 }
 
 #undef LOCTEXT_NAMESPACE
-
