Substrate - added a message when disabling Substrate to make sure users are aware of the potential issues in case it is done inadvertidely.

#rnx
#tests pc editor

diff --git a/Engine/Source/Runtime/Engine/Private/RendererSettings.cpp b/Engine/Source/Runtime/Engine/Private/RendererSettings.cpp
--- a/Engine/Source/Runtime/Engine/Private/RendererSettings.cpp
+++ b/Engine/Source/Runtime/Engine/Private/RendererSettings.cpp
@@ -369,6 +369,10 @@ void URendererSettings::PostEditChangeProperty(FPropertyChangedEvent& PropertyCh
 			{
 				FMessageDialog::Open(EAppMsgType::Ok, LOCTEXT("Substrate Beta", "Warning: Substrate is in a Beta state. Be aware that any materials saved when Substrate is enabled won't be rendered correctly if Substrate is disabled later on."));
 			}
+			else
+			{
+				FMessageDialog::Open(EAppMsgType::Ok, LOCTEXT("Substrate", "Warning: Disabling Substrate will cause all materials authored with Substrate nodes to appear incorrect, as they are not backward compatible with the legacy root node and GBuffer format."));
+			}
 		}
 
 		ExportValuesToConsoleVariables(PropertyChangedEvent.Property);
