Fixes for Material Aggregate update issues:
Re-enable custom UVs in the aggregate structure, and move them back to the pixel shader for the time being until a full update for them to the vertex shader mapping can be implemented.
Fix the shading model selection being stomped, currently ignoring the owning material.

diff --git a/Engine/Source/Runtime/Core/Private/UObject/FortniteMainBranchObjectVersion.cpp b/Engine/Source/Runtime/Core/Private/UObject/FortniteMainBranchObjectVersion.cpp
--- a/Engine/Source/Runtime/Core/Private/UObject/FortniteMainBranchObjectVersion.cpp
+++ b/Engine/Source/Runtime/Core/Private/UObject/FortniteMainBranchObjectVersion.cpp
@@ -14,6 +14,6 @@ TMap<FGuid, FGuid> FFortniteMainBranchObjectVersion::GetSystemGuids()
 	SystemGuids.Add(DevGuids.Niagara_LatestScriptCompileVersion, FGuid("97F0FD7149D90D42B65D6557AF477D1A"));
 	SystemGuids.Add(DevGuids.SkeletalMeshDerivedDataVersion, FGuid("C998F2F3B2884A02B7E290B8C25835F0"));
 	SystemGuids.Add(DevGuids.STATICMESH_DERIVEDDATA_VER, FGuid("3ABF32149F9D4D838750368AB95573BE"));
-	SystemGuids.Add(DevGuids.MaterialTranslationDDCVersion, FGuid("7F010615E10C4835B08B2395E7274119"));
+	SystemGuids.Add(DevGuids.MaterialTranslationDDCVersion, FGuid("AECF2BF62A8A49B8929C91FE8AA32F44"));
 	return SystemGuids;
 }
diff --git a/Engine/Source/Runtime/Engine/Private/Materials/MaterialAggregate.cpp b/Engine/Source/Runtime/Engine/Private/Materials/MaterialAggregate.cpp
--- a/Engine/Source/Runtime/Engine/Private/Materials/MaterialAggregate.cpp
+++ b/Engine/Source/Runtime/Engine/Private/Materials/MaterialAggregate.cpp
@@ -189,6 +189,14 @@ struct FMaterialAttributePropertyIndexMap
 		PushAttribute(MP_CustomData1);
 		PushAttribute(MP_AmbientOcclusion);
 		PushAttribute(MP_Refraction);
+		PushAttribute(MP_CustomizedUVs0);
+		PushAttribute(MP_CustomizedUVs1);
+		PushAttribute(MP_CustomizedUVs2);
+		PushAttribute(MP_CustomizedUVs3);
+		PushAttribute(MP_CustomizedUVs4);
+		PushAttribute(MP_CustomizedUVs5);
+		PushAttribute(MP_CustomizedUVs6);
+		PushAttribute(MP_CustomizedUVs7);
 		PushAttribute(MP_PixelDepthOffset);
 		PushAttribute(MP_ShadingModel);
 		PushAttribute(MP_Displacement);
diff --git a/Engine/Source/Runtime/Engine/Private/Materials/MaterialIR.cpp b/Engine/Source/Runtime/Engine/Private/Materials/MaterialIR.cpp
--- a/Engine/Source/Runtime/Engine/Private/Materials/MaterialIR.cpp
+++ b/Engine/Source/Runtime/Engine/Private/Materials/MaterialIR.cpp
@@ -26,7 +26,7 @@ const TCHAR* LexToString(EStage Stage)
 
 bool MaterialOutputEvaluatesInStage(EMaterialProperty Property, EStage Stage)
 {	
-	if (FMaterialAttributeDefinitionMap::GetShaderFrequency(Property) == EShaderFrequency::SF_Vertex)
+	if (Property == MP_WorldPositionOffset)
 	{
 		return Stage == Stage_Vertex;
 	}
diff --git a/Engine/Source/Runtime/Engine/Private/Materials/MaterialIREmitter.cpp b/Engine/Source/Runtime/Engine/Private/Materials/MaterialIREmitter.cpp
--- a/Engine/Source/Runtime/Engine/Private/Materials/MaterialIREmitter.cpp
+++ b/Engine/Source/Runtime/Engine/Private/Materials/MaterialIREmitter.cpp
@@ -1,6 +1,7 @@
 // Copyright Epic Games, Inc. All Rights Reserved.
 
 #include "Materials/MaterialIREmitter.h"
+#include "Materials/Material.h"
 #include "Materials/MaterialIRModule.h"
 #include "Materials/MaterialIRModuleBuilder.h"
 #include "Materials/MaterialIRInternal.h"
@@ -1266,6 +1267,21 @@ static FValueRef EmitAttributeDefaultValue(FEmitter& Emitter, const FMaterialAgg
 	}
 }
 
+static FValueRef MapMaterialAggregateDefaultValues(FEmitter& Emitter, UMaterial* Material, const UMaterialAggregate* InAggregate, const int32& Index)
+{
+	ensure(Index >= 0 && Index < InAggregate->Attributes.Num());
+	if (Material && InAggregate == UMaterialAggregate::GetMaterialAttributes())
+	{
+		switch (UMaterialAggregate::AttributeIndexToMaterialProperty(Index))
+		{
+		case MP_ShadingModel: return Emitter.ConstantInt(Material->GetShadingModels().GetFirstShadingModel());
+		default: break;
+		}
+	}
+
+	return EmitAttributeDefaultValue(Emitter, InAggregate->Attributes[Index]);
+}
+
 FValueRef FEmitter::Aggregate(const UMaterialAggregate* InAggregate, FValueRef InPrototype, TConstArrayView<FValueRef> AttributeValues)
 {
 	// Check that the specified prototype (if any) aggregate matches the one provided.
@@ -1300,7 +1316,7 @@ FValueRef FEmitter::Aggregate(const UMaterialAggregate* InAggregate, FValueRef I
 		else
 		{
 			// ...otherwise construct the default value as indicated in the attribute.
-			Components[i] = EmitAttributeDefaultValue(*this, InAggregate->Attributes[i]);
+			Components[i] = MapMaterialAggregateDefaultValues(*this, Material, InAggregate, i);
 		}
 	}
 
@@ -1348,7 +1364,7 @@ FValueRef FEmitter::Aggregate(const UMaterialAggregate* InAggregate, FValueRef I
 		else
 		{
 			// ...otherwise construct the default value.
-			Components[i] = EmitAttributeDefaultValue(*this, InAggregate->Attributes[i]);
+			Components[i] = MapMaterialAggregateDefaultValues(*this, Material, InAggregate, i);
 		}
 	}
 
