Improve iteration time for the network metrics database by removing the NetworkMetricsDatabase.h include from NetDriver.h.

diff --git a/Engine/Plugins/Runtime/MultiServerReplication/Source/MultiServerReplication/Private/MultiServerProxy.cpp b/Engine/Plugins/Runtime/MultiServerReplication/Source/MultiServerReplication/Private/MultiServerProxy.cpp
--- a/Engine/Plugins/Runtime/MultiServerReplication/Source/MultiServerReplication/Private/MultiServerProxy.cpp
+++ b/Engine/Plugins/Runtime/MultiServerReplication/Source/MultiServerReplication/Private/MultiServerProxy.cpp
@@ -9,6 +9,7 @@
 #include "Misc/CommandLine.h"
 #include "Net/DataChannel.h"
 #include "Net/Core/Misc/NetSubObjectRegistry.h"
+#include "Net/NetworkMetricsDatabase.h"
 #include "Net/NetworkMetricsDefs.h"
 
 DEFINE_LOG_CATEGORY_STATIC(LogNetProxy, Log, All);
diff --git a/Engine/Source/Runtime/Engine/Classes/Engine/NetDriver.h b/Engine/Source/Runtime/Engine/Classes/Engine/NetDriver.h
--- a/Engine/Source/Runtime/Engine/Classes/Engine/NetDriver.h
+++ b/Engine/Source/Runtime/Engine/Classes/Engine/NetDriver.h
@@ -7,7 +7,9 @@
 #include "CoreMinimal.h"
 #include "Engine/EngineBaseTypes.h"
 #include "Engine/EngineTypes.h"
+#if UE_ENABLE_INCLUDE_ORDER_DEPRECATED_IN_5_7
 #include "Net/NetworkMetricsDatabase.h"
+#endif
 #include "HAL/IConsoleManager.h"
 #include "Math/RandomStream.h"
 #include "UObject/ObjectMacros.h"
@@ -334,6 +336,8 @@ class FReplicationChangelistMgr;
 class FVoicePacket;
 class StatelessConnectHandlerComponent;
 class UNetConnection;
+class UNetworkMetricsDatabase;
+class UNetworkMetricsBaseListener;
 class UReplicationDriver;
 struct FNetworkObjectInfo;
 class UChannel;
@@ -794,21 +798,6 @@ enum class EProcessRemoteFunctionFlags : uint32
 };
 ENUM_CLASS_FLAGS(EProcessRemoteFunctionFlags);
 
-/** A metrics listener that writes a metric to the 'Replication' CSV category. */
-UCLASS(MinimalAPI)
-class UNetworkMetricsCSV_Replication : public UNetworkMetricsCSV
-{
-	GENERATED_BODY()
-
-public:
-	UNetworkMetricsCSV_Replication()
-	{
-		SetCategory("Replication");
-	}
-
-	virtual ~UNetworkMetricsCSV_Replication() = default;
-};
-
 UCLASS(Abstract, customConstructor, transient, MinimalAPI, config=Engine)
 class UNetDriver : public UObject, public FExec
 {
diff --git a/Engine/Source/Runtime/Engine/Private/DataChannel.cpp b/Engine/Source/Runtime/Engine/Private/DataChannel.cpp
--- a/Engine/Source/Runtime/Engine/Private/DataChannel.cpp
+++ b/Engine/Source/Runtime/Engine/Private/DataChannel.cpp
@@ -29,6 +29,7 @@
 #include "DrawDebugHelpers.h"
 #include "Net/NetworkProfiler.h"
 #include "Net/DataReplication.h"
+#include "Net/NetworkMetricsDatabase.h"
 #include "Net/NetworkMetricsDefs.h"
 #include "Engine/ActorChannel.h"
 #include "Engine/ControlChannel.h"
diff --git a/Engine/Source/Runtime/Engine/Private/NetDriver.cpp b/Engine/Source/Runtime/Engine/Private/NetDriver.cpp
--- a/Engine/Source/Runtime/Engine/Private/NetDriver.cpp
+++ b/Engine/Source/Runtime/Engine/Private/NetDriver.cpp
@@ -47,6 +47,7 @@
 #include "Engine/LevelScriptActor.h"
 #include "Engine/NetworkSettings.h"
 #include "Net/NetEmulationHelper.h"
+#include "Net/NetworkMetricsDatabase.h"
 #include "Net/NetworkMetricsDefs.h"
 #include "Net/NetworkMetricsConfig.h"
 #include "Net/NetSubObjectRegistryGetter.h"
diff --git a/Engine/Source/Runtime/Engine/Public/Net/NetworkMetricsDatabase.h b/Engine/Source/Runtime/Engine/Public/Net/NetworkMetricsDatabase.h
--- a/Engine/Source/Runtime/Engine/Public/Net/NetworkMetricsDatabase.h
+++ b/Engine/Source/Runtime/Engine/Public/Net/NetworkMetricsDatabase.h
@@ -243,6 +243,21 @@ private:
 	int32 CategoryIndex;
 };
 
+/** A metrics listener that writes a metric to the 'Replication' CSV category. */
+UCLASS(MinimalAPI)
+class UNetworkMetricsCSV_Replication : public UNetworkMetricsCSV
+{
+	GENERATED_BODY()
+
+public:
+	UNetworkMetricsCSV_Replication()
+	{
+		SetCategory("Replication");
+	}
+
+	virtual ~UNetworkMetricsCSV_Replication() = default;
+};
+
 /**
  * A metrics listener that reports an array of metrics to PerfCounters.
  * 
