Skeletal Mesh: Move LOD info FSkeletalMeshSourceModel

[RN] FSkeletalMeshLODInfo has been folded into FSkeletalMeshSourceModel

diff --git a/Engine/Plugins/ChaosClothAssetEditor/Source/ChaosClothAssetTools/Private/ChaosClothAsset/SkeletalMeshConverter.cpp b/Engine/Plugins/ChaosClothAssetEditor/Source/ChaosClothAssetTools/Private/ChaosClothAsset/SkeletalMeshConverter.cpp
--- a/Engine/Plugins/ChaosClothAssetEditor/Source/ChaosClothAssetTools/Private/ChaosClothAsset/SkeletalMeshConverter.cpp
+++ b/Engine/Plugins/ChaosClothAssetEditor/Source/ChaosClothAssetTools/Private/ChaosClothAsset/SkeletalMeshConverter.cpp
@@ -86,7 +86,7 @@ bool UClothAssetEditorSkeletalMeshConverter::ExportToSkeletalMesh(const UChaosCl
 
 	SkeletalMesh.PreEditChange(nullptr);
 	SkeletalMesh.GetImportedModel()->LODModels.Empty();
-	SkeletalMesh.ResetLODInfo();
+	SkeletalMesh.SetNumSourceModels(0);
 
 	TArray<FMeshDescription> LocalDescriptions;
 	for (const FDynamicMesh3& DynamicMesh : DynamicMeshes)
diff --git a/Engine/Plugins/Experimental/Dataflow/Source/DataflowNodes/Private/Dataflow/DataflowCollectionEditSkinWeightsNode.cpp b/Engine/Plugins/Experimental/Dataflow/Source/DataflowNodes/Private/Dataflow/DataflowCollectionEditSkinWeightsNode.cpp
--- a/Engine/Plugins/Experimental/Dataflow/Source/DataflowNodes/Private/Dataflow/DataflowCollectionEditSkinWeightsNode.cpp
+++ b/Engine/Plugins/Experimental/Dataflow/Source/DataflowNodes/Private/Dataflow/DataflowCollectionEditSkinWeightsNode.cpp
@@ -38,7 +38,7 @@ static bool CreateSkeletalMesh(USkeletalMesh* SkeletalMesh, const TArray<UE::Geo
 	SkeletalMesh->PreEditChange(nullptr);
 	SkeletalMesh->GetImportedModel()->LODModels.Empty();
 #endif
-	SkeletalMesh->ResetLODInfo();
+	SkeletalMesh->SetNumSourceModels(0);
 
 	TArray<FMeshDescription> LocalDescriptions;
 	for (const UE::Geometry::FDynamicMesh3& DynamicMesh : DynamicMeshes)
diff --git a/Engine/Plugins/Importers/AlembicImporter/Source/AlembicLibrary/Private/AbcImporter.cpp b/Engine/Plugins/Importers/AlembicImporter/Source/AlembicLibrary/Private/AbcImporter.cpp
--- a/Engine/Plugins/Importers/AlembicImporter/Source/AlembicLibrary/Private/AbcImporter.cpp
+++ b/Engine/Plugins/Importers/AlembicImporter/Source/AlembicLibrary/Private/AbcImporter.cpp
@@ -632,7 +632,7 @@ TArray<UObject*> FAbcImporter::ImportAsSkeletalMesh(UObject* InParent, EObjectFl
 		ImportedModel->LODModels.Empty();
 		ImportedModel->EmptyOriginalReductionSourceMeshData();
 		ImportedModel->LODModels.Add(new FSkeletalMeshLODModel());
-		SkeletalMesh->ResetLODInfo();
+		SkeletalMesh->SetNumSourceModels(0);
 
 		FSkeletalMeshLODInfo& NewLODInfo = SkeletalMesh->AddLODInfo();
 		NewLODInfo.ReductionSettings.NumOfTrianglesPercentage = 1.0f;
diff --git a/Engine/Plugins/Mutable/Source/CustomizableObjectEditor/Private/MuCOE/CustomizableObjectInstanceBakingUtils.cpp b/Engine/Plugins/Mutable/Source/CustomizableObjectEditor/Private/MuCOE/CustomizableObjectInstanceBakingUtils.cpp
--- a/Engine/Plugins/Mutable/Source/CustomizableObjectEditor/Private/MuCOE/CustomizableObjectInstanceBakingUtils.cpp
+++ b/Engine/Plugins/Mutable/Source/CustomizableObjectEditor/Private/MuCOE/CustomizableObjectInstanceBakingUtils.cpp
@@ -894,7 +894,7 @@ bool BakeCustomizableObjectInstance(
 
 		if (USkeletalMesh* SkeletalMesh = Cast<USkeletalMesh>(DuplicatedMesh))
 		{
-			SkeletalMesh->ResetLODInfo();
+			SkeletalMesh->SetNumSourceModels(0);
 			for (int32 LODIndex = 0; LODIndex < Mesh->GetLODNum(); ++LODIndex)
         	{
         		SkeletalMesh->AddLODInfo(*Mesh->GetLODInfo(LODIndex));
diff --git a/Engine/Plugins/Runtime/HairStrands/Source/HairStrandsCore/Private/GroomDeformerBuilder.cpp b/Engine/Plugins/Runtime/HairStrands/Source/HairStrandsCore/Private/GroomDeformerBuilder.cpp
--- a/Engine/Plugins/Runtime/HairStrands/Source/HairStrandsCore/Private/GroomDeformerBuilder.cpp
+++ b/Engine/Plugins/Runtime/HairStrands/Source/HairStrandsCore/Private/GroomDeformerBuilder.cpp
@@ -367,7 +367,7 @@ static USkeletalMesh* BuildSkeletalMesh(UGroomAsset* GroomAsset)
 
 	// Reset the skeletal mesh
 	SkeletalMesh->GetImportedModel()->LODModels.Reset();
-	SkeletalMesh->ResetLODInfo();
+	SkeletalMesh->SetNumSourceModels(0);
 	SkeletalMesh->ReleaseResources();
 
 	// Build the LOD render data
diff --git a/Engine/Plugins/Runtime/USDCore/Source/USDUtilities/Private/USDSkeletalDataConversion.cpp b/Engine/Plugins/Runtime/USDCore/Source/USDUtilities/Private/USDSkeletalDataConversion.cpp
--- a/Engine/Plugins/Runtime/USDCore/Source/USDUtilities/Private/USDSkeletalDataConversion.cpp
+++ b/Engine/Plugins/Runtime/USDCore/Source/USDUtilities/Private/USDSkeletalDataConversion.cpp
@@ -2772,7 +2772,7 @@ bool UsdToUnreal::ConvertSkeletalImportData(
 
 	FSkeletalMeshModel* ImportedResource = InOutSkeletalMesh->GetImportedModel();
 	ImportedResource->LODModels.Empty();
-	InOutSkeletalMesh->ResetLODInfo();
+	InOutSkeletalMesh->SetNumSourceModels(0);
 	bool bHasVertexColors = false;
 	for (int32 LODIndex = 0; LODIndex < InLODIndexToSkeletalMeshImportData.Num(); ++LODIndex)
 	{
diff --git a/Engine/Source/Developer/SkeletalMeshUtilitiesCommon/Private/LODUtilities.cpp b/Engine/Source/Developer/SkeletalMeshUtilitiesCommon/Private/LODUtilities.cpp
--- a/Engine/Source/Developer/SkeletalMeshUtilitiesCommon/Private/LODUtilities.cpp
+++ b/Engine/Source/Developer/SkeletalMeshUtilitiesCommon/Private/LODUtilities.cpp
@@ -420,7 +420,7 @@ void FLODUtilities::RemoveLOD(FSkeletalMeshUpdateContext& UpdateContext, int32 D
 		}
 
 		SkelMeshModel->LODModels.RemoveAt(DesiredLOD);
-		SkeletalMesh->RemoveLODInfo(DesiredLOD);
+		SkeletalMesh->RemoveSourceModel(DesiredLOD);
 		RefreshLODChange(SkeletalMesh);
 
 		// Adjust the force LOD to point on the same one, if we are forcing a LOD greater then the one we delete, we want to continue pointing on it
diff --git a/Engine/Source/Editor/Persona/Private/PersonaMeshDetails.cpp b/Engine/Source/Editor/Persona/Private/PersonaMeshDetails.cpp
--- a/Engine/Source/Editor/Persona/Private/PersonaMeshDetails.cpp
+++ b/Engine/Source/Editor/Persona/Private/PersonaMeshDetails.cpp
@@ -6,7 +6,6 @@
 #include "Misc/MessageDialog.h"
 #include "Modules/ModuleManager.h"
 #include "UObject/UObjectHash.h"
-#include "UObject/UObjectIterator.h"
 #include "Framework/Commands/UIAction.h"
 #include "Widgets/Text/STextBlock.h"
 #include "Framework/MultiBox/MultiBoxBuilder.h"
@@ -15,13 +14,11 @@
 #include "IDetailGroup.h"
 #include "DetailLayoutBuilder.h"
 #include "DetailCategoryBuilder.h"
-#include "Widgets/SToolTip.h"
 #include "IDocumentation.h"
 #include "Components/SkinnedMeshComponent.h"
 #include "Animation/DebugSkelMeshComponent.h"
 #include "IPersonaPreviewScene.h"
 #include "Animation/AnimBlueprint.h"
-#include "SlateOptMacros.h"
 #include "Framework/Application/SlateApplication.h"
 #include "Widgets/Images/SImage.h"
 #include "Widgets/Layout/SBox.h"
@@ -31,13 +28,9 @@
 #include "Widgets/Input/SComboButton.h"
 #include "Widgets/Input/SVectorInputBox.h"
 #include "Editor/UnrealEdEngine.h"
-#include "EditorDirectories.h"
 #include "UnrealEdGlobals.h"
-#include "IDetailsView.h"
 #include "MaterialList.h"
 #include "PropertyCustomizationHelpers.h"
-#include "DesktopPlatformModule.h"
-#include "Interfaces/IMainFrameModule.h"
 #include "ScopedTransaction.h"
 #include "Editor.h"
 #include "Animation/AnimBlueprintGeneratedClass.h"
@@ -53,10 +46,6 @@
 #include "FbxMeshUtils.h"
 #include "NaniteLayout.h"
 
-#include "MeshDescription.h"
-#include "MeshAttributes.h"
-#include "MeshAttributeArray.h"
-
 #include "Widgets/Input/STextComboBox.h"
 
 #include "Engine/SkeletalMeshLODSettings.h"
@@ -64,13 +53,8 @@
 #include "Animation/AnimSequence.h"
 #include "IPersonaToolkit.h"
 #include "Widgets/Input/SNumericEntryBox.h"
-#include "IPersonaPreviewScene.h"
-#include "IDocumentation.h"
 #include "JsonObjectConverter.h"
 #include "Widgets/Layout/SUniformWrapPanel.h"
-#include "Widgets/Input/SNumericDropDown.h"
-#include "ComponentReregisterContext.h"
-#include "ClothingSystemEditorInterfaceModule.h"
 #include "ClothingAssetFactoryInterface.h"
 #include "Widgets/Layout/SExpandableArea.h"
 #include "SKismetInspector.h"
@@ -80,12 +64,9 @@
 #include "IMeshReductionManagerModule.h"
 #include "SkinWeightProfileHelpers.h"
 #include "Factories/FbxSkeletalMeshImportData.h"
-#include "PropertyCustomizationHelpers.h"
-#include "ComponentReregisterContext.h"
 #include "IPropertyUtilities.h"
 #include "LODInfoUILayout.h"
 #include "SWarningOrErrorBox.h"
-#include "Interfaces/ITargetPlatform.h"
 #include "Interfaces/ITargetPlatformManagerModule.h"
 #include "Misc/CoreMisc.h"
 #include "InterchangeManager.h"
@@ -97,13 +78,13 @@
 
 DEFINE_LOG_CATEGORY(LogSkeletalMeshPersonaMeshDetail);
 
-#define SUFFIXE_DEFAULT_MATERIAL TEXT(" (Default)")
 
+static const TCHAR* GSuffixDefaultMaterial = TEXT(" (Default)");
 
-static bool GEnableSKeletalMeshMaterialSlotOverlayMaterialUI = false;
-static FAutoConsoleVariableRef CCvarEnableSKeletalMeshMaterialSlotOverlayMaterialUI(
+static bool GEnableSkeletalMeshMaterialSlotOverlayMaterialUI = false;
+static FAutoConsoleVariableRef CVarEnableSkeletalMeshMaterialSlotOverlayMaterialUI(
 	TEXT("SkeletalMesh.FeatureFlags.UI.ShowMaterialSlotOverlayMaterial"),
-	GEnableSKeletalMeshMaterialSlotOverlayMaterialUI,
+	GEnableSkeletalMeshMaterialSlotOverlayMaterialUI,
 	TEXT("Whether skeletal mesh editor show material overlay per material slot UI."),
 	ECVF_Default);
 
@@ -119,8 +100,8 @@ enum SK_CustomDataKey
 
 namespace PersonaMeshDetailsConstants
 {
-	/** Number to extend the num lods slider beyond the currently existing number of lods in the mesh */
-	const static int32 LodSliderExtension = 5;
+	/** Number to extend the num LODs slider beyond the currently existing number of LODs in the mesh */
+	constexpr static int32 LodSliderExtension = 5;
 }
 
 static int32 GetDefaultMaterialIndex(const USkeletalMesh* SkeletalMesh, int32 LODIndex, int32 SectionIndex)
@@ -149,7 +130,7 @@ static int32 GetDefaultMaterialIndex(const USkeletalMesh* SkeletalMesh, int32 LO
 /** Returns true if automatic mesh reduction is available. */
 static bool IsAutoMeshReductionAvailable()
 {
-	static bool bAutoMeshReductionAvailable = FModuleManager::Get().LoadModuleChecked<IMeshReductionManagerModule>("MeshReductionInterface").GetSkeletalMeshReductionInterface() != NULL;
+	static bool bAutoMeshReductionAvailable = FModuleManager::Get().LoadModuleChecked<IMeshReductionManagerModule>("MeshReductionInterface").GetSkeletalMeshReductionInterface() != nullptr;
 	return bAutoMeshReductionAvailable;
 }
 
@@ -338,9 +319,8 @@ private:
 			return LOCTEXT("ReimportButton_NewFile_LODNotValid_ToolTip", "Cannot reimport, LOD has been deleted.");
 		}
 
-		FSkeletalMeshLODInfo& LODInfo = *(SkelMesh->GetLODInfo(LODIndex));
-
-		FString Filename = FPaths::GetCleanFilename(LODInfo.SourceImportFilename);
+		const FSkeletalMeshLODInfo& LODInfo = *(SkelMesh->GetLODInfo(LODIndex));
+		const FString Filename = FPaths::GetCleanFilename(LODInfo.SourceImportFilename);
 
 		return FText::Format(LOCTEXT("ReimportButton_NewFile_ToolTip", "Re-import LOD {0} using the current source file ({1})"), FText::AsNumber(LODIndex), FText::FromString(Filename));
 	}
@@ -1200,7 +1180,7 @@ TSharedRef<SWidget> FSkeletalMeshReductionSettingsLayout::FillReductionMethodMen
 	{
 		return SNullWidget::NullWidget;
 	}
-	FMenuBuilder MenuBuilder(true, NULL);
+	FMenuBuilder MenuBuilder(true, nullptr);
 	int32 EnumCount = EnumReductionMethod->NumEnums();
 	//Skip the last enum since this is the _MAX
 	for (int32 EnumIndex = 0; EnumIndex < EnumCount - 1; ++EnumIndex)
@@ -1230,7 +1210,7 @@ TSharedRef<SWidget> FSkeletalMeshReductionSettingsLayout::FillReductionImportanc
 	{
 		return SNullWidget::NullWidget;
 	}
-	FMenuBuilder MenuBuilder(true, NULL);
+	FMenuBuilder MenuBuilder(true, nullptr);
 	int32 EnumCount = EnumImportance->NumEnums();
 	//Skip the last enum since this is the _MAX
 	for (int32 EnumIndex = 0; EnumIndex < EnumCount - 1; ++EnumIndex)
@@ -1285,7 +1265,7 @@ TSharedRef<class SWidget> FSkeletalMeshReductionSettingsLayout::FillReductionTer
 	{
 		return SNullWidget::NullWidget;
 	}
-	FMenuBuilder MenuBuilder(true, NULL);
+	FMenuBuilder MenuBuilder(true, nullptr);
 	int32 EnumCount = EnumTerminationCriterion->NumEnums();
 	//Skip the last enum since this is the _MAX
 	for (int32 EnumIndex = 0; EnumIndex < EnumCount - 1; ++EnumIndex)
@@ -1816,7 +1796,7 @@ void FSkeletalMeshBuildSettingsLayout::GenerateChildContent(IDetailChildrenBuild
 	}
 
 	{
-		FDetailWidgetRow& PositionThreshold = AddFloatRow(ChildrenBuilder,
+		AddFloatRow(ChildrenBuilder,
 			LOCTEXT("PositionThreshold_Row", "ThresholdPositionRow"),
 			LOCTEXT("PositionThreshold", "Threshold Position"),
 			LOCTEXT("PositionThreshold_ToolTip", "Threshold to compare vertex position equality."),
@@ -1827,7 +1807,7 @@ void FSkeletalMeshBuildSettingsLayout::GenerateChildContent(IDetailChildrenBuild
 	}
 
 	{
-		FDetailWidgetRow& TangentNormalThreshold = AddFloatRow(ChildrenBuilder,
+		AddFloatRow(ChildrenBuilder,
 			LOCTEXT("TangentNormalThreshold_Row", "ThresholdTangentNormalRow"),
 			LOCTEXT("TangentNormalThreshold", "Threshold Tangent and Normal"),
 			LOCTEXT("TangentNormalThreshold_ToolTip", "Threshold to compare normal, tangent or bi-normal equality."),
@@ -1838,7 +1818,7 @@ void FSkeletalMeshBuildSettingsLayout::GenerateChildContent(IDetailChildrenBuild
 	}
 
 	{
-		FDetailWidgetRow& UVThreshold = AddFloatRow(ChildrenBuilder,
+		AddFloatRow(ChildrenBuilder,
 			LOCTEXT("UVThreshold_Row", "ThresholdUVRow"),
 			LOCTEXT("UVThreshold", "Threshold UV"),
 			LOCTEXT("UVThreshold_ToolTip", "Threshold to compare vertex UV equality."),
@@ -1849,7 +1829,7 @@ void FSkeletalMeshBuildSettingsLayout::GenerateChildContent(IDetailChildrenBuild
 	}
 
 	{
-		FDetailWidgetRow& MorphPositionThreshold = AddFloatRow(ChildrenBuilder,
+		AddFloatRow(ChildrenBuilder,
 			LOCTEXT("MorphPositionThreshold_Row", "MorphThresholdPositionRow"),
 			LOCTEXT("MorphPositionThreshold", "Morph Threshold Position"),
 			LOCTEXT("MorphPositionThreshold_ToolTip", "Threshold to compare vertex position equality when computing morph target deltas."),
@@ -1860,7 +1840,7 @@ void FSkeletalMeshBuildSettingsLayout::GenerateChildContent(IDetailChildrenBuild
 	}
 
 	{
-		FDetailWidgetRow& BoneInfluenceLimit = AddIntegerRow(ChildrenBuilder,
+		AddIntegerRow(ChildrenBuilder,
 			LOCTEXT("BoneInfluenceLimit_Row", "BoneInfluenceLimitRow"),
 			LOCTEXT("BoneInfluenceLimit", "Bone Influence Limit"),
 			LOCTEXT("BoneInfluenceLimit_ToolTip", "Limit the number of bone influences a vertex can have. If 0, the Default Bone Influence Limit from the project settings will be used."),
@@ -2774,11 +2754,11 @@ void FPersonaMeshDetails::CustomizeLODInfoSetingsDetails(IDetailLayoutBuilder& D
 
 	TAttribute<bool> EnabledAttrib = TAttribute<bool>::Create(TAttribute<bool>::FGetter::CreateSP(this, &FPersonaMeshDetails::IsLODInfoEditingEnabled, LODIndex));
 
-	const bool bForceComponentRefreshTrue = true;
-	const bool bForceComponentRefreshFalse = false;
+	constexpr bool bForceComponentRefreshTrue = true;
+	constexpr bool bForceComponentRefreshFalse = false;
 	// enable/disable handler - because we want to make sure not editable if LOD sharing is on
 	TSharedPtr<IPropertyHandle> ScreenSizeHandle = LODInfoChild->GetChildHandle(GET_MEMBER_NAME_CHECKED(FSkeletalMeshLODInfo, ScreenSize));
-	IDetailPropertyRow& ScreenSizeRow = LODInfoGroup.AddPropertyRow(ScreenSizeHandle->AsShared()).EditCondition(EnabledAttrib, NULL);
+	IDetailPropertyRow& ScreenSizeRow = LODInfoGroup.AddPropertyRow(ScreenSizeHandle->AsShared()).EditCondition(EnabledAttrib, nullptr);
 	ScreenSizeRow.IsEnabled(EnabledAttrib);
 	PreventAttributePostEditChange(ScreenSizeHandle, LODIndex, ScreenSizeHandle->GetProperty()->GetFName(), bForceComponentRefreshTrue);
 	DetailLayout.HideProperty(ScreenSizeHandle);
@@ -3062,7 +3042,8 @@ void FPersonaMeshDetails::AddLODLevelCategories(IDetailLayoutBuilder& DetailLayo
 				.IsEnabled(this, &FPersonaMeshDetails::IsLODCustomModeEnable, LODIndex)
 			];
 
-			TSharedRef<IPropertyHandle> LODInfoProperty = DetailLayout.GetProperty(FName("LODInfo"), USkeletalMesh::StaticClass());
+			TSharedRef<IPropertyHandle> LODInfoProperty = DetailLayout.GetProperty(FName("SourceModels"), USkeletalMesh::StaticClass());
+			
 			uint32 NumChildren = 0;
 			LODInfoProperty->GetNumChildren(NumChildren);
 			//If the skeleton change during a re-import the skleleton preview mesh will be force reload even if we are in the middle of a re-import
@@ -3234,12 +3215,6 @@ void FPersonaMeshDetails::AddLODLevelCategories(IDetailLayoutBuilder& DetailLayo
 
 FText FPersonaMeshDetails::GetLODCustomModeNameContent(int32 LODIndex) const
 {
-	int32 CurrentLodIndex = 0;
-	if (GetPersonaToolkit()->GetPreviewMeshComponent() != nullptr)
-	{
-		CurrentLodIndex = GetPersonaToolkit()->GetPreviewMeshComponent()->GetForcedLOD();
-	}
-	int32 RealCurrentLODIndex = (CurrentLodIndex == 0 ? 0 : CurrentLodIndex - 1);
 	if (LODIndex == INDEX_NONE)
 	{
 		return LOCTEXT("GetLODCustomModeNameContent_None", "Custom");
@@ -3249,11 +3224,6 @@ FText FPersonaMeshDetails::GetLODCustomModeNameContent(int32 LODIndex) const
 
 ECheckBoxState FPersonaMeshDetails::IsLODCustomModeCheck(int32 LODIndex) const
 {
-	int32 CurrentLodIndex = 0;
-	if (GetPersonaToolkit()->GetPreviewMeshComponent() != nullptr)
-	{
-		CurrentLodIndex = GetPersonaToolkit()->GetPreviewMeshComponent()->GetForcedLOD();
-	}
 	if (LODIndex == INDEX_NONE)
 	{
 		return (GetPersonaToolkit()->GetCustomData(CustomDataKey_LODEditMode) > 0) ? ECheckBoxState::Checked : ECheckBoxState::Unchecked;
@@ -3338,7 +3308,7 @@ TOptional<int32> FPersonaMeshDetails::GetLodSliderMaxValue() const
 void FPersonaMeshDetails::CustomizeSkinWeightProfiles(IDetailLayoutBuilder& DetailLayout)
 {
 	TSharedRef<IPropertyHandle> SkinWeightProfilesProperty = DetailLayout.GetProperty(FName("SkinWeightProfiles"), USkeletalMesh::StaticClass());
-	IDetailCategoryBuilder& SkinWeightCategory = DetailLayout.EditCategory("SkinWeights", LOCTEXT("SkinWeightsCategory", "Skin Weights"));
+	DetailLayout.EditCategory("SkinWeights", LOCTEXT("SkinWeightsCategory", "Skin Weights"));
 
 	IDetailPropertyRow& Row = DetailLayout.AddPropertyToCategory(SkinWeightProfilesProperty);	
 	Row.CustomWidget(true)
@@ -3380,7 +3350,6 @@ void FPersonaMeshDetails::CustomizeSkinWeightProfiles(IDetailLayoutBuilder& Deta
 
 TSharedRef<SWidget> FPersonaMeshDetails::CreateSkinWeightProfileMenuContent()
 {
-	bool bAddedMenuItem = false;
 	FMenuBuilder AddProfileMenuBuilder(true, nullptr, nullptr, true);
 
 	// Menu entry for importing skin weights from an FBX file
@@ -3601,10 +3570,10 @@ void FPersonaMeshDetails::CustomizeLODSettingsCategories(IDetailLayoutBuilder& D
 
 	TSharedPtr<IPropertyHandle> MinLODPropertyHandle = DetailLayout.GetProperty(USkeletalMesh::GetMinLodMemberName(), USkeletalMesh::StaticClass());
 	IDetailPropertyRow& MinLODRow = LODSettingsCategory.AddProperty(MinLODPropertyHandle);
-	MinLODRow.EditCondition(IsPerPlatformMinLodEnabled, NULL);
+	MinLODRow.EditCondition(IsPerPlatformMinLodEnabled, nullptr);
 	DetailLayout.HideProperty(MinLODPropertyHandle);
 
-	TSharedRef<IPropertyHandle> LODInfoProperty = DetailLayout.GetProperty(FName("LODInfo"), USkeletalMesh::StaticClass());
+	TSharedRef<IPropertyHandle> LODInfoProperty = DetailLayout.GetProperty(FName("SourceModels"), USkeletalMesh::StaticClass());
 	DetailLayout.HideProperty(LODInfoProperty);
 
 	TSharedPtr<IPropertyHandle> QualityLevelMinLODPropertyHandle = DetailLayout.GetProperty(USkeletalMesh::GetQualityLevelMinLodMemberName(), USkeletalMesh::StaticClass());
@@ -3613,7 +3582,7 @@ void FPersonaMeshDetails::CustomizeLODSettingsCategories(IDetailLayoutBuilder& D
 	LODSettingsCategory.AddCustomRow(LOCTEXT("QualityLevelMinLOD", "Quality Level Min LOD"))
 		.RowTag("QualityLevelMinLOD")
 		.IsEnabled(IsQualityLevelLodEnabled)
-		.EditCondition(IsQualityLevelLodEnabled, NULL)
+		.EditCondition(IsQualityLevelLodEnabled, nullptr)
 		.NameContent()
 		[
 			SNew(SHorizontalBox)
@@ -3727,7 +3696,7 @@ FReply FPersonaMeshDetails::OnSaveLODSettings()
 				FSavePackageArgs SaveArgs;
 				SaveArgs.TopLevelFlags = RF_Standalone;
 				SaveArgs.SaveFlags = SAVE_NoError;
-				UPackage::SavePackage(Package, NULL, *PackageFileName, SaveArgs);
+				UPackage::SavePackage(Package, nullptr, *PackageFileName, SaveArgs);
 
 				// set the property back to SkelMesh;
 				SkelMesh->SetLODSettings(NewLODSettingAsset);
@@ -4237,7 +4206,7 @@ FReply FPersonaMeshDetails::ApplyLODChanges(int32 LODIndex)
 
 			RestoreNonReducedLOD(LODIndex);
 		}
-		SkelMesh->MarkPackageDirty();
+		(void)SkelMesh->MarkPackageDirty();
 	}
 
 	RequestLayoutUpdate();
@@ -4328,12 +4297,10 @@ FReply FPersonaMeshDetails::RegenerateLOD(int32 LODIndex)
 	USkeletalMesh* SkelMesh = GetPersonaToolkit()->GetMesh();
 	check(SkelMesh);
 
-	bool bIsReductionActive = false;
-
 	if (SkelMesh->IsValidLODIndex(LODIndex))
 	{
-		FSkeletalMeshLODInfo& CurrentLODInfo = *(SkelMesh->GetLODInfo(LODIndex));
-		bIsReductionActive = SkelMesh->IsReductionActive(LODIndex);
+		const FSkeletalMeshLODInfo& CurrentLODInfo = *(SkelMesh->GetLODInfo(LODIndex));
+		const bool bIsReductionActive = SkelMesh->IsReductionActive(LODIndex);
 		if (CurrentLODInfo.bHasBeenSimplified == false && (LODIndex > 0 || bIsReductionActive))
 		{
 			if (LODIndex > 0)
@@ -4347,7 +4314,7 @@ FReply FPersonaMeshDetails::RegenerateLOD(int32 LODIndex)
 			}
 			else if (bIsReductionActive)
 			{
-				//Ask user a special permission when the base LOD can be reduce 
+				//Ask user a special permission when the base LOD can be reduced 
 				const FText Text(LOCTEXT("Warning_ReductionApplyingToImportedMesh_ReduceNonGenBaseLOD", "Are you sure you'd like to apply mesh reduction to the non-generated base LOD?"));
 				EAppReturnType::Type Ret = FMessageDialog::Open(EAppMsgType::YesNo, Text);
 				if (Ret == EAppReturnType::No)
@@ -4699,7 +4666,7 @@ void FPersonaMeshDetails::CustomizeDetails( IDetailLayoutBuilder& DetailLayout )
 	PostProcessHandle->MarkHiddenByCustomization();
 
 	// Hide the existing NaniteSettings property so we can use the customization instead
-	IDetailCategoryBuilder& MeshCategory = DetailLayout.EditCategory("Mesh");
+	DetailLayout.EditCategory("Mesh");
 	TSharedRef<IPropertyHandle> NaniteSettingsProperty = DetailLayout.GetProperty(GET_MEMBER_NAME_CHECKED(USkeletalMesh, NaniteSettings), USkeletalMesh::StaticClass());
 	NaniteSettingsProperty->MarkHiddenByCustomization();
 
@@ -4876,7 +4843,7 @@ void FPersonaMeshDetails::OnMeshRebuildCompleted(USkeletalMesh* InMesh)
 void FPersonaMeshDetails::HideUnnecessaryProperties(IDetailLayoutBuilder& DetailLayout)
 {
 	// LODInfo doesn't need to be showed anymore because it was moved to each LOD category
-	TSharedRef<IPropertyHandle> LODInfoProperty = DetailLayout.GetProperty(FName("LODInfo"), USkeletalMesh::StaticClass());
+	TSharedRef<IPropertyHandle> LODInfoProperty = DetailLayout.GetProperty(FName("SourceModels"), USkeletalMesh::StaticClass());
 	DetailLayout.HideProperty(LODInfoProperty);
 	uint32 NumChildren = 0;
 	LODInfoProperty->GetNumChildren(NumChildren);
@@ -5187,7 +5154,7 @@ void FPersonaMeshDetails::OnGetSectionsForView(ISectionListBuilder& OutSections,
 					{
 						if (DefaultSectionMaterialIndex == CurrentIterMaterialIndex)
 						{
-							FString BuildDefaultName = SkeletalMaterial.MaterialSlotName.ToString() + SUFFIXE_DEFAULT_MATERIAL;
+							FString BuildDefaultName = SkeletalMaterial.MaterialSlotName.ToString() + GSuffixDefaultMaterial;
 							AvailableSectionName.Add(CurrentIterMaterialIndex, FName(*BuildDefaultName));
 						}
 						else
@@ -5303,7 +5270,7 @@ TSharedRef<SWidget> FPersonaMeshDetails::OnGenerateMaterialListExtraBottomWidget
 	}
 	TSharedRef<SVerticalBox> VerticalWidget = SNew(SVerticalBox);
 	const FSkeletalMaterial& SkeletalMaterial = SkeletalMeshPtr->GetMaterials()[MaterialIndex];
-	if(GEnableSKeletalMeshMaterialSlotOverlayMaterialUI || SkeletalMaterial.OverlayMaterialInterface)
+	if(GEnableSkeletalMeshMaterialSlotOverlayMaterialUI || SkeletalMaterial.OverlayMaterialInterface)
 	{
 		VerticalWidget->AddSlot()
 		.AutoHeight()
@@ -6078,7 +6045,7 @@ TSharedRef<SWidget> FPersonaMeshDetails::OnGenerateLodMenuForLodPicker()
 	{
 		return SNullWidget::NullWidget;
 	}
-	FMenuBuilder MenuBuilder(true, NULL);
+	FMenuBuilder MenuBuilder(true, nullptr);
 
 	FUIAction AutoLodAction(FExecuteAction::CreateSP(this, &FPersonaMeshDetails::SetCurrentLOD, 0));
 	MenuBuilder.AddMenuEntry(LOCTEXT("LODAutoLabel", "Auto (LOD 0)"), LOCTEXT("OnGenerateLodMenuForSectionList_Auto_ToolTip", "With Auto LOD selected, LOD 0's properties are visible for editing."), FSlateIcon(), AutoLodAction);
@@ -6618,7 +6585,7 @@ void FPersonaMeshDetails::OnSectionChanged(int32 LODIndex, int32 SectionIndex, i
 		check(TotalSectionCount > SectionIndex);
 
 		FString NewMaterialSlotNameString = NewMaterialSlotName.ToString();
-		NewMaterialSlotNameString.RemoveFromEnd(SUFFIXE_DEFAULT_MATERIAL, ESearchCase::CaseSensitive);
+		NewMaterialSlotNameString.RemoveFromEnd(GSuffixDefaultMaterial, ESearchCase::CaseSensitive);
 		FName CleanNewMaterialSlotName(*NewMaterialSlotNameString);
 
 		int32 NewSkeletalMaterialIndex = INDEX_NONE;
@@ -7141,4 +7108,5 @@ bool FPersonaMeshDetails::FilterOutBakePose(const FAssetData& AssetData, USkelet
 {
 	return !(Skeleton && Skeleton->IsCompatibleForEditor(AssetData));
 }
-#undef LOCTEXT_NAMESPACE
\ No newline at end of file
+
+#undef LOCTEXT_NAMESPACE
diff --git a/Engine/Source/Editor/UnrealEd/Private/Fbx/FbxSkeletalMeshImport.cpp b/Engine/Source/Editor/UnrealEd/Private/Fbx/FbxSkeletalMeshImport.cpp
--- a/Engine/Source/Editor/UnrealEd/Private/Fbx/FbxSkeletalMeshImport.cpp
+++ b/Engine/Source/Editor/UnrealEd/Private/Fbx/FbxSkeletalMeshImport.cpp
@@ -1893,7 +1893,7 @@ USkeletalMesh* UnFbx::FFbxImporter::ImportSkeletalMesh(FImportSkeletalMeshArgs &
 		ImportedResource->EmptyOriginalReductionSourceMeshData();
 		ImportedResource->LODModels.Empty();
 
-		SkeletalMesh->ResetLODInfo();
+		SkeletalMesh->SetNumSourceModels(0);
 		SkeletalMesh->GetMaterials().Empty();
 		SkeletalMesh->GetRefSkeleton().Empty();
 		SkeletalMesh->SetSkeleton(nullptr);
@@ -1939,7 +1939,9 @@ USkeletalMesh* UnFbx::FFbxImporter::ImportSkeletalMesh(FImportSkeletalMeshArgs &
 	// process bone influences from import data
 	SkeletalMeshImportUtils::ProcessImportMeshInfluences(*SkelMeshImportDataPtr, SkeletalMesh->GetPathName());
 
-	SkeletalMesh->ResetLODInfo();
+	// Clean out all LODs.
+	SkeletalMesh->SetNumSourceModels(0);
+	
 	FSkeletalMeshLODInfo& NewLODInfo = SkeletalMesh->AddLODInfo();
 	NewLODInfo.ReductionSettings.NumOfTrianglesPercentage = 1.0f;
 	NewLODInfo.ReductionSettings.NumOfVertPercentage = 1.0f;
diff --git a/Engine/Source/Runtime/Engine/Classes/Engine/SkeletalMesh.h b/Engine/Source/Runtime/Engine/Classes/Engine/SkeletalMesh.h
--- a/Engine/Source/Runtime/Engine/Classes/Engine/SkeletalMesh.h
+++ b/Engine/Source/Runtime/Engine/Classes/Engine/SkeletalMesh.h
@@ -94,7 +94,7 @@ enum class ESkeletalMeshAsyncProperties : uint64
 	MaxNumStreamedLODs = 1 << 17,
 	MaxNumOptionalLODs = 1 << 18,
 	ImportedModel = 1 << 19,
-	LODInfo = 1 << 20,
+	// Spare = 1 << 20,
 	SkinWeightProfiles = 1 << 21,
 	CachedComposedRefPoseMatrices = 1 << 22,
 	SamplingInfo = 1 << 23,
@@ -442,44 +442,19 @@ class USkeletalMesh : public USkinnedAsset, public IInterface_CollisionDataProvi
 
 	// This is declared so we can use TUniquePtr<FSkeletalMeshRenderData> with just a forward declare of that class
 	ENGINE_API USkeletalMesh(FVTableHelper& Helper);
-	ENGINE_API ~USkeletalMesh();
+	ENGINE_API virtual ~USkeletalMesh() override;
 
 #if WITH_EDITOR
 	/** Notification when anything changed */
 	DECLARE_MULTICAST_DELEGATE(FOnMeshChanged);
 #endif
 private:
-#if WITH_EDITORONLY_DATA
 	/** The imported/editable geometric source data for generating renderable data. */
 	UE_DEPRECATED(5.4, "This must be protected for async build, always use the accessors even internally.")
-	UPROPERTY()
+	UPROPERTY(EditAnywhere, EditFixedSize, Category=LevelOfDetailModels)
 	TArray<FSkeletalMeshSourceModel> SourceModels;
 
-	/** Set the number of source models. If adding more LODs than currently exist, empty LODs will be created
-	 *  to fill in the gap. If the new size is smaller than the existing LODs, the source geometry for LODs
-	 *  at and higher than the number given will be thrown away.
-	 *  NOTE: Internal for now, use AddLODInfo to add LODs (incl. source models).
-	 */
-	void SetNumSourceModels(const int32 InNumSourceModels);
-
-	/** Used by both SetNumSourceModels and AddLODInfo to set the current number of source models.
-	 *  Does not take an async lock, expects the calling function to be holding it already.
-	 *  \param InNumSourceModels The number of source models to set. Can be lower or higher than the
-	 *     current number of entries and models will be added/removed as needed.
-	 */
-	void SetNumSourceModelsInternal(const int32 InNumSourceModels);
-
-	/** Remove the source model for the given LOD. LODs above it will be adjusted down.
-	 *  NOTE: Internal for now, use RemoveLODInfo to remove specific LODs.  
-	 */
-	void RemoveSourceModel(const int32 InLODIndex);
-	
-	/** Used by RemoveLODInfo to remove source models if a LOD is removed.
-	 *  Does not take an async lock, expects the calling function to be holding it already.
-	 *  \param InLODIndex The LOD index whose model data is being removed.
-	 */
-	void RemoveSourceModelInternal(const int32 InLODIndex);
-	
+#if WITH_EDITORONLY_DATA
 	/** Processed skeletal mesh geometry information (not used at runtime). */
 	/** NOTE: This data is slated for removal in future versions of UE. Do not rely on it. Use either the
 	 *  render data, or source models, as needed.
@@ -499,25 +474,36 @@ private:
 	/** Returns true if the mesh has valid Nanite render data. */
 	ENGINE_API bool HasValidNaniteData() const;
 
-#if WITH_EDITORONLY_DATA
-	/*
-	 * This editor data asset is save under the skeletalmesh(skel mesh is the owner), the editor data asset is always loaded.
-	 * There is only one editor data asset possible per skeletalmesh.
-	 * The reason we store the editor data in a separate asset is because the size of it can be very big and affect the editor performance. (undo/redo transactions)
-	 */
-	UE_DEPRECATED(4.27, "Please do not access this member directly; use the public ImportData API.")
-	UPROPERTY()
-	mutable TObjectPtr<USkeletalMeshEditorData> MeshEditorDataObject_DEPRECATED;
+public:
 
-#endif //WITH_EDITORONLY_DATA
+	/** Add a new source model and set the reduction setting according to the state of the bAutoGenerated argument. 
+	 * \param bAutoGenerated If set to true, the reduction settings are updated to point to a previous LOD with
+	 *    editable data and the triangle reduction set to 50% for each LOD between the two (e.g. if this adds LOD 2,
+	 *    and the first non-autogenerated LOD is 0, then the reduction settings will be at 25%).
+	 * \return The newly created source model.
+	 */
+	ENGINE_API FSkeletalMeshSourceModel& AddSourceModel(const bool bAutoGenerated = false);
+		
+	/** Set the number of source models. If adding more LODs than currently exist, empty LODs will be created
+	 *  to fill in the gap. If the new size is smaller than the existing LODs, the source geometry for LODs
+	 *  at and higher than the number given will be thrown away.
+	 */
+	ENGINE_API void SetNumSourceModels(const int32 InNumSourceModels);
 
-public:
+	/** Remove the source model for the given LOD. LODs above it will be adjusted down.
+	 */
+	ENGINE_API void RemoveSourceModel(const int32 InLODIndex);
 
-#if WITH_EDITORONLY_DATA
 	/** Returns the number of source models.
 	 */
 	ENGINE_API int32 GetNumSourceModels() const;
 
+	/** Returns a boolean indicating whether the given LOD index is valid or not.
+	 * \param InLODIndex The potential LOD index to check for validity.
+	 * \return \c true if the given LOD index in InLODIndex is valid.
+	*/
+	ENGINE_API bool IsSourceModelValid(const int32 InLODIndex) const;	
+	
 	/** Returns all source models as immutable. See GetSourceModel for description of what each source model contains.
 	 *  \return The list of all source models. 
 	 */
@@ -541,7 +527,8 @@ public:
 	 *  \param InLODIndex The LOD index to retrieve a FSkeletalMeshSourceModel object for. 
 	 */
 	ENGINE_API FSkeletalMeshSourceModel& GetSourceModel(const int32 InLODIndex);
-	
+
+#if WITH_EDITORONLY_DATA
 	/** Returns the MeshDescription object associated with the given LOD index. If there's no import data
 	 *  at this LOD, and the LOD is auto-generated, then a \c nullptr is returned.
 	 *  \param InLODIndex The LOD index to retrieve a FMeshDescription object for. 
@@ -1016,10 +1003,6 @@ private:
 #endif
 
 private:
-	/** Struct containing information for each LOD level, such as materials to use, and when use the LOD. */
-	UE_DEPRECATED(5.0, "This must be protected for async build, always use the accessors even internally.")
-	UPROPERTY(EditAnywhere, EditFixedSize, Category=LevelOfDetail)
-	TArray<FSkeletalMeshLODInfo> LODInfo;
 
 #if WITH_EDITOR
 	FSimpleMulticastDelegate OnVertexAttributesArrayChanged;
@@ -2973,44 +2956,27 @@ public:
 	/* 
 	 * Remove LOD info of given index
 	 */
+	UE_DEPRECATED(5.7, "Use RemoveSourceModel() instead.")
 	ENGINE_API void RemoveLODInfo(int32 Index);
 	
 	/*
 	 * Reset whole entry
 	 */
+	UE_DEPRECATED(5.7, "Use SetNumSourceModels(0) instead.")
 	ENGINE_API void ResetLODInfo();
 
-	static FName GetLODInfoMemberName()
-	{
-		PRAGMA_DISABLE_DEPRECATION_WARNINGS
-		return GET_MEMBER_NAME_CHECKED(USkeletalMesh, LODInfo);
-		PRAGMA_ENABLE_DEPRECATION_WARNINGS
-	}
-
 	/*
 	 * Returns whole array of LODInfo non-const. USkinnedAsset interface.
 	 */
 	UE_DEPRECATED(5.5, "Use GetLODInfo and GetLODNum instead.")
-	virtual TArray<FSkeletalMeshLODInfo>& GetLODInfoArray() override
-	{ 
-		WaitUntilAsyncPropertyReleased(ESkeletalMeshAsyncProperties::LODInfo);
-		PRAGMA_DISABLE_DEPRECATION_WARNINGS
-		return LODInfo;
-		PRAGMA_ENABLE_DEPRECATION_WARNINGS
-	}
+	ENGINE_API virtual TArray<FSkeletalMeshLODInfo>& GetLODInfoArray() override;
 
 	/*
 	 * Returns whole array of LODInfo const. USkinnedAsset interface.
 	 */
 	UE_DEPRECATED(5.5, "Use GetLODInfo and GetLODNum instead.")
-	virtual const TArray<FSkeletalMeshLODInfo>& GetLODInfoArray() const override
-	{
-		WaitUntilAsyncPropertyReleased(ESkeletalMeshAsyncProperties::LODInfo, ESkinnedAssetAsyncPropertyLockType::ReadOnly);
-		PRAGMA_DISABLE_DEPRECATION_WARNINGS
-		return LODInfo;
-		PRAGMA_ENABLE_DEPRECATION_WARNINGS
-	}
-	
+	ENGINE_API virtual const TArray<FSkeletalMeshLODInfo>& GetLODInfoArray() const override;
+
 	/* 
 	 * Get LODInfo of the given index non-const. USkinnedAsset interface.
 	 */
@@ -3021,6 +2987,14 @@ public:
 	 */	
 	ENGINE_API virtual const FSkeletalMeshLODInfo* GetLODInfo(int32 Index) const override;
 
+private:
+
+	/** Update the given LOD based on a previous LOD, setting it to auto-generated based on the
+	 *  first non-autogenerated LOD, with 50% reduction in triangle count per-LOD.
+	 */
+	void CopyAutogeneratedLODInfoFromPreviousLOD(int32 InLODIndex);
+public:
+	
 #if WITH_EDITOR
 	ENGINE_API FSimpleMulticastDelegate& GetOnVertexAttributesArrayChanged();
 #endif
@@ -3229,6 +3203,21 @@ private:
 	 * Complete the postload process - Can't be done in parallel. USkinnedAsset interface.
 	 */
 	ENGINE_API virtual void FinishPostLoadInternal(FSkinnedAssetPostLoadContext& Context) override;
+
+private:
+	/**
+	 * Fully deprecated properties.
+	 */
+	 
+#if WITH_EDITORONLY_DATA
+	UPROPERTY()
+	TObjectPtr<USkeletalMeshEditorData> MeshEditorDataObject_DEPRECATED;
+
+#endif //WITH_EDITORONLY_DATA
+	
+	UPROPERTY()
+	TArray<FSkeletalMeshLODInfo> LODInfo_DEPRECATED;
+	
 };
 
 struct FSkeletalMeshBuildParameters
diff --git a/Engine/Source/Runtime/Engine/Classes/Engine/SkeletalMeshSourceModel.h b/Engine/Source/Runtime/Engine/Classes/Engine/SkeletalMeshSourceModel.h
--- a/Engine/Source/Runtime/Engine/Classes/Engine/SkeletalMeshSourceModel.h
+++ b/Engine/Source/Runtime/Engine/Classes/Engine/SkeletalMeshSourceModel.h
@@ -4,6 +4,7 @@
 
 #include "MeshDescription.h"
 #include "MeshDescriptionBaseBulkData.h"
+#include "SkinnedAssetCommon.h"
 
 #include "SkeletalMeshSourceModel.generated.h"
 
@@ -24,7 +25,8 @@ public:
 
 
 USTRUCT()
-struct FSkeletalMeshSourceModel
+struct FSkeletalMeshSourceModel :
+	public FSkeletalMeshLODInfo
 {
 	GENERATED_BODY()
 
@@ -187,13 +189,19 @@ private:
 	void UpdateBonesDataIfNeeded(FMeshDescription& InOutMeshDescription) const;
     
 #endif
-	
+
+	/** The number of triangles on the mesh, stored here so that users can quickly query this information
+	 *  without having to unpack the entire mesh. This is set when CommitMeshDescription is called. */
 	UPROPERTY()
 	int32 TriangleCount = 0;
 	
+	/** The number of vertices on the mesh, stored here so that users can quickly query this information
+	 *  without having to unpack the entire mesh. This is set when CommitMeshDescription is called. */
 	UPROPERTY()
 	int32 VertexCount = 0;
 
+	/** The bounds of the stored mesh, here so that users can quickly query this information
+	 *  without having to unpack the entire mesh. This is set when CommitMeshDescription is called. */
 	UPROPERTY()
 	FBoxSphereBounds Bounds;
 
diff --git a/Engine/Source/Runtime/Engine/Classes/Engine/SkinnedAssetCommon.h b/Engine/Source/Runtime/Engine/Classes/Engine/SkinnedAssetCommon.h
--- a/Engine/Source/Runtime/Engine/Classes/Engine/SkinnedAssetCommon.h
+++ b/Engine/Source/Runtime/Engine/Classes/Engine/SkinnedAssetCommon.h
@@ -135,11 +135,11 @@ struct FSkeletalMeshLODInfo
 	 * sphere of the model. i.e. 0.5 means half the screen's maximum dimension.
 	 */
 	UPROPERTY(EditAnywhere, Category=SkeletalMeshLODInfo)
-	FPerPlatformFloat ScreenSize;
+	FPerPlatformFloat ScreenSize{ 1.0 };
 
 	/**	Used to avoid 'flickering' when on LOD boundary. Only taken into account when moving from complex->simple. */
 	UPROPERTY(EditAnywhere, Category=SkeletalMeshLODInfo, meta=(DisplayName="LOD Hysteresis"))
-	float LODHysteresis;
+	float LODHysteresis = 0.0f;
 
 	/** Mapping table from this LOD's materials to the USkeletalMesh materials array.
 	 * section index is the key
@@ -180,7 +180,7 @@ struct FSkeletalMeshLODInfo
 
 	/** How much to consideration to give BonesToPrioritize and SectionsToPrioritize.  The weight is an additional vertex simplification penalty where 0 means nothing. */
 	UPROPERTY(EditAnywhere, Category = ReductionSettings, meta = (UIMin = "0.0", ClampMin = "0.0"))
-	float WeightOfPrioritization;
+	float WeightOfPrioritization = 1.0f;
 
 	/** Pose which should be used to reskin vertex influences for which the bones will be removed in this LOD level, uses ref-pose by default */
 	UPROPERTY(EditAnywhere, Category = ReductionSettings)
@@ -214,6 +214,10 @@ struct FSkeletalMeshLODInfo
 
 #endif //WITH_EDITORONLY_DATA
 
+	/** List of vertex attributes to include for rendering and what type they should be */
+	UPROPERTY(EditAnywhere, Category = SkeletalMeshLODInfo, AdvancedDisplay, EditFixedSize, Meta=(NoResetToDefault))
+	TArray<FSkeletalMeshVertexAttributeInfo> VertexAttributes;	
+	
 	/** Whether to disable morph targets for this LOD. */
 	UPROPERTY()
 	uint8 bHasBeenSimplified:1;
@@ -237,10 +241,6 @@ struct FSkeletalMeshLODInfo
 	UPROPERTY(EditAnywhere, Category = SkeletalMeshLODInfo)
 	uint8 bAllowMeshDeformer : 1;
 	
-	/** List of vertex attributes to include for rendering and what type they should be */
-	UPROPERTY(EditAnywhere, Category = SkeletalMeshLODInfo, AdvancedDisplay, EditFixedSize, Meta=(NoResetToDefault))
-	TArray<FSkeletalMeshVertexAttributeInfo> VertexAttributes;	
-	
 	/**
 	Mesh supports uniformly distributed sampling in constant time.
 	Memory cost is 8 bytes per triangle.
@@ -350,12 +350,7 @@ struct FSkeletalMeshLODInfo
 #endif
 
 	FSkeletalMeshLODInfo()
-		: ScreenSize(1.0)
-		, LODHysteresis(0.0f)
-		, WeightOfPrioritization(1.f)
-		, BakePose(nullptr)
-		, BakePoseOverride(nullptr)
-		, bHasBeenSimplified(false)
+		: bHasBeenSimplified(false)
 		, bHasPerLODVertexColors(false)
 		, bAllowCPUAccess(false)
 		, bBuildHalfEdgeBuffers(false)
@@ -365,9 +360,6 @@ struct FSkeletalMeshLODInfo
 		, bImportWithBaseMesh(false)
 #endif
 	{
-#if WITH_EDITORONLY_DATA
-		BuildGUID.Invalidate();
-#endif
 	}
 
 };
diff --git a/Engine/Source/Runtime/Engine/Private/SkeletalMesh.cpp b/Engine/Source/Runtime/Engine/Private/SkeletalMesh.cpp
--- a/Engine/Source/Runtime/Engine/Private/SkeletalMesh.cpp
+++ b/Engine/Source/Runtime/Engine/Private/SkeletalMesh.cpp
@@ -6,11 +6,9 @@
 
 #include "Engine/SkeletalMesh.h"
 
-#include "Animation/MeshDeformerGeometry.h"
 #include "DataDrivenShaderPlatformInfo.h"
 #include "Engine/Blueprint.h"
 #include "Engine/SkeletalMeshSampling.h"
-#include "EngineStats.h"
 #include "Engine/SkinnedAsset.h"
 #include "Engine/SkinnedAssetAsyncCompileUtils.h"
 #include "Engine/SkinnedAssetCommon.h"
@@ -25,11 +23,8 @@
 #include "Math/ScaleRotationTranslationMatrix.h"
 #include "UObject/AssetRegistryTagsContext.h"
 #include "UObject/FrameworkObjectVersion.h"
-#include "PhysicsEngine/SphereElem.h"
 #include "UObject/NiagaraObjectVersion.h"
-#include "PhysicsEngine/SphylElem.h"
 #include "UObject/ObjectSaveContext.h"
-#include "PrimitiveViewRelevance.h"
 #include "UObject/Package.h"
 #include "RenderUtils.h"
 #include "AssetCompilingManager.h"
@@ -38,12 +33,9 @@
 #include "EngineUtils.h"
 #include "EditorSupportDelegates.h"
 #include "SkeletalMeshLegacyCustomVersions.h"
-#include "SkeletalMeshSceneProxy.h"
 #include "SkeletalRenderPublic.h"
-#include "PrimitiveDrawingUtils.h"
 #include "PhysicsEngine/BodySetup.h"
 #include "PhysicsEngine/PhysicsAsset.h"
-#include "PhysicsEngine/SkeletalBodySetup.h"
 #include "Engine/AssetUserData.h"
 #include "Animation/NodeMappingContainer.h"
 #include "Rendering/SkeletalMeshRenderData.h"
@@ -56,9 +48,25 @@
 #include "UObject/FortniteMainBranchObjectVersion.h"
 #include "UObject/ReleaseObjectVersion.h"
 #include "UObject/ICookInfo.h"
-#include "GPUSkinCacheVisualizationData.h"
 #include "GPUSkinCache.h"
-#include "MeshPaintVisualize.h"
+#include "UObject/DevObjectVersion.h"
+#include "UnrealEngine.h"
+#include "EditorFramework/AssetImportData.h"
+#include "Engine/SkeletalMeshSocket.h"
+#include "ClothingAssetBase.h"
+#include "Async/Async.h"
+#include "Misc/UObjectToken.h"
+#include "Misc/RuntimeErrors.h"
+#include "PlatformInfo.h"
+#include "Animation/SkinWeightProfileManager.h"
+#include "BoneWeights.h"
+#include "SkeletalMeshAttributes.h"
+#include "Algo/AnyOf.h"
+#include "Logging/StructuredLog.h"
+#include "Rendering/SkeletalMeshHalfEdgeBufferAccessor.h"
+#include "Animation/MeshDeformerCollection.h"
+#include "Animation/MeshDeformer.h"
+#include "Streaming/UVChannelDensity.h"
 
 #if WITH_EDITOR
 #include "AssetRegistry/IAssetRegistry.h"
@@ -76,40 +84,14 @@
 #include "Engine/RendererSettings.h"
 #include "Misc/DataValidation.h"
 #include "Hash/xxhash.h"
+#include "ScopedTransaction.h"
 #else
 #include "Interfaces/ITargetPlatform.h"
 #endif // #if WITH_EDITOR
 
-#include "UObject/DevObjectVersion.h"
-#include "UnrealEngine.h"
-
-#include "EditorFramework/AssetImportData.h"
-#include "Engine/SkeletalMeshSocket.h"
-#include "Components/BrushComponent.h"
-
-#include "ClothingAssetBase.h"
-#include "Async/Async.h"
-#include "Misc/UObjectToken.h"
-
-#if WITH_EDITOR
-#include "ScopedTransaction.h"
-#endif
-#include "SkeletalDebugRendering.h"
-#include "Misc/RuntimeErrors.h"
-#include "PlatformInfo.h"
-
 #if RHI_RAYTRACING
 #include "RayTracingInstance.h"
 #endif
-#include "Animation/SkinWeightProfileManager.h"
-#include "BoneWeights.h"
-#include "SkeletalMeshAttributes.h"
-#include "Algo/AnyOf.h"
-#include "Logging/StructuredLog.h"
-#include "Rendering/SkeletalMeshHalfEdgeBufferAccessor.h"
-#include "Animation/MeshDeformerCollection.h"
-#include "Animation/MeshDeformer.h"
-#include "Streaming/UVChannelDensity.h"
 
 #include UE_INLINE_GENERATED_CPP_BY_NAME(SkeletalMesh)
 
@@ -1362,9 +1344,9 @@ void RefreshSkelMeshOnPhysicsAssetChange(const USkeletalMesh* InSkeletalMesh)
 		for (FThreadSafeObjectIterator Iter(USkeletalMeshComponent::StaticClass()); Iter; ++Iter)
 		{
 			USkeletalMeshComponent* SkeletalMeshComponent = Cast<USkeletalMeshComponent>(*Iter);
-			// if PhysicsAssetOverride is NULL, it uses SkeletalMesh Physics Asset, so I'll need to update here
+			// if PhysicsAssetOverride is null, it uses SkeletalMesh Physics Asset, so I'll need to update here
 			if  (SkeletalMeshComponent->GetSkeletalMeshAsset() == InSkeletalMesh &&
-				 SkeletalMeshComponent->PhysicsAssetOverride == NULL)
+				 SkeletalMeshComponent->PhysicsAssetOverride == nullptr)
 			{
 				// it needs to recreate IF it already has been created
 				if (SkeletalMeshComponent->IsPhysicsStateCreated())
@@ -1407,13 +1389,12 @@ void USkeletalMesh::PostEditChangeProperty(FPropertyChangedEvent& PropertyChange
 		//Ignore those call when we have an active delay stack
 		return;
 	}
-	//Block any re-entrant call by incrementing PostEditChangeStackCounter. It will be decrement when we will go out of scope.
-	const bool bCallPostEditChange = false;
-	const bool bReRegisterComponents = false;
+	
+	//Block any re-entrant call by incrementing PostEditChangeStackCounter. It will be decremented when we exit scope.
+	constexpr bool bCallPostEditChange = false;
+	constexpr bool bReRegisterComponents = false;
 	FScopedSkeletalMeshPostEditChange BlockRecursiveCallScope(this, bCallPostEditChange, bReRegisterComponents);
 
-	bool bFullPrecisionUVsReallyChanged = false;
-
 	FProperty* PropertyThatChanged = PropertyChangedEvent.Property;
 	
 	if (GIsEditor &&
@@ -1508,10 +1489,6 @@ void USkeletalMesh::PostEditChangeProperty(FPropertyChangedEvent& PropertyChange
 			{
 				GetSamplingInfoInternal().BuildRegions(this);
 			}
-			else if (PropertyChangedEvent.MemberProperty->GetFName() == GetLODInfoMemberName())
-			{
-				GetSamplingInfoInternal().BuildWholeMesh(this);
-			}
 			else if (PropertyChangedEvent.MemberProperty->GetFName() == GET_MEMBER_NAME_CHECKED(FSkeletalMeshLODInfo, bSupportUniformlyDistributedSampling))
 			{
 				GetSamplingInfoInternal().BuildWholeMesh(this);
@@ -2035,7 +2012,7 @@ void USkeletalMesh::Serialize( FArchive& Ar )
 					else
 					{
 						//Fall back in case we use an archive that the cooking target has not been set (i.e. Duplicate archive)
-						check(RunningPlatform != NULL);
+						check(RunningPlatform != nullptr);
 						LocalSkeletalMeshRenderData = &::GetPlatformSkeletalMeshRenderData(this, RunningPlatform, bIsSerializeSaving);
 					}
 #endif
@@ -2548,7 +2525,7 @@ void USkeletalMesh::CalculateRequiredBones(FSkeletalMeshLODModel& LODModel, cons
 	{
 		// Make sure it's not in BonesToRemove
 		// @Todo change this to one TArray
-		if (!BonesToRemove || BonesToRemove->Find(i) == NULL)
+		if (!BonesToRemove || BonesToRemove->Find(i) == nullptr)
 		{
 			LODModel.RequiredBones.Add(i);
 		}
@@ -2654,6 +2631,8 @@ void USkeletalMesh::RemoveLegacyClothingSections()
 	}
 }
 
+#endif
+
 ///////////////////////////////////////////////////////////////////////
 //// Source Model API
 
@@ -2665,17 +2644,49 @@ int32 USkeletalMesh::GetNumSourceModels() const
 	PRAGMA_ENABLE_DEPRECATION_WARNINGS
 }
 
-void USkeletalMesh::SetNumSourceModels(const int32 InNumSourceModels)
+bool USkeletalMesh::IsSourceModelValid(const int32 InLODIndex) const
 {
 	WaitUntilAsyncPropertyReleased(ESkeletalMeshAsyncProperties::SourceModels);
-	SetNumSourceModelsInternal(InNumSourceModels);
+	PRAGMA_DISABLE_DEPRECATION_WARNINGS
+	return SourceModels.IsValidIndex(InLODIndex);
+	PRAGMA_ENABLE_DEPRECATION_WARNINGS
+}
+
+
+FSkeletalMeshSourceModel& USkeletalMesh::AddSourceModel(const bool bAutoGenerated)
+{
+	WaitUntilAsyncPropertyReleased(ESkeletalMeshAsyncProperties::SourceModels|ESkeletalMeshAsyncProperties::ImportedModel);
+	PRAGMA_DISABLE_DEPRECATION_WARNINGS
+
+	const int32 LODIndex = SourceModels.Num(); 
+	FSkeletalMeshSourceModel& Model = SourceModels.AddDefaulted_GetRef();
+	Model.Initialize(this);
+
+	const USkeletalMeshLODSettings* DefaultSetting = GetDefaultLODSetting();
+
+	// Set the LOD settings from the default LOD settings and populate with that.
+	DefaultSetting->SetLODSettingsToMesh(this, LODIndex);
+	
+	// If the new LOD is requested to be auto-generated, then set the reduction settings to use the same settings
+	// but decimate 2x as much.
+	if (LODIndex > 0 && bAutoGenerated)
+	{
+		CopyAutogeneratedLODInfoFromPreviousLOD(LODIndex);
+	}
+	
+	PRAGMA_ENABLE_DEPRECATION_WARNINGS
+
+	return Model;
 }
 
-void USkeletalMesh::SetNumSourceModelsInternal(const int32 InNumSourceModels)
+
+void USkeletalMesh::SetNumSourceModels(const int32 InNumSourceModels)
 {
+	WaitUntilAsyncPropertyReleased(ESkeletalMeshAsyncProperties::SourceModels|ESkeletalMeshAsyncProperties::ImportedModel);
 	PRAGMA_DISABLE_DEPRECATION_WARNINGS
 	const int32 OldNumSourceModels = SourceModels.Num();
 
+#if WITH_EDITOR
 	// If we're lowering the LOD count, clear out any mesh data stored immediately rather
 	// than wait for GC to do the work.
 	if (OldNumSourceModels > InNumSourceModels)
@@ -2684,9 +2695,15 @@ void USkeletalMesh::SetNumSourceModelsInternal(const int32 InNumSourceModels)
 		{
 			FSkeletalMeshSourceModel& SourceModel = SourceModels[Index];
 			SourceModel.ClearAllMeshData();
+
+			if (ImportedModel->InlineReductionCacheDatas.IsValidIndex(Index))
+			{
+				ImportedModel->InlineReductionCacheDatas.RemoveAt(Index);
+			}
 		}
 	}
-
+#endif
+	
 	SourceModels.SetNum(InNumSourceModels);
 
 	for (int32 Index = OldNumSourceModels; Index < InNumSourceModels; Index++)
@@ -2696,22 +2713,22 @@ void USkeletalMesh::SetNumSourceModelsInternal(const int32 InNumSourceModels)
 	PRAGMA_ENABLE_DEPRECATION_WARNINGS
 }
 
-void USkeletalMesh::RemoveSourceModel(const int32 InLODIndex)
-{
-	WaitUntilAsyncPropertyReleased(ESkeletalMeshAsyncProperties::SourceModels);
-	PRAGMA_DISABLE_DEPRECATION_WARNINGS
-	RemoveSourceModelInternal(InLODIndex);	
-	PRAGMA_ENABLE_DEPRECATION_WARNINGS
-}
 
-void USkeletalMesh::RemoveSourceModelInternal(const int32 InLODIndex)
+void USkeletalMesh::RemoveSourceModel(const int32 InLODIndex)
 {
+	WaitUntilAsyncPropertyReleased(ESkeletalMeshAsyncProperties::SourceModels|ESkeletalMeshAsyncProperties::ImportedModel);
 	PRAGMA_DISABLE_DEPRECATION_WARNINGS
 	if (ensure(SourceModels.IsValidIndex(InLODIndex)))
 	{
+#if WITH_EDITOR
 		FSkeletalMeshSourceModel& SourceModel = SourceModels[InLODIndex];
-
+		
 		SourceModel.ClearAllMeshData();
+		if (ImportedModel->InlineReductionCacheDatas.IsValidIndex(InLODIndex))
+		{
+			ImportedModel->InlineReductionCacheDatas.RemoveAt(InLODIndex);
+		}
+#endif
 		SourceModels.RemoveAt(InLODIndex);
 	}
 	PRAGMA_ENABLE_DEPRECATION_WARNINGS
@@ -2756,6 +2773,8 @@ FSkeletalMeshSourceModel& USkeletalMesh::GetSourceModel(const int32 InLODIndex)
 ///////////////////////////////////////////////////////////////////////
 //// Mesh Description API
 
+#if WITH_EDITOR
+
 FMeshDescription* USkeletalMesh::GetMeshDescription(const int32 InLODIndex) const
 {
 	if (!ensure(IsValidLODIndex(InLODIndex)))
@@ -4028,6 +4047,23 @@ void USkeletalMesh::ExecutePostLoadInternal(FSkinnedAssetPostLoadContext& Contex
 	// This scope allows us to use any locked properties without causing stalls
 	FSkinnedAssetAsyncBuildScope AsyncBuildScope(this);
 
+	// If loading older skeletal meshes, copy over the old LOD info into the source models, since source models are now the ground truth for LOD info.
+	if (!LODInfo_DEPRECATED.IsEmpty())
+	{
+		// Ensure we have source model storage that matches the number of LODs defined on this mesh since we're not doing any translation from
+		// old models. 
+		SetNumSourceModels(LODInfo_DEPRECATED.Num());
+
+		for (int32 LODIndex = 0; LODIndex < LODInfo_DEPRECATED.Num(); LODIndex++)
+		{
+			FSkeletalMeshLODInfo& TargetLODInfo = GetSourceModel(LODIndex);
+			TargetLODInfo = LODInfo_DEPRECATED[LODIndex];
+		}
+		
+		// Clean out the old LOD info to ensure no-one mistakenly does anything with it.
+		LODInfo_DEPRECATED.Empty();
+	}
+
 	if (!GetOutermost()->bIsCookedForEditor)
 	{
 		RemoveLegacyClothingSections();
@@ -4036,8 +4072,6 @@ void USkeletalMesh::ExecutePostLoadInternal(FSkinnedAssetPostLoadContext& Contex
 
 		PostLoadValidateUserSectionData();
 
-		PostLoadVerifyAndFixBadTangent();
-
 		// Fixup missing material slot names and import slot names, so that mesh editing 
 		// preserves material assignments.
 		IMeshUtilities* MeshUtilities = FModuleManager::Get().LoadModulePtr<IMeshUtilities>("MeshUtilities");
@@ -4048,18 +4082,14 @@ void USkeletalMesh::ExecutePostLoadInternal(FSkinnedAssetPostLoadContext& Contex
 
 		PostLoadRecoverConvertLODModelsToMeshDescription();		
 
+		PostLoadVerifyAndFixBadTangent();
+		
 		if (GetResourceForRendering() == nullptr)
 		{
 			CacheDerivedData(&Context);
 			Context.bHasCachedDerivedData = true;
 		}
 	}
-	else
-	{
-		// Ensure we have source model storage that matches the number of LODs defined on this mesh since we're not doing any translation from
-		// old models. 
-		SetNumSourceModels(GetLODNum());
-	}
 
 	// check the MinLOD values are all within range
 	FPerQualityLevelInt QualityLocalMinLOD;
@@ -5159,7 +5189,7 @@ FMatrix USkeletalMesh::GetComposedRefPoseMatrix( FName InBoneName ) const
 		{
 			USkeletalMeshSocket const* Socket = FindSocket(InBoneName);
 
-			if(Socket != NULL)
+			if(Socket != nullptr)
 			{
 				BoneIndex = GetRefSkeleton().FindBoneIndex(Socket->BoneName);
 
@@ -5368,7 +5398,7 @@ bool USkeletalMesh::IsSocketOnMesh(const FName& InSocketName) const
 	{
 		USkeletalMeshSocket* Socket = Sockets[SocketIdx];
 
-		if(Socket != NULL && Socket->SocketName == InSocketName)
+		if(Socket != nullptr && Socket->SocketName == InSocketName)
 		{
 			return true;
 		}
@@ -5554,7 +5584,7 @@ void USkeletalMesh::ClearAllCachedCookedPlatformData()
 		// We need to keep the ddc editor data LODModel for rendering; it can be different (The number of sections, the number of vertices, the number of morphtargets) because of chunking, build or reduction setting that are or will be per platform.
 		//Normally this call should be able to read values out of ddc rather than rebuilding, because the ddc for the running platform was cached when we loaded the asset.
 		ITargetPlatform* RunningPlatform = GetTargetPlatformManagerRef().GetRunningTargetPlatform();
-		check(RunningPlatform != NULL);
+		check(RunningPlatform != nullptr);
 		const FString RunningPlatformDerivedDataKey = BuildDerivedDataKey(RunningPlatform);
 		FSkeletalMeshRenderData RunningPlatformSkeletalMeshRenderData;
 		constexpr bool bIsSerializeSaving = false;
@@ -6026,7 +6056,7 @@ bool USkeletalMesh::GetPhysicsTriMeshData(FTriMeshCollisionData* CollisionData,
 
 void USkeletalMesh::AddAssetUserData( UAssetUserData* InUserData)
 {
-	if (InUserData != NULL)
+	if (InUserData != nullptr)
 	{
 		RemoveUserDataOfClass(InUserData->GetClass());
 		AssetUserData.Add(InUserData);
@@ -6039,12 +6069,12 @@ UAssetUserData* USkeletalMesh::GetAssetUserDataOfClass(TSubclassOf<UAssetUserDat
 	for (int32 DataIdx = 0; DataIdx < ArrayPtr->Num(); DataIdx++)
 	{
 		UAssetUserData* Datum = (*ArrayPtr)[DataIdx];
-		if (Datum != NULL && Datum->IsA(InUserDataClass))
+		if (Datum != nullptr && Datum->IsA(InUserDataClass))
 		{
 			return Datum;
 		}
 	}
-	return NULL;
+	return nullptr;
 }
 
 void USkeletalMesh::RemoveUserDataOfClass(TSubclassOf<UAssetUserData> InUserDataClass)
@@ -6052,7 +6082,7 @@ void USkeletalMesh::RemoveUserDataOfClass(TSubclassOf<UAssetUserData> InUserData
 	for (int32 DataIdx = 0; DataIdx < AssetUserData.Num(); DataIdx++)
 	{
 		UAssetUserData* Datum = AssetUserData[DataIdx];
-		if (Datum != NULL && Datum->IsA(InUserDataClass))
+		if (Datum != nullptr && Datum->IsA(InUserDataClass))
 		{
 			AssetUserData.RemoveAt(DataIdx);
 			return;
@@ -6063,7 +6093,7 @@ void USkeletalMesh::RemoveUserDataOfClass(TSubclassOf<UAssetUserData> InUserData
 	for (int32 DataIdx = 0; DataIdx < AssetUserDataEditorOnly.Num(); DataIdx++)
 	{
 		UAssetUserData* Datum = AssetUserDataEditorOnly[DataIdx];
-		if (Datum != NULL && Datum->IsA(InUserDataClass))
+		if (Datum != nullptr && Datum->IsA(InUserDataClass))
 		{
 			AssetUserDataEditorOnly.RemoveAt(DataIdx);
 			return;
@@ -6221,18 +6251,61 @@ class UNodeMappingContainer* USkeletalMesh::GetNodeMappingContainer(class UBluep
 
 FSkeletalMeshLODInfo* USkeletalMesh::GetLODInfo(int32 Index)
 {
-	WaitUntilAsyncPropertyReleased(ESkeletalMeshAsyncProperties::LODInfo);
-	PRAGMA_DISABLE_DEPRECATION_WARNINGS
-	return LODInfo.IsValidIndex(Index) ? &LODInfo[Index] : nullptr;  
-	PRAGMA_ENABLE_DEPRECATION_WARNINGS
+	if (IsSourceModelValid(Index))
+	{
+		return &GetSourceModel(Index);
+	}
+	return nullptr;
 }
 	
 const FSkeletalMeshLODInfo* USkeletalMesh::GetLODInfo(int32 Index) const
 {
-	WaitUntilAsyncPropertyReleased(ESkeletalMeshAsyncProperties::LODInfo, ESkinnedAssetAsyncPropertyLockType::ReadOnly);
-	PRAGMA_DISABLE_DEPRECATION_WARNINGS
-	return LODInfo.IsValidIndex(Index) ? &LODInfo[Index] : nullptr;
-	PRAGMA_ENABLE_DEPRECATION_WARNINGS
+	if (IsSourceModelValid(Index))
+	{
+		return &GetSourceModel(Index);
+	}
+	return nullptr;
+}
+
+
+void USkeletalMesh::CopyAutogeneratedLODInfoFromPreviousLOD(
+	const int32 InLODIndex
+	)
+{
+	// The first LOD cannot be auto-generated.
+	if (InLODIndex == 0 || !ensure(IsValidLODIndex(InLODIndex)))
+	{
+		return;
+	}
+
+	FSkeletalMeshLODInfo& LODInfo = *GetLODInfo(InLODIndex);
+	
+	// 
+	const int32 PreviousLODIndex = InLODIndex - 1;
+	const FSkeletalMeshLODInfo& PreviousLODInfo = *GetLODInfo(PreviousLODIndex);
+	LODInfo.ScreenSize.Default = PreviousLODInfo.ScreenSize.Default * 0.5f;
+	LODInfo.LODHysteresis = PreviousLODInfo.LODHysteresis;
+	LODInfo.BakePose = PreviousLODInfo.BakePose;
+	LODInfo.BakePoseOverride = PreviousLODInfo.BakePoseOverride;
+	LODInfo.BonesToRemove = PreviousLODInfo.BonesToRemove;
+	LODInfo.BonesToPrioritize = PreviousLODInfo.BonesToPrioritize;
+	LODInfo.SectionsToPrioritize = PreviousLODInfo.SectionsToPrioritize;
+	// now find reduction setting
+	for (int32 SubLODIndex = PreviousLODIndex; SubLODIndex >= 0; --SubLODIndex)
+	{
+		const FSkeletalMeshLODInfo& SubLODInfo = *GetLODInfo(SubLODIndex);
+		if (SubLODInfo.bHasBeenSimplified)
+		{
+			// copy from previous index of LOD info reduction setting
+			// this may not match with previous copy - as we're only looking for simplified version
+			LODInfo.ReductionSettings = SubLODInfo.ReductionSettings;
+			// and make it 50 % of that
+			LODInfo.ReductionSettings.NumOfTrianglesPercentage = FMath::Clamp(LODInfo.ReductionSettings.NumOfTrianglesPercentage * 0.5f, 0.f, 1.f);
+			// increase maxdeviation, 1.5 is random number
+			LODInfo.ReductionSettings.MaxDeviationPercentage = FMath::Clamp(LODInfo.ReductionSettings.MaxDeviationPercentage * 1.5f, 0.f, 1.f);
+			break;
+		}
+	}
 }
 
 #if WITH_EDITOR
@@ -6276,20 +6349,14 @@ const USkeletalMeshLODSettings* USkeletalMesh::GetDefaultLODSetting() const
 
 bool USkeletalMesh::IsValidLODIndex(int32 Index) const
 {
-	WaitUntilAsyncPropertyReleased(ESkeletalMeshAsyncProperties::LODInfo, ESkinnedAssetAsyncPropertyLockType::ReadOnly);
-	PRAGMA_DISABLE_DEPRECATION_WARNINGS
-	return LODInfo.IsValidIndex(Index);
-	PRAGMA_ENABLE_DEPRECATION_WARNINGS
+	return IsSourceModelValid(Index);
 }
 /* 
 	* Returns total number of LOD. USkinnedAsset interface.
 	*/
 int32 USkeletalMesh::GetLODNum() const
 {
-	WaitUntilAsyncPropertyReleased(ESkeletalMeshAsyncProperties::LODInfo, ESkinnedAssetAsyncPropertyLockType::ReadOnly);
-	PRAGMA_DISABLE_DEPRECATION_WARNINGS
-	return LODInfo.Num();
-	PRAGMA_ENABLE_DEPRECATION_WARNINGS
+	return GetNumSourceModels();
 }
 
 int32 USkeletalMesh::GetNumNaniteVertices() const
@@ -6393,92 +6460,45 @@ TArray<FString> USkeletalMesh::K2_GetAllSkinWeightProfileNames() const
 
 FSkeletalMeshLODInfo& USkeletalMesh::AddLODInfo()
 {
-	const int32 NewIndex = AddLODInfo({});
-
-	const USkeletalMeshLODSettings* DefaultSetting = GetDefaultLODSetting();
-	// if failed to get setting, that means, we don't have proper setting 
-	// in that case, use last index setting
-	FSkeletalMeshLODInfo& NewLODInfo = *GetLODInfo(NewIndex);
-	if (!DefaultSetting->SetLODSettingsToMesh(this, NewIndex))
-	{
-		if (NewIndex > 0)
-		{
-			// copy previous copy
-			const int32 LastIndex = NewIndex - 1;
-			const FSkeletalMeshLODInfo& LastLODInfo = *GetLODInfo(LastIndex);
-			NewLODInfo.ScreenSize.Default = LastLODInfo.ScreenSize.Default * 0.5f;
-			NewLODInfo.LODHysteresis = LastLODInfo.LODHysteresis;
-			NewLODInfo.BakePose = LastLODInfo.BakePose;
-			NewLODInfo.BakePoseOverride = LastLODInfo.BakePoseOverride;
-			NewLODInfo.BonesToRemove = LastLODInfo.BonesToRemove;
-			NewLODInfo.BonesToPrioritize = LastLODInfo.BonesToPrioritize;
-			NewLODInfo.SectionsToPrioritize = LastLODInfo.SectionsToPrioritize;
-			// now find reduction setting
-			for (int32 SubLOD = LastIndex; SubLOD >= 0; --SubLOD)
-			{
-				const FSkeletalMeshLODInfo& SubLODInfo = *GetLODInfo(SubLOD);
-				if (SubLODInfo.bHasBeenSimplified)
-				{
-					// copy from previous index of LOD info reduction setting
-					// this may not match with previous copy - as we're only looking for simplified version
-					NewLODInfo.ReductionSettings = SubLODInfo.ReductionSettings;
-					// and make it 50 % of that
-					NewLODInfo.ReductionSettings.NumOfTrianglesPercentage = FMath::Clamp(NewLODInfo.ReductionSettings.NumOfTrianglesPercentage * 0.5f, 0.f, 1.f);
-					// increase maxdeviation, 1.5 is random number
-					NewLODInfo.ReductionSettings.MaxDeviationPercentage = FMath::Clamp(NewLODInfo.ReductionSettings.MaxDeviationPercentage * 1.5f, 0.f, 1.f);
-					break;
-				}
-			}
-
-		}
-		// if this is the first LOD, then just use default setting of the struct
-	}
-
-	return NewLODInfo;
+	const bool bAutogenerated = GetNumSourceModels() != 0;
+	return AddSourceModel(bAutogenerated);
 }
 
 int32 USkeletalMesh::AddLODInfo(const FSkeletalMeshLODInfo& NewLODInfo) 
 {
-	WaitUntilAsyncPropertyReleased(ESkeletalMeshAsyncProperties::LODInfo|ESkeletalMeshAsyncProperties::SourceModels);
-	PRAGMA_DISABLE_DEPRECATION_WARNINGS
-	int32 NewIndex = LODInfo.Add(NewLODInfo);
-#if WITH_EDITOR
-	SetNumSourceModelsInternal(LODInfo.Num());
-#endif
-	PRAGMA_ENABLE_DEPRECATION_WARNINGS
-	
+	constexpr bool bAutogenerated = false; 
+	const int32 NewIndex = GetNumSourceModels();
+	FSkeletalMeshSourceModel& Model = AddSourceModel(bAutogenerated);
+	static_cast<FSkeletalMeshLODInfo&>(Model) = NewLODInfo;
 	return NewIndex;
 }
 
 void USkeletalMesh::RemoveLODInfo(int32 Index)
 {
-	WaitUntilAsyncPropertyReleased(ESkeletalMeshAsyncProperties::LODInfo|ESkeletalMeshAsyncProperties::SourceModels|ESkeletalMeshAsyncProperties::ImportedModel);
-	
-	PRAGMA_DISABLE_DEPRECATION_WARNINGS
-	if (LODInfo.IsValidIndex(Index))
+	RemoveSourceModel(Index);
+}
+
+void USkeletalMesh::ResetLODInfo()
+{
+	SetNumSourceModels(0);
+}
+
+TArray<FSkeletalMeshLODInfo>& USkeletalMesh::GetLODInfoArray()
+{
+	thread_local TArray<FSkeletalMeshLODInfo> LODInfoArray;
+
+	LODInfoArray.SetNumUninitialized(GetNumSourceModels());
+	for (const FSkeletalMeshSourceModel& SourceModel: GetAllSourceModels())
 	{
-#if WITH_EDITOR
-		RemoveSourceModelInternal(Index);
-		
-		if (ImportedModel->InlineReductionCacheDatas.IsValidIndex(Index))
-		{
-			ImportedModel->InlineReductionCacheDatas.RemoveAt(Index);
-		}
-#endif // WITH_EDITOR
-		LODInfo.RemoveAt(Index);
+		LODInfoArray.Add(SourceModel);
 	}
-	PRAGMA_ENABLE_DEPRECATION_WARNINGS
+	return LODInfoArray;
 }
 
-void USkeletalMesh::ResetLODInfo()
+const TArray<FSkeletalMeshLODInfo>& USkeletalMesh::GetLODInfoArray() const
 {
-	WaitUntilAsyncPropertyReleased(ESkeletalMeshAsyncProperties::LODInfo|ESkeletalMeshAsyncProperties::SourceModels);
 	PRAGMA_DISABLE_DEPRECATION_WARNINGS
-	
-	LODInfo.Reset();
-#if WITH_EDITOR
-	SetNumSourceModelsInternal(0);
-#endif
+	return const_cast<USkeletalMesh*>(this)->GetLODInfoArray();
 	PRAGMA_ENABLE_DEPRECATION_WARNINGS
 }
 
@@ -6996,7 +7016,7 @@ bool USkeletalMeshSocket::AttachActor(AActor* Actor, class USkeletalMeshComponen
 	#if WITH_EDITOR
 				if (GIsEditor)
 				{
-					Actor->PreEditChange(NULL);
+					Actor->PreEditChange(nullptr);
 					Actor->PostEditChange();
 				}
 	#endif // WITH_EDITOR
diff --git a/Engine/Source/Runtime/Engine/Private/SkeletalMeshMerge.cpp b/Engine/Source/Runtime/Engine/Private/SkeletalMeshMerge.cpp
--- a/Engine/Source/Runtime/Engine/Private/SkeletalMeshMerge.cpp
+++ b/Engine/Source/Runtime/Engine/Private/SkeletalMeshMerge.cpp
@@ -980,7 +980,7 @@ void FSkeletalMeshMerge::ReleaseResources(int32 Slack)
 		Resource->LODRenderData.Empty(Slack);
 	}
 
-	MergeMesh->ResetLODInfo();
+	MergeMesh->SetNumSourceModels(0);
 	MergeMesh->GetMaterials().Empty();
 }
 
diff --git a/Engine/Source/Runtime/Engine/Private/SkeletalMeshSourceModel.cpp b/Engine/Source/Runtime/Engine/Private/SkeletalMeshSourceModel.cpp
--- a/Engine/Source/Runtime/Engine/Private/SkeletalMeshSourceModel.cpp
+++ b/Engine/Source/Runtime/Engine/Private/SkeletalMeshSourceModel.cpp
@@ -85,6 +85,56 @@ void FSkeletalMeshSourceModel::Initialize(USkeletalMesh* InOwner)
 #endif
 }
 
+
+TConstArrayView<FName> FSkeletalMeshSourceModel::GetSkinWeightProfileNames() const
+{
+	if (CachedSkinWeightProfileNames.Num() == 1 && CachedSkinWeightProfileNames[0] == NAME_None)
+	{
+#if WITH_EDITOR
+		const_cast<FSkeletalMeshSourceModel*>(this)->UpdateCachedMeshStatisticsFromBulkIfNeeded();
+#else
+		static TArray<FName> EmptyArray;
+		return EmptyArray; 
+#endif
+	}
+	return CachedSkinWeightProfileNames;
+}
+
+TConstArrayView<FName> FSkeletalMeshSourceModel::GetMorphTargetNames() const
+{
+	if (CachedMorphTargetNames.Num() == 1 && CachedMorphTargetNames[0] == NAME_None)
+	{
+#if WITH_EDITOR
+		const_cast<FSkeletalMeshSourceModel*>(this)->UpdateCachedMeshStatisticsFromBulkIfNeeded();
+#else
+		static TArray<FName> EmptyArray;
+		return EmptyArray; 
+#endif
+	}
+	return CachedMorphTargetNames;
+}
+
+TOptional<FGuid> FSkeletalMeshSourceModel::GetVertexColorGuid() const
+{
+#if WITH_EDITOR
+	if (!CachedHasVertexColors.IsSet())
+	{
+		const_cast<FSkeletalMeshSourceModel*>(this)->UpdateCachedMeshStatisticsFromBulkIfNeeded();
+	}
+#endif
+
+	if (CachedHasVertexColors.Get(false))
+	{
+		return CachedVertexColorGuid;
+	}
+	
+	return {};
+}
+
+
+////////////////////////////////////////////////////////////////////////////////
+///// EDITOR ONLY FUNCTIONS
+
 #if WITH_EDITOR
 bool FSkeletalMeshSourceModel::HasMeshDescription() const
 {
@@ -232,39 +282,6 @@ const FMeshDescriptionBulkData* FSkeletalMeshSourceModel::GetMeshDescriptionBulk
 }
 
 
-TConstArrayView<FName> FSkeletalMeshSourceModel::GetSkinWeightProfileNames() const
-{
-	if (CachedSkinWeightProfileNames.Num() == 1 && CachedSkinWeightProfileNames[0] == NAME_None)
-	{
-		const_cast<FSkeletalMeshSourceModel*>(this)->UpdateCachedMeshStatisticsFromBulkIfNeeded();
-	}
-	return CachedSkinWeightProfileNames;
-}
-
-TConstArrayView<FName> FSkeletalMeshSourceModel::GetMorphTargetNames() const
-{
-	if (CachedMorphTargetNames.Num() == 1 && CachedMorphTargetNames[0] == NAME_None)
-	{
-		const_cast<FSkeletalMeshSourceModel*>(this)->UpdateCachedMeshStatisticsFromBulkIfNeeded();
-	}
-	return CachedMorphTargetNames;
-}
-
-TOptional<FGuid> FSkeletalMeshSourceModel::GetVertexColorGuid() const
-{
-	if (!CachedHasVertexColors.IsSet())
-	{
-		const_cast<FSkeletalMeshSourceModel*>(this)->UpdateCachedMeshStatisticsFromBulkIfNeeded();
-	}
-
-	if (CachedHasVertexColors.Get(false))
-	{
-		return CachedVertexColorGuid;
-	}
-	
-	return {};
-}
-
 USkeletalMesh* FSkeletalMeshSourceModel::GetOwner() const
 {
 	if (!ensure(MeshDescriptionBulkData))
diff --git a/Engine/Source/Runtime/Experimental/GeometryCollectionEngine/Private/GeometryCollection/GeometryCollectionEngineConversion.cpp b/Engine/Source/Runtime/Experimental/GeometryCollectionEngine/Private/GeometryCollection/GeometryCollectionEngineConversion.cpp
--- a/Engine/Source/Runtime/Experimental/GeometryCollectionEngine/Private/GeometryCollection/GeometryCollectionEngineConversion.cpp
+++ b/Engine/Source/Runtime/Experimental/GeometryCollectionEngine/Private/GeometryCollection/GeometryCollectionEngineConversion.cpp
@@ -1944,9 +1944,8 @@ void FGeometryCollectionEngineConversion::ConvertCollectionToSkeletalMesh(const
 	const TManagedArray<FTransform3f>* Transforms = TransformFacade.FindTransforms();
 
 	// Reset the skeletal mesh
-	OutSkeletalMesh.ClearAllMeshDescriptions();
+	OutSkeletalMesh.SetNumSourceModels(0);
 	OutSkeletalMesh.GetImportedModel()->LODModels.Reset();
-	OutSkeletalMesh.ResetLODInfo();
 	OutSkeletalMesh.ReleaseResources();
 
 	// first extract the skeleton 
