[Audio] Tweak ForEachSubsystem to allow either Interface or UClass as a parameter

#rnx

diff --git a/Engine/Source/Runtime/Engine/Public/Subsystems/AudioEngineSubsystem.h b/Engine/Source/Runtime/Engine/Public/Subsystems/AudioEngineSubsystem.h
--- a/Engine/Source/Runtime/Engine/Public/Subsystems/AudioEngineSubsystem.h
+++ b/Engine/Source/Runtime/Engine/Public/Subsystems/AudioEngineSubsystem.h
@@ -82,13 +82,22 @@ class FAudioSubsystemCollection : public FSubsystemCollection<UAudioEngineSubsys
 {
 public:
 
-	template<class InterfaceToCastTo>
-	void ForEachSubsystem(TFunctionRef<bool(InterfaceToCastTo*)> InFunction) const
+	template<class ArgType>
+	void ForEachSubsystem(TFunctionRef<bool(ArgType*)> InFunction) const
 	{
-		UClass* InterfaceClass = InterfaceToCastTo::UClassType::StaticClass();
-		ForEachSubsystemOfClass(InterfaceClass, [Function=MoveTemp(InFunction)](USubsystem* Subsystem)
+		UClass* ArgClass = nullptr;
+		if constexpr (TIsIInterface<ArgType>::Value)
 		{
-			Function(CastChecked<InterfaceToCastTo>(Subsystem));
+			ArgClass = ArgType::UClassType::StaticClass();
+		}
+		else
+		{
+			ArgClass = ArgType::StaticClass();
+		}
+
+		ForEachSubsystemOfClass(ArgClass, [Function=MoveTemp(InFunction)](USubsystem* Subsystem)
+		{
+			Function(CastChecked<ArgType>(Subsystem));
 		});
 	}
 };
