Avoid provoking a full flush of the async pipeline.

diff --git a/Engine/Source/Runtime/Engine/Private/WorldPartition/WorldPartitionHandle.cpp b/Engine/Source/Runtime/Engine/Private/WorldPartition/WorldPartitionHandle.cpp
--- a/Engine/Source/Runtime/Engine/Private/WorldPartition/WorldPartitionHandle.cpp
+++ b/Engine/Source/Runtime/Engine/Private/WorldPartition/WorldPartitionHandle.cpp
@@ -167,7 +167,12 @@ FWorldPartitionLoadingContext::FDeferred::~FDeferred()
 						AllAsyncLoadIds.Add(AsyncLoadId);
 					}
 				}
-				FlushAsyncLoading(AllAsyncLoadIds);
+				
+				//Only flush if there are items in the array. An empty array will result in a complete flush which is not the intended result.
+				if(AllAsyncLoadIds.Num() > 0)
+				{
+					FlushAsyncLoading(AllAsyncLoadIds);
+				}
 
 				ULevel* Level = nullptr;
 				for (FWorldPartitionActorDescInstance* ActorDescInstance : SourceList)
