SplineComponent: Disable rainbow splines

The rainbow code path still exists but is compiled out for now. Will revisit optional rainbows later.

diff --git a/Engine/Source/Runtime/Engine/Private/Components/SplineComponent.cpp b/Engine/Source/Runtime/Engine/Private/Components/SplineComponent.cpp
--- a/Engine/Source/Runtime/Engine/Private/Components/SplineComponent.cpp
+++ b/Engine/Source/Runtime/Engine/Private/Components/SplineComponent.cpp
@@ -2419,10 +2419,14 @@ void USplineComponent::Draw(FPrimitiveDrawInterface* PDI, const FSceneView* View
 	const int32 GrabHandleSize = 6;
 	FVector OldKeyPos(0);
 
-	auto GetLineColor = [](int32 KeyIdx, int32 NumSegments, int32 NumSteps) -> FLinearColor
+	auto GetLineColor = [LineColor](int32 KeyIdx, int32 NumSegments, int32 NumSteps) -> FLinearColor
 	{
+#if 1
+		return LineColor;
+#else
 		float Normalized = (KeyIdx + (NumSegments * NumSteps)) / static_cast<float>(NumSegments * NumSteps);
 		return FLinearColor::MakeFromHSV8(static_cast<uint8>(Normalized * 255), 255, 255);
+#endif
 	};
 	
 	const int32 NumPoints = SplineInfo.Points.Num();
