Dstm: Add physics substepping for simulation. Adjust behavior of pushing/shooting objects, so the force is consistent across different fps.

diff --git a/Engine/Source/Runtime/Engine/Classes/GameFramework/CharacterMovementComponent.h b/Engine/Source/Runtime/Engine/Classes/GameFramework/CharacterMovementComponent.h
--- a/Engine/Source/Runtime/Engine/Classes/GameFramework/CharacterMovementComponent.h
+++ b/Engine/Source/Runtime/Engine/Classes/GameFramework/CharacterMovementComponent.h
@@ -253,6 +253,14 @@ public:
 	UPROPERTY(Category="Character Movement: Walking", EditAnywhere, BlueprintReadWrite, meta=(ClampMin="0", UIMin="0"))
 	float GroundFriction;
 
+#if UE_WITH_REMOTE_OBJECT_HANDLE
+	/**
+	* An extra scaling factor on the impact physics force applied to pushing object when the component is doing substepping 
+	*/
+	float PhysicsForceSubsteppingFactor = 1.0f;
+#endif
+	
+
 	/** Saved location of object we are standing on, for UpdateBasedMovement() to determine if base moved in the last frame, and therefore pawn needs an update. */
 	FQuat OldBaseQuat;
 
diff --git a/Engine/Source/Runtime/Engine/Private/Components/CharacterMovementComponent.cpp b/Engine/Source/Runtime/Engine/Private/Components/CharacterMovementComponent.cpp
--- a/Engine/Source/Runtime/Engine/Private/Components/CharacterMovementComponent.cpp
+++ b/Engine/Source/Runtime/Engine/Private/Components/CharacterMovementComponent.cpp
@@ -5593,6 +5593,12 @@ void UCharacterMovementComponent::PhysWalking(float deltaTime, int32 Iterations)
 		const float timeTick = GetSimulationTimeStep(remainingTime, Iterations);
 		remainingTime -= timeTick;
 
+#if UE_WITH_REMOTE_OBJECT_HANDLE
+		//Scale down impact force if CharacterMoveComponent is taking multiple substeps.
+		const float LastFrameDt = GetWorld()->GetDeltaSeconds();
+		PhysicsForceSubsteppingFactor = timeTick / LastFrameDt;
+#endif
+		
 		// Save current values
 		UPrimitiveComponent * const OldBase = GetMovementBase();
 		const FVector PreviousBaseLocation = (OldBase != NULL) ? OldBase->GetComponentLocation() : FVector::ZeroVector;
@@ -7762,6 +7768,10 @@ void UCharacterMovementComponent::ApplyImpactPhysicsForces(const FHitResult& Imp
 					PushForceModificator *= BodyMass;
 				}
 				const bool bClientAuthScheduledPushForces = CharacterMovementCVars::bEnableClientAuthScheduledPushForces && Chaos::FPhysicsSolverBase::IsNetworkPhysicsPredictionEnabled();
+	
+#if UE_WITH_REMOTE_OBJECT_HANDLE
+				PushForceModificator *= PhysicsForceSubsteppingFactor;
+#endif
 
 				Force *= PushForceModificator;
 				const float ZeroVelocityTolerance = 1.0f;
