Quick fix for mistaken convention in edgraphnode clipboard representation. I had to move the strings into the class if I wanted to keep the flags private, then accidently replaced the member name with the string name

#rnx

diff --git a/Engine/Source/Runtime/Engine/Private/EdGraph/EdGraphPin.cpp b/Engine/Source/Runtime/Engine/Private/EdGraph/EdGraphPin.cpp
--- a/Engine/Source/Runtime/Engine/Private/EdGraph/EdGraphPin.cpp
+++ b/Engine/Source/Runtime/Engine/Private/EdGraph/EdGraphPin.cpp
@@ -2392,8 +2392,8 @@ bool UEdGraphPin::ImportText_PinArray(const TCHAR*& Buffer, TArray<UEdGraphPin*>
 }
 
 #if WITH_EDITORONLY_DATA
-FString UEdGraphPin::bHasSnappedChildName = GET_MEMBER_NAME_STRING_CHECKED(UEdGraphPin, bHasSnappedChildName);
-FString UEdGraphPin::bHasSnappedParentName = GET_MEMBER_NAME_STRING_CHECKED(UEdGraphPin, bHasSnappedParentName);
+FString UEdGraphPin::bHasSnappedChildName = GET_MEMBER_NAME_STRING_CHECKED(UEdGraphPin, bHasSnappedChild);
+FString UEdGraphPin::bHasSnappedParentName = GET_MEMBER_NAME_STRING_CHECKED(UEdGraphPin, bHasSnappedParent);
 #endif
 
 int32 GCVarDisableOrphanPins = 0;
