Removed IsAsyncLoadingThread check from UActorComponent::PostInitProperties.

Some systems, like the chaos visual debugger, create components in a worker thread to prime its object pool.

diff --git a/Engine/Source/Runtime/Engine/Private/Components/ActorComponent.cpp b/Engine/Source/Runtime/Engine/Private/Components/ActorComponent.cpp
--- a/Engine/Source/Runtime/Engine/Private/Components/ActorComponent.cpp
+++ b/Engine/Source/Runtime/Engine/Private/Components/ActorComponent.cpp
@@ -604,10 +604,6 @@ void UActorComponent::PostInitProperties()
 			Storage->AddCompatibleObject(this);
 		}
 	}
-	else
-	{
-		check(IsInAsyncLoadingThread());
-	}
 #endif	// WITH_EDITOR
 
 	// Instance components will be added during the owner's initialization
