[Audio] Re-enable SetOwner notification, replaced with a Delegate instead of IActiveSoundUpdateInterface for thread safety.

#rnx

diff --git a/Engine/Source/Runtime/Engine/Private/AudioDevice.cpp b/Engine/Source/Runtime/Engine/Private/AudioDevice.cpp
--- a/Engine/Source/Runtime/Engine/Private/AudioDevice.cpp
+++ b/Engine/Source/Runtime/Engine/Private/AudioDevice.cpp
@@ -133,7 +133,7 @@ FAutoConsoleVariableRef CVarDisableLegacyReverb(
 	TEXT("0: Enabled, 1: Disabled"),
 	ECVF_Default);
 
-static int32 DisableOwnerNotify = 1;
+static int32 DisableOwnerNotify = 0;
 FAutoConsoleVariableRef CVarDisableOwnerNotify(
 	TEXT("au.DisableOwnerNotify"),
 	DisableOwnerNotify,
@@ -3785,11 +3785,10 @@ void FAudioDevice::NotifySubsystemsVirtualizedSoundDeleting(FActiveSound& Active
 
 void FAudioDevice::NotifySubsystemsSetOwner(FActiveSound& ActiveSound, const UObject* Owner) const
 {
-	SubsystemCollection.ForEachSubsystem<IActiveSoundUpdateInterface>([&ActiveSound, Owner](IActiveSoundUpdateInterface* ActiveSoundUpdate)
-	{
-		ActiveSoundUpdate->NotifyActiveSoundOwner(ActiveSound, Owner);
-		return true;
-	});
+	// DEPRECATED in 5.7
+	// Calling SubsystemCollection.ForEachSubsystem here is not safe because
+	// it is on the Game Thread while all other calls are on the Audio Thread.
+	// There would be risk of concurrent calls, which will corrupt SubsystemCollection bIsIterating.
 }
 
 void FAudioDevice::SetBaseSoundMix(USoundMix* NewMix)
@@ -6808,7 +6807,7 @@ void FAudioDevice::PlaySoundAtLocation(USoundBase* Sound, UWorld* World, float V
 		NewActiveSound.SetOwner(OwningActor);
 		if (!DisableOwnerNotify)
 		{
-			NotifySubsystemsSetOwner(NewActiveSound, OwningActor);
+			OnAudioDeviceSetSoundOwner.Broadcast(NewActiveSound, OwningActor);
 		}
 
 #if UE_AUDIO_PROFILERTRACE_ENABLED
diff --git a/Engine/Source/Runtime/Engine/Public/AudioDevice.h b/Engine/Source/Runtime/Engine/Public/AudioDevice.h
--- a/Engine/Source/Runtime/Engine/Public/AudioDevice.h
+++ b/Engine/Source/Runtime/Engine/Public/AudioDevice.h
@@ -417,6 +417,8 @@ struct FAudioDeviceRenderInfo
 DECLARE_MULTICAST_DELEGATE_OneParam(FOnAudioDevicePreRender, const FAudioDeviceRenderInfo&);
 DECLARE_MULTICAST_DELEGATE_OneParam(FOnAudioDevicePostRender, const FAudioDeviceRenderInfo&);
 
+DECLARE_MULTICAST_DELEGATE_TwoParams(FOnAudioDeviceSetSoundOwner, const FActiveSound&, const UObject*);
+
 class FAudioDevice : public FExec
 {
 public:
@@ -904,13 +906,14 @@ protected:
 		NotifySubsystemsActiveSoundDeleting(ActiveSound);
 	}
 
-	/**
-	 * Notifies subsystems an owning object has been associated with a sound.
-	 */
+	UE_DEPRECATED(5.7, "NotifySubsystemsSetOwner is deprecated. Use OnAudioDeviceSetSoundOwner instead.")
 	ENGINE_API void NotifySubsystemsSetOwner(FActiveSound& ActiveSound, const UObject* Owner) const;
 
 public:
 
+	// Broadcast on the Game Thread when an owning object has been associated with a sound.
+	FOnAudioDeviceSetSoundOwner OnAudioDeviceSetSoundOwner;
+
 	/**
 	 * Gets the default reverb and interior settings for the provided world.  Returns true if settings with WorldID were located
 	 */
