Substrate - added ok/cancel choice when enabling disabling Substrate.

#rnx
#tests PC editor

diff --git a/Engine/Source/Runtime/Engine/Private/RendererSettings.cpp b/Engine/Source/Runtime/Engine/Private/RendererSettings.cpp
--- a/Engine/Source/Runtime/Engine/Private/RendererSettings.cpp
+++ b/Engine/Source/Runtime/Engine/Private/RendererSettings.cpp
@@ -367,11 +367,19 @@ void URendererSettings::PostEditChangeProperty(FPropertyChangedEvent& PropertyCh
 		{
 			if (bEnableSubstrate)
 			{
-				FMessageDialog::Open(EAppMsgType::Ok, LOCTEXT("Substrate Beta", "Warning: Substrate is in a Beta state. Be aware that any materials saved when Substrate is enabled won't be rendered correctly if Substrate is disabled later on."));
+				const EAppReturnType::Type Choice = FMessageDialog::Open(EAppMsgType::OkCancel, LOCTEXT("Substrate Beta", "Warning: Substrate is in a Beta state. Be aware that any materials saved when Substrate is enabled won't be rendered correctly if Substrate is disabled later on."));
+				if (Choice == EAppReturnType::Cancel)
+				{
+					bEnableSubstrate = false;
+				}
 			}
 			else
 			{
-				FMessageDialog::Open(EAppMsgType::Ok, LOCTEXT("Substrate", "Warning: Disabling Substrate will cause all materials authored with Substrate nodes to appear incorrect, as they are not backward compatible with the legacy root node and GBuffer format."));
+				const EAppReturnType::Type Choice = FMessageDialog::Open(EAppMsgType::OkCancel, LOCTEXT("Substrate", "Warning: Disabling Substrate will cause all materials authored with Substrate nodes to appear incorrect, as they are not backward compatible with the legacy root node and GBuffer format."));
+				if (Choice == EAppReturnType::Cancel)
+				{
+					bEnableSubstrate = true;
+				}
 			}
 		}
 
