Added generic AddElem funciton for AggGeom and helpers to access the correct containers based on input element type

diff --git a/Engine/Source/Runtime/Engine/Classes/PhysicsEngine/AggregateGeom.h b/Engine/Source/Runtime/Engine/Classes/PhysicsEngine/AggregateGeom.h
--- a/Engine/Source/Runtime/Engine/Classes/PhysicsEngine/AggregateGeom.h
+++ b/Engine/Source/Runtime/Engine/Classes/PhysicsEngine/AggregateGeom.h
@@ -78,6 +78,84 @@ struct FKAggregateGeom
 
 	SIZE_T GetAllocatedSize() const { return SphereElems.GetAllocatedSize() + SphylElems.GetAllocatedSize() + BoxElems.GetAllocatedSize() + ConvexElems.GetAllocatedSize() + TaperedCapsuleElems.GetAllocatedSize() + LevelSetElems.GetAllocatedSize() + SkinnedLevelSetElems.GetAllocatedSize() + MLLevelSetElems.GetAllocatedSize() + SkinnedTriangleMeshElems.GetAllocatedSize(); }
 
+	template <typename Callable>
+	auto VisitShapeAndContainer(FKShapeElem& InElement, Callable&& InCallable)
+	{
+		switch(InElement.GetShapeType())
+		{
+		case EAggCollisionShape::Sphere:
+			return InCallable(static_cast<FKSphereElem&>(InElement), SphereElems);
+		case EAggCollisionShape::Box:
+			return InCallable(static_cast<FKBoxElem&>(InElement), BoxElems);
+		case EAggCollisionShape::Sphyl:
+			return InCallable(static_cast<FKSphylElem&>(InElement), SphylElems);
+		case EAggCollisionShape::Convex:
+			return InCallable(static_cast<FKConvexElem&>(InElement), ConvexElems);
+		case EAggCollisionShape::TaperedCapsule:
+			return InCallable(static_cast<FKTaperedCapsuleElem&>(InElement), TaperedCapsuleElems);
+		case EAggCollisionShape::LevelSet:
+			return InCallable(static_cast<FKLevelSetElem&>(InElement), LevelSetElems);
+		case EAggCollisionShape::SkinnedLevelSet:
+			return InCallable(static_cast<FKSkinnedLevelSetElem&>(InElement), SkinnedLevelSetElems);
+		case EAggCollisionShape::MLLevelSet:
+			return InCallable(static_cast<FKMLLevelSetElem&>(InElement), MLLevelSetElems);
+		case EAggCollisionShape::SkinnedTriangleMesh:
+			return InCallable(static_cast<FKSkinnedTriangleMeshElem&>(InElement), SkinnedTriangleMeshElems);
+		default:
+			check(false);
+		}
+
+		using RetType = TInvokeResult_T<Callable, FKSphereElem&, TArray<FKSphereElem>&>;
+		if constexpr(!std::is_same_v<RetType, void>)
+		{
+			return RetType{};
+		}
+	}
+
+	template <typename Callable>
+	auto VisitShapeAndContainer(const FKShapeElem& InElement, Callable&& InCallable)
+	{
+		switch(InElement.GetShapeType())
+		{
+		case EAggCollisionShape::Sphere:
+			return InCallable(static_cast<const FKSphereElem&>(InElement), SphereElems);
+		case EAggCollisionShape::Box:
+			return InCallable(static_cast<const FKBoxElem&>(InElement), BoxElems);
+		case EAggCollisionShape::Sphyl:
+			return InCallable(static_cast<const FKSphylElem&>(InElement), SphylElems);
+		case EAggCollisionShape::Convex:
+			return InCallable(static_cast<const FKConvexElem&>(InElement), ConvexElems);
+		case EAggCollisionShape::TaperedCapsule:
+			return InCallable(static_cast<const FKTaperedCapsuleElem&>(InElement), TaperedCapsuleElems);
+		case EAggCollisionShape::LevelSet:
+			return InCallable(static_cast<const FKLevelSetElem&>(InElement), LevelSetElems);
+		case EAggCollisionShape::SkinnedLevelSet:
+			return InCallable(static_cast<const FKSkinnedLevelSetElem&>(InElement), SkinnedLevelSetElems);
+		case EAggCollisionShape::MLLevelSet:
+			return InCallable(static_cast<const FKMLLevelSetElem&>(InElement), MLLevelSetElems);
+		case EAggCollisionShape::SkinnedTriangleMesh:
+			return InCallable(static_cast<const FKSkinnedTriangleMeshElem&>(InElement), SkinnedTriangleMeshElems);
+		default:
+			check(false);
+		}
+
+		using RetType = TInvokeResult_T<Callable, const FKSphereElem&, TArray<FKSphereElem>&>;
+		if constexpr(!std::is_same_v<RetType, void>)
+		{
+			return RetType{};
+		}
+	}
+
+	// Add a typed element to the appropriate container
+	template <typename T>
+	void AddElement(const T& Elem) requires (TIsDerivedFrom<T, FKShapeElem>::Value || std::is_same_v<T, FKShapeElem>)
+	{
+		VisitShapeAndContainer(Elem, [] <typename ElemType> (const ElemType& TypedElem, TArray<ElemType>& Container)
+		{
+			Container.Add(TypedElem);
+		});
+	}
+
 	FKShapeElem* GetElement(const EAggCollisionShape::Type Type, const int32 Index)
 	{
 		switch (Type)
