Substrate - fixed bad bsdf features reported in editor.

#rnx
#tests PC editor

diff --git a/Engine/Source/Runtime/Engine/Private/Materials/MaterialExpressionSubstrate.cpp b/Engine/Source/Runtime/Engine/Private/Materials/MaterialExpressionSubstrate.cpp
--- a/Engine/Source/Runtime/Engine/Private/Materials/MaterialExpressionSubstrate.cpp
+++ b/Engine/Source/Runtime/Engine/Private/Materials/MaterialExpressionSubstrate.cpp
@@ -588,6 +588,26 @@ void UMaterialExpressionSubstrateShadingModels::GatherSubstrateMaterialInfo(FSub
 	}
 }
 
+static void AppendLegacyShadingModelToBSDFFeature(const FMaterialShadingModelField& ShadingModels, const bool bAnisotropyEnabled, ESubstrateBsdfFeature& BsdfFeatures)
+{
+	if (ShadingModels.HasShadingModel(MSM_Eye))			{ BsdfFeatures |= ESubstrateBsdfFeature::Eye; }
+	if (ShadingModels.HasShadingModel(MSM_Hair))		{ BsdfFeatures |= ESubstrateBsdfFeature::Hair; }
+	if (ShadingModels.HasShadingModel(MSM_Cloth))		{ BsdfFeatures |= ESubstrateBsdfFeature::Fuzz; }
+	if (ShadingModels.HasShadingModel(MSM_ClearCoat))	{ BsdfFeatures |= ESubstrateBsdfFeature::SecondRoughnessOrSimpleClearCoat; }
+
+	if (ShadingModels.HasShadingModel(MSM_Subsurface) || ShadingModels.HasShadingModel(MSM_SubsurfaceProfile) ||
+		ShadingModels.HasShadingModel(MSM_TwoSidedFoliage) || ShadingModels.HasShadingModel(MSM_Eye) ||
+		ShadingModels.HasShadingModel(MSM_PreintegratedSkin))
+	{
+		BsdfFeatures |= ESubstrateBsdfFeature::SSS;
+	}
+
+	if (bAnisotropyEnabled && (ShadingModels.HasShadingModel(MSM_DefaultLit) || ShadingModels.HasShadingModel(MSM_ClearCoat)))
+	{
+		BsdfFeatures |= ESubstrateBsdfFeature::Anisotropy;
+	}
+}
+
 FSubstrateOperator* UMaterialExpressionSubstrateShadingModels::SubstrateGenerateMaterialTopologyTreeCommon(
 	class FMaterialCompiler* Compiler, struct FGuid ThisExpressionGuid, class UMaterialExpression* Parent, int32 OutputIndex,
 	const FExpressionInput& EmissiveColor,
@@ -637,9 +657,9 @@ FSubstrateOperator* UMaterialExpressionSubstrateShadingModels::SubstrateGenerate
 
 		// Be sure to track eye/hair feature, even though they are not part of Slab BSDF. This is important later for issuing 
 		// the correct material complexity, as hair/eye requires 'complex' complexity (not 'simple') for correct packing
-		ESubstrateBsdfFeature BsdfFeatures = ESubstrateBsdfFeature::SSS | ESubstrateBsdfFeature::Fuzz;
-		if (ShadingModels.HasShadingModel(MSM_Eye))  { BsdfFeatures |= ESubstrateBsdfFeature::Eye; }
-		if (ShadingModels.HasShadingModel(MSM_Hair)) { BsdfFeatures |= ESubstrateBsdfFeature::Hair; }
+		ESubstrateBsdfFeature BsdfFeatures = ESubstrateBsdfFeature::None;
+		AppendLegacyShadingModelToBSDFFeature(ShadingModels, Anisotropy.IsConnected(), BsdfFeatures);
+
 		FSubstrateOperator* Operator = AddDefaultWorstCase(BsdfFeatures);
 		if (ShadingModels.HasShadingModel(MSM_Eye) || bHasShadingModelFromExpression)
 		{
@@ -4131,9 +4151,8 @@ FSubstrateOperator* UMaterialExpressionSubstrateConvertMaterialAttributes::Subst
 	const bool bHasShadingModelFromExpression = FMaterialAttributeDefinitionMap::IsAttributeInBitmask(CachedConnectedMaterialAttributesInputs, MP_ShadingModel) || bShadingModelFromMaterialExpression; // We keep HasShadingModelFromExpression in case all shading models cannot be safely recovered from material functions.
 	if ((ShadingModels.CountShadingModels() > 1) || bHasShadingModelFromExpression)
 	{
-		ESubstrateBsdfFeature Features = ESubstrateBsdfFeature::SSS | ESubstrateBsdfFeature::Fuzz;
-		if (ShadingModels.HasShadingModel(MSM_Eye)) { Features |= ESubstrateBsdfFeature::Eye; }
-		if (ShadingModels.HasShadingModel(MSM_Hair)) { Features |= ESubstrateBsdfFeature::Hair; }
+		ESubstrateBsdfFeature Features = ESubstrateBsdfFeature::None;
+		AppendLegacyShadingModelToBSDFFeature(ShadingModels, bHasAnisotropy, Features);
 		return AddDefaultWorstCase(Features);
 	}
 	else
