[HLOD] Add world partition validation to detect actors having HLOD Layer which doesn't support that actor class
Rules can be specified using config file and violations trigger a warning.
Added first rule - Instancing HLOD doesn't support SplineMeshActors

diff --git a/Engine/Config/BaseEngine.ini b/Engine/Config/BaseEngine.ini
--- a/Engine/Config/BaseEngine.ini
+++ b/Engine/Config/BaseEngine.ini
@@ -3923,6 +3923,7 @@ SendOSData=true
 [/Script/Engine.WorldPartitionSettings]
 EditorHashDefaultClass=/Script/CoreUObject.Class'/Script/Engine.WorldPartitionEditorSpatialHash'
 RuntimeHashDefaultClass=/Script/CoreUObject.Class'/Script/Engine.WorldPartitionRuntimeHashSet'
+UnsupportedActorClassesPerHLODLayerType=((Instancing, (ActorClasses=("/Script/CoreUObject.Class'/Script/Engine.SplineMeshActor'"))))
 
 [/Script/Engine.WorldPartitionRuntimeCellTransformerSettings]
 +IgnoredComponentClasses=/Script/CoreUObject.Class'/Script/Engine.LevelInstanceComponent'
diff --git a/Engine/Plugins/Editor/DataValidation/Source/DataValidation/Private/WorldPartitionChangelistValidator.cpp b/Engine/Plugins/Editor/DataValidation/Source/DataValidation/Private/WorldPartitionChangelistValidator.cpp
--- a/Engine/Plugins/Editor/DataValidation/Source/DataValidation/Private/WorldPartitionChangelistValidator.cpp
+++ b/Engine/Plugins/Editor/DataValidation/Source/DataValidation/Private/WorldPartitionChangelistValidator.cpp
@@ -21,6 +21,7 @@
 #include "WorldPartition/DataLayer/ExternalDataLayerAsset.h"
 #include "WorldPartition/DataLayer/ExternalDataLayerHelper.h"
 #include "WorldPartition/ContentBundle/ContentBundlePaths.h"
+#include "WorldPartition/HLOD/HLODLayer.h"
 
 #include UE_INLINE_GENERATED_CPP_BY_NAME(WorldPartitionChangelistValidator)
 
@@ -555,4 +556,26 @@ void UWorldPartitionChangelistValidator::OnInvalidHLODLayer(const IWorldPartitio
 	}
 }
 
+void UWorldPartitionChangelistValidator::OnUnsupportedHLODLayer(const IWorldPartitionActorDescInstanceView& ActorDescView, const UHLODLayer* HLODLayer, EHLODLayerUnsupportedReason Reason)
+{
+	if (Filter(ActorDescView))
+	{
+		FText CurrentWarning;
+		FText LayerTypeText = HLODLayer ? StaticEnum<EHLODLayerType>()->GetDisplayNameTextByValue((int64)HLODLayer->GetLayerType()) : FText::FromString(TEXT("Unknown"));
+
+		switch(Reason)
+		{
+		case EHLODLayerUnsupportedReason::ActorClassUnsupportedByLayerType:
+			CurrentWarning = FText::Format(LOCTEXT("DataValidation.Changelist.WorldPartition.UnsupportedHLODLayer", "Actor {0} of class '{1}' is unsupported by HLOD Layer {2} of type '{3}'"),
+				FText::FromString(GetFullActorName(ActorDescView)),
+				FText::FromString(ActorDescView.GetActorNativeClass()->GetName()),
+				FText::FromString(ActorDescView.GetHLODLayer().ToString()),
+				LayerTypeText);
+			break;
+		};
+
+		AssetWarning(CurrentAsset, CurrentWarning);
+	}
+}
+
 #undef LOCTEXT_NAMESPACE
\ No newline at end of file
diff --git a/Engine/Plugins/Editor/DataValidation/Source/DataValidation/Private/WorldPartitionChangelistValidator.h b/Engine/Plugins/Editor/DataValidation/Source/DataValidation/Private/WorldPartitionChangelistValidator.h
--- a/Engine/Plugins/Editor/DataValidation/Source/DataValidation/Private/WorldPartitionChangelistValidator.h
+++ b/Engine/Plugins/Editor/DataValidation/Source/DataValidation/Private/WorldPartitionChangelistValidator.h
@@ -53,6 +53,7 @@ protected:
 	UE_API virtual void OnLevelInstanceInvalidWorldAsset(const IWorldPartitionActorDescInstanceView& ActorDescView, FName WorldAsset, ELevelInstanceInvalidReason Reason) override;
 	UE_API virtual void OnInvalidActorFilterReference(const IWorldPartitionActorDescInstanceView& ActorDescView, const IWorldPartitionActorDescInstanceView& ReferenceActorDescView) override;
 	UE_API virtual void OnInvalidHLODLayer(const IWorldPartitionActorDescInstanceView& ActorDescView) override;
+	UE_API virtual void OnUnsupportedHLODLayer(const IWorldPartitionActorDescInstanceView& ActorDescView, const UHLODLayer* HLODLayer, EHLODLayerUnsupportedReason Reason) override;
 };
 
 #undef UE_API
diff --git a/Engine/Source/Runtime/Engine/Private/WorldPartition/ErrorHandling/WorldPartitionStreamingGenerationLogErrorHandler.cpp b/Engine/Source/Runtime/Engine/Private/WorldPartition/ErrorHandling/WorldPartitionStreamingGenerationLogErrorHandler.cpp
--- a/Engine/Source/Runtime/Engine/Private/WorldPartition/ErrorHandling/WorldPartitionStreamingGenerationLogErrorHandler.cpp
+++ b/Engine/Source/Runtime/Engine/Private/WorldPartition/ErrorHandling/WorldPartitionStreamingGenerationLogErrorHandler.cpp
@@ -5,6 +5,7 @@
 #include "WorldPartition/WorldPartitionLog.h"
 #include "WorldPartition/DataLayer/DataLayerInstanceWithAsset.h"
 #include "WorldPartition/DataLayer/WorldDataLayers.h"
+#include "WorldPartition/HLOD/HLODLayer.h"
 
 #define UE_ASSET_LOG_ACTORDESCVIEW(CategoryName, Verbosity, ActorDescView, Format, ...) \
 	UE_ASSET_LOG(LogWorldPartition, Log, *ActorDescView.GetActorPackage().ToString(), Format, ##__VA_ARGS__)
@@ -140,4 +141,17 @@ void FStreamingGenerationLogErrorHandler::OnInvalidHLODLayer(const IWorldPartiti
 {
 	UE_ASSET_LOG_ACTORDESCVIEW(LogWorldPartition, Log, ActorDescView, TEXT("Actor %s has an invalid HLOD layer %s"), *GetActorName(ActorDescView), *ActorDescView.GetHLODLayer().ToString());
 }
+
+void FStreamingGenerationLogErrorHandler::OnUnsupportedHLODLayer(const IWorldPartitionActorDescInstanceView& ActorDescView, const UHLODLayer* HLODLayer, EHLODLayerUnsupportedReason Reason)
+{
+	FString LayerType = HLODLayer ? StaticEnum<EHLODLayerType>()->GetDisplayNameTextByValue((int64)HLODLayer->GetLayerType()).ToString() : TEXT("Unknown");
+
+	switch (Reason)
+	{
+	case EHLODLayerUnsupportedReason::ActorClassUnsupportedByLayerType:
+		UE_ASSET_LOG_ACTORDESCVIEW(LogWorldPartition, Log, ActorDescView, TEXT("Actor %s of class '%s' is unsupported by HLOD Layer %s of type '%s'"), 
+			*GetActorName(ActorDescView), *ActorDescView.GetActorNativeClass()->GetName(), *ActorDescView.GetHLODLayer().ToString(), *LayerType);
+		break;
+	}
+}
 #endif
diff --git a/Engine/Source/Runtime/Engine/Private/WorldPartition/ErrorHandling/WorldPartitionStreamingGenerationProxyErrorHandler.cpp b/Engine/Source/Runtime/Engine/Private/WorldPartition/ErrorHandling/WorldPartitionStreamingGenerationProxyErrorHandler.cpp
--- a/Engine/Source/Runtime/Engine/Private/WorldPartition/ErrorHandling/WorldPartitionStreamingGenerationProxyErrorHandler.cpp
+++ b/Engine/Source/Runtime/Engine/Private/WorldPartition/ErrorHandling/WorldPartitionStreamingGenerationProxyErrorHandler.cpp
@@ -82,4 +82,9 @@ void FStreamingGenerationProxyErrorHandler::OnInvalidHLODLayer(const IWorldParti
 {
 	InnerErrorHandler->OnInvalidHLODLayer(ActorDescView);
 }
+
+void FStreamingGenerationProxyErrorHandler::OnUnsupportedHLODLayer(const IWorldPartitionActorDescInstanceView& ActorDescView, const UHLODLayer* HLODLayer, EHLODLayerUnsupportedReason Reason)
+{
+	InnerErrorHandler->OnUnsupportedHLODLayer(ActorDescView, HLODLayer, Reason);
+}
 #endif
diff --git a/Engine/Source/Runtime/Engine/Private/WorldPartition/ErrorHandling/WorldPartitionStreamingGenerationTokenizedMessageErrorHandler.cpp b/Engine/Source/Runtime/Engine/Private/WorldPartition/ErrorHandling/WorldPartitionStreamingGenerationTokenizedMessageErrorHandler.cpp
--- a/Engine/Source/Runtime/Engine/Private/WorldPartition/ErrorHandling/WorldPartitionStreamingGenerationTokenizedMessageErrorHandler.cpp
+++ b/Engine/Source/Runtime/Engine/Private/WorldPartition/ErrorHandling/WorldPartitionStreamingGenerationTokenizedMessageErrorHandler.cpp
@@ -9,6 +9,7 @@
 #include "WorldPartition/DataLayer/DataLayerAsset.h"
 #include "WorldPartition/DataLayer/DataLayerInstanceWithAsset.h"
 #include "WorldPartition/DataLayer/WorldDataLayers.h"
+#include "WorldPartition/HLOD/HLODLayer.h"
 
 #define LOCTEXT_NAMESPACE "WorldPartition"
 
@@ -285,6 +286,24 @@ void ITokenizedMessageErrorHandler::OnInvalidHLODLayer(const IWorldPartitionActo
 	HandleTokenizedMessage(MoveTemp(Message));
 }
 
+void ITokenizedMessageErrorHandler::OnUnsupportedHLODLayer(const IWorldPartitionActorDescInstanceView& ActorDescView, const UHLODLayer* HLODLayer, EHLODLayerUnsupportedReason Reason)
+{
+	FText LayerTypeText = HLODLayer ? StaticEnum<EHLODLayerType>()->GetDisplayNameTextByValue((int64)HLODLayer->GetLayerType()) : FText::FromString(TEXT("Unknown"));
+
+	switch (Reason)
+	{
+	case EHLODLayerUnsupportedReason::ActorClassUnsupportedByLayerType:
+		TSharedRef<FTokenizedMessage> Message = FTokenizedMessage::Create(EMessageSeverity::Warning);
+		Message->AddToken(FTextToken::Create(LOCTEXT("TokenMessage_WorldPartition_Actor", "Actor")))
+			->AddToken(FActorToken::Create(ActorDescView.GetActorSoftPath().ToString(), ActorDescView.GetGuid(), FText::FromString(GetActorName(ActorDescView))))
+			->AddToken(FTextToken::Create(FText::Format(LOCTEXT("TokenMessage_WorldPartition_ActorClassUnsupportedByLayerType_Warning", "of class '{0}' is unsupported by HLOD Layer"), FText::FromString(ActorDescView.GetActorNativeClass()->GetName()))))
+			->AddToken(FAssetNameToken::Create(ActorDescView.GetHLODLayer().ToString()))
+			->AddToken(FTextToken::Create(FText::Format(LOCTEXT("TokenMessage_WorldPartition_ActorClassUnsupportedByLayerType_LayerType", "of type '{0}'"), LayerTypeText)));
+		HandleTokenizedMessage(MoveTemp(Message));
+		break;
+	}
+}
+
 #undef LOCTEXT_NAMESPACE
 
 #endif
diff --git a/Engine/Source/Runtime/Engine/Private/WorldPartition/WorldPartitionStreamingGeneration.cpp b/Engine/Source/Runtime/Engine/Private/WorldPartition/WorldPartitionStreamingGeneration.cpp
--- a/Engine/Source/Runtime/Engine/Private/WorldPartition/WorldPartitionStreamingGeneration.cpp
+++ b/Engine/Source/Runtime/Engine/Private/WorldPartition/WorldPartitionStreamingGeneration.cpp
@@ -22,6 +22,7 @@
 #include "ProfilingDebugging/ScopedTimers.h"
 #include "WorldPartition/IWorldPartitionEditorModule.h"
 #include "WorldPartition/WorldPartitionRuntimeHash.h"
+#include "WorldPartition/WorldPartitionSettings.h"
 #include "WorldPartition/WorldPartitionStreamingPolicy.h"
 #include "WorldPartition/WorldPartitionLevelStreamingPolicy.h"
 #include "WorldPartition/DataLayer/DataLayerManager.h"
@@ -1737,6 +1738,32 @@ class FWorldPartitionStreamingGenerator
 				    }
 				}
 
+				if (ActorDescView.GetIsSpatiallyLoaded() && ActorDescView.GetHLODLayer().IsValid())
+				{
+					if (UHLODLayer* HLODLayer = Cast<UHLODLayer>(ActorDescView.GetHLODLayer().TryLoad()))
+					{
+						if (const FHLODLayerTypeUnsupportedActorClasses* UnsupportedActorClasses = UWorldPartitionSettings::Get()->GetUnsupportedActorClassesPerHLODLayerType().Find(HLODLayer->GetLayerType()))
+						{
+							for (TSubclassOf<AActor> ActorClass : UnsupportedActorClasses->ActorClasses)
+							{
+								if (ActorDescView.GetActorNativeClass()->IsChildOf(ActorClass))
+								{
+									if (PassType == EPassType::ErrorReporting)
+									{
+										ErrorHandler->OnUnsupportedHLODLayer(ActorDescView, HLODLayer, IStreamingGenerationErrorHandler::EHLODLayerUnsupportedReason::ActorClassUnsupportedByLayerType);
+									}
+									else
+									{
+										ActorDescView.SetForcedNoHLODLayer();
+									}
+
+									NbErrorsDetected++;
+								}
+							}
+						}
+					}
+				}
+
 				if (TSet<FGuid>* FilteredActors = ContainerFilteredActors.Find(ContainerCollectionInstanceDescriptor.ID))
 				{
 					const bool IsReferencerFiltered = FilteredActors->Contains(ActorDescView.GetGuid());
diff --git a/Engine/Source/Runtime/Engine/Public/WorldPartition/ErrorHandling/WorldPartitionStreamingGenerationErrorHandler.h b/Engine/Source/Runtime/Engine/Public/WorldPartition/ErrorHandling/WorldPartitionStreamingGenerationErrorHandler.h
--- a/Engine/Source/Runtime/Engine/Public/WorldPartition/ErrorHandling/WorldPartitionStreamingGenerationErrorHandler.h
+++ b/Engine/Source/Runtime/Engine/Public/WorldPartition/ErrorHandling/WorldPartitionStreamingGenerationErrorHandler.h
@@ -9,6 +9,7 @@ class AWorldDataLayers;
 class UDataLayerAsset;
 class UDataLayerInstance;
 class UDataLayerInstanceWithAsset;
+class UHLODLayer;
 
 // deprecated
 class FWorldPartitionActorDescView;
@@ -129,6 +130,16 @@ public:
 	 */
 	virtual void OnInvalidHLODLayer(const IWorldPartitionActorDescInstanceView& ActorDescView) = 0;
 
+	enum class EHLODLayerUnsupportedReason
+	{
+		ActorClassUnsupportedByLayerType
+	};
+
+	/**
+	 * Called when an actor has a HLOD layer that doesn't support that actor.
+	 */
+	virtual void OnUnsupportedHLODLayer(const IWorldPartitionActorDescInstanceView& ActorDescView, const UHLODLayer* HLODLayer, EHLODLayerUnsupportedReason Reason) = 0;
+
 	// Helpers
 	static ENGINE_API FString GetActorName(const IWorldPartitionActorDescInstanceView& ActorDescView);
 	static ENGINE_API FString GetFullActorName(const IWorldPartitionActorDescInstanceView& ActorDescView);
diff --git a/Engine/Source/Runtime/Engine/Public/WorldPartition/ErrorHandling/WorldPartitionStreamingGenerationLogErrorHandler.h b/Engine/Source/Runtime/Engine/Public/WorldPartition/ErrorHandling/WorldPartitionStreamingGenerationLogErrorHandler.h
--- a/Engine/Source/Runtime/Engine/Public/WorldPartition/ErrorHandling/WorldPartitionStreamingGenerationLogErrorHandler.h
+++ b/Engine/Source/Runtime/Engine/Public/WorldPartition/ErrorHandling/WorldPartitionStreamingGenerationLogErrorHandler.h
@@ -23,5 +23,6 @@ public:
 	ENGINE_API virtual void OnLevelInstanceInvalidWorldAsset(const IWorldPartitionActorDescInstanceView& ActorDescView, FName WorldAsset, ELevelInstanceInvalidReason Reason) override;
 	ENGINE_API virtual void OnInvalidActorFilterReference(const IWorldPartitionActorDescInstanceView& ActorDescView, const IWorldPartitionActorDescInstanceView& ReferenceActorDescView) override;
 	ENGINE_API virtual void OnInvalidHLODLayer(const IWorldPartitionActorDescInstanceView& ActorDescView) override;
+	ENGINE_API virtual void OnUnsupportedHLODLayer(const IWorldPartitionActorDescInstanceView& ActorDescView, const UHLODLayer* HLODLayer, EHLODLayerUnsupportedReason Reason) override;
 };
 #endif
diff --git a/Engine/Source/Runtime/Engine/Public/WorldPartition/ErrorHandling/WorldPartitionStreamingGenerationNullErrorHandler.h b/Engine/Source/Runtime/Engine/Public/WorldPartition/ErrorHandling/WorldPartitionStreamingGenerationNullErrorHandler.h
--- a/Engine/Source/Runtime/Engine/Public/WorldPartition/ErrorHandling/WorldPartitionStreamingGenerationNullErrorHandler.h
+++ b/Engine/Source/Runtime/Engine/Public/WorldPartition/ErrorHandling/WorldPartitionStreamingGenerationNullErrorHandler.h
@@ -23,5 +23,6 @@ public:
 	virtual void OnLevelInstanceInvalidWorldAsset(const IWorldPartitionActorDescInstanceView& ActorDescView, FName WorldAsset, ELevelInstanceInvalidReason Reason) override {};
 	virtual void OnInvalidActorFilterReference(const IWorldPartitionActorDescInstanceView& ActorDescView, const IWorldPartitionActorDescInstanceView& ReferenceActorDescView) override {};
 	virtual void OnInvalidHLODLayer(const IWorldPartitionActorDescInstanceView& ActorDescView) override {}
+	virtual void OnUnsupportedHLODLayer(const IWorldPartitionActorDescInstanceView& ActorDescView, const UHLODLayer* HLODLayer, EHLODLayerUnsupportedReason Reason) override {};
 };
 #endif
diff --git a/Engine/Source/Runtime/Engine/Public/WorldPartition/ErrorHandling/WorldPartitionStreamingGenerationProxyErrorHandler.h b/Engine/Source/Runtime/Engine/Public/WorldPartition/ErrorHandling/WorldPartitionStreamingGenerationProxyErrorHandler.h
--- a/Engine/Source/Runtime/Engine/Public/WorldPartition/ErrorHandling/WorldPartitionStreamingGenerationProxyErrorHandler.h
+++ b/Engine/Source/Runtime/Engine/Public/WorldPartition/ErrorHandling/WorldPartitionStreamingGenerationProxyErrorHandler.h
@@ -27,6 +27,7 @@ public:
 	ENGINE_API virtual void OnLevelInstanceInvalidWorldAsset(const IWorldPartitionActorDescInstanceView& ActorDescView, FName WorldAsset, ELevelInstanceInvalidReason Reason) override;
 	ENGINE_API virtual void OnInvalidActorFilterReference(const IWorldPartitionActorDescInstanceView& ActorDescView, const IWorldPartitionActorDescInstanceView& ReferenceActorDescView) override;
 	ENGINE_API virtual void OnInvalidHLODLayer(const IWorldPartitionActorDescInstanceView& ActorDescView) override;
+	ENGINE_API virtual void OnUnsupportedHLODLayer(const IWorldPartitionActorDescInstanceView& ActorDescView, const UHLODLayer* HLODLayer, EHLODLayerUnsupportedReason Reason) override;
 
 protected:
 	IStreamingGenerationErrorHandler* InnerErrorHandler;
diff --git a/Engine/Source/Runtime/Engine/Public/WorldPartition/ErrorHandling/WorldPartitionStreamingGenerationTokenizedMessageErrorHandler.h b/Engine/Source/Runtime/Engine/Public/WorldPartition/ErrorHandling/WorldPartitionStreamingGenerationTokenizedMessageErrorHandler.h
--- a/Engine/Source/Runtime/Engine/Public/WorldPartition/ErrorHandling/WorldPartitionStreamingGenerationTokenizedMessageErrorHandler.h
+++ b/Engine/Source/Runtime/Engine/Public/WorldPartition/ErrorHandling/WorldPartitionStreamingGenerationTokenizedMessageErrorHandler.h
@@ -24,6 +24,7 @@ class ITokenizedMessageErrorHandler : public IStreamingGenerationErrorHandler
 	ENGINE_API virtual void OnLevelInstanceInvalidWorldAsset(const IWorldPartitionActorDescInstanceView& ActorDescView, FName WorldAsset, ELevelInstanceInvalidReason Reason) override;
 	ENGINE_API virtual void OnInvalidActorFilterReference(const IWorldPartitionActorDescInstanceView& ActorDescView, const IWorldPartitionActorDescInstanceView& ReferenceActorDescView) override;
 	ENGINE_API virtual void OnInvalidHLODLayer(const IWorldPartitionActorDescInstanceView& ActorDescView) override;
+	ENGINE_API virtual void OnUnsupportedHLODLayer(const IWorldPartitionActorDescInstanceView& ActorDescView, const UHLODLayer* HLODLayer, EHLODLayerUnsupportedReason Reason) override;
 
 protected:
 	virtual void HandleTokenizedMessage(TSharedRef<FTokenizedMessage>&& ErrorMessage) = 0;
diff --git a/Engine/Source/Runtime/Engine/Public/WorldPartition/WorldPartitionSettings.h b/Engine/Source/Runtime/Engine/Public/WorldPartition/WorldPartitionSettings.h
--- a/Engine/Source/Runtime/Engine/Public/WorldPartition/WorldPartitionSettings.h
+++ b/Engine/Source/Runtime/Engine/Public/WorldPartition/WorldPartitionSettings.h
@@ -2,6 +2,7 @@
 #pragma once
 
 #include "Engine/DeveloperSettings.h"
+#include "WorldPartition/HLOD/HLODLayer.h"
 #include "WorldPartition/WorldPartition.h"
 #include "WorldPartitionSettings.generated.h"
 
@@ -9,6 +10,15 @@
 
 class UWorldPartitionPropertyOverridePolicy;
 
+USTRUCT()
+struct FHLODLayerTypeUnsupportedActorClasses
+{
+	GENERATED_BODY()
+
+	UPROPERTY(Config)
+	TArray<TSubclassOf<AActor>> ActorClasses;
+};
+
 UCLASS(MinimalAPI, config = Engine, defaultconfig, DisplayName = "World Partition")
 class UWorldPartitionSettings : public UDeveloperSettings
 {
@@ -30,6 +40,12 @@ public:
 	{
 		return PropertyOverridePolicy;
 	}
+
+	const TMap<EHLODLayerType, FHLODLayerTypeUnsupportedActorClasses>& GetUnsupportedActorClassesPerHLODLayerType() const
+	{
+		return UnsupportedActorClassesPerHLODLayerType;
+	}
+
 protected:
 	friend class ULevelInstanceSettings;
 	void SetPropertyOverridePolicy(UWorldPartitionPropertyOverridePolicy* InPropertyOverridePolicy)
@@ -59,6 +75,9 @@ protected:
 
 	UPROPERTY(Transient)
 	TObjectPtr<UWorldPartitionPropertyOverridePolicy> PropertyOverridePolicy;
+
+	UPROPERTY(Config)
+	TMap<EHLODLayerType, FHLODLayerTypeUnsupportedActorClasses> UnsupportedActorClassesPerHLODLayerType;
 };
 
 #undef UE_API
