[StandaloneHLOD] Fix logic that generates the list of Standalone HLOD levels returned by ShouldProcessAdditionalWorlds, to correctly handle cases where the number of HLOD Setups differs from the maximum HLOD hierarchy depth.

diff --git a/Engine/Source/Editor/UnrealEd/Private/WorldPartition/WorldPartitionHLODsBuilder.cpp b/Engine/Source/Editor/UnrealEd/Private/WorldPartition/WorldPartitionHLODsBuilder.cpp
--- a/Engine/Source/Editor/UnrealEd/Private/WorldPartition/WorldPartitionHLODsBuilder.cpp
+++ b/Engine/Source/Editor/UnrealEd/Private/WorldPartition/WorldPartitionHLODsBuilder.cpp
@@ -203,13 +203,13 @@ bool UWorldPartitionHLODsBuilder::ShouldProcessAdditionalWorlds(UWorld* InWorld,
 			FString FolderPath, PackagePrefix;
 			UWorldPartitionStandaloneHLODSubsystem::GetStandaloneHLODFolderPathAndPackagePrefix(InWorld->GetPackage()->GetName(), FolderPath, PackagePrefix);
 
-			WorldPartitionRuntimeHashSet->ForEachHLODLayer([&FolderPath, &PackagePrefix, &OutPackageNames](FName RuntimePartitionName, FName HLODSetupName, int32 HLODSetupIndex)
+			int32 HLODDepth = WorldPartitionRuntimeHashSet->ComputeHLODHierarchyDepth();
+
+			for (int32 HLODSetupIndex = 0; HLODSetupIndex < HLODDepth; HLODSetupIndex++)
 			{
 				const FString HLODLevelPackageName = FString::Printf(TEXT("%s/%s%d"), *FolderPath, *PackagePrefix, HLODSetupIndex);
 				OutPackageNames.Add(HLODLevelPackageName);
-
-				return true;
-			});
+			}
 
 			return true;
 		}
diff --git a/Engine/Source/Runtime/Engine/Private/WorldPartition/RuntimeHashSet/WorldPartitionRuntimeHashSet.cpp b/Engine/Source/Runtime/Engine/Private/WorldPartition/RuntimeHashSet/WorldPartitionRuntimeHashSet.cpp
--- a/Engine/Source/Runtime/Engine/Private/WorldPartition/RuntimeHashSet/WorldPartitionRuntimeHashSet.cpp
+++ b/Engine/Source/Runtime/Engine/Private/WorldPartition/RuntimeHashSet/WorldPartitionRuntimeHashSet.cpp
@@ -1193,6 +1193,33 @@ void UWorldPartitionRuntimeHashSet::ForEachHLODLayer(TFunctionRef<bool(FName, FN
 	}
 }
 
+int32 UWorldPartitionRuntimeHashSet::ComputeHLODHierarchyDepth() const
+{
+	int32 MaxHLODDepth = 0;
+	for (const FRuntimePartitionDesc& RuntimePartition : RuntimePartitions)
+	{
+		if (RuntimePartition.MainLayer)
+		{
+			for (const FRuntimePartitionHLODSetup& HLODSetup : RuntimePartition.HLODSetups)
+			{
+				for (const TObjectPtr<const UHLODLayer>& HLODLayer : HLODSetup.HLODLayers)
+				{
+					const UHLODLayer* CurrentHLODLayer = HLODLayer;
+					int32 CurrentHLODDepth = 0;
+					while (CurrentHLODLayer)
+					{
+						CurrentHLODDepth++;
+						CurrentHLODLayer = CurrentHLODLayer->GetParentLayer();
+					}
+					MaxHLODDepth = FMath::Max(MaxHLODDepth, CurrentHLODDepth);
+				}
+			}
+		}
+	}
+	
+	return MaxHLODDepth;
+}
+
 UWorldPartitionRuntimeHashSet::FCellUniqueId UWorldPartitionRuntimeHashSet::GetCellUniqueId(const URuntimePartition::FCellDescInstance& InCellDescInstance) const
 {
 	FCellUniqueId CellUniqueId;
diff --git a/Engine/Source/Runtime/Engine/Public/WorldPartition/RuntimeHashSet/WorldPartitionRuntimeHashSet.h b/Engine/Source/Runtime/Engine/Public/WorldPartition/RuntimeHashSet/WorldPartitionRuntimeHashSet.h
--- a/Engine/Source/Runtime/Engine/Public/WorldPartition/RuntimeHashSet/WorldPartitionRuntimeHashSet.h
+++ b/Engine/Source/Runtime/Engine/Public/WorldPartition/RuntimeHashSet/WorldPartitionRuntimeHashSet.h
@@ -164,6 +164,7 @@ public:
 	ENGINE_API virtual void Draw3D(const TArray<FWorldPartitionStreamingSource>& Sources) const override;
 	ENGINE_API virtual const FGuid* GetStandaloneHLODActorSourceCellOverride(const FGuid& InActorGuid) const override;
 	ENGINE_API void ForEachHLODLayer(TFunctionRef<bool(FName, FName, int32)> Func) const; // RuntimePartitionName, HLODSetupName, HLODSetupIndex
+	ENGINE_API int32 ComputeHLODHierarchyDepth() const;
 
 #if WITH_EDITOR
 	// Streaming generation interface
