Added support for user-defined HLSL functions in the new translator, used to add support for Custom expressions in the new translator.

This change introduces the idea of user defined functions (created at emit time) in the translator. Right now, only one kind of user function is supported and that is the one defined by HLSL.

#tests Tested that the MaterialExpressionCustom works identically to how it works in the old translator except for the SceneTexture special case.
#rnx

diff --git a/Engine/Shaders/Private/MaterialTemplate.ush b/Engine/Shaders/Private/MaterialTemplate.ush
--- a/Engine/Shaders/Private/MaterialTemplate.ush
+++ b/Engine/Shaders/Private/MaterialTemplate.ush
@@ -3539,6 +3539,8 @@ SHADER_DISABLE_WARNINGS
 
 %{uniform_material_expressions}
 
+%{custom_functions}
+
 SHADER_POP_WARNINGS_STATE
 
 #if USE_ANALYTIC_DERIVATIVES && TEXTURE_SAMPLE_DEBUG
diff --git a/Engine/Source/Runtime/Engine/Private/Materials/MaterialExpressionsIR.cpp b/Engine/Source/Runtime/Engine/Private/Materials/MaterialExpressionsIR.cpp
--- a/Engine/Source/Runtime/Engine/Private/Materials/MaterialExpressionsIR.cpp
+++ b/Engine/Source/Runtime/Engine/Private/Materials/MaterialExpressionsIR.cpp
@@ -1746,12 +1746,9 @@ static void BuildTransformVectorBase(
 			OutputValue = Em.CastToScalarKind(OutputValue, MIR::EScalarKind::Double);
 		}
 	}
-	else
+	else if (OutputValue->Type.IsDouble())
 	{
-		if (OutputValue->Type.IsDouble())
-		{
-			OutputValue = Em.CastToFloatKind(OutputValue);
-		}
+		OutputValue = Em.CastToFloatKind(OutputValue);
 	}
 
 	Em.Output(0, OutputValue);
@@ -3129,4 +3126,139 @@ void UMaterialExpressionSphereMask::Build(MIR::FEmitter& Em)
 	Em.Output(0, Em.Saturate(MaskUnclamped));
 }
 
+// Takes a description user string and turns it into a valid C/HLSL identifier.
+// Note: we allow the identifier to start by numbers.
+static FString DescriptionToIdentifier(FStringView Source)
+{
+	FString Out;
+	Out.Reserve(Source.Len());
+	// Append an underscore if the source starts by a igit
+	if (!Source.IsEmpty() && FChar::IsDigit(Source[0]))
+	{
+		Out.AppendChar('_');
+	}
+	for (TCHAR Ch : Source)
+	{
+		Out.AppendChar(FChar::IsAlnum(Ch) ? Ch : '_');
+	}
+	return Out;
+}
+
+// Custom material output to MIR type conversion.
+static MIR::FType CustomMaterialOutputTypeToMIR(ECustomMaterialOutputType Type)
+{
+	switch (Type) 
+	{
+		case CMOT_Float1: return MIR::FType::MakeFloatScalar();
+		case CMOT_Float2: return MIR::FType::MakeFloatVector(2);
+		case CMOT_Float3: return MIR::FType::MakeFloatVector(3);
+		case CMOT_Float4: return MIR::FType::MakeFloatVector(4);
+		case CMOT_MaterialAttributes: return MIR::FType::MakeAggregate(UMaterialAggregate::GetMaterialAttributes());
+		default: UE_MIR_UNREACHABLE();
+	}
+}
+
+void UMaterialExpressionCustom::Build(MIR::FEmitter& Em)
+{
+	// Convert the description to a valid HLSL identifier
+	FString Name = DescriptionToIdentifier(Description);
+
+	MIR::TTemporaryArray<MIR::FValueRef> InputArgs { Inputs.Num() };
+
+	// Prepare a description of the user-defined HLSL function for the emitter.
+	MIR::FFunctionHLSLDesc FuncDesc;
+	FuncDesc.Name = Name;
+	FuncDesc.ReturnType = CustomMaterialOutputTypeToMIR(OutputType);
+	FuncDesc.Code = Code;
+
+	// Turn each expression input into an input-only parameter.
+	for (FCustomInput const& Input : Inputs)
+	{
+		// Ignore input parameters with "None" name.
+		if (Input.InputName.IsNone())
+		{
+			continue;
+		}
+
+		// Read the input argument
+		InputArgs[FuncDesc.NumInputOnlyParams] = Em.Input(&Input.Input);
+
+		// Declare an input-only parameter
+		if (!FuncDesc.PushInputOnlyParameter(Input.InputName, InputArgs[FuncDesc.NumInputOnlyParams]->Type))
+		{
+			Em.Errorf(TEXT("Too many inputs. Custom expressions can have at most %d input/output pins."), MIR::MaxNumFunctionParameters);
+			return;
+		}
+	}
+
+	// Some Input() call might have generated an error
+	UE_MIR_CHECKPOINT(Em);
+
+	// Turn each expression additional output into a output-only parameter.
+	for (FCustomOutput const& AdditionalOutput : AdditionalOutputs)
+	{
+		// Ignore output parameters with "None" name.
+		if (AdditionalOutput.OutputName.IsNone())
+		{
+			continue;
+		}
+
+		if (!FuncDesc.PushOutputOnlyParameter(AdditionalOutput.OutputName, CustomMaterialOutputTypeToMIR(AdditionalOutput.OutputType)))
+		{
+			Em.Errorf(TEXT("Too many input/outputs. Custom expressions can have at most %d input/output pins."), MIR::MaxNumFunctionParameters);
+			return;
+		}
+	}
+
+	// Generate the array of additional defines
+	MIR::TTemporaryArray<MIR::FFunctionHLSLDefine> Defines{ AdditionalDefines.Num() };
+	for (int32 i = 0; i < AdditionalDefines.Num(); ++i)
+	{
+		if (AdditionalDefines[i].DefineName.IsEmpty())
+		{
+			Em.Errorf(TEXT("Define with index '%d' has no valid name."), i);
+		}
+
+		if (AdditionalDefines[i].DefineValue.IsEmpty())
+		{
+			Em.Errorf(TEXT("Define with index '%d' has no valid value."), i);
+		}
+
+		Defines[i] = { AdditionalDefines[i].DefineName, AdditionalDefines[i].DefineValue };
+	}
+	FuncDesc.Defines = Defines;
+
+	// Generate the array of additional includes
+	MIR::TTemporaryArray<FStringView> Includes{ IncludeFilePaths.Num() };
+	for (int32 i = 0; i < IncludeFilePaths.Num(); ++i)
+	{
+		if (IncludeFilePaths[i].IsEmpty())
+		{
+			Em.Errorf(TEXT("Define with index '%d' has no valid value."), i);
+		}
+
+		Includes[i] = IncludeFilePaths[i];
+	}
+
+	FuncDesc.Includes = Includes;
+
+	UE_MIR_CHECKPOINT(Em); // Make sure checks above did not fail
+
+	// Declare the HLSL function with the description we generated
+	const MIR::FFunction* Func = Em.FunctionHLSL(FuncDesc);
+	
+	UE_MIR_CHECKPOINT(Em); // To guarantee a function was succesfully emitted.
+
+	FValueRef Call = Em.Call(Func, { InputArgs.GetData(), (int32)FuncDesc.NumInputOnlyParams });
+
+	// Output the call return value through the first output pin
+	Em.Output(0, Call);
+
+	// Output the additional outputs through subsequent output pins
+	for (uint32 i = 0; i < Func->GetNumOutputParameters(); ++i)
+	{
+		Em.Output(i + 1, Em.CallParameterOutput(Call, i));
+	}
+}
+
 #endif // WITH_EDITOR
diff --git a/Engine/Source/Runtime/Engine/Private/Materials/MaterialIR.cpp b/Engine/Source/Runtime/Engine/Private/Materials/MaterialIR.cpp
--- a/Engine/Source/Runtime/Engine/Private/Materials/MaterialIR.cpp
+++ b/Engine/Source/Runtime/Engine/Private/Materials/MaterialIR.cpp
@@ -59,6 +59,8 @@ const TCHAR* LexToString(EValueKind Kind)
 		case VK_PromoteSubstrateParameter: return TEXT("PromoteSubstrateParameter");
 		case VK_StageSwitch: return TEXT("StageSwitch");
 		case VK_HardwarePartialDerivative: return TEXT("HardwarePartialDerivative");
+		case VK_Call: return TEXT("Call");
+		case VK_CallParameterOutput: return TEXT("CallParameterOutput");
 		/* invalid entries */
 		case VK_InstructionBegin:
 		case VK_InstructionEnd:
@@ -120,7 +122,9 @@ uint32 FValue::GetSizeInBytes() const
 		case VK_PromoteSubstrateParameter: return sizeof(FPromoteSubstrateParameter);
 		case VK_StageSwitch: return sizeof(FStageSwitch);
 		case VK_HardwarePartialDerivative: return sizeof(FHardwarePartialDerivative);
-		/* invalid entries */
+		case VK_Call: return sizeof(FCall);
+		case VK_CallParameterOutput: return sizeof(FCallParameterOutput);
+			/* invalid entries */
 		case VK_InstructionBegin:
 		case VK_InstructionEnd:
 			UE_MIR_UNREACHABLE();
@@ -210,6 +214,19 @@ TConstArrayView<FValue*> FValue::GetUses() const
 			return { &This->Arg, FHardwarePartialDerivative::NumStaticUses };
 		}
 
+		case VK_Call:
+		{
+			auto This = static_cast<const FCall*>(this);
+			return { This->Arguments, This->NumArguments };
+
+		}
+
+		case VK_CallParameterOutput:
+		{
+			auto This = static_cast<const FCallParameterOutput*>(this);
+			return { &This->Call, FCallParameterOutput::NumStaticUses };
+		}
+
 		default: UE_MIR_UNREACHABLE();
 	}
 }
@@ -815,6 +832,22 @@ void FStageSwitch::SetArgs(FValue* PixelStageArg, FValue* OtherStagesArg)
 	}
 }
 
+bool FFunction::Equals(const FFunction* Other) const
+{
+	return Kind == Other->Kind
+		&& Name == Other->Name
+		&& ReturnType == Other->ReturnType
+		&& NumInputOnlyParams == Other->NumInputOnlyParams
+		&& NumInputAndOutputParams == Other->NumInputAndOutputParams
+		&& NumParameters == Other->NumParameters
+		&& FMemory::Memcmp(Parameters, Other->Parameters, NumParameters * sizeof(FFunctionParameter)) == 0;
+}
+
+bool FFunctionHLSL::Equals(const FFunctionHLSL* Other) const
+{
+	return FFunction::Equals(Other) && Code == Other->Code;
+}
+
 } // namespace MIR
 
 #endif // #if WITH_EDITOR
diff --git a/Engine/Source/Runtime/Engine/Private/Materials/MaterialIRDebug.cpp b/Engine/Source/Runtime/Engine/Private/Materials/MaterialIRDebug.cpp
--- a/Engine/Source/Runtime/Engine/Private/Materials/MaterialIRDebug.cpp
+++ b/Engine/Source/Runtime/Engine/Private/Materials/MaterialIRDebug.cpp
@@ -198,32 +198,35 @@ static bool InstrHasVariableArgCount(MIR::EValueKind Kind)
 	return Kind == MIR::VK_Operator;
 }
 
-/// Helper struct to wrap the state used during IR to text dumping.
+// Helper struct to wrap the state used during IR to text dumping.
 struct FDebugDumpIRState
 {
-	/// The module we are printing the IR for.
+	// The module we are printing the IR for.
 	const FMaterialIRModule* Module{};
 
-	/// Output string containing the generated result.
+	// Output string containing the generated result.
 	FString Out; 
 
-	/// String used for temporary operations. Clear before use.
+	// String used for temporary operations. Clear before use.
 	FString Temp;
 
-	/// Maps values to an incrementing id. Used to give values a "name" for future referencing (e.g. "%6")
+	// Maps values to an incrementing id. Used to give values a "name" for future referencing (e.g. "%6")
 	TMap<const FValue*, uint32> ValueToIdMap;
 
-	/// Counter used to assign an id to encountered values.
+	// Counter used to assign an id to encountered values.
 	uint32 InstrIdCounter = 0;
 
-	/// Array of encountered parameters. Used later on to generate the a recap of all referenced parameters.
+	// Array of encountered parameters. Used later on to generate the a recap of all referenced parameters.
 	TArray<TPair<uint32, const FValue*>> ReferencedParameters;
+	
+	// Array of encountered user HLSL functions.
+	TArray<const MIR::FFunctionHLSL*> ReferencedFunctionHLSLs;
 
-	/// Stage we're currently emitting to.
+	// Stage we're currently emitting to.
 	MIR::EStage CurrentStage;
 
-	/// Prints a block of instructions to `Out`. Indentation indicates how many levels of indentation to put
-	/// to the left of printed instruction.
+	// Prints a block of instructions to `Out`. Indentation indicates how many levels of indentation to put
+	// to the left of printed instruction.
 	void AppendBlock(const FBlock& Block, int32 Indentation)
 	{
 		for (FInstruction* Instr = Block.Instructions; Instr; Instr = Instr->Next[CurrentStage])
@@ -294,7 +297,7 @@ struct FDebugDumpIRState
 		}
 	}
 
-	/// Appends extra information regarding the instruction.
+	// Appends extra information regarding the instruction.
 	void AppendInstructionProperties(const FInstruction* Instr)
 	{
 		if (const FSetMaterialOutput* SetMaterialOutput = Instr->As<FSetMaterialOutput>())
@@ -338,10 +341,22 @@ struct FDebugDumpIRState
 		{
 			Out.Append(Derivative->Axis == MIR::EDerivativeAxis::X ? TEXT(" \"ddx\"") : TEXT(" \"ddy\""));
 		}
+		else if (const FCall* Call = Instr->As<FCall>())
+		{
+			if (Call->Function->Kind == MIR::FFunctionKind::HLSL)
+			{
+				ReferencedFunctionHLSLs.AddUnique(static_cast<const MIR::FFunctionHLSL*>(Call->Function));
+				Out.Appendf(TEXT(" FunctionHLSL=\"%s\""), Call->Function->Name.GetData());
+			}
+		}
+		else if (const FCallParameterOutput* CallOutput = Instr->As<FCallParameterOutput>())
+		{
+			Out.Appendf(TEXT(" Output=\"%s\""), *CallOutput->Call->As<FCall>()->Function->GetOutputParameter(CallOutput->Index).Name.ToString());
+		}
 	}
 
-	/// Appends a reference to the specified value. This will look like "%x" if Value is an
-	/// instruction, otherwise it will inline information regarding the value.
+	// Appends a reference to the specified value. This will look like "%x" if Value is an
+	// instruction, otherwise it will inline information regarding the value.
 	void AppendValueReference(const FValue* Value)
 	{
 		if (ValueToIdMap.Contains(Value))
@@ -398,7 +413,7 @@ struct FDebugDumpIRState
 		}
 	}
 
-	/// Gets the instruction reference.
+	// Gets the instruction reference.
 	uint32 ReferenceInstruction(const FInstruction* Instr)
 	{
 		uint32 Id;
@@ -426,7 +441,7 @@ struct FDebugDumpIRState
 		Out += LeftColumn;
 	}
 
-	/// Prints the parameter recap section.
+	// Prints the parameter recap section.
 	void DumpReferencedParameters()
 	{
 		if (ReferencedParameters.IsEmpty())
@@ -462,6 +477,29 @@ struct FDebugDumpIRState
 			}
 		}
 	}
+
+	void DumpFunctionHLSLs()
+	{
+		if (ReferencedFunctionHLSLs.IsEmpty())
+		{
+			return;
+		}
+
+		Out.Append(TEXT("\n; Referenced user HLSL functions\n"));
+
+		for (const MIR::FFunctionHLSL* Function : ReferencedFunctionHLSLs)
+		{
+			Out.Appendf(TEXT("FunctionHLSL Name=\"C%d_%s\" ReturnType=\"%s\"\n"), Function->UniqueId, Function->Name.GetData(), *Function->ReturnType.GetSpelling());
+			for (int i = 0; i < Function->NumParameters; ++i)
+			{
+				const TCHAR* Keyword = (i < Function->NumInputOnlyParams) ? TEXT("In")
+					: (i < Function->NumInputAndOutputParams) ? TEXT("InOut")
+					: TEXT("Out");
+				
+				Out.Appendf(TEXT("\tParam %s Name=\"%s\" Type=\"%s\"\n"), Keyword, *Function->Parameters[i].Name.ToString(), *Function->Parameters[i].Type.GetSpelling());
+			}
+		}
+	}
 };
 
 FString DebugDumpIR(FStringView MaterialName, const FMaterialIRModule& Module)
@@ -482,7 +520,7 @@ FString DebugDumpIR(FStringView MaterialName, const FMaterialIRModule& Module)
 
 	// Print referenced material parameters recap if any
 	State.DumpReferencedParameters();
-
+	State.DumpFunctionHLSLs();
 	return State.Out;
 }
 
diff --git a/Engine/Source/Runtime/Engine/Private/Materials/MaterialIREmitter.cpp b/Engine/Source/Runtime/Engine/Private/Materials/MaterialIREmitter.cpp
--- a/Engine/Source/Runtime/Engine/Private/Materials/MaterialIREmitter.cpp
+++ b/Engine/Source/Runtime/Engine/Private/Materials/MaterialIREmitter.cpp
@@ -97,6 +97,15 @@ struct FEmitter::FPrivate
 	}
 };
 
+// Creates a copy of specified array using the Module allocator and returns it.
+template <typename T>
+TConstArrayView<T> MakeArrayCopy(FEmitter& Emitter, TConstArrayView<T> Array)
+{
+	T* Data = (T*)FEmitter::FPrivate::Allocate(Emitter, sizeof(T) * Array.Num(), alignof(T));
+	FMemory::Memcpy(Data, Array.GetData(), sizeof(T) * Array.Num());
+	return { Data, Array.Num() };
+}
+
 // Creates a “prototype” value of type T with the specified MIR type.
 // Emit this value later with a matching EmitPrototype() call, which returns the actual value instance
 // after potential deduplication.
@@ -225,6 +234,48 @@ FValueRef FValueRef::ToPoison() const
 	return To(FPoison::Get());
 }
 
+/*--------------------------------- FFunctionHLSLDesc ---------------------------------*/
+
+bool FFunctionHLSLDesc::PushInputOnlyParameter(FName InName, FType InType)
+{
+	// You must input-only parameters first, before the others.
+	check(!NumInputOutputParams && !NumOutputOnlyParams);
+
+	if (GetNumParameters() == MIR::MaxNumFunctionParameters)
+	{
+		return false;
+	}
+
+	Parameters[NumInputOnlyParams++] = { InName, InType };
+	return true;
+}
+	
+bool FFunctionHLSLDesc::PushInputOutputParameter(FName InName, FType InType)
+{
+	checkf(!NumOutputOnlyParams, TEXT("You must input-output parameters before you push any output-only parameters."));
+
+	if (GetNumParameters() == MIR::MaxNumFunctionParameters)
+	{
+		return false;
+	}
+
+	Parameters[NumInputOnlyParams + NumInputOutputParams] = { InName, InType };
+	++NumInputOutputParams;
+	return true;
+}
+	
+bool FFunctionHLSLDesc::PushOutputOnlyParameter(FName InName, FType InType)
+{
+	if (GetNumParameters() == MIR::MaxNumFunctionParameters)
+	{
+		return false;
+	}
+
+	Parameters[NumInputOnlyParams + NumInputOutputParams + NumOutputOnlyParams] = { InName, InType };
+	++NumOutputOnlyParams;
+	return true;
+}
+
 /*----------------------------------- Error handling -----------------------------------*/
 
 void FEmitter::Error(FValueRef Source, FStringView Message)
@@ -2866,6 +2917,7 @@ static FValueRef CastToPrimitive(FEmitter& Emitter, FValueRef Value, FType Targe
 	{
 		// Get the first component of value. We already know value's type is primitive, so this will return a scalar.
 		Value = Emitter.Subscript(Value, 0);
+
 		ValuePrimitiveType = Value->Type.GetPrimitive();
 		check(ValuePrimitiveType.IsScalar());
 		
@@ -3759,6 +3811,101 @@ FValueRef FEmitter::PromoteSubstrateParameter()
 	return EmitPrototype(*this, Prototype);
 }
 
+const FFunction* FEmitter::FunctionHLSL(const FFunctionHLSLDesc& Desc)
+{
+	check(Desc.GetNumParameters() <= MIR::MaxNumFunctionParameters);
+
+	FFunctionHLSL Prototype = {};
+	Prototype.Name = Desc.Name;
+	Prototype.ReturnType = Desc.ReturnType;
+	Prototype.Code = Desc.Code;
+	Prototype.NumInputOnlyParams = Desc.NumInputOnlyParams;
+	Prototype.NumInputAndOutputParams = Desc.NumInputOnlyParams + Desc.NumInputOutputParams;
+	Prototype.NumParameters = Prototype.NumInputAndOutputParams + Desc.NumOutputOnlyParams;
+	Prototype.UniqueId = Module->FunctionHLSLs.Num();
+	Prototype.Defines = Desc.Defines;
+	Prototype.Includes = Desc.Includes;
+
+	// Copy over the parameter declarations from the description to the function prototype
+	for (uint32 i = 0; i < Prototype.NumParameters; ++i)
+	{
+		Prototype.Parameters[i] = Desc.Parameters[i];
+	}
+
+	// TODO: Optimize the lookup to be constant
+	for (const FFunctionHLSL* Function : Module->FunctionHLSLs)
+	{
+		if (Function->Equals(&Prototype))
+		{
+			return Function;
+		}
+	}
+
+	// Create the new MIR HLSL function instance and set it up
+	FFunctionHLSL* Function = new (FPrivate::Allocate(*this, sizeof(FFunctionHLSL), alignof(FFunctionHLSL))) FFunctionHLSL { Prototype };
+	Function->Name = Module->InternString(Function->Name);
+	Function->Code = Desc.Code;
+	Function->Defines = MakeArrayCopy(*this, Desc.Defines);
+	Function->Includes = MakeArrayCopy(*this, Desc.Includes);
+
+	// Add it to the module list
+	Module->FunctionHLSLs.Push(Function);
+
+	return Function;
+}
+
+FValueRef FEmitter::Call(const FFunction* Function, TConstArrayView<FValueRef> InputArguments)
+{
+	if (!Function)
+	{
+		return Poison();
+	}
+
+	if (InputArguments.Num() != Function->NumInputAndOutputParams)
+	{
+		Errorf(TEXT("Function called with incorrect number of arguments. Expected %d but got %d."), Function->NumInputAndOutputParams, InputArguments.Num());
+		return Poison();
+	}
+
+	FCall Call = MakePrototype<FCall>(Function->ReturnType);
+	Call.Function = Function;
+	Call.NumArguments = InputArguments.Num();
+
+	for (int i = 0; i < InputArguments.Num(); ++i)
+	{
+		Call.Arguments[i] = InputArguments[i];
+	}
+
+	return EmitPrototype(*this, Call);
+}
+
+FValueRef FEmitter::CallParameterOutput(FValueRef InCall, uint32 ParameterIndex)
+{
+	if (IsAnyNotValid(InCall))
+	{
+		return InCall.ToPoison();
+	}
+
+	FCall* Call = InCall->As<FCall>();
+	if (!Call)
+	{
+		Errorf(TEXT("Expected function call, found a '%s' value instead."), MIR::LexToString(InCall->Kind));
+		return InCall.ToPoison();
+	}
+	
+	if (ParameterIndex >= Call->Function->GetNumOutputParameters())
+	{
+		Errorf(InCall, TEXT("Invalid output index %d. Function has %d outputs."), ParameterIndex, Call->Function->GetNumOutputParameters());
+		return InCall.ToPoison();
+	}
+
+	FCallParameterOutput Proto = MakePrototype<FCallParameterOutput>(Call->Function->GetOutputParameter(ParameterIndex).Type);
+	Proto.Call = Call;
+	Proto.Index = ParameterIndex;
+
+	return EmitPrototype(*this, Proto);
+}
+
 void FEmitter::Initialize()
 {
 	// Create and reference the true/false constants.
diff --git a/Engine/Source/Runtime/Engine/Private/Materials/MaterialIRToHLSLTranslator.cpp b/Engine/Source/Runtime/Engine/Private/Materials/MaterialIRToHLSLTranslator.cpp
--- a/Engine/Source/Runtime/Engine/Private/Materials/MaterialIRToHLSLTranslator.cpp
+++ b/Engine/Source/Runtime/Engine/Private/Materials/MaterialIRToHLSLTranslator.cpp
@@ -228,6 +228,7 @@
 	switch (Instr->Kind)
 	{
 		case MIR::VK_VTPageTableRead: return false;
+		case MIR::VK_Call: return false;
 		default: break;
 	}
 
@@ -263,7 +264,6 @@
 	// HLSL code for evaluating non-normal material attributes per stage
 	FString EvaluateOtherMaterialAttributesHLSL[MIR::NumStages];
 
-
 	void GeneratePixelAttributesHLSL()
 	{
 		for (EMaterialProperty Property : UMaterialAggregate::GetMaterialAttributesProperties())
@@ -336,7 +336,7 @@
                 continue;
             }
             
-			if (Instr->NumUsers[CurrentStage] >= 1)
+			if (Instr->NumUsers[CurrentStage] >= 1 && !Instr->As<MIR::FCall>())
 			{
 				// Allocate a new local index for this instruction
 				uint32 LocalIndex = NumLocals++;
@@ -372,6 +372,8 @@
 			if (SetMaterialOutput && SetMaterialOutput->Property == MP_Normal) 
 			{
 				EvaluateNormalMaterialAttributeHLSL[CurrentStage] = MoveTemp(Printer.Buffer);
+				Printer.Tabs = 1;
+				Printer << Indentation;
 			}
         }
 
@@ -380,30 +382,32 @@
 		return NoOp;
 	}
 	
-	ENoOp LowerValue(const MIR::FValue* InValue)
+	ENoOp LowerValue(const MIR::FValue* Value)
 	{
-		if (const MIR::FInstruction* Instr = MIR::AsInstruction(InValue))
-		{
+		// Instruction results may be shared among other dependent instructions.
+		if (const MIR::FInstruction* Instr = MIR::AsInstruction(Value))
+		{ 
+			// If this instruction has only one user (dependant) and is foldable, then generate the full HLSL for the instruction in place.
 			if (Instr->NumUsers[CurrentStage] <= 1 && IsFoldable(Instr, CurrentStage))
 			{
 				Printer << LowerInstruction(Instr);
 			}
 			else
 			{
+				// Otherwise, this instruction has already been generated before, we emit a reference to the local that stores its result.
 				Printer << TEXTVIEW("_") << InstrToLocalIndex[Instr];
 			}
-
 			return NoOp;
 		}
 
-		switch (InValue->Kind)
+		switch (Value->Kind)
 		{
-			case MIR::VK_Constant: LowerConstant(static_cast<const MIR::FConstant*>(InValue)); break;
-			case MIR::VK_ExternalInput: LowerExternalInput(static_cast<const MIR::FExternalInput*>(InValue)); break;
-			case MIR::VK_MaterialParameterCollection: LowerMaterialParameterCollection(static_cast<const MIR::FMaterialParameterCollection*>(InValue)); break;
-			case MIR::VK_ScreenTexture: LowerScreenTexture(static_cast<const MIR::FScreenTexture*>(InValue)); break;
-			case MIR::VK_TextureObject: LowerTextureObject(static_cast<const MIR::FTextureObject*>(InValue)); break;
-			case MIR::VK_UniformParameter: LowerUniformParameter(static_cast<const MIR::FUniformParameter*>(InValue)); break;
+			case MIR::VK_Constant: LowerConstant(static_cast<const MIR::FConstant*>(Value)); break;
+			case MIR::VK_ExternalInput: LowerExternalInput(static_cast<const MIR::FExternalInput*>(Value)); break;
+			case MIR::VK_MaterialParameterCollection: LowerMaterialParameterCollection(static_cast<const MIR::FMaterialParameterCollection*>(Value)); break;
+			case MIR::VK_ScreenTexture: LowerScreenTexture(static_cast<const MIR::FScreenTexture*>(Value)); break;
+			case MIR::VK_TextureObject: LowerTextureObject(static_cast<const MIR::FTextureObject*>(Value)); break;
+			case MIR::VK_UniformParameter: LowerUniformParameter(static_cast<const MIR::FUniformParameter*>(Value)); break;
 			default: UE_MIR_UNREACHABLE();
 		}
 
@@ -426,6 +430,8 @@
 			case MIR::VK_PromoteSubstrateParameter: LowerPromoteSubstrateParameter(static_cast<const MIR::FPromoteSubstrateParameter*>(Instr)); break;
 			case MIR::VK_StageSwitch: LowerStageSwitch(static_cast<const MIR::FStageSwitch*>(Instr)); break;
 			case MIR::VK_HardwarePartialDerivative: LowerHardwarePartialDerivative(static_cast<const MIR::FHardwarePartialDerivative*>(Instr)); break;
+			case MIR::VK_Call: LowerCall(static_cast<const MIR::FCall*>(Instr)); break;
+			case MIR::VK_CallParameterOutput: LowerCallOutput(static_cast<const MIR::FCallParameterOutput*>(Instr)); break;
 			default: UE_MIR_UNREACHABLE();
 		}
 	
@@ -1806,6 +1812,129 @@
 		}
 	}
 
+	void LowerCall(const MIR::FCall* Call)
+	{
+		int32 ParamLocal = NumLocals;
+		
+		// Generate locals to store the output and input-output parameters
+		for (int32 i = Call->Function->NumInputOnlyParams; i < Call->Function->NumParameters; ++i)
+		{
+			Printer << LowerType(Call->Function->Parameters[i].Type) << TEXT(" _") << (NumLocals++);
+			
+			if (i < Call->NumArguments)
+			{
+				Printer << TEXT(" = ") << LowerValue(Call->Arguments[i]);
+			}
+
+			Printer << EndOfStatement;
+		}
+
+		// Print the local that will store the result and assign it to the call to the custom function.
+		Printer << LowerType(Call->Function->ReturnType) << TEXT(" _") << NumLocals << TEXT(" = ") << TEXT("C") << Call->Function->UniqueId << TEXT("_") << Call->Function->Name;
+		
+		Printer << BeginArgs << TEXT("Parameters");
+
+		// Print function call arguments
+		for (int32 i = 0; i < Call->Function->NumParameters; ++i)
+		{
+			Printer << TEXT(", ");
+
+			// Outputs and input-output parameters are stored in special locals. Refer to them
+			if (i >= Call->Function->NumInputOnlyParams)
+			{
+				Printer << TEXT("_") << ParamLocal++;
+			}
+			else
+			{
+				// Input-only parameters instead can inline their value instead
+				Printer << LowerValue(Call->Arguments[i]);
+			}
+		}
+		Printer << EndArgs;
+
+		// Assign a local to the call result
+		InstrToLocalIndex.Add(Call, NumLocals++);
+	}
+
+	void LowerCallOutput(const MIR::FCallParameterOutput* CallOutput)
+	{
+		const MIR::FCall* Call = CallOutput->Call->As<MIR::FCall>();
+		int32 ParamIndex = InstrToLocalIndex[Call];
+
+		// Void functions don't use a local
+		if (!Call->Function->ReturnType.IsVoid())
+		{
+			ParamIndex -= 1;
+		}
+
+		// Compute the index of the local that stores the additional output parameter 
+		ParamIndex = ParamIndex - Call->Function->NumInputAndOutputParams + Call->Function->NumInputOnlyParams + CallOutput->Index;
+
+		Printer << TEXT("_") << ParamIndex;
+	}
+
+	ENoOp LowerType(const MIR::FType& Type)
+	{
+		if (TOptional<MIR::FPrimitive> PrimitiveType = Type.AsPrimitive())
+		{
+			if (PrimitiveType->IsDouble())
+			{
+				if (PrimitiveType->IsScalar())
+				{
+					Printer << TEXTVIEW("FWSScalar");
+				}
+				else if (PrimitiveType->IsRowVector())
+				{
+					Printer << TEXTVIEW("FWSVector") << PrimitiveType->NumComponents();
+				}
+				else if (PrimitiveType->bIsLWCInverseMatrix)
+				{
+					Printer << TEXTVIEW("FWSInverseMatrix");
+				}
+				else
+				{
+					Printer << TEXTVIEW("FWSMatrix");
+				}
+			}
+			else
+			{
+				switch (PrimitiveType->ScalarKind)
+				{
+					case MIR::EScalarKind::Bool:  Printer << TEXTVIEW("bool"); break;
+					case MIR::EScalarKind::Int:   Printer << TEXTVIEW("int"); break;
+					case MIR::EScalarKind::Float: Printer << TEXTVIEW("MaterialFloat"); break;
+				}
+
+				if (PrimitiveType->NumRows == 1 && PrimitiveType->NumColumns > 1)
+				{
+					Printer << PrimitiveType->NumColumns;
+				}
+				else if (PrimitiveType->IsMatrix())
+				{
+					Printer << PrimitiveType->NumRows << TEXTVIEW("x") << PrimitiveType->NumColumns;
+				}
+			}
+		}
+		else if (Type.IsVoid())
+		{
+			Printer << TEXTVIEW("void");
+		}
+		else if (Type.IsSubstrateData())
+		{
+			Printer << TEXTVIEW("FSubstrateData");
+		}
+		else if (Type.IsVTPageTableResult())
+		{
+			Printer << TEXTVIEW("VTPageTableResult");
+		}
+		else
+		{
+			UE_MIR_UNREACHABLE();
+		}
+
+		return NoOp;
+	}
+	
 	/* Finalization */
 
 	void GenerateTemplateStringParameters(TMap<FString, FString>& Params)
@@ -1899,68 +2028,73 @@
 			);
 		}
 		Params.Add(TEXT("uniform_material_expressions"), MoveTemp(UniformMaterialExpressions));
-
 		Params.Add(TEXT("user_scene_texture_remap"), UE::MaterialTranslatorUtils::GenerateUserSceneTextureRemapHLSLDefines(Module->GetCompilationOutput()));
 	}
 	
-	ENoOp LowerType(const MIR::FType& Type)
+	// Generates the definitions of all the custom HLSL functions in the module and puts the resulting string
+	// into the "custom_functions" source template parameter.
+	void GenerateCustomFunctionsHLSL(TMap<FString, FString>& Params)
 	{
-		if (TOptional<MIR::FPrimitive> PrimitiveType = Type.AsPrimitive())
+		Printer.Tabs = 0;
+
+		for (const MIR::FFunctionHLSL* Function : Module->GetFunctionHLSLs())
 		{
-			if (PrimitiveType->IsDouble())
+			// Print the user-specified defines
+			for (const MIR::FFunctionHLSLDefine& Define : Function->Defines)
 			{
-				if (PrimitiveType->IsScalar())
-				{
-					Printer << TEXTVIEW("FWSScalar");
-				}
-				else if (PrimitiveType->IsRowVector())
-				{
-					Printer << TEXTVIEW("FWSVector") << PrimitiveType->NumComponents();
-				}
-				else if (PrimitiveType->bIsLWCInverseMatrix)
-				{
-					Printer << TEXTVIEW("FWSInverseMatrix");
-				}
-				else
-				{
-					Printer << TEXTVIEW("FWSMatrix");
-				}
+				Printer << TEXTVIEW("#ifndef ") << Define.Name << NewLine;
+				Printer << TEXTVIEW("\t#define ") << Define.Name << TEXTVIEW(" ") << Define.Value << NewLine;
+				Printer << TEXTVIEW("#endif") << NewLine;
 			}
-			else
+
+			// Print the user-specified include directives
+			for (FStringView Include : Function->Includes)
 			{
-				switch (PrimitiveType->ScalarKind)
-				{
-					case MIR::EScalarKind::Bool:  Printer << TEXTVIEW("bool"); break;
-					case MIR::EScalarKind::Int:   Printer << TEXTVIEW("int"); break;
-					case MIR::EScalarKind::Float: Printer << TEXTVIEW("MaterialFloat"); break;
-				}
+				Printer << TEXTVIEW("#include \"") << Include << TEXTVIEW("\"") << NewLine;
+			}
 
-				if (PrimitiveType->NumRows == 1 && PrimitiveType->NumColumns > 1)
-				{
-					Printer << PrimitiveType->NumColumns;
-				}
-				else if (PrimitiveType->IsMatrix())
-				{
-					Printer << PrimitiveType->NumRows << TEXTVIEW("x") << PrimitiveType->NumColumns;
-				}
+			// Write the custom function signature, for example "C5_MyCustomNode".
+			// - 'C' is only a "namespace" for custom functions
+			// - '5' is a unique id used to disambiguate distinct custom functions with the same name.
+			Printer << LowerType(Function->ReturnType) << TEXTVIEW(" C") << Function->UniqueId << TEXT("_") << Function->Name << BeginArgs << TEXT("FMaterialPixelParameters Parameters");
+			
+			// Write the parameter declarations
+			for (uint16 i = 0; i < Function->NumParameters; ++i)
+			{
+				Printer << TEXTVIEW(", ");
+
+				// Print the io keyword
+				Printer << (i < Function->NumInputOnlyParams ? TEXTVIEW("")
+						:  i < Function->NumInputAndOutputParams ? TEXTVIEW("inout ")
+						: TEXTVIEW("out "));
+
+				// Type and name
+				Printer << LowerType(Function->Parameters[i].Type) << TEXTVIEW(" ") << *Function->Parameters[i].Name.ToString();
 			}
-		}
-		else if (Type.IsSubstrateData())
-		{
-			Printer << TEXTVIEW("FSubstrateData");
-		}
-		else if (Type.IsVTPageTableResult())
-		{
-			Printer << TEXTVIEW("VTPageTableResult");
-		}
-		else
-		{
-			UE_MIR_UNREACHABLE();
+
+			Printer << EndArgs << NewLine << OpenBrace;
+			
+			// If the function does not contain a "return" keyword, add one.
+			bool bContainsReturn = Function->Code.Contains(TEXT("return"));
+			if (!bContainsReturn)
+			{
+				Printer << TEXT("return") << NewLine;
+			}
+			
+			// Write the function code
+			Printer << Function->Code << NewLine;
+
+			if (!bContainsReturn)
+			{
+				Printer << TEXT(";");
+			}
+
+			Printer << NewLine << CloseBrace << TEXTVIEW("\n\n");
 		}
 
-		return NoOp;
+		Params.Add(TEXT("custom_functions"), MoveTemp(Printer.Buffer));
 	}
-	
+
 	void GetShaderCompilerEnvironment(FShaderCompilerEnvironment& OutEnvironment)
 	{
 		const FMaterialCompilationOutput& CompilationOutput = Module->GetCompilationOutput();
@@ -2155,6 +2289,7 @@
 	Private.GenerateOtherStageHLSL(MIR::Stage_Pixel);
 	Private.GenerateOtherStageHLSL(MIR::Stage_Compute);
 	Private.GenerateTemplateStringParameters(OutParameters);
+	Private.GenerateCustomFunctionsHLSL(OutParameters);
 	Private.GetShaderCompilerEnvironment(OutEnvironment);
 }
 
diff --git a/Engine/Source/Runtime/Engine/Private/Materials/MaterialIRTypes.cpp b/Engine/Source/Runtime/Engine/Private/Materials/MaterialIRTypes.cpp
--- a/Engine/Source/Runtime/Engine/Private/Materials/MaterialIRTypes.cpp
+++ b/Engine/Source/Runtime/Engine/Private/Materials/MaterialIRTypes.cpp
@@ -16,8 +16,8 @@ const TCHAR* TypeKindToString(ETypeKind Kind)
 		case ETypeKind::Void: return TEXT("Void");
 		case ETypeKind::Primitive: return TEXT("Primitive");
 		case ETypeKind::Aggregate: return TEXT("Aggregate");
-		case ETypeKind::ParameterCollection: return TEXT("ParameterCollection");
 		case ETypeKind::Texture2D: return TEXT("Texture2D");
+		case ETypeKind::ParameterCollection: return TEXT("ParameterCollection");
 		case ETypeKind::SubstrateData: return TEXT("SubstrateData");
 		case ETypeKind::VTPageTableResult: return TEXT("VTPageTableResult");
 		default: UE_MIR_UNREACHABLE();
@@ -90,11 +90,11 @@ FType FType::FromMaterialValueType(EMaterialValueType Type)
 	switch (Type)
 	{
 		case MCT_VoidStatement: return FType::MakeVoid();
+		case MCT_Float: return FType::MakeFloatScalar();
 		case MCT_Float1: return FType::MakeFloatScalar();
 		case MCT_Float2: return FType::MakeFloatVector(2);
 		case MCT_Float3: return FType::MakeFloatVector(3);
 		case MCT_Float4: return FType::MakeFloatVector(4);
-		case MCT_Float: return FType::MakeFloatVector(4);
 		case MCT_Texture2D: return FType::MakeTexture2D();
 		case MCT_StaticBool: return FType::MakeBoolScalar();
 		case MCT_UInt1: return FType::MakeIntScalar();
diff --git a/Engine/Source/Runtime/Engine/Public/Materials/MaterialExpressionCustom.h b/Engine/Source/Runtime/Engine/Public/Materials/MaterialExpressionCustom.h
--- a/Engine/Source/Runtime/Engine/Public/Materials/MaterialExpressionCustom.h
+++ b/Engine/Source/Runtime/Engine/Public/Materials/MaterialExpressionCustom.h
@@ -98,6 +98,7 @@ class UMaterialExpressionCustom : public UMaterialExpression
 	
 	//~ Begin UMaterialExpression Interface
 #if WITH_EDITOR
+	virtual void Build(MIR::FEmitter& Emitter) override;
 	virtual int32 Compile(class FMaterialCompiler* Compiler, int32 OutputIndex) override;
 	virtual void GetCaption(TArray<FString>& OutCaptions) const override;
 	virtual FText GetCreationName() const override { return FText::FromString(TEXT("Custom")); }
diff --git a/Engine/Source/Runtime/Engine/Public/Materials/MaterialIR.h b/Engine/Source/Runtime/Engine/Public/Materials/MaterialIR.h
--- a/Engine/Source/Runtime/Engine/Public/Materials/MaterialIR.h
+++ b/Engine/Source/Runtime/Engine/Public/Materials/MaterialIR.h
@@ -124,6 +124,8 @@ enum EValueKind
 	VK_PromoteSubstrateParameter,
 	VK_StageSwitch,
 	VK_HardwarePartialDerivative,
+	VK_Call,
+	VK_CallParameterOutput,
 
 	VK_InstructionEnd,
 };
@@ -576,7 +578,7 @@ struct FInstruction : FValue
 	FBlock* GetDesiredBlockForUse(EStage Stage, int32 UseIndex);
 };
 
-template <EValueKind TTypeKind, uint32 TNumStaticUses>
+template <EValueKind TTypeKind, uint32 TNumStaticUses = 0>
 struct TInstruction : FInstruction
 {
 	// The kind of this instruction.
@@ -963,5 +965,115 @@ struct FHardwarePartialDerivative : TInstruction<VK_HardwarePartialDerivative, 1
 	EDerivativeAxis Axis;
 };
 
+// Kind of user-specified function.
+enum class FFunctionKind
+{
+	HLSL,
+};
+
+// Describes a function parameter.
+struct FFunctionParameter
+{
+	// The parameter name (used as identifier).
+	FName Name;
+	
+	// The parameter type.
+	FType Type;
+};
+
+// Maximum number of parameters a user function can have.
+constexpr uint32 MaxNumFunctionParameters = 32;
+
+// Defines a user-specified function.
+struct FFunction
+{
+	// The kind of function this is.
+	FFunctionKind Kind;
+
+	// This function's name (could be used as identifier).
+	FStringView Name;
+
+	// The function return value type
+	FType ReturnType;
+
+	// Number of input-only parameters.
+	uint16 NumInputOnlyParams;
+
+	// Number of input-output parameters.
+	uint16 NumInputAndOutputParams;
+
+	// Total number of parameters (input-only + input-output + output-only)
+	uint16 NumParameters;
+
+	// A module-unique id assigned to this function to disambiguate it with other functions with identical name.
+	uint16 UniqueId;
+
+	// The parameters this function has (in order: input-only, then input-output, finally output-only).
+	FFunctionParameter Parameters[MaxNumFunctionParameters];
+
+	// Returns the number of output parameters (input-output + output-only).
+	uint32 GetNumOutputParameters() const { return NumParameters - NumInputOnlyParams; }
+
+	// Returns the input-output/output-only parameter with given index.
+	const FFunctionParameter& GetOutputParameter(uint32 Index) const { return Parameters[Index + NumInputOnlyParams]; }
+
+	// Returns whether parameter with specified index is input.
+	bool IsInputParameter(uint32 Index) const { return Index < NumInputAndOutputParams; }
+	
+	// Returns whether parameter with specified index is output.
+	bool IsOutputParameter(uint32 Index) const { return Index >= NumInputOnlyParams; }
+
+	// Whether this function equals specified (they have the same data).
+	bool Equals(const FFunction* Other) const;
+};
+
+// Describes an HLSL preprocessor define, .e.g "#define MyDefine 123"
+struct FFunctionHLSLDefine
+{
+	// The define name.
+	FStringView Name;
+	
+	// The define value.
+	FStringView Value;
+};
+
+// A user-function defined with user-provided HLSL code.
+struct FFunctionHLSL : FFunction
+{
+	// The HLSL code of the body of this function.
+	FStringView Code;
+
+	TConstArrayView<FFunctionHLSLDefine> Defines;
+
+	TConstArrayView<FStringView> Includes;
+
+	// Whether this function equals specified (they have the same data).
+	bool Equals(const FFunctionHLSL* Other) const;
+};
+
+// A call instruction to a user-function.
+struct FCall : TInstruction<VK_Call, 0>
+{
+	// The user-function to call.
+	const FFunction* Function;
+
+	// The arguments to pass to the function parameters.
+	FValue* Arguments[MaxNumFunctionParameters];
+	
+	// The number of arguments.
+	int NumArguments;
+
+};
+
+// Instruction that evaluates some user-function call output parameter result.
+struct FCallParameterOutput : TInstruction<VK_CallParameterOutput, 1>
+{
+	// The function call.
+	FValue* Call;
+	
+	// The index of the output parameter to fetch the value from.
+	int32 	Index;
+};
+
 } // namespace MIR
 #endif // WITH_EDITOR
diff --git a/Engine/Source/Runtime/Engine/Public/Materials/MaterialIREmitter.h b/Engine/Source/Runtime/Engine/Public/Materials/MaterialIREmitter.h
--- a/Engine/Source/Runtime/Engine/Public/Materials/MaterialIREmitter.h
+++ b/Engine/Source/Runtime/Engine/Public/Materials/MaterialIREmitter.h
@@ -109,16 +109,59 @@ struct FValueRef
 	const FExpressionInput* Input{};
 };
 
-//
+// An assignment to a material aggregate attribute, used when constructing an aggregate instance.
 struct FAttributeAssignment
 {
-	//
+	// Then name of the assigned aggregate attribute.
 	FName Name;
 
-	//
+	// The value it is assigned to.
 	FValueRef Value;
 };
 
+// Describes a user-specified HLSL function.
+struct FFunctionHLSLDesc
+{
+	// The function name, used for debugging purposes and referenced in the generated HLSL. 
+	FStringView Name{};
+
+	// The function return type.
+	FType ReturnType{};
+
+	// The number of parameters that are input-only.
+	uint32 NumInputOnlyParams{};
+
+	// The number of input-output parameters after the input-only parameters.
+	uint32 NumInputOutputParams{};
+
+	// The number of output-only parameters after the input-output parameters.
+	uint32 NumOutputOnlyParams{};
+
+	// Array of descriptions of the function parameters.
+	FFunctionParameter Parameters[MIR::MaxNumFunctionParameters];
+
+	// The user-specified HLSL code inside the function body.
+	FStringView Code{};
+
+	// Array of user-specified defines to declare before this function.
+	TConstArrayView<FFunctionHLSLDefine> Defines;
+
+	// Array of user-specified #include directives to declare before this function.
+	TConstArrayView<FStringView> Includes;
+
+	// Pushes an input-only parameter and returns whether there was an available free parameter to push this one into.
+	bool PushInputOnlyParameter(FName Name, FType Type);
+	
+	// Pushes an input-output parameter and returns whether there was an available free parameter to push this one into.
+	bool PushInputOutputParameter(FName Name, FType Type);
+	
+	// Pushes an output-only parameter and returns whether there was an available free parameter to push this one into.
+	bool PushOutputOnlyParameter(FName Name, FType Type);
+
+	// Returns the total number of parameters.
+	uint32 GetNumParameters() const { return NumInputOnlyParams + NumInputOutputParams + NumOutputOnlyParams; }
+};
+
 // The FEmitter is a helper object responsible for emitting MIR values and instructions.
 // It is primarily used within UMaterialExpression::Build() implementations to lower the
 // high-level semantics of a material expression into corresponding lower-level MIR nodes.
@@ -554,6 +597,17 @@ public:
 	// Emits an instruction that promotes a Substrate parameter, encoded as MCT_Substrate.
 	FValueRef PromoteSubstrateParameter();
 
+	// Defines a user-specified HLSL function to embed in the generated HLSL when generated IR Module is translated to HLSL.
+	// This will create a local function with the specified description. You can issue calls to the returned functions using the Call() API.
+	const FFunction* FunctionHLSL(const FFunctionHLSLDesc& Desc);
+
+	// Emits a call to a specified user function using specified input arguments.
+	// Note that the InputArguments must provide an argument for *all* function input parameters (that is including the input-output ones).
+	FValueRef Call(const FFunction* Function, TConstArrayView<FValueRef> InputArguments);
+
+	// Emits an instruction that retrieves the output parameter value from a function call.
+	FValueRef CallParameterOutput(FValueRef Call, uint32 ParameterIndex);
+
 	/*--------------------------- Unary operators shortcuts ----------------------------*/
 
 	FValueRef BitwiseNot(FValueRef A) { return Operator(UO_BitwiseNot, A); }
diff --git a/Engine/Source/Runtime/Engine/Public/Materials/MaterialIRModule.h b/Engine/Source/Runtime/Engine/Public/Materials/MaterialIRModule.h
--- a/Engine/Source/Runtime/Engine/Public/Materials/MaterialIRModule.h
+++ b/Engine/Source/Runtime/Engine/Public/Materials/MaterialIRModule.h
@@ -89,6 +89,9 @@ public:
 	// Returns the list of material parameter collections used by this material.
 	TConstArrayView<UMaterialParameterCollection*> GetParameterCollections() const { return ParameterCollections; }
 
+	// Returns the array of HLSL-defined user-functions in the module.
+	TConstArrayView<const MIR::FFunctionHLSL*> GetFunctionHLSLs() const { return FunctionHLSLs; };
+
 	// Stores a user-defined string and returns a view of it.
 	FStringView InternString(FStringView InString);
 
@@ -144,6 +147,9 @@ private:
 	// Parameter collections used by this material.
 	TArray<UMaterialParameterCollection*> ParameterCollections;
 
+	// Array of HLSL-defined user-functions in the module.
+	TArray<const MIR::FFunctionHLSL*> FunctionHLSLs;
+
 	// List of compilation errors.
 	TArray<FError> Errors;
 
diff --git a/Engine/Source/Runtime/Engine/Public/Materials/MaterialIRTypes.h b/Engine/Source/Runtime/Engine/Public/Materials/MaterialIRTypes.h
--- a/Engine/Source/Runtime/Engine/Public/Materials/MaterialIRTypes.h
+++ b/Engine/Source/Runtime/Engine/Public/Materials/MaterialIRTypes.h
@@ -21,8 +21,8 @@ enum class ETypeKind : uint8
 
 	// Opaque Objects
 	
-	ParameterCollection,	// Material parameter collection object.
 	Texture2D,				// A 2D texture object.
+	ParameterCollection,	// Material parameter collection object.
 	SubstrateData,			// Substrate-specific data structure.
 	VTPageTableResult,		// A virtual texture page table lookup result.
 };
@@ -203,6 +203,9 @@ struct FType
 	// Returns whether this is the poison type, indicating an error.
 	bool IsPoison() const { return Kind == ETypeKind::Poison; }
 	
+	// Returns whether this is the void type.
+	bool IsVoid() const { return Kind == ETypeKind::Void; }
+
 	// Returns true if the given type is a primitive that supports arithmetic operations.
 	bool IsPrimitive() const { return Is(ETypeKind::Primitive); }
 
