Added command to print a struct's type hierarchy.

#rnx

diff --git a/Engine/Source/Runtime/Engine/Private/PrintObjectUtils.cpp b/Engine/Source/Runtime/Engine/Private/PrintObjectUtils.cpp
--- a/Engine/Source/Runtime/Engine/Private/PrintObjectUtils.cpp
+++ b/Engine/Source/Runtime/Engine/Private/PrintObjectUtils.cpp
@@ -383,15 +383,13 @@ void PrintObjectIDO(UObject* Object, EPrintObjectFlag Flags, FOutputDevice* InOu
 	}
 }
 
-void PrintClassDefaultObject(UClass* Class, EPrintObjectFlag Flags, FOutputDevice* InOutputDevice)
+void PrintClassDefaultObject(const UClass* Class, EPrintObjectFlag Flags, FOutputDevice* InOutputDevice)
 {
 	FOutputDevice* OutputDevice = InOutputDevice ? InOutputDevice : GLog;
 	if (OutputDevice)
 	{
 		if (Class)
 		{
-			FString ClassPath = Class->GetPathName();
-
 			UObject* CDO = Class->GetDefaultObject(/*bCreateIfNeeded*/false);
 
 			bool bVerbose = EnumHasAnyFlags(Flags, EPrintObjectFlag::Verbose);
@@ -413,6 +411,35 @@ void PrintClassDefaultObject(UClass* Class, EPrintObjectFlag Flags, FOutputDevic
 	}
 }
 
+void PrintStructHierarchy(const UStruct* Struct, EPrintObjectFlag Flags, FOutputDevice* InOutputDevice)
+{
+	FOutputDevice* OutputDevice = InOutputDevice ? InOutputDevice : GLog;
+	if (OutputDevice)
+	{
+		if (Struct)
+		{
+			bool bVerbose = EnumHasAnyFlags(Flags, EPrintObjectFlag::Verbose);
+
+			const UStruct* CurrentStruct = Struct;
+
+			TStringBuilder<16> PaddingBuilder;
+
+			while (CurrentStruct)
+			{
+				LogObjectState(CurrentStruct, OutputDevice, bVerbose, *PaddingBuilder);
+
+				CurrentStruct = CurrentStruct->GetSuperStruct();
+
+				PaddingBuilder.Append(TEXT("  "));
+			}
+		}
+		else
+		{
+			OutputDevice->Logf(TEXT("PrintStructHierarchy: NULL struct"));
+		}
+	}
+}
+
 static void ParseObjectIdAndFlags(const TArray<FString>& Args, FString& OutObjectIdentifier, EPrintObjectFlag& OutFlags)
 {
 	OutFlags = EPrintObjectFlag::None;
@@ -449,7 +476,7 @@ static void LogHelp(FOutputDevice& OutputDevice)
 
 FAutoConsoleCommandWithArgsAndOutputDevice CVarCommandPrintObjectsInOuter(
 	TEXT("Obj.DumpObjectsInOuter"),
-	TEXT("Lists all objects under a specified parent (the parent object must be specified as a path, e.g., /MyLevel/MyLevel.MyLevel:PersistentLevel)."),
+	TEXT("Lists all objects under a specified parent (the parent object must be specified as a name or path, e.g., /MyLevel/MyLevel.MyLevel:PersistentLevel)."),
 	FConsoleCommandWithArgsAndOutputDeviceDelegate::CreateLambda([](const TArray<FString>& Args, FOutputDevice& OutputDevice)
 	{
 		if (Args.Num() > 0)
@@ -498,7 +525,7 @@ FAutoConsoleCommandWithArgsAndOutputDevice CVarCommandPrintObjectsWithName(
 
 FAutoConsoleCommandWithArgsAndOutputDevice CVarCommandPrintObjectProperties(
 	TEXT("Obj.DumpProperties"),
-	TEXT("Lists the properties of an object (the object must be specified as a path, e.g., /MyLevel/MyLevel.MyLevel:PersistentLevel.MyActor). Pass InitState=true to iclude the properties' initialization state."),
+	TEXT("Lists the properties of an object (the object must be specified as a name or path, e.g., /MyLevel/MyLevel.MyLevel:PersistentLevel.MyActor). Pass InitState=true to iclude the properties' initialization state."),
 	FConsoleCommandWithArgsAndOutputDeviceDelegate::CreateLambda([](const TArray<FString>& Args, FOutputDevice& OutputDevice)
 	{
 		if (Args.Num() > 0)
@@ -527,7 +554,7 @@ FAutoConsoleCommandWithArgsAndOutputDevice CVarCommandPrintObjectProperties(
 
 FAutoConsoleCommandWithArgsAndOutputDevice CVarCommandPrintObjectArchetype(
 	TEXT("Obj.DumpArchetype"),
-	TEXT("Outputs an object's archetype (the object must be specified as a path, e.g., /MyLevel/MyLevel.MyLevel:PersistentLevel.MyActor)."),
+	TEXT("Outputs an object's archetype (the object must be specified as a name or path, e.g., /MyLevel/MyLevel.MyLevel:PersistentLevel.MyActor)."),
 	FConsoleCommandWithArgsAndOutputDeviceDelegate::CreateLambda([](const TArray<FString>& Args, FOutputDevice& OutputDevice)
 	{
 		if (Args.Num() > 0)
@@ -556,7 +583,7 @@ FAutoConsoleCommandWithArgsAndOutputDevice CVarCommandPrintObjectArchetype(
 
 FAutoConsoleCommandWithArgsAndOutputDevice CVarCommandPrintObjectIDO(
 	TEXT("Obj.DumpIDO"),
-	TEXT("Outputs an object's IDO (Instance Data Object) (the object must be specified as a path, e.g., /MyLevel/MyLevel.MyLevel:PersistentLevel.MyActor)."),
+	TEXT("Outputs an object's IDO (Instance Data Object) (the object must be specified as a name or path, e.g., /MyLevel/MyLevel.MyLevel:PersistentLevel.MyActor)."),
 	FConsoleCommandWithArgsAndOutputDeviceDelegate::CreateLambda([](const TArray<FString>& Args, FOutputDevice& OutputDevice)
 	{
 		if (Args.Num() > 0)
@@ -611,6 +638,35 @@ FAutoConsoleCommandWithArgsAndOutputDevice CVarCommandPrintCDO(
 			LogHelp(OutputDevice);
 		}
 	}));
+
+FAutoConsoleCommandWithArgsAndOutputDevice CVarCommandPrintStructHierarchy(
+	TEXT("Obj.DumpStructHierarchy"),
+	TEXT("Outputs a struct (or class) type hierarchy."),
+	FConsoleCommandWithArgsAndOutputDeviceDelegate::CreateLambda([](const TArray<FString>& Args, FOutputDevice& OutputDevice)
+	{
+		if (Args.Num() > 0)
+		{
+			FString StructPath;
+			EPrintObjectFlag Flags;
+			ParseObjectIdAndFlags(Args, StructPath, Flags);
+
+			UStruct* Struct = FindFirstObjectSafe<UStruct>(*StructPath);
+
+			if (Struct)
+			{
+				PrintStructHierarchy(Struct, Flags, &OutputDevice);
+			}
+			else
+			{
+				OutputDevice.Logf(TEXT("PrintStructHierarchy: failed to find any types for path: '%s'"), *StructPath);
+			}
+		}
+		else
+		{
+			OutputDevice.Logf(TEXT("PrintStructHierarchy: no struct name or path specified (example usage: Obj.DumpStructHierarchy MyActor)"));
+			LogHelp(OutputDevice);
+		}
+	}));
 	
 } // namespace UE
 
diff --git a/Engine/Source/Runtime/Engine/Public/PrintObjectUtils.h b/Engine/Source/Runtime/Engine/Public/PrintObjectUtils.h
--- a/Engine/Source/Runtime/Engine/Public/PrintObjectUtils.h
+++ b/Engine/Source/Runtime/Engine/Public/PrintObjectUtils.h
@@ -98,7 +98,16 @@ UE_INTERNAL void PrintObjectIDO(UObject* Object, EPrintObjectFlag Flags = EPrint
  * @param Flags For specifying optional behaviors.
  * @param OutputDevice The output device to print to (can be NULL).
  */
-UE_INTERNAL void PrintClassDefaultObject(UClass* Class, EPrintObjectFlag Flags = EPrintObjectFlag::None, FOutputDevice* OutputDevice = nullptr);
+UE_INTERNAL void PrintClassDefaultObject(const UClass* Class, EPrintObjectFlag Flags = EPrintObjectFlag::None, FOutputDevice* OutputDevice = nullptr);
+
+/**
+ * Outputs a struct's type hierarchy.
+ *
+ * @param Struct The struct to use.
+ * @param Flags For specifying optional behaviors.
+ * @param OutputDevice The output device to print to (can be NULL).
+ */
+UE_INTERNAL void PrintStructHierarchy(const UStruct* Struct, EPrintObjectFlag Flags = EPrintObjectFlag::None, FOutputDevice* OutputDevice = nullptr);
 
 } // namespace UE
 
