Fix for minor comments and error message inaccuracies spotted by BugHawk.

#rnx

diff --git a/Engine/Source/Runtime/Engine/Private/Materials/MaterialExpressionsIR.cpp b/Engine/Source/Runtime/Engine/Private/Materials/MaterialExpressionsIR.cpp
--- a/Engine/Source/Runtime/Engine/Private/Materials/MaterialExpressionsIR.cpp
+++ b/Engine/Source/Runtime/Engine/Private/Materials/MaterialExpressionsIR.cpp
@@ -3127,12 +3127,11 @@ void UMaterialExpressionSphereMask::Build(MIR::FEmitter& Em)
 }
 
 // Takes a description user string and turns it into a valid C/HLSL identifier.
-// Note: we allow the identifier to start by numbers.
 static FString DescriptionToIdentifier(FStringView Source)
 {
 	FString Out;
 	Out.Reserve(Source.Len());
-	// Append an underscore if the source starts by a igit
+	// Append an underscore if the source starts by a digit
 	if (!Source.IsEmpty() && FChar::IsDigit(Source[0]))
 	{
 		Out.AppendChar('_');
@@ -3234,7 +3233,7 @@ void UMaterialExpressionCustom::Build(MIR::FEmitter& Em)
 	{
 		if (IncludeFilePaths[i].IsEmpty())
 		{
-			Em.Errorf(TEXT("Define with index '%d' has no valid value."), i);
+			Em.Errorf(TEXT("Include with index '%d' has no valid value."), i);
 		}
 
 		Includes[i] = IncludeFilePaths[i];
