SplineComponent - Add GetDefaultSpline, a static function for returning the value of Spline on the CDO

diff --git a/Engine/Source/Runtime/Engine/Classes/Components/SplineComponent.h b/Engine/Source/Runtime/Engine/Classes/Components/SplineComponent.h
--- a/Engine/Source/Runtime/Engine/Classes/Components/SplineComponent.h
+++ b/Engine/Source/Runtime/Engine/Classes/Components/SplineComponent.h
@@ -404,6 +404,7 @@ public:
 	/** Reset the spline to its default shape (a spline of two points) */
 	ENGINE_API void ResetToDefault();
 	ENGINE_API bool CanResetToDefault() const;
+	ENGINE_API static FSpline GetDefaultSpline();
 
 	/** Update the spline tangents and SplineReparamTable */
 	UFUNCTION(BlueprintCallable, Category = Spline)
diff --git a/Engine/Source/Runtime/Engine/Private/Components/SplineComponent.cpp b/Engine/Source/Runtime/Engine/Private/Components/SplineComponent.cpp
--- a/Engine/Source/Runtime/Engine/Private/Components/SplineComponent.cpp
+++ b/Engine/Source/Runtime/Engine/Private/Components/SplineComponent.cpp
@@ -152,6 +152,11 @@ bool USplineComponent::CanResetToDefault() const
 	return Spline != CastChecked<USplineComponent>(GetArchetype())->Spline;
 }
 
+FSpline USplineComponent::GetDefaultSpline()
+{
+	return StaticClass()->GetDefaultObject<USplineComponent>()->Spline;
+}
+
 void USplineComponent::SetDefaultSpline()
 {
 	const FVector StartPoint(0.0, 0.0, 0.0);
@@ -2351,11 +2356,6 @@ FPrimitiveSceneProxy* USplineComponent::CreateSceneProxy()
 		{
 			QUICK_SCOPE_CYCLE_COUNTER(STAT_SplineSceneProxy_GetDynamicMeshElements);
 
-			if (IsSelected())
-			{
-				return;
-			}
-
 			for (int32 ViewIndex = 0; ViewIndex < Views.Num(); ViewIndex++)
 			{
 				if (VisibilityMap & (1 << ViewIndex))
@@ -2380,7 +2380,7 @@ FPrimitiveSceneProxy* USplineComponent::CreateSceneProxy()
 		virtual FPrimitiveViewRelevance GetViewRelevance(const FSceneView* View) const override
 		{
 			FPrimitiveViewRelevance Result;
-			Result.bDrawRelevance = bDrawDebug && !IsSelected() && IsShown(View) && View->Family->EngineShowFlags.Splines;
+			Result.bDrawRelevance = bDrawDebug && IsShown(View) && View->Family->EngineShowFlags.Splines;
 			Result.bDynamicRelevance = true;
 			Result.bShadowRelevance = IsShadowCast(View);
 			Result.bEditorPrimitiveRelevance = UseEditorCompositing(View);
