[MIR] Implement remaining texture object types in new material translator.

Also fix some issues with the transform expressions:
1. Matrix dimension down-casting (e.g. `(float3x3)My4x4Matrix`) was not handled correctly.
2. `MCT_Float` was interpreted as float4 when it should be interpreted as float in the new translator.
3. Matrices must be multiplied explicitly with the BO_MatrixMultiply operator instead of BO_Multiply. This also can't be folded since `mul()` in HLSL is not foldable per operand.

#test Custom materials for volume-, 2D array-, and Cube array textures. Also ran Rendering/Materials automation tests which all succeeded.
#rnx

diff --git a/Engine/Source/Runtime/Engine/Private/Materials/MaterialExpressionsIR.cpp b/Engine/Source/Runtime/Engine/Private/Materials/MaterialExpressionsIR.cpp
--- a/Engine/Source/Runtime/Engine/Private/Materials/MaterialExpressionsIR.cpp
+++ b/Engine/Source/Runtime/Engine/Private/Materials/MaterialExpressionsIR.cpp
@@ -2,6 +2,8 @@
 
 #if WITH_EDITOR
 
+#include "MaterialExpressionsIR.h"
+
 #include "Curves/CurveLinearColorAtlas.h"
 #include "Engine/Engine.h"
 #include "Engine/Font.h"
@@ -1248,7 +1250,7 @@ static MIR::ETextureReadMode TextureGatherModeToMIR(::ETextureGatherMode Mode)
 	}
 }
 
-static void BuildTextureSampleExpression(MIR::FEmitter& Em, UMaterialExpressionTextureSample* Expr, FValueRef Texture)
+static void BuildTextureSampleExpression(MIR::FEmitter& Em, UMaterialExpressionTextureSample* Expr, FValueRef Texture, EMaterialValueType TextureType)
 {
 	FValueRef TexCoords = Em.TryInput(&Expr->Coordinates);
 	if (!TexCoords)
@@ -1269,15 +1271,23 @@ static void BuildTextureSampleExpression(MIR::FEmitter& Em, UMaterialExpressionT
 	} 
 	else
 	{
-		// Determine if automatic view mip bias should be used, by trying ot acquire a its input as a static boolean.
+		// Determine if automatic view mip bias should be used, by trying to acquire a its input as a static boolean.
 		bool bAutomaticViewMipBias = Em.ToConstantBool(Em.InputDefaultBool(&Expr->AutomaticViewMipBiasValue, Expr->AutomaticViewMipBias));
 	
+		// If not 2D texture, disable AutomaticViewMipBias.
+		if (!(TextureType & (MCT_Texture2D | MCT_TextureVirtual | MCT_TextureMeshPaint)))
+		{
+			bAutomaticViewMipBias = false;
+		}
+
 		UE_MIR_CHECKPOINT(Em);
 
 		// Get the mip value level (either through the expression input or using the given constant if disconnected).
-		FValueRef MipValue = (Expr->MipValueMode == TMVM_MipLevel || Expr->MipValueMode == TMVM_MipBias)
-			? Em.CheckIsScalar(Em.InputDefaultInt(&Expr->MipValue, Expr->ConstMipValue))
-			: Em.Poison();
+		FValueRef MipValue;
+		if (Expr->MipValueMode == TMVM_MipLevel || Expr->MipValueMode == TMVM_MipBias)
+		{
+			MipValue = Em.CheckIsScalar(Em.InputDefaultInt(&Expr->MipValue, Expr->ConstMipValue));
+		}
 
 		switch (Expr->MipValueMode)
 		{
@@ -1341,7 +1351,7 @@ void UMaterialExpressionTextureSample::Build(MIR::FEmitter& Em)
 
 	UE_MIR_CHECKPOINT(Em);
 
-	BuildTextureSampleExpression(Em, this, TextureValue);
+	BuildTextureSampleExpression(Em, this, TextureValue, Texture->GetMaterialType());
 }
 
 static FValueRef BuildTextureObjectParameter(MIR::FEmitter& Em, UMaterialExpressionTextureSampleParameter* Expr)
@@ -1365,9 +1375,15 @@ static FValueRef BuildTextureObjectParameter(MIR::FEmitter& Em, UMaterialExpress
 
 void UMaterialExpressionTextureSampleParameter::Build(MIR::FEmitter& Em)
 {
+	FString ErrorMessage;
+	if (!TextureIsValid(Texture, ErrorMessage))
+	{
+		Em.Error(ErrorMessage);
+		return;
+	}
 	FValueRef ParameterValue = BuildTextureObjectParameter(Em, this);
 	UE_MIR_CHECKPOINT(Em);
-	BuildTextureSampleExpression(Em, this, ParameterValue);
+	BuildTextureSampleExpression(Em, this, ParameterValue, Texture->GetMaterialType());
 }
 
 void UMaterialExpressionTextureSampleParameterCube::Build(MIR::FEmitter& Em)
@@ -1457,19 +1473,24 @@ static FValueRef EmitFixedExternalCode(MIR::FEmitter& Em, FName InExternalCodeId
 	return Em.InlineHLSL(MaterialExternalCodeRegistry::Get().FindExternalCode(InExternalCodeIdentifier), {}, MIR::EValueFlags::SubstituteTags);
 }
 
+static FValueRef EmitMatrixCastTo3x3(MIR::FEmitter& Em, FValueRef MatrixValue)
+{
+	return Em.InlineHLSL(MIR::FType::MakeFloat(3, 3), MatrixValue->Type.IsDouble() ? TEXTVIEW("DFToFloat3x3($0)") : TEXTVIEW("(float3x3)$0"), { MatrixValue });
+}
+
 static FValueRef EmitMatrixMultiply(MIR::FEmitter& Em, FValueRef VectorValue, FValueRef MatrixValue, bool bHasWComponent)
 {
 	return bHasWComponent
-		? Em.Swizzle(Em.Multiply(VectorValue, MatrixValue), MIR::FSwizzleMask::XYZ())  // mul(Float4(V, 1.0), V).xyz
-		: Em.Multiply(VectorValue, Em.Cast(MatrixValue, MIR::FType::MakeFloat(3, 3))); // mul(V, (Float3x3)M)
+		? Em.Swizzle(Em.MatrixMultiply(VectorValue, MatrixValue), MIR::FSwizzleMask::XYZ())  // mul(Float4(V, 1.0), V).xyz
+		: Em.MatrixMultiply(VectorValue, EmitMatrixCastTo3x3(Em, MatrixValue)); // mul(V, (Float3x3)M)
 }
 
 static FValueRef EmitMultiplyTransposeMatrix(MIR::FEmitter& Em, FValueRef MatrixValue, FValueRef VectorValue, bool bHasWComponent)
 {
 	// TODO: this should be removed when the Transpose operator is added to the translator.
 	return bHasWComponent
-		? Em.Swizzle(Em.Multiply(MatrixValue, VectorValue), MIR::FSwizzleMask::XYZ())  // mul(M, Float4(V, 1.0)).xyz
-		: Em.Multiply(Em.Cast(MatrixValue, MIR::FType::MakeFloat(3, 3)), VectorValue); // mul((Float3x3)M, V)
+		? Em.Swizzle(Em.MatrixMultiply(MatrixValue, VectorValue), MIR::FSwizzleMask::XYZ())  // mul(M, Float4(V, 1.0)).xyz
+		: Em.MatrixMultiply(EmitMatrixCastTo3x3(Em, MatrixValue), VectorValue); // mul((Float3x3)M, V)
 }
 
 static FValueRef EmitMultiplyTranslatedMatrix(MIR::FEmitter& Em, FValueRef VectorValue, FValueRef MatrixPreTranslation, bool bHasWComponent)
@@ -1483,8 +1504,7 @@ static FValueRef EmitMultiplyTranslatedMatrix(MIR::FEmitter& Em, FValueRef Vecto
 	else
 	{
 		// mul(V, DFToFloat3x3(M))
-		MatrixPreTranslation = Em.InlineHLSL(MIR::FType::MakeFloat(3, 3), TEXTVIEW("DFToFloat3x3($0)"), { MatrixPreTranslation });
-		return Em.MatrixMultiply(VectorValue, MatrixPreTranslation);
+		return Em.MatrixMultiply(VectorValue, EmitMatrixCastTo3x3(Em, MatrixPreTranslation));
 	}
 }
 
diff --git a/Engine/Source/Runtime/Engine/Private/Materials/MaterialIREmitter.cpp b/Engine/Source/Runtime/Engine/Private/Materials/MaterialIREmitter.cpp
--- a/Engine/Source/Runtime/Engine/Private/Materials/MaterialIREmitter.cpp
+++ b/Engine/Source/Runtime/Engine/Private/Materials/MaterialIREmitter.cpp
@@ -887,10 +887,11 @@ FValueRef FEmitter::ConstantDouble(TDouble InX)
 /*--------------------- Other non-instruction values emission ---------------------*/
 
 FValueRef FEmitter::Poison()
-{ 	if (CVarMaterialIRDebugBreakOnPoison.GetValueOnGameThread())
- 	{
- 		UE_DEBUG_BREAK();
- 	}
+{
+	if (CVarMaterialIRDebugBreakOnPoison.GetValueOnGameThread())
+	{
+		UE_DEBUG_BREAK();
+	}
 	return FPoison::Get();
 }
 
@@ -910,25 +911,10 @@ FValueRef FEmitter::MaterialParameterCollection(UMaterialParameterCollection* Co
 
 FValueRef FEmitter::TextureObject(UTexture* Texture, EMaterialSamplerType SamplerType)
 {
-	switch (Texture->GetMaterialType())
-	{
-		case MCT_Texture2D:
-		case MCT_TextureVirtual:
-		{
-			FTextureObject Proto = MakePrototype<FTextureObject>(FType::MakeTexture2D());
-			Proto.Texture = Texture;
-			Proto.SamplerType = SamplerType;
-			return EmitPrototype(*this, Proto);
-		}
-		break;
-
-		default:
-		{
-			Error(TEXT("Only 2D and Virtual Textures are supported for now."));
-		}
-		break;
-	}
-	return Poison();
+	FTextureObject Proto = MakePrototype<FTextureObject>(FType::MakeTexture2D());
+	Proto.Texture = Texture;
+	Proto.SamplerType = SamplerType;
+	return EmitPrototype(*this, Proto);
 }
 
 FValueRef FEmitter::Parameter(FName Name, FMaterialParameterMetadata& Metadata, EMaterialSamplerType SamplerType)
@@ -1587,6 +1573,7 @@ static T FoldScalarOperator(FEmitter& Emitter, EOperator Operator, T A, T B, T C
 		case BO_Add: return A + B;
 		case BO_Subtract: return A - B;
 		case BO_Multiply: return A * B;
+		case BO_MatrixMultiply: return A * B; // mul() is also supported for scalars
 		case BO_Divide: return A / B;
 		case BO_Min: return FMath::Min<T>(A, B);
 		case BO_Max: return FMath::Max<T>(A, B);
@@ -2331,7 +2318,10 @@ static FType ValidateOperatorAndGetResultType(FEmitter& Emitter, EOperator Op, F
 				return FType::MakePoison();
 			}
 
-			return FType::MakePrimitive(FirstArgumentType.GetPrimitive().ScalarKind, FirstArgumentType.GetPrimitive().NumRows, SecondArgumentType.GetPrimitive().NumColumns);
+			FPrimitive LhsPrimitiveType = FirstArgumentType.GetPrimitive();
+			FPrimitive RhsPrimitiveType = SecondArgumentType.GetPrimitive();
+
+			return FType::MakePrimitive(LhsPrimitiveType.ScalarKind, LhsPrimitiveType.NumRows, RhsPrimitiveType.NumColumns);
 		}
 	}
 
@@ -2342,7 +2332,7 @@ static FType ValidateOperatorAndGetResultType(FEmitter& Emitter, EOperator Op, F
 // 	op(v, w) == [op(v_0, w_0), ..., op(v_n, w_n)]
 static bool IsComponentwiseOperator(EOperator Op)
 {
-	return Op != BO_Dot && Op != BO_Cross;
+	return Op != BO_Dot && Op != BO_Cross && Op != BO_MatrixMultiply;
 }
 
 // Tries to fold the operator by applying the operator componentwise on arguments components.
@@ -3726,8 +3716,8 @@ FValueRef FEmitter::AnalyticalPartialDerivative(FValueRef Value, EDerivativeAxis
 
 		case VK_Scalar:
 		{
-			FScalar* AsCast = Value->As<FScalar>();
-			return Cast(AnalyticalPartialDerivative(AsCast->Arg, Axis), AsCast->Type);
+			FScalar* AsScalar = Value->As<FScalar>();
+			return Cast(AnalyticalPartialDerivative(AsScalar->Arg, Axis), AsScalar->Type);
 		}
 
 		case VK_InlineHLSL:
diff --git a/Engine/Source/Runtime/Engine/Private/Materials/MaterialIRToHLSLTranslator.cpp b/Engine/Source/Runtime/Engine/Private/Materials/MaterialIRToHLSLTranslator.cpp
--- a/Engine/Source/Runtime/Engine/Private/Materials/MaterialIRToHLSLTranslator.cpp
+++ b/Engine/Source/Runtime/Engine/Private/Materials/MaterialIRToHLSLTranslator.cpp
@@ -1548,11 +1548,31 @@ struct FPrivate : FMaterialIRToHLSLTranslation
 		}
 	}
 
+	ENoOp LowerTextureMaterialType(EMaterialValueType TextureType, bool bForResourceDeclarations = false)
+	{
+		switch (TextureType)
+		{
+		case MCT_Texture2D: Printer << TEXTVIEW("Texture2D"); break;
+		case MCT_TextureCube: Printer << TEXTVIEW("TextureCube"); break;
+		case MCT_Texture2DArray: Printer << TEXTVIEW("Texture2DArray"); break;
+		case MCT_TextureCubeArray: Printer << TEXTVIEW("TextureCubeArray"); break;
+		case MCT_VolumeTexture: Printer << (bForResourceDeclarations ? TEXTVIEW("VolumeTexture") : TEXTVIEW("Texture3D")); break;
+		case MCT_TextureExternal: Printer << TEXTVIEW("Texture2D");
+		}
+		return NoOp;
+	}
+
+	void LowerTextureType(const MIR::FTextureObject* TextureObject)
+	{
+		check(TextureObject->Texture);
+		LowerTextureMaterialType(TextureObject->Texture->GetMaterialType());
+	}
+
 	void LowerStandardTextureRead(const MIR::FTextureRead* TextureRead)
 	{
 		switch (TextureRead->TextureObject->Type.GetKind())
 		{
-			case MIR::ETypeKind::Texture2D: Printer << TEXTVIEW("Texture2D"); break;
+			case MIR::ETypeKind::Texture2D: LowerTextureType(TextureRead->TextureObject->As<MIR::FTextureObject>()); break;
 			default: UE_MIR_UNREACHABLE();
 		}
 
@@ -1780,17 +1800,7 @@ struct FPrivate : FMaterialIRToHLSLTranslation
 
 	ENoOp LowerTextureReference(EMaterialValueType TextureType, int32 TextureParameterIndex)
 	{
-		Printer << TEXTVIEW("Material.");
-
-		switch (TextureType)
-		{
-			case MCT_Texture2D: Printer << TEXTVIEW("Texture2D_"); break;
-			case MCT_TextureVirtual: Printer << TEXTVIEW("VirtualTexturePhysical_"); break;
-			default: UE_MIR_UNREACHABLE();
-		}
-
-		Printer << TextureParameterIndex;
-
+		Printer << TEXTVIEW("Material.") << LowerTextureMaterialType(TextureType, true) << TEXTVIEW("_") << TextureParameterIndex;
 		return NoOp;
 	}
 	
diff --git a/Engine/Source/Runtime/Engine/Private/Materials/MaterialIRTypes.cpp b/Engine/Source/Runtime/Engine/Private/Materials/MaterialIRTypes.cpp
--- a/Engine/Source/Runtime/Engine/Private/Materials/MaterialIRTypes.cpp
+++ b/Engine/Source/Runtime/Engine/Private/Materials/MaterialIRTypes.cpp
@@ -90,7 +90,7 @@ FType FType::FromMaterialValueType(EMaterialValueType Type)
 	switch (Type)
 	{
 		case MCT_VoidStatement: return FType::MakeVoid();
-		case MCT_Float: return FType::MakeFloatScalar();
+		case MCT_Float: return FType::MakeFloatScalar(); // MCT_Float represents any scalar or vector, but GetMaterialTypeFromInputType() maps FunctionInput_Scalar to MCT_Float instead of MCT_Float1, so interpret it as a scalar here
 		case MCT_Float1: return FType::MakeFloatScalar();
 		case MCT_Float2: return FType::MakeFloatVector(2);
 		case MCT_Float3: return FType::MakeFloatVector(3);
