[MIR] Fix include header in MaterialExpressionsIR.cpp.

#rnx

diff --git a/Engine/Source/Runtime/Engine/Private/Materials/MaterialExpressionsIR.cpp b/Engine/Source/Runtime/Engine/Private/Materials/MaterialExpressionsIR.cpp
--- a/Engine/Source/Runtime/Engine/Private/Materials/MaterialExpressionsIR.cpp
+++ b/Engine/Source/Runtime/Engine/Private/Materials/MaterialExpressionsIR.cpp
@@ -2,8 +2,6 @@
 
 #if WITH_EDITOR
 
-#include "MaterialExpressionsIR.h"
-
 #include "Curves/CurveLinearColorAtlas.h"
 #include "Engine/Engine.h"
 #include "Engine/Font.h"
