Some refactoring to the code in MaterialIREmitter that emits aggregate values.

#rnx

diff --git a/Engine/Source/Runtime/Engine/Private/Materials/MaterialIREmitter.cpp b/Engine/Source/Runtime/Engine/Private/Materials/MaterialIREmitter.cpp
--- a/Engine/Source/Runtime/Engine/Private/Materials/MaterialIREmitter.cpp
+++ b/Engine/Source/Runtime/Engine/Private/Materials/MaterialIREmitter.cpp
@@ -1281,47 +1281,32 @@ FValueRef FEmitter::Aggregate(const UMaterialAggregate* InAggregate)
 	return Aggregate(InAggregate, {}, TConstArrayView<FValueRef>{});
 }
 
-static FValueRef EmitAttributeDefaultValue(FEmitter& Emitter, const FMaterialAggregateAttribute& Attribute)
+static FValueRef EmitAttributeDefaultValue(FEmitter& Emitter, const UMaterial* Material, const FMaterialAggregateAttribute& Attribute)
 {
 	switch (Attribute.Type)
 	{
-	case EMaterialAggregateAttributeType::Bool1: return Emitter.ConstantBool((bool)Attribute.DefaultValue.X);
-	case EMaterialAggregateAttributeType::Bool2: return Emitter.ConstantBool2((bool)Attribute.DefaultValue.X, (bool)Attribute.DefaultValue.Y);
-	case EMaterialAggregateAttributeType::Bool3: return Emitter.ConstantBool3((bool)Attribute.DefaultValue.X, (bool)Attribute.DefaultValue.Y, (bool)Attribute.DefaultValue.Z);
-	case EMaterialAggregateAttributeType::Bool4: return Emitter.ConstantBool4((bool)Attribute.DefaultValue.X, (bool)Attribute.DefaultValue.Y, (bool)Attribute.DefaultValue.Z, (bool)Attribute.DefaultValue.W);
-	case EMaterialAggregateAttributeType::ShadingModel:
-	case EMaterialAggregateAttributeType::UInt1: return Emitter.ConstantInt((TInteger)Attribute.DefaultValue.X);
-	case EMaterialAggregateAttributeType::UInt2: return Emitter.ConstantInt2({ (TInteger)Attribute.DefaultValue.X, (TInteger)Attribute.DefaultValue.Y });
-	case EMaterialAggregateAttributeType::UInt3: return Emitter.ConstantInt3({ (TInteger)Attribute.DefaultValue.X, (TInteger)Attribute.DefaultValue.Y, (TInteger)Attribute.DefaultValue.Z });
-	case EMaterialAggregateAttributeType::UInt4: return Emitter.ConstantInt4({ (TInteger)Attribute.DefaultValue.X, (TInteger)Attribute.DefaultValue.Y, (TInteger)Attribute.DefaultValue.Z, (TInteger)Attribute.DefaultValue.W });
-	case EMaterialAggregateAttributeType::Float1: return Emitter.ConstantFloat((TFloat)Attribute.DefaultValue.X);
-	case EMaterialAggregateAttributeType::Float2: return Emitter.ConstantFloat2({ (TFloat)Attribute.DefaultValue.X, (TFloat)Attribute.DefaultValue.Y });
-	case EMaterialAggregateAttributeType::Float3: return Emitter.ConstantFloat3({ (TFloat)Attribute.DefaultValue.X, (TFloat)Attribute.DefaultValue.Y, (TFloat)Attribute.DefaultValue.Z });
-	case EMaterialAggregateAttributeType::Float4: return Emitter.ConstantFloat4(UE::Math::TVector4<TFloat>{ Attribute.DefaultValue.X, Attribute.DefaultValue.Y, Attribute.DefaultValue.Z, Attribute.DefaultValue.W });
-	case EMaterialAggregateAttributeType::MaterialAttributes: return Emitter.Aggregate(UMaterialAggregate::GetMaterialAttributes());
-	case EMaterialAggregateAttributeType::Aggregate: return Emitter.Aggregate(Attribute.Aggregate);
-	default: UE_MIR_UNREACHABLE();
-	}
-}
-
-static FValueRef MapMaterialAggregateDefaultValues(FEmitter& Emitter, UMaterial* Material, const UMaterialAggregate* InAggregate, const int32& Index)
-{
-	ensure(Index >= 0 && Index < InAggregate->Attributes.Num());
-	if (Material && InAggregate == UMaterialAggregate::GetMaterialAttributes())
-	{
-		switch (UMaterialAggregate::AttributeIndexToMaterialProperty(Index))
-		{
-		case MP_ShadingModel: return Emitter.ConstantInt(Material->GetShadingModels().GetFirstShadingModel());
-		default: break;
-		}
+		case EMaterialAggregateAttributeType::Bool1: return Emitter.ConstantBool((bool)Attribute.DefaultValue.X);
+		case EMaterialAggregateAttributeType::Bool2: return Emitter.ConstantBool2((bool)Attribute.DefaultValue.X, (bool)Attribute.DefaultValue.Y);
+		case EMaterialAggregateAttributeType::Bool3: return Emitter.ConstantBool3((bool)Attribute.DefaultValue.X, (bool)Attribute.DefaultValue.Y, (bool)Attribute.DefaultValue.Z);
+		case EMaterialAggregateAttributeType::Bool4: return Emitter.ConstantBool4((bool)Attribute.DefaultValue.X, (bool)Attribute.DefaultValue.Y, (bool)Attribute.DefaultValue.Z, (bool)Attribute.DefaultValue.W);
+		case EMaterialAggregateAttributeType::UInt1: return Emitter.ConstantInt((TInteger)Attribute.DefaultValue.X);
+		case EMaterialAggregateAttributeType::UInt2: return Emitter.ConstantInt2({ (TInteger)Attribute.DefaultValue.X, (TInteger)Attribute.DefaultValue.Y });
+		case EMaterialAggregateAttributeType::UInt3: return Emitter.ConstantInt3({ (TInteger)Attribute.DefaultValue.X, (TInteger)Attribute.DefaultValue.Y, (TInteger)Attribute.DefaultValue.Z });
+		case EMaterialAggregateAttributeType::UInt4: return Emitter.ConstantInt4({ (TInteger)Attribute.DefaultValue.X, (TInteger)Attribute.DefaultValue.Y, (TInteger)Attribute.DefaultValue.Z, (TInteger)Attribute.DefaultValue.W });
+		case EMaterialAggregateAttributeType::Float1: return Emitter.ConstantFloat((TFloat)Attribute.DefaultValue.X);
+		case EMaterialAggregateAttributeType::Float2: return Emitter.ConstantFloat2({ (TFloat)Attribute.DefaultValue.X, (TFloat)Attribute.DefaultValue.Y });
+		case EMaterialAggregateAttributeType::Float3: return Emitter.ConstantFloat3({ (TFloat)Attribute.DefaultValue.X, (TFloat)Attribute.DefaultValue.Y, (TFloat)Attribute.DefaultValue.Z });
+		case EMaterialAggregateAttributeType::Float4: return Emitter.ConstantFloat4(UE::Math::TVector4<TFloat>{ Attribute.DefaultValue.X, Attribute.DefaultValue.Y, Attribute.DefaultValue.Z, Attribute.DefaultValue.W });
+		case EMaterialAggregateAttributeType::ShadingModel: return Emitter.ConstantInt(Material->GetShadingModels().GetFirstShadingModel());
+		case EMaterialAggregateAttributeType::MaterialAttributes: return Emitter.Aggregate(UMaterialAggregate::GetMaterialAttributes());
+		case EMaterialAggregateAttributeType::Aggregate: return Emitter.Aggregate(Attribute.Aggregate);
+		default: UE_MIR_UNREACHABLE();
 	}
-
-	return EmitAttributeDefaultValue(Emitter, InAggregate->Attributes[Index]);
 }
 
 FValueRef FEmitter::Aggregate(const UMaterialAggregate* InAggregate, FValueRef InPrototype, TConstArrayView<FValueRef> AttributeValues)
 {
-	// Check that the specified prototype (if any) aggregate matches the one provided.
+	// Check that the szpecified prototype (if any) aggregate matches the one provided.
 	InPrototype = CheckIsAggregate(InPrototype, InAggregate);
 
 	// If a prototype is provided and there are no attribute assignments, this is a no-op, return the prototype.
@@ -1340,20 +1325,20 @@ FValueRef FEmitter::Aggregate(const UMaterialAggregate* InAggregate, FValueRef I
 		// Get the the ith aggregate aggregate MIR type.
 		FType AttributeType = GetMaterialAggregateAttributeType(InAggregate, i);
 
-		if (i < AttributeValues.Num())
+		if (i < AttributeValues.Num() && AttributeValues[i])
 		{
 			// Set the this aggregate component to the specified value cast to the attribute type, if present...
 			Components[i] = Cast(AttributeValues[i], AttributeType);
 		}
 		else if (InPrototype)
 		{
-			// ... otherwise use the component value as in the prototype, if defined...
+			// ... otherwise use the component value as in the prototype, if od...
 			Components[i] = Subscript(InPrototype, i);
 		}
 		else
 		{
 			// ...otherwise construct the default value as indicated in the attribute.
-			Components[i] = MapMaterialAggregateDefaultValues(*this, Material, InAggregate, i);
+			Components[i] = EmitAttributeDefaultValue(*this, Material, InAggregate->Attributes[i]);
 		}
 	}
 
@@ -1365,47 +1350,29 @@ FValueRef FEmitter::Aggregate(const UMaterialAggregate* InAggregate, FValueRef I
 	// Check that the specified prototype (if any) aggregate matches the one provided.
 	InPrototype = CheckIsAggregate(InPrototype, InAggregate);
 
-	// If a prototype is provided and there are no attribute assignments, this is a no-op, return the prototype.
-	if (InPrototype.IsValid() && AttributeAssignments.IsEmpty())
+	// If prototype was provided and is poison, return it..
+	if (InPrototype.IsPoison())
 	{
-		return InPrototype;
+		return InPrototype.ToPoison();
 	}
 
-	// Create the new composite to store the aggregate attribute values.
-	FComposite* AggregateValue = MakeCompositePrototype(*this, FType::MakeAggregate(InAggregate), InAggregate->Attributes.Num());
-	
-	// Assign all components of the new composite value.
-	TArrayView<FValue*> Components = AggregateValue->GetMutableComponents();
-	for (int32 i = 0; i < Components.Num(); ++i)
-	{
-		// Search the assignments for this component's name.
-		const FAttributeAssignment* Assignment = AttributeAssignments.FindByPredicate(
-			[AttributeName = InAggregate->Attributes[i].Name](const FAttributeAssignment& Item)
-			{
-				return Item.Name == AttributeName;
-			});
-
-		// Get the the ith aggregate aggregate MIR type.
-		FType AttributeType = GetMaterialAggregateAttributeType(InAggregate, i);
+	// Allocate temporary storage to hold the attribute values
+	TTemporaryArray<FValueRef> AttributeValues { InAggregate->Attributes.Num() };
+	AttributeValues.Zero();
 
-		if (Assignment && Assignment->Value)
-		{
-			// Set the this aggregate component to the value specified in the assignments, if present...
-			Components[i] = Cast(Assignment->Value, AttributeType);
-		}
-		else if (!Assignment && InPrototype)
-		{
-			// ... otherwise use the component value as in the prototype, if defined...
-			Components[i] = Subscript(InPrototype, i);
-		}
-		else
+	// Linearize the assignments into the array of attribute values.
+	for (const FAttributeAssignment& Assignment : AttributeAssignments)
+	{
+		// Find the attribute index by name
+		int32 AttrIndex = InAggregate->FindAttributeIndexByName(Assignment.Name);
+		if (AttrIndex != INDEX_NONE)
 		{
-			// ...otherwise construct the default value.
-			Components[i] = MapMaterialAggregateDefaultValues(*this, Material, InAggregate, i);
+			// If found, set the slot to the assignment value.
+			AttributeValues[AttrIndex] = Assignment.Value;
 		}
 	}
 
-	return EmitPrototype(*this, *AggregateValue);
+	return Aggregate(InAggregate, InPrototype, AttributeValues);
 }
 
 /*--------------------------------- Operator emission ---------------------------------*/
diff --git a/Engine/Source/Runtime/Engine/Public/Materials/MaterialIRCommon.h b/Engine/Source/Runtime/Engine/Public/Materials/MaterialIRCommon.h
--- a/Engine/Source/Runtime/Engine/Public/Materials/MaterialIRCommon.h
+++ b/Engine/Source/Runtime/Engine/Public/Materials/MaterialIRCommon.h
@@ -83,6 +83,11 @@ struct TTemporaryArray : public TArrayView<T>
 	{
 		return { TArrayView<T>::GetData(), TArrayView<T>::Num() };
 	}
+
+	void Zero()
+	{
+		FMemory::Memzero(TArrayView<T>::GetData(), sizeof(T) * TArrayView<T>::Num());
+	}
 };
 
 }
