Adding missing material value type description for Texture Collection

diff --git a/Engine/Source/Runtime/Engine/Private/Materials/MaterialExpressions.cpp b/Engine/Source/Runtime/Engine/Private/Materials/MaterialExpressions.cpp
--- a/Engine/Source/Runtime/Engine/Private/Materials/MaterialExpressions.cpp
+++ b/Engine/Source/Runtime/Engine/Private/Materials/MaterialExpressions.cpp
@@ -470,6 +470,8 @@ extern ENGINE_API void GetMaterialValueTypeDescriptions(const EMaterialValueType
 		OutDescriptions.Add(LOCTEXT("ShadingModel", "Shading Model"));
 	if (MaterialValueType & MCT_Substrate)
 		OutDescriptions.Add(LOCTEXT("Substrate", "Substrate Material"));
+	if (MaterialValueType & MCT_TextureCollection)
+		OutDescriptions.Add(LOCTEXT("TextureCollection", "Texture Collection"));
 	if (MaterialValueType & MCT_Unknown)
 		OutDescriptions.Add(LOCTEXT("Unknown", "Unknown"));
 }
