World Partition - HLOD: Fix incorrect handling of non-spatially loaded HLODs following CL 44426691

diff --git a/Engine/Source/Editor/WorldPartitionEditor/Private/WorldPartition/Customizations/WorldPartitionRuntimeHashSetDetailsCustomization.cpp b/Engine/Source/Editor/WorldPartitionEditor/Private/WorldPartition/Customizations/WorldPartitionRuntimeHashSetDetailsCustomization.cpp
--- a/Engine/Source/Editor/WorldPartitionEditor/Private/WorldPartition/Customizations/WorldPartitionRuntimeHashSetDetailsCustomization.cpp
+++ b/Engine/Source/Editor/WorldPartitionEditor/Private/WorldPartition/Customizations/WorldPartitionRuntimeHashSetDetailsCustomization.cpp
@@ -65,6 +65,11 @@ void FRuntimePartitionHLODSetupDetails::CustomizeChildren(TSharedRef<IPropertyHa
 				}
 			}
 		}
+		else if (ChildHandle->GetProperty()->GetFName() == GET_MEMBER_NAME_CHECKED(FRuntimePartitionHLODSetup, RowDisplayName))
+		{
+			// Property is only used as the "TitlePropery" when showing as an array item
+			bShouldAddProperty = false;
+		}
 		
 		if (bShouldAddProperty)
 		{
diff --git a/Engine/Source/Runtime/Engine/Private/WorldPartition/RuntimeHashSet/WorldPartitionRuntimeHashSet.cpp b/Engine/Source/Runtime/Engine/Private/WorldPartition/RuntimeHashSet/WorldPartitionRuntimeHashSet.cpp
--- a/Engine/Source/Runtime/Engine/Private/WorldPartition/RuntimeHashSet/WorldPartitionRuntimeHashSet.cpp
+++ b/Engine/Source/Runtime/Engine/Private/WorldPartition/RuntimeHashSet/WorldPartitionRuntimeHashSet.cpp
@@ -1000,6 +1000,7 @@ void UWorldPartitionRuntimeHashSet::PostEditChangeChainProperty(FPropertyChanged
 			}
 
 			RuntimePartitionHLODSetup.PartitionLayer->Name = RuntimePartitionHLODSetup.Name;
+			RuntimePartitionHLODSetup.RowDisplayName = RuntimePartitionHLODSetup.PartitionLayer->Name;
 		}
 		else
 		{
@@ -1156,6 +1157,9 @@ void UWorldPartitionRuntimeHashSet::FixupHLODSetup(FRuntimePartitionDesc& Runtim
 					HLODSetup.PartitionLayer->Name = HLODSetup.Name;
 				}
 
+				// Row display name (used as title when showing as an array item) must match the partition name
+				HLODSetup.RowDisplayName = HLODSetup.Name;
+
 				// Make sure that HLOD partitions settings are updated to match the main layer settings
 				HLODSetup.PartitionLayer->UpdateHLODRuntimePartitionFrom(RuntimePartition.MainLayer);
 			}
@@ -1168,8 +1172,9 @@ void UWorldPartitionRuntimeHashSet::FixupHLODSetup(FRuntimePartitionDesc& Runtim
 			FRuntimePartitionHLODSetup& PersistentHLODSetup = RuntimePartition.HLODSetups.Emplace_GetRef();
 			PersistentHLODSetup.PartitionLayer = NewObject<URuntimePartitionPersistent>(this, NAME_None);
 			PersistentHLODSetup.PartitionLayer->LoadingRange = 0;
-			PersistentHLODSetup.PartitionLayer->Name = "Non-Spatially Loaded HLODs";
-			PersistentHLODSetup.Name = PersistentHLODSetup.PartitionLayer->Name;
+			PersistentHLODSetup.PartitionLayer->Name = NAME_None;
+			PersistentHLODSetup.Name = NAME_None;
+			PersistentHLODSetup.RowDisplayName = "Non-Spatially Loaded HLODs";
 			PersistentHLODSetup.HLODLayers = MoveTemp(PersistentHLODLayers);
 			PersistentHLODSetup.bIsSpatiallyLoaded = false;
 		}
diff --git a/Engine/Source/Runtime/Engine/Private/WorldPartition/RuntimeHashSet/WorldPartitionRuntimeHashSetHLODGeneration.cpp b/Engine/Source/Runtime/Engine/Private/WorldPartition/RuntimeHashSet/WorldPartitionRuntimeHashSetHLODGeneration.cpp
--- a/Engine/Source/Runtime/Engine/Private/WorldPartition/RuntimeHashSet/WorldPartitionRuntimeHashSetHLODGeneration.cpp
+++ b/Engine/Source/Runtime/Engine/Private/WorldPartition/RuntimeHashSet/WorldPartitionRuntimeHashSetHLODGeneration.cpp
@@ -2,6 +2,7 @@
 
 #include "WorldPartition/RuntimeHashSet/WorldPartitionRuntimeHashSet.h"
 #include "WorldPartition/RuntimeHashSet/RuntimePartition.h"
+#include "WorldPartition/RuntimeHashSet/RuntimePartitionPersistent.h"
 
 #include "WorldPartition/WorldPartition.h"
 #include "WorldPartition/WorldPartitionLog.h"
@@ -367,6 +368,12 @@ bool UWorldPartitionRuntimeHashSet::SetupHLODActors(const IStreamingGenerationCo
 					{
 						return NAME_None;
 					}
+
+					if (HLODRuntimePartition->IsA<URuntimePartitionPersistent>())
+					{
+						return NAME_None;
+					}
+
 					return FName(*FString::Printf(TEXT("%s:%s"), *MainRuntimePartition->Name.ToString(), *HLODRuntimePartition->Name.ToString()));
 				};				
 
diff --git a/Engine/Source/Runtime/Engine/Public/WorldPartition/RuntimeHashSet/WorldPartitionRuntimeHashSet.h b/Engine/Source/Runtime/Engine/Public/WorldPartition/RuntimeHashSet/WorldPartitionRuntimeHashSet.h
--- a/Engine/Source/Runtime/Engine/Public/WorldPartition/RuntimeHashSet/WorldPartitionRuntimeHashSet.h
+++ b/Engine/Source/Runtime/Engine/Public/WorldPartition/RuntimeHashSet/WorldPartitionRuntimeHashSet.h
@@ -38,6 +38,10 @@ struct FRuntimePartitionHLODSetup
 	/** Associated HLOD Layer objects */
 	UPROPERTY(EditAnywhere, Category = RuntimeSettings, DisplayName = "HLOD Layers To Include")
 	TArray<TObjectPtr<const UHLODLayer>> HLODLayers;
+
+	/** Used as the "TitlePropery" when showing as an array item */
+	UPROPERTY(Transient, VisibleAnywhere, Category = RuntimeSettings)
+	FName RowDisplayName;
 #endif
 
 	UPROPERTY(EditAnywhere, Category = RuntimeSettings, Instanced, Meta = (NoResetToDefault))
@@ -68,7 +72,7 @@ struct FRuntimePartitionDesc
 	TObjectPtr<URuntimePartition> MainLayer;
 
 	/** HLOD setups used by this partition, one for each layers in the hierarchy */
-	UPROPERTY(EditAnywhere, Category = RuntimeSettings, Meta = (DisplayName = "HLOD Partitions", EditCondition = "MainLayer != nullptr", HideEditConditionToggle, TitleProperty="Name"))
+	UPROPERTY(EditAnywhere, Category = RuntimeSettings, Meta = (DisplayName = "HLOD Partitions", EditCondition = "MainLayer != nullptr", HideEditConditionToggle, TitleProperty="RowDisplayName"))
 	TArray<FRuntimePartitionHLODSetup> HLODSetups;
 
 #if WITH_EDITOR
