FSpline: Flesh out the interface

1) Adds support for caching update settings. This allows updating the spline using the cached settings without needing to specify new settings

2) Add support for getting/setting orientation in the same way that the spline component deals with rotations (see GetQuaternionAtSplinePoint).

3) Add support for closing/opening the spline more directly (required cached updating settings). Adds support for checking if the spline is currently closed.

diff --git a/Engine/Source/Runtime/Engine/Classes/Curves/Spline.h b/Engine/Source/Runtime/Engine/Classes/Curves/Spline.h
--- a/Engine/Source/Runtime/Engine/Classes/Curves/Spline.h
+++ b/Engine/Source/Runtime/Engine/Classes/Curves/Spline.h
@@ -59,6 +59,9 @@ struct FSpline
 	ENGINE_API float GetParameterAtIndex(int32 Index) const;
 	ENGINE_API float GetParameterAtDistance(float Distance) const;
 	ENGINE_API float GetDistanceAtParameter(float Parameter) const;
+
+	ENGINE_API FQuat GetOrientation(int32 Index) const;
+	ENGINE_API void SetOrientation(int32 Index, const FQuat& InOrientation);
 	
 	/* Parameter Interface */
 
@@ -67,6 +70,8 @@ struct FSpline
 	ENGINE_API FQuat EvaluateRotation(float Param) const;
 	ENGINE_API FVector EvaluateScale(float Param) const;
 
+	ENGINE_API FQuat GetOrientation(float Param) const;
+	
 	/* Attribute Interface */
 
 	// Non-Templated functions
@@ -118,6 +123,9 @@ struct FSpline
 
 	/** Returns total length along this spline */
 	ENGINE_API float GetSplineLength() const;
+
+	/** Returns the total number of segments on this spline. */
+	ENGINE_API int32 GetNumSegments() const;
 	
 	/** Returns the total number of control points on this spline. */
 	ENGINE_API int32 GetNumControlPoints() const;
@@ -131,6 +139,9 @@ struct FSpline
 	/** Reset the scale attribute channel to default values. */
 	ENGINE_API void ResetScale();
 
+	ENGINE_API void SetClosedLoop(bool bClosed);
+	ENGINE_API bool IsClosedLoop() const;
+	
 	struct FUpdateSplineParams
 	{
 		bool bClosedLoop = false;
@@ -138,12 +149,15 @@ struct FSpline
 		int32 ReparamStepsPerSegment = 10;
 		bool bLoopPositionOverride = false;
 		float LoopPosition = 0.0f;
-		const FVector Scale3D = FVector(1.0f);
+		FVector Scale3D = FVector(1.0f);
 	};
 	
 	/** Update the spline's internal data according to the passed-in params. */
 	ENGINE_API void UpdateSpline(const FUpdateSplineParams& InParams);
 
+	/** Update the spline's internal data according to the most recently used update params (or default params if never updated). */
+	ENGINE_API void UpdateSpline();
+
 	ENGINE_API TSharedPtr<UE::Geometry::Spline::TSplineInterface<FVector>> GetSplineInterface() const;
 
 private:
@@ -165,6 +179,8 @@ private:
 	
 	uint32 Version;
 
+	FUpdateSplineParams CachedUpdateSplineParams;
+	
 	// probably better implemented as a TSharedPtr<ISplineInterface> or something
 	struct {													// Invalid when CurrentImplementation is 0 or 3.
 		TSharedPtr<FLegacySpline> LegacyData;					// Valid when CurrentImplementation is 1.
diff --git a/Engine/Source/Runtime/Engine/Private/Curves/Spline.cpp b/Engine/Source/Runtime/Engine/Private/Curves/Spline.cpp
--- a/Engine/Source/Runtime/Engine/Private/Curves/Spline.cpp
+++ b/Engine/Source/Runtime/Engine/Private/Curves/Spline.cpp
@@ -1032,6 +1032,29 @@ float FSpline::GetDistanceAtParameter(float Parameter) const
 	}
 }
 
+FQuat FSpline::GetOrientation(int32 Index) const
+{
+	return GetOrientation(GetParameterAtIndex(Index));
+}
+
+void FSpline::SetOrientation(int32 Index, const FQuat& InOrientation)
+{
+	if (Index < 0 || Index >= GetNumControlPoints())
+	{
+		return;
+	}
+	
+	// Store rotation which transforms the world up vector to the local up vector.
+	SetRotation(Index, FQuat::FindBetween(FVector::UpVector, InOrientation.GetUpVector()));
+
+	// Align tangents with rotation, preserving magnitude.
+	const float InTangentMag = GetInTangent(Index).Length();
+	const float OutTangentMag = GetOutTangent(Index).Length();
+	const FVector NewTangent = InOrientation.GetForwardVector().GetSafeNormal();
+	SetInTangent(Index, InTangentMag * NewTangent);
+	SetOutTangent(Index, OutTangentMag * NewTangent);
+}
+
 FVector FSpline::Evaluate(float Param) const
 {
 	switch (CurrentImplementation)
@@ -1099,6 +1122,19 @@ FVector FSpline::EvaluateScale(float Param) const
 	}	
 }
 
+FQuat FSpline::GetOrientation(float Param) const
+{
+	FQuat Rotation = EvaluateRotation(Param);
+	Rotation.Normalize();
+
+	FVector Direction = EvaluateDerivative(Param);
+	Direction = Direction.GetSafeNormal();
+
+	const FVector UpVector = Rotation.RotateVector(FVector::UpVector);
+	
+	return FRotationMatrix::MakeFromXZ(Direction, UpVector).ToQuat();
+}
+
 bool FSpline::HasAttributeChannel(FName AttributeName) const
 {
 	if (NewData && SupportsAttributes())
@@ -1508,6 +1544,16 @@ float FSpline::GetSplineLength() const
 	}
 }
 
+int32 FSpline::GetNumSegments() const
+{
+	switch (CurrentImplementation)
+	{
+	default: return 0;
+	case 1: return LegacyData->GetSplinePointsPosition().bIsLooped ? LegacyData->GetNumControlPoints() : FMath::Max(0, LegacyData->GetNumControlPoints() - 1);
+	case 2: return NewData->GetSpline().GetNumberOfSegments();
+	}
+}
+
 int32 FSpline::GetNumControlPoints() const
 {
 	switch (CurrentImplementation)
@@ -1548,6 +1594,17 @@ void FSpline::ResetScale()
 	}
 }
 
+void FSpline::SetClosedLoop(bool bClosed)
+{
+	CachedUpdateSplineParams.bClosedLoop = bClosed;
+	UpdateSpline();
+}
+
+bool FSpline::IsClosedLoop() const
+{
+	return CachedUpdateSplineParams.bClosedLoop;
+}
+
 void FSpline::UpdateSpline(const FUpdateSplineParams& InParams)
 {
 	switch (CurrentImplementation)
@@ -1557,9 +1614,15 @@ void FSpline::UpdateSpline(const FUpdateSplineParams& InParams)
     case 2: NewData->UpdateSpline(InParams); break;
     }
 
+	CachedUpdateSplineParams = InParams;
 	++Version;
 }
 
+void FSpline::UpdateSpline()
+{
+	UpdateSpline(CachedUpdateSplineParams);
+}
+
 TSharedPtr<TSplineInterface<FVector>> FSpline::GetSplineInterface() const
 {
 	if (IsNew())
@@ -2017,6 +2080,7 @@ FSplinePoint FNewSpline::GetPoint(int32 Index) const
     	return Point;
     }
 
+	Point.InputKey = GetParameterAtIndex(Index);
     Point.Position = GetSpline().GetValue(Index);
 	Point.ArriveTangent = GetSpline().GetTangentIn(Index);
     Point.LeaveTangent = GetSpline().GetTangentOut(Index);
