[Materials] Make UMaterialExpressionColorRamp private.

This expression causes issues in a cook when used with a broken color curve.
To prevent this recently added and non-ciritcal expression from being used in new materials, we hide it until this issue is fully resolved.

#test All ShadingModels/Materials automation tests succeeded; Lyra cook succeeded; Custom matieral with ColorRamp node still works but can no longer be instantiated.
#rnx

diff --git a/Engine/Source/Runtime/Engine/Public/Materials/MaterialExpressionColorRamp.h b/Engine/Source/Runtime/Engine/Public/Materials/MaterialExpressionColorRamp.h
--- a/Engine/Source/Runtime/Engine/Public/Materials/MaterialExpressionColorRamp.h
+++ b/Engine/Source/Runtime/Engine/Public/Materials/MaterialExpressionColorRamp.h
@@ -10,7 +10,7 @@
 #include "Curves/CurveLinearColor.h"
 #include "MaterialExpressionColorRamp.generated.h"
 
-UCLASS(MinimalAPI)
+UCLASS(MinimalAPI, meta = (Private))
 class UMaterialExpressionColorRamp : public UMaterialExpression
 {
 	GENERATED_BODY()
