[37.10] Making SafeMoveUpdatedComponent virtual so that it can be overridden in child classes.

#rnx

diff --git a/Engine/Source/Runtime/Engine/Classes/GameFramework/MovementComponent.h b/Engine/Source/Runtime/Engine/Classes/GameFramework/MovementComponent.h
--- a/Engine/Source/Runtime/Engine/Classes/GameFramework/MovementComponent.h
+++ b/Engine/Source/Runtime/Engine/Classes/GameFramework/MovementComponent.h
@@ -309,9 +309,11 @@ public:
 	 * If this adjustment succeeds, the original movement will be attempted again.
 	 * @note The overload taking rotation as an FQuat is slightly faster than the version using FRotator (which will be converted to an FQuat).
 	 * @note The 'Teleport' flag is currently always treated as 'None' (not teleporting) when used in an active FScopedMovementUpdate.
+	 * @note When overriding, it is recommended to either call super or at least maintain the intended operation that includes resolving initial penetration.
+	 * @note Also when overriding, make sure to include `using UMovementComponent::SafeMoveUpdatedComponent` next to the declaration. The FRotator version will be hidden otherwise and cause compilation issues when invoked.
 	 * @return result of the final MoveUpdatedComponent() call.
 	 */
-	ENGINE_API bool SafeMoveUpdatedComponent(const FVector& Delta, const FQuat& NewRotation,    bool bSweep, FHitResult& OutHit, ETeleportType Teleport = ETeleportType::None);
+	ENGINE_API virtual bool SafeMoveUpdatedComponent(const FVector& Delta, const FQuat& NewRotation, bool bSweep, FHitResult& OutHit, ETeleportType Teleport = ETeleportType::None);
 	bool SafeMoveUpdatedComponent(const FVector& Delta, const FRotator& NewRotation, bool bSweep, FHitResult& OutHit, ETeleportType Teleport = ETeleportType::None);
 
 	/**
