Fix for Function Output Expressions mapping the incorrect output index when checking the connected value types.

diff --git a/Engine/Source/Runtime/Engine/Private/Materials/MaterialExpressions.cpp b/Engine/Source/Runtime/Engine/Private/Materials/MaterialExpressions.cpp
--- a/Engine/Source/Runtime/Engine/Private/Materials/MaterialExpressions.cpp
+++ b/Engine/Source/Runtime/Engine/Private/Materials/MaterialExpressions.cpp
@@ -18600,11 +18600,11 @@ EMaterialValueType UMaterialExpressionFunctionOutput::GetOutputValueType(int32 O
 	const FExpressionInput& ExpressionInputRef = A.GetTracedInput();
 	if (ExpressionInputRef.Expression)
 	{
-		if (ExpressionInputRef.Expression->IsResultMaterialAttributes(OutputIndex))
+		if (ExpressionInputRef.Expression->IsResultMaterialAttributes(ExpressionInputRef.OutputIndex))
 		{
 			return MCT_MaterialAttributes;
 		}
-		if (ExpressionInputRef.Expression->IsResultSubstrateMaterial(OutputIndex))
+		if (ExpressionInputRef.Expression->IsResultSubstrateMaterial(ExpressionInputRef.OutputIndex))
 		{
 			return MCT_Substrate;
 		}
