Enabled dynamic mesh bounds by default for platforms supporting Nanite.

diff --git a/Engine/Source/Runtime/Engine/Private/GPUSkinCache.cpp b/Engine/Source/Runtime/Engine/Private/GPUSkinCache.cpp
--- a/Engine/Source/Runtime/Engine/Private/GPUSkinCache.cpp
+++ b/Engine/Source/Runtime/Engine/Private/GPUSkinCache.cpp
@@ -208,11 +208,14 @@ FAutoConsoleVariableRef CVarSkinCacheAsyncCompute(
 	ECVF_RenderThreadSafe
 );
 
-bool GSkinCacheDynamicMeshBounds = false;
-FAutoConsoleVariableRef CVarSkinCacheDynamicMeshBounds(
+static int32 GSkinCacheDynamicMeshBounds = 2;
+static FAutoConsoleVariable CVarSkinCacheDynamicMeshBounds(
 	TEXT("r.SkinCache.DynamicMeshBounds"),
 	GSkinCacheDynamicMeshBounds,
-	TEXT("If enabled (default: off) skin cache will generate tight mesh bounds."),
+	TEXT("If enabled, skin cache will generate tight mesh bounds that are used during instance culling.\n")
+	TEXT(" 0 - disabled support compiled out of all shaders.\n")
+	TEXT(" 1 - enabled for all platforms that support atomic operations.\n")
+	TEXT(" 2 - auto (default), enabled for higher end platforms (currently those supporting Nanite).\n"),
 	ECVF_ReadOnly
 );
 
@@ -1092,7 +1095,7 @@ public:
 		OutEnvironment.SetDefine(TEXT("GPUSKIN_BONE_INDEX_UINT16"), bBoneIndex16);
 		OutEnvironment.SetDefine(TEXT("GPUSKIN_BONE_WEIGHTS_UINT16"), bBoneWeights16);
 
-		DynamicMeshBoundsModifyCompilationEnvironment(OutEnvironment);
+		DynamicMeshBoundsModifyCompilationEnvironment(Parameters, OutEnvironment);
 
 		if (FDataDrivenShaderPlatformInfo::GetSupportsRealTypes(Parameters.Platform) == ERHIFeatureSupport::RuntimeGuaranteed)
 		{
@@ -1264,11 +1267,16 @@ FDynamicMeshBoundsShaderParameters GetDynamicMeshBoundsShaderParameters(FRDGBuil
 	return ShaderParameters;
 }
 
-void DynamicMeshBoundsModifyCompilationEnvironment(FShaderCompilerEnvironment& OutEnvironment)
+static bool DoesPlatformSupportDynamicMeshBounds(EShaderPlatform ShaderPlatform)
 {
-	OutEnvironment.SetDefine(TEXT("ENABLE_DYNAMIC_MESH_BOUNDS"), GSkinCacheDynamicMeshBounds ? 1 : 0);
+	return GSkinCacheDynamicMeshBounds == 1 
+		|| (GSkinCacheDynamicMeshBounds == 2 && DoesPlatformSupportNanite(ShaderPlatform));
 }
 
+void DynamicMeshBoundsModifyCompilationEnvironment(const FGlobalShaderPermutationParameters& Parameters, FShaderCompilerEnvironment& OutEnvironment)
+{
+	OutEnvironment.SetDefine(TEXT("ENABLE_DYNAMIC_MESH_BOUNDS"), DoesPlatformSupportDynamicMeshBounds(Parameters.Platform) ? 1 : 0);
+}
 
 int32 FGPUSkinCache::AllocateDynamicMeshBoundsSlot(int32 NumSlots)
 {
@@ -1292,7 +1300,7 @@ FGPUSkinCache::FGPUSkinCache(ERHIFeatureLevel::Type InFeatureLevel, bool bInRequ
 	, ExtraRequiredMemory(0)
 	, FlushCounter(0)
 	, bRequiresMemoryLimit(bInRequiresMemoryLimit)
-	, DynamicMeshBoundsBuffer(GSkinCacheDynamicMeshBounds ? MakePimpl<FDynamicMeshBoundsBuffer>() : nullptr)
+	, DynamicMeshBoundsBuffer(DoesPlatformSupportDynamicMeshBounds(GetFeatureLevelShaderPlatform(InFeatureLevel)) ? MakePimpl<FDynamicMeshBoundsBuffer>() : nullptr)
 	, CurrentStagingBufferIndex(0)
 	, FeatureLevel(InFeatureLevel)
 	, World(InWorld)
@@ -1767,11 +1775,11 @@ ERHIPipeline FGPUSkinCache::GetDispatchPipeline(FRDGBuilder& GraphBuilder)
 
 struct FGPUSkinCache::FTaskData
 {
-	FTaskData(FRDGBuilder& GraphBuilder, ERHIPipeline InPipeline)
+	FTaskData(FRDGBuilder& GraphBuilder, EShaderPlatform ShaderPlatform, ERHIPipeline InPipeline)
 		: Pipeline(InPipeline)
 		, bAsyncCommandList(GraphBuilder.IsParallelSetupEnabled() && !GNumDispatchesToCapture)
 	{
-		if (GSkinCacheDynamicMeshBounds && FSkeletalMeshUpdater::IsEnabled())
+		if (DoesPlatformSupportDynamicMeshBounds(ShaderPlatform) && FSkeletalMeshUpdater::IsEnabled())
 		{
 			DynamicMeshBoundsBlackboardStruct = &GraphBuilder.Blackboard.Create<FDynamicMeshBoundsBlackboardStruct>();
 		}
@@ -1847,7 +1855,7 @@ struct FGPUSkinCache::FTaskData
 
 UE::Tasks::FTask FGPUSkinCache::Dispatch(FRDGBuilder& GraphBuilder, const UE::Tasks::FTask& PrerequisitesTask, ERHIPipeline InPipeline)
 {
-	FTaskData* TaskData = &GraphBuilder.Blackboard.Create<FTaskData>(GraphBuilder, InPipeline);
+	FTaskData* TaskData = &GraphBuilder.Blackboard.Create<FTaskData>(GraphBuilder, GetFeatureLevelShaderPlatform(FeatureLevel), InPipeline);
 
 	TaskData->SetupTask = GraphBuilder.AddSetupTask([this, TaskData]
 	{
diff --git a/Engine/Source/Runtime/Engine/Public/GPUSkinCache.h b/Engine/Source/Runtime/Engine/Public/GPUSkinCache.h
--- a/Engine/Source/Runtime/Engine/Public/GPUSkinCache.h
+++ b/Engine/Source/Runtime/Engine/Public/GPUSkinCache.h
@@ -63,6 +63,7 @@ struct FClothSimulData;
 struct FSkelMeshRenderSection;
 struct FVertexBufferAndSRV;
 struct FRayTracingGeometrySegment;
+struct FGlobalShaderPermutationParameters;
 
 extern bool ShouldWeCompileGPUSkinVFShaders(EShaderPlatform Platform, ERHIFeatureLevel::Type FeatureLevel);
 
@@ -115,7 +116,11 @@ ENGINE_API FDynamicMeshBoundsShaderParameters GetDynamicMeshBoundsShaderParamete
 
 // Set required shader defines for any shader using the dynamic mesh bounds (through FDynamicMeshBoundsShaderParameters), 
 // without this the operations are compiled out as well as if r.SkinCache.DynamicMeshBounds is off.
-ENGINE_API void DynamicMeshBoundsModifyCompilationEnvironment(FShaderCompilerEnvironment& OutEnvironment);
+ENGINE_API void DynamicMeshBoundsModifyCompilationEnvironment(const FGlobalShaderPermutationParameters& Parameters, FShaderCompilerEnvironment& OutEnvironment);
+
+UE_DEPRECATED(5.7, "This is just a stub. Use the above version that takes a FGlobalShaderPermutationParameters argument.")
+inline void DynamicMeshBoundsModifyCompilationEnvironment(FShaderCompilerEnvironment& OutEnvironment) { };
+
 
 class FGPUSkinCache
 {
diff --git a/Engine/Source/Runtime/Renderer/Private/InstanceCulling/InstanceCullingContext.cpp b/Engine/Source/Runtime/Renderer/Private/InstanceCulling/InstanceCullingContext.cpp
--- a/Engine/Source/Runtime/Renderer/Private/InstanceCulling/InstanceCullingContext.cpp
+++ b/Engine/Source/Runtime/Renderer/Private/InstanceCulling/InstanceCullingContext.cpp
@@ -563,7 +563,7 @@ public:
 		const bool bPrimitiveHasTileOffsetData = CVarPrimitiveHasTileOffsetData->GetValueOnAnyThread() != 0;
 		OutEnvironment.SetDefine(TEXT("PRIMITIVE_HAS_TILEOFFSET_DATA"), bPrimitiveHasTileOffsetData ? 1 : 0);
 
-		DynamicMeshBoundsModifyCompilationEnvironment(OutEnvironment);
+		DynamicMeshBoundsModifyCompilationEnvironment(Parameters, OutEnvironment);
 	}
 
 	BEGIN_SHADER_PARAMETER_STRUCT(FParameters, )
diff --git a/Engine/Source/Runtime/Renderer/Private/VirtualShadowMaps/VirtualShadowMapArray.cpp b/Engine/Source/Runtime/Renderer/Private/VirtualShadowMaps/VirtualShadowMapArray.cpp
--- a/Engine/Source/Runtime/Renderer/Private/VirtualShadowMaps/VirtualShadowMapArray.cpp
+++ b/Engine/Source/Runtime/Renderer/Private/VirtualShadowMaps/VirtualShadowMapArray.cpp
@@ -3089,7 +3089,7 @@ public:
 		OutEnvironment.SetDefine(TEXT("INDIRECT_ARGS_NUM_WORDS"), FInstanceCullingContext::IndirectArgsNumWords);
 		OutEnvironment.SetDefine(TEXT("VF_SUPPORTS_PRIMITIVE_SCENE_DATA"), 1);
 
-		DynamicMeshBoundsModifyCompilationEnvironment(OutEnvironment);
+		DynamicMeshBoundsModifyCompilationEnvironment(Parameters, OutEnvironment);
 	}
 
 	BEGIN_SHADER_PARAMETER_STRUCT(FHZBShaderParameters, )
