Material node SamplePhysicsField returns 0 instead of unsupported on mobile platforms that support SM5 materials

diff --git a/Engine/Source/Runtime/Engine/Private/Materials/HLSLMaterialTranslator.cpp b/Engine/Source/Runtime/Engine/Private/Materials/HLSLMaterialTranslator.cpp
--- a/Engine/Source/Runtime/Engine/Private/Materials/HLSLMaterialTranslator.cpp
+++ b/Engine/Source/Runtime/Engine/Private/Materials/HLSLMaterialTranslator.cpp
@@ -12349,6 +12349,18 @@ int32 FHLSLMaterialTranslator::DistanceFieldApproxAO(int32 PositionArg, EPositio
 
 int32 FHLSLMaterialTranslator::SamplePhysicsField(int32 PositionArg, EPositionOrigin PositionOrigin, const int32 OutputType, const int32 TargetIndex)
 {
+	if (FeatureLevel == ERHIFeatureLevel::ES3_1 && MobileSupportsSM5MaterialNodes(Platform))
+	{
+		if (OutputType == EFieldOutputType::Field_Output_Vector)
+		{
+			return Constant3(0.0f, 0.0f, 0.0f);
+		}
+		else if (OutputType == EFieldOutputType::Field_Output_Scalar || OutputType == EFieldOutputType::Field_Output_Integer)
+		{
+			return Constant(0.0f);
+		}
+	}
+
 	if (ErrorUnlessFeatureLevelSupported(ERHIFeatureLevel::SM5) == INDEX_NONE)
 	{
 		return INDEX_NONE;
