Move SequenceNumber unit test from "Networking" to "Net.Misc"

#rnx

diff --git a/Engine/Source/Runtime/Engine/Private/Net/Tests/Util/SequenceNumberTests.cpp b/Engine/Source/Runtime/Engine/Private/Net/Tests/Util/SequenceNumberTests.cpp
--- a/Engine/Source/Runtime/Engine/Private/Net/Tests/Util/SequenceNumberTests.cpp
+++ b/Engine/Source/Runtime/Engine/Private/Net/Tests/Util/SequenceNumberTests.cpp
@@ -5,7 +5,7 @@
 
 #if WITH_DEV_AUTOMATION_TESTS
 
-IMPLEMENT_SIMPLE_AUTOMATION_TEST(FNetSequenceNumberTest, "Network.SequenceNumberTest", EAutomationTestFlags_ApplicationContextMask | EAutomationTestFlags::EngineFilter)
+IMPLEMENT_SIMPLE_AUTOMATION_TEST(FNetSequenceNumberTest, "Net.Misc.SequenceNumberTest", EAutomationTestFlags_ApplicationContextMask | EAutomationTestFlags::EngineFilter)
 
 bool FNetSequenceNumberTest::RunTest(const FString& Parameters)
 {
