Fix crash when accessing null object in LowerTextureType due to missing MIR::FTextureObject

diff --git a/Engine/Source/Runtime/Engine/Private/Materials/MaterialIRToHLSLTranslator.cpp b/Engine/Source/Runtime/Engine/Private/Materials/MaterialIRToHLSLTranslator.cpp
--- a/Engine/Source/Runtime/Engine/Private/Materials/MaterialIRToHLSLTranslator.cpp
+++ b/Engine/Source/Runtime/Engine/Private/Materials/MaterialIRToHLSLTranslator.cpp
@@ -1572,7 +1572,20 @@ struct FPrivate : FMaterialIRToHLSLTranslation
 	{
 		switch (TextureRead->TextureObject->Type.GetKind())
 		{
-			case MIR::ETypeKind::Texture2D: LowerTextureType(TextureRead->TextureObject->As<MIR::FTextureObject>()); break;
+			case MIR::ETypeKind::Texture2D: 
+			{
+				const MIR::FTextureObject* TextureObject = TextureRead->TextureObject->As<MIR::FTextureObject>();
+				
+				if (!TextureObject)
+				{
+					Printer << TEXTVIEW("Texture2D");
+				}
+				else
+				{
+					LowerTextureType(TextureObject);
+				}
+				break;
+			}
 			default: UE_MIR_UNREACHABLE();
 		}
 
