[Input] Deprecate the "Speech Mappings". These speech mappings were originally added to support Hololens, which is no longer a supported platform and can be removed.

[RN] Deprecate the "Speech Mappings". These speech mappings were originally added to support Hololens, which is no longer a supported platform and can be removed.

diff --git a/Engine/Source/Runtime/Engine/Classes/GameFramework/InputSettings.h b/Engine/Source/Runtime/Engine/Classes/GameFramework/InputSettings.h
--- a/Engine/Source/Runtime/Engine/Classes/GameFramework/InputSettings.h
+++ b/Engine/Source/Runtime/Engine/Classes/GameFramework/InputSettings.h
@@ -173,6 +173,7 @@ private:
 	TSet<FName> DeprecatedActionAndAxisNames;
 
 	/** List of Speech Mappings */
+	UE_DEPRECATED(5.7, "SpeechMappings are deprecated. These were only used by a now deprecated platform.")
 	UPROPERTY(config, EditAnywhere, Category = "Bindings")
 	TArray<struct FInputActionSpeechMapping> SpeechMappings;
 
@@ -278,14 +279,20 @@ public:
 	ENGINE_API const TArray <FInputActionKeyMapping>& GetActionMappings() const;
 	/** Ask for all the axis mappings */
 	ENGINE_API const TArray <FInputAxisKeyMapping>& GetAxisMappings() const;
+
 	/** Ask for all the speech mappings */
+	UE_DEPRECATED(5.7, "Speech Mappings are deprecated. These were only used by a now deprecated platform.")
+	PRAGMA_DISABLE_DEPRECATION_WARNINGS
 	ENGINE_API const TArray <FInputActionSpeechMapping>& GetSpeechMappings() const;
+	PRAGMA_ENABLE_DEPRECATION_WARNINGS
 
 	/** Finds unique action name based on existing action names */
 	ENGINE_API bool DoesActionExist(const FName InActionName);
 	/** Finds unique axis name based on existing action names */
 	ENGINE_API bool DoesAxisExist(const FName InAxisName);
+	
 	/** Finds unique speech name based on existing speech names */
+	UE_DEPRECATED(5.7, "Speech Mappings are deprecated. These were only used by a now deprecated platform.")
 	ENGINE_API bool DoesSpeechExist(const FName InSpeechName);
 
 
diff --git a/Engine/Source/Runtime/Engine/Classes/GameFramework/PlayerInput.h b/Engine/Source/Runtime/Engine/Classes/GameFramework/PlayerInput.h
--- a/Engine/Source/Runtime/Engine/Classes/GameFramework/PlayerInput.h
+++ b/Engine/Source/Runtime/Engine/Classes/GameFramework/PlayerInput.h
@@ -270,8 +270,8 @@ struct FInputAxisKeyMapping
  *
  * @see https://docs.unrealengine.com/latest/INT/Gameplay/Input/index.html
  */
-USTRUCT( BlueprintType )
-struct FInputActionSpeechMapping
+USTRUCT(BlueprintType)
+struct UE_DEPRECATED(5.7, "Speech mappings are deprecated. These were only used by a now deprecated platform.") FInputActionSpeechMapping
 {
 	GENERATED_USTRUCT_BODY()
 
@@ -301,6 +301,7 @@ private:
 	FName SpeechKeyword;
 public:
 
+	PRAGMA_DISABLE_DEPRECATION_WARNINGS
 	bool operator==(const FInputActionSpeechMapping& Other) const
 	{
 		return (   ActionName == Other.ActionName
@@ -325,6 +326,14 @@ public:
 		: ActionName(InActionName)
 		, SpeechKeyword(InSpeechKeyword)
 	{}
+	
+	~FInputActionSpeechMapping() = default;
+	FInputActionSpeechMapping(const FInputActionSpeechMapping&) = default;
+	FInputActionSpeechMapping(FInputActionSpeechMapping&&) = default;
+	FInputActionSpeechMapping& operator=(const FInputActionSpeechMapping&) = default;
+	FInputActionSpeechMapping& operator=(FInputActionSpeechMapping&&) = default;
+	
+	PRAGMA_ENABLE_DEPRECATION_WARNINGS
 };
 
 /** Paramaters to be considered when calling UPlayerInput::InputKey. */
diff --git a/Engine/Source/Runtime/Engine/Private/UserInterface/InputSettings.cpp b/Engine/Source/Runtime/Engine/Private/UserInterface/InputSettings.cpp
--- a/Engine/Source/Runtime/Engine/Private/UserInterface/InputSettings.cpp
+++ b/Engine/Source/Runtime/Engine/Private/UserInterface/InputSettings.cpp
@@ -213,6 +213,7 @@ void UInputSettings::PostEditChangeChainProperty(FPropertyChangedChainEvent& Pro
 
 	const FName MemberPropertyName = PropertyChangedEvent.PropertyChain.GetActiveMemberNode()->GetValue()->GetFName();
 
+	PRAGMA_DISABLE_DEPRECATION_WARNINGS
 	if (MemberPropertyName == GET_MEMBER_NAME_CHECKED(UInputSettings, ActionMappings) || MemberPropertyName == GET_MEMBER_NAME_CHECKED(UInputSettings, AxisMappings) ||
 		MemberPropertyName == GET_MEMBER_NAME_CHECKED(UInputSettings, AxisConfig) || MemberPropertyName == GET_MEMBER_NAME_CHECKED(UInputSettings, SpeechMappings))
 	{
@@ -227,6 +228,7 @@ void UInputSettings::PostEditChangeChainProperty(FPropertyChangedChainEvent& Pro
 	{
 		FPlatformApplicationMisc::EnableMotionData(bEnableMotionControls);
 	}
+	PRAGMA_ENABLE_DEPRECATION_WARNINGS
 }
 
 #endif
@@ -235,7 +237,6 @@ void UInputSettings::SaveKeyMappings()
 {
 	ActionMappings.Sort();
 	AxisMappings.Sort();
-	SpeechMappings.Sort();
 	SaveConfig();
 }
 
@@ -335,10 +336,12 @@ void UInputSettings::GetActionNames(TArray<FName>& ActionNames) const
 		ActionNames.AddUnique(ActionMapping.ActionName);
 	}
 
+	PRAGMA_DISABLE_DEPRECATION_WARNINGS
 	for (const FInputActionSpeechMapping& SpeechMapping : SpeechMappings)
 	{
 		ActionNames.AddUnique(SpeechMapping.GetActionName());
 	}
+	PRAGMA_ENABLE_DEPRECATION_WARNINGS
 }
 
 bool UInputSettings::IsActionMappingDeprecated(const FName ActionName) const
@@ -409,7 +412,9 @@ FName UInputSettings::GetUniqueActionName(const FName BaseActionMappingName)
 		NewActionMappingName = FName(BaseActionMappingName, ++NewMappingCount);
 		bFoundUniqueName = true;
 
+		PRAGMA_DISABLE_DEPRECATION_WARNINGS
 		bFoundUniqueName = !(DoesActionExist(NewActionMappingName) || DoesSpeechExist(NewActionMappingName));
+		PRAGMA_ENABLE_DEPRECATION_WARNINGS
 	} while (!bFoundUniqueName);
 
 	return NewActionMappingName;
@@ -461,11 +466,14 @@ const TArray <FInputAxisKeyMapping>& UInputSettings::GetAxisMappings() const
 	return AxisMappings;
 }
 
+PRAGMA_DISABLE_DEPRECATION_WARNINGS
 const TArray <FInputActionSpeechMapping>& UInputSettings::GetSpeechMappings() const
 {
 	return SpeechMappings;
 }
+PRAGMA_ENABLE_DEPRECATION_WARNINGS
 
+PRAGMA_DISABLE_DEPRECATION_WARNINGS
 struct FMatchMappingByName
 {
 	FMatchMappingByName(const FName InName)
@@ -490,6 +498,7 @@ struct FMatchMappingByName
 
 	FName Name;
 };
+PRAGMA_ENABLE_DEPRECATION_WARNINGS
 
 /** Finds unique action name based on existing action names */
 bool UInputSettings::DoesActionExist(const FName InActionName)
@@ -504,10 +513,12 @@ bool UInputSettings::DoesAxisExist(const FName InAxisName)
 }
 
 /** Finds unique axis name based on existing action names */
+PRAGMA_DISABLE_DEPRECATION_WARNINGS
 bool UInputSettings::DoesSpeechExist(const FName InSpeechName)
 {
 	return (SpeechMappings.FindByPredicate(FMatchMappingByName(InSpeechName)) != nullptr);
 }
+PRAGMA_ENABLE_DEPRECATION_WARNINGS
 
 /** Get the member name for the details panel */
 const FName UInputSettings::GetActionMappingsPropertyName()
diff --git a/Engine/Source/Runtime/Engine/Private/UserInterface/PlayerInput.cpp b/Engine/Source/Runtime/Engine/Private/UserInterface/PlayerInput.cpp
--- a/Engine/Source/Runtime/Engine/Private/UserInterface/PlayerInput.cpp
+++ b/Engine/Source/Runtime/Engine/Private/UserInterface/PlayerInput.cpp
@@ -879,6 +879,8 @@ void UPlayerInput::ForceRebuildingKeyMaps(const bool bRestoreDefaults)
 			AxisMappings = InputSettings->GetAxisMappings();
 			ActionMappings = InputSettings->GetActionMappings();
 
+			// Speech mappings are deprecated in 5.7
+			PRAGMA_DISABLE_DEPRECATION_WARNINGS
 			//append on speech action mappings
 			const TArray<FInputActionSpeechMapping>& SpeechMappings = InputSettings->GetSpeechMappings();
 			for (const FInputActionSpeechMapping& SpeechMapping : SpeechMappings)
@@ -887,6 +889,7 @@ void UPlayerInput::ForceRebuildingKeyMaps(const bool bRestoreDefaults)
 				ConvertedSpeechToActionMap.ActionName = SpeechMapping.GetActionName();
 				ConvertedSpeechToActionMap.Key = SpeechMapping.GetKeyName();
 			}
+			PRAGMA_ENABLE_DEPRECATION_WARNINGS
 		}
 	}
 
