Dstm: Fix pawn going to the sky when crossing border on a crate. Mark OldBaseLocation and OlsBaseQuat as UPropertys so they can be migrated with the pawn when it is crossing the boundary.

diff --git a/Engine/Source/Runtime/Engine/Classes/GameFramework/CharacterMovementComponent.h b/Engine/Source/Runtime/Engine/Classes/GameFramework/CharacterMovementComponent.h
--- a/Engine/Source/Runtime/Engine/Classes/GameFramework/CharacterMovementComponent.h
+++ b/Engine/Source/Runtime/Engine/Classes/GameFramework/CharacterMovementComponent.h
@@ -262,9 +262,11 @@ public:
 	
 
 	/** Saved location of object we are standing on, for UpdateBasedMovement() to determine if base moved in the last frame, and therefore pawn needs an update. */
+	UPROPERTY(Transient)
 	FQuat OldBaseQuat;
 
 	/** Saved location of object we are standing on, for UpdateBasedMovement() to determine if base moved in the last frame, and therefore pawn needs an update. */
+	UPROPERTY(Transient)
 	FVector OldBaseLocation;
 
 	/** The maximum ground speed when walking. Also determines maximum lateral speed when falling. */
