Modify implementation of DefaultMaterialAttributes in material translator.
Using typed chunks instead of void prevents the translator from duplicating output shader code for each pass through the graph.

diff --git a/Engine/Source/Runtime/Engine/Private/Materials/HLSLMaterialTranslator.cpp b/Engine/Source/Runtime/Engine/Private/Materials/HLSLMaterialTranslator.cpp
--- a/Engine/Source/Runtime/Engine/Private/Materials/HLSLMaterialTranslator.cpp
+++ b/Engine/Source/Runtime/Engine/Private/Materials/HLSLMaterialTranslator.cpp
@@ -12650,15 +12650,14 @@ int32 FHLSLMaterialTranslator::DefaultMaterialAttributes()
 			continue;
 		}
 		
-		FVector4f Default = FMaterialAttributeDefinitionMap::GetDefaultValue(AttributeGUID);
-		
-		AddCodeChunk(
-			MCT_VoidStatement, 
-			TEXT("%s.%s = float4(%f, %f, %f, %f);"),
-			 *AttributesStr,
-			 *FMaterialAttributeDefinitionMap::GetAttributeName(AttributeGUID),
-			 Default.X, Default.Y, Default.Z, Default.W
-		 );
+		if (FMaterialAttributeDefinitionMap::GetShaderFrequency(AttributeGUID) != ShaderFrequency)
+		{
+			continue;
+		}
+
+		const FString PropertyName = FMaterialAttributeDefinitionMap::GetAttributeName(AttributeGUID);
+		const int32 Default = FMaterialAttributeDefinitionMap::CompileDefaultExpression(this, AttributeGUID);
+		Attributes = AddCodeChunk(MCT_MaterialAttributes, TEXT("FMaterialAttributes_Set%s(%s, %s)"), *PropertyName, *GetParameterCode(Attributes), *GetParameterCode(Default));
 	}
 
 	return Attributes;
