Fix link error in some modules from missing SlateCore dep for inline method by moving code into Engine module instead of inline.

diff --git a/Engine/Source/Runtime/Engine/Classes/Engine/LocalPlayer.h b/Engine/Source/Runtime/Engine/Classes/Engine/LocalPlayer.h
--- a/Engine/Source/Runtime/Engine/Classes/Engine/LocalPlayer.h
+++ b/Engine/Source/Runtime/Engine/Classes/Engine/LocalPlayer.h
@@ -172,7 +172,7 @@ class ULocalPlayer : public UPlayer
 public:
 
 	/** DO NOT USE. This constructor is for internal usage only for hot-reload purposes. */
-	ULocalPlayer(FVTableHelper& Helper) : Super(Helper), SlateOperations(FReply::Unhandled()) {}
+	ENGINE_API ULocalPlayer(FVTableHelper& Helper);
 
 	/** The FUniqueNetId which this player is associated with. */
 	FUniqueNetIdRepl CachedUniqueNetId;
diff --git a/Engine/Source/Runtime/Engine/Private/LocalPlayer.cpp b/Engine/Source/Runtime/Engine/Private/LocalPlayer.cpp
--- a/Engine/Source/Runtime/Engine/Private/LocalPlayer.cpp
+++ b/Engine/Source/Runtime/Engine/Private/LocalPlayer.cpp
@@ -235,6 +235,12 @@ ULocalPlayer::ULocalPlayer(const FObjectInitializer& ObjectInitializer)
 	PendingLevelPlayerControllerClass = APlayerController::StaticClass();
 }
 
+ULocalPlayer::ULocalPlayer(FVTableHelper& Helper)
+	: Super(Helper)
+	, SlateOperations(FReply::Unhandled())
+{
+}
+
 void ULocalPlayer::PlayerAdded(UGameViewportClient* InViewportClient, int32 InControllerID)
 {
 	ViewportClient = InViewportClient;
