Don't drop Nanite bulkdata when running cook commandlet, so it doesn't have to be reloaded immediately after. It will still be dropped when it is done by ClearAllCachedCookedPlatformData.

diff --git a/Engine/Source/Runtime/Engine/Private/StaticMesh.cpp b/Engine/Source/Runtime/Engine/Private/StaticMesh.cpp
--- a/Engine/Source/Runtime/Engine/Private/StaticMesh.cpp
+++ b/Engine/Source/Runtime/Engine/Private/StaticMesh.cpp
@@ -4390,7 +4390,7 @@ void FStaticMeshRenderData::Cache(const ITargetPlatform* TargetPlatform, UStatic
 				if (bSavedToDDC)
 				{
 					NaniteResources.SetHasBuildFromDDCError(false);
-					if (NaniteResources.HasStreamingData())
+					if (NaniteResources.HasStreamingData() && !IsRunningCookCommandlet())
 					{
 						// Drop streaming data from memory when it has been successfully committed to DDC
 						NaniteResources.DropBulkData();
