Fixes crash in the new translator when the Material happens to be in invalid state and some expression input references an input that isn't found on the expression.

- Also fixing M_TestVolume material. It had a problematic function in some invalid state.

#tests M_TestVolume opens in EngineTest without crashing (but won't translate yet)
#rnx

diff --git a/Engine/Source/Runtime/Engine/Private/Materials/MaterialIREmitter.cpp b/Engine/Source/Runtime/Engine/Private/Materials/MaterialIREmitter.cpp
--- a/Engine/Source/Runtime/Engine/Private/Materials/MaterialIREmitter.cpp
+++ b/Engine/Source/Runtime/Engine/Private/Materials/MaterialIREmitter.cpp
@@ -176,7 +176,6 @@ static int32 SlowFindExpressionInputIndex(UMaterialExpression* Expression, const
 			return It.Index;
 		}
 	}
-	check(false && "No input found.");
 	return -1;
 }
 
@@ -184,7 +183,7 @@ static int32 SlowFindExpressionInputIndex(UMaterialExpression* Expression, const
 static FName SlowFindInputName(UMaterialExpression* Expression, const FExpressionInput* InInput)
 {
 	int32 InputIndex = SlowFindExpressionInputIndex(Expression, InInput);
-	return Expression->GetInputName(InputIndex);
+	return InputIndex != INDEX_NONE ? Expression->GetInputName(InputIndex) : FName{};
 }
 
 /*------------------------------------- FValueRef --------------------------------------*/
@@ -1306,7 +1305,7 @@ static FValueRef EmitAttributeDefaultValue(FEmitter& Emitter, const UMaterial* M
 
 FValueRef FEmitter::Aggregate(const UMaterialAggregate* InAggregate, FValueRef InPrototype, TConstArrayView<FValueRef> AttributeValues)
 {
-	// Check that the szpecified prototype (if any) aggregate matches the one provided.
+	// Check that the specified prototype (if any) aggregate matches the one provided.
 	InPrototype = CheckIsAggregate(InPrototype, InAggregate);
 
 	// If a prototype is provided and there are no attribute assignments, this is a no-op, return the prototype.
@@ -1332,7 +1331,7 @@ FValueRef FEmitter::Aggregate(const UMaterialAggregate* InAggregate, FValueRef I
 		}
 		else if (InPrototype)
 		{
-			// ... otherwise use the component value as in the prototype, if od...
+			// ... otherwise use the component value as in the prototype if provided
 			Components[i] = Subscript(InPrototype, i);
 		}
 		else
