Fix UPoseAsset::SetBasePoseName not returning true on success

diff --git a/Engine/Source/Runtime/Engine/Private/Animation/PoseAsset.cpp b/Engine/Source/Runtime/Engine/Private/Animation/PoseAsset.cpp
--- a/Engine/Source/Runtime/Engine/Private/Animation/PoseAsset.cpp
+++ b/Engine/Source/Runtime/Engine/Private/Animation/PoseAsset.cpp
@@ -1182,6 +1182,7 @@ bool UPoseAsset::SetBasePoseName(const FName& NewBasePoseName)
 		if (NewIndex != INDEX_NONE)
 		{
 			BasePoseIndex = NewIndex;
+			return true;
 		}
 		else
 		{
