Material Cache, silenced SA

diff --git a/Engine/Source/Runtime/Engine/Private/Materials/MaterialExpressions.cpp b/Engine/Source/Runtime/Engine/Private/Materials/MaterialExpressions.cpp
--- a/Engine/Source/Runtime/Engine/Private/Materials/MaterialExpressions.cpp
+++ b/Engine/Source/Runtime/Engine/Private/Materials/MaterialExpressions.cpp
@@ -20409,10 +20409,10 @@ void UMaterialExpressionMaterialCache::ConstructOutputs()
 			Outputs.Add(FExpressionOutput(
 				*GetMaterialCacheAttributeDecoration(Attribute),
 				1,
-				ComponentCount > 0,
-				ComponentCount > 1,
-				ComponentCount > 2,
-				ComponentCount > 3
+				ComponentCount > 0, //-V547
+				ComponentCount > 1, //-V547
+				ComponentCount > 2, //-V547
+				ComponentCount > 3 //-V547
 			));
 		}
 	}
