Virtual Texture, added r.VT.ForceRecreatePool to forcefully recreate the backing pool for testing purposes

diff --git a/Engine/Source/Runtime/Engine/Private/VT/VirtualTexturePoolConfig.cpp b/Engine/Source/Runtime/Engine/Private/VT/VirtualTexturePoolConfig.cpp
--- a/Engine/Source/Runtime/Engine/Private/VT/VirtualTexturePoolConfig.cpp
+++ b/Engine/Source/Runtime/Engine/Private/VT/VirtualTexturePoolConfig.cpp
@@ -85,6 +85,15 @@ struct FVirtualTexturePoolVersion
 };
 FVirtualTexturePoolVersion GVirtualTexturePoolVersion;
 
+static FAutoConsoleCommand CVarVTForceRecreatePoolCommand(
+	TEXT("r.VT.ForceRecreatePool"),
+	TEXT("Forcefully recreates the backing virtual pool"),
+	FConsoleCommandWithArgsDelegate::CreateLambda([](const TArray<FString>&)
+	{
+		GVirtualTexturePoolVersion.Increment_GameThread();
+		OnVirtualTexturePoolConfigUpdate();
+	})
+);
 
 UVirtualTexturePoolConfig::UVirtualTexturePoolConfig(const FObjectInitializer& ObjectInitializer)
 	: Super(ObjectInitializer)
