Material Translator, fixed GetParameterType crash on invalid indices

diff --git a/Engine/Source/Runtime/Engine/Private/Materials/HLSLMaterialTranslator.cpp b/Engine/Source/Runtime/Engine/Private/Materials/HLSLMaterialTranslator.cpp
--- a/Engine/Source/Runtime/Engine/Private/Materials/HLSLMaterialTranslator.cpp
+++ b/Engine/Source/Runtime/Engine/Private/Materials/HLSLMaterialTranslator.cpp
@@ -16191,6 +16191,11 @@ int32 FHLSLMaterialTranslator::CustomOutput(const FString& CustomOutputName, con
 
 int32 FHLSLMaterialTranslator::CustomOutput(class UMaterialExpressionCustomOutput* Custom, int32 OutputIndex, int32 OutputCode, EMaterialCustomOutputFlags Flags)
 {
+	if (OutputCode == INDEX_NONE)
+	{
+		return INDEX_NONE;
+	}
+	
 	return CustomOutput(Custom->GetDescription(), HLSLTypeString(GetParameterType(OutputCode)), Custom->GetFunctionName(), OutputIndex, OutputCode, Flags);
 }
 
