[AnimSequenceBase]
- Restore the unconditional postload for the time being since there is a circular dependency that would otherwise cause problems.
- UDataModel A (PostLoad) -> UAnimSequence A (PostLoad) which needs to call PostLoad on UDataModel A even though its already on the stack

diff --git a/Engine/Source/Runtime/Engine/Private/Animation/AnimSequenceBase.cpp b/Engine/Source/Runtime/Engine/Private/Animation/AnimSequenceBase.cpp
--- a/Engine/Source/Runtime/Engine/Private/Animation/AnimSequenceBase.cpp
+++ b/Engine/Source/Runtime/Engine/Private/Animation/AnimSequenceBase.cpp
@@ -162,7 +162,7 @@ void UAnimSequenceBase::PostLoad()
 				if (DataModel != nullptr)
 				{
 					DataModel->ConditionalPreload();
-					DataModel->ConditionalPostLoad();
+					DataModel->PostLoad();
 					DataModel->ConditionalPostLoadSubobjects();
 					PopulateWithExistingModel(DataModel.Get());
 
