Fix a potential crash in map error checking where an ABrush has a null Polys reference.

diff --git a/Engine/Source/Runtime/Core/Private/Misc/MapErrors.cpp b/Engine/Source/Runtime/Core/Private/Misc/MapErrors.cpp
--- a/Engine/Source/Runtime/Core/Private/Misc/MapErrors.cpp
+++ b/Engine/Source/Runtime/Core/Private/Misc/MapErrors.cpp
@@ -24,6 +24,7 @@ FLazyName FMapErrors::InvalidHairStrandsMaterial(TEXT("InvalidHairStrandsMateria
 FLazyName FMapErrors::MultipleSkyAtmosphereTypes(TEXT("MultipleSkyAtmosphereTypes"));
 FLazyName FMapErrors::InvalidTrace(TEXT("InvalidTrace"));
 FLazyName FMapErrors::BrushZeroPolygons(TEXT("BrushZeroPolygons"));
+FLazyName FMapErrors::BrushNullPolygons(TEXT("BrushNullPolygons"));
 FLazyName FMapErrors::CleanBSPMaterials(TEXT("CleanBSPMaterials"));
 FLazyName FMapErrors::BrushComponentNull(TEXT("BrushComponentNull"));
 FLazyName FMapErrors::PlanarBrush(TEXT("PlanarBrush"));
diff --git a/Engine/Source/Runtime/Core/Public/Misc/MapErrors.h b/Engine/Source/Runtime/Core/Public/Misc/MapErrors.h
--- a/Engine/Source/Runtime/Core/Public/Misc/MapErrors.h
+++ b/Engine/Source/Runtime/Core/Public/Misc/MapErrors.h
@@ -75,6 +75,9 @@ struct FMapErrors
 	/** {Brush Actor} : Brush has zero polygons - please delete!: This warning indicates that you have a brush in the level that doesn't have any polygons associated with it.  The brush should be deleted as it isn't doing anything useful. */
     static CORE_API FLazyName BrushZeroPolygons;
     
+	/** {Brush Actor} : Brush has NULL polygons property - please delete!: This warning indicates that you have a brush in the level with a "None" polygons reference.  The brush should be deleted as it isn't doing anything useful. */
+	static CORE_API FLazyName BrushNullPolygons;
+    
 	/** Run 'Clean BSP Materials' to clear {count} references: This warning indicates that there are material references on brush faces that aren't contributing to the BSP, and that applying the Tools->'Clean BSP Materials' operation can clean up these references. */
     static CORE_API FLazyName CleanBSPMaterials;
     
diff --git a/Engine/Source/Runtime/Engine/Private/ErrorChecking.cpp b/Engine/Source/Runtime/Engine/Private/ErrorChecking.cpp
--- a/Engine/Source/Runtime/Engine/Private/ErrorChecking.cpp
+++ b/Engine/Source/Runtime/Engine/Private/ErrorChecking.cpp
@@ -123,43 +123,57 @@ void ABrush::CheckForErrors()
 		// NOTE : don't report NULL texture references on the builder brush - it doesn't matter there
 		if( Brush && !FActorEditorUtils::IsABuilderBrush(this) && !IsBrushShape() )
 		{
-			// A brush without any polygons in it isn't useful.  Should be deleted.
-			if( Brush->Polys->Element.Num() == 0 )
+			if ( Brush->Polys )
 			{
-				FFormatNamedArguments Arguments;
-				Arguments.Add(TEXT("ActorName"), FText::FromString(GetName()));
-				FMessageLog("MapCheck").Warning()
-					->AddToken(FUObjectToken::Create(this))
-					->AddToken(FTextToken::Create(FText::Format(LOCTEXT( "MapCheck_Message_BrushZeroPolygons", "{ActorName} : Brush has zero polygons - please delete" ), Arguments) ))
-					->AddToken(FMapErrorToken::Create(FMapErrors::BrushZeroPolygons));
-			}
-
-			// Check for non-coplanar polygons.
-			for( int32 x = 0 ; x < Brush->Polys->Element.Num() ; ++x )
-			{
-				FPoly* Poly = &(Brush->Polys->Element[x]);
-				bool	Coplanar = 1;
+				// A brush without any polygons in it isn't useful.  Should be deleted.
+				if( Brush->Polys->Element.Num() == 0 )
+				{
+					FFormatNamedArguments Arguments;
+					Arguments.Add(TEXT("ActorName"), FText::FromString(GetName()));
+					FMessageLog("MapCheck").Warning()
+						->AddToken(FUObjectToken::Create(this))
+						->AddToken(FTextToken::Create(FText::Format(LOCTEXT( "MapCheck_Message_BrushZeroPolygons", "{ActorName} : Brush has zero polygons - please delete" ), Arguments) ))
+						->AddToken(FMapErrorToken::Create(FMapErrors::BrushZeroPolygons));
+				}
 
-				for(int32 VertexIndex = 0;VertexIndex < Poly->Vertices.Num();VertexIndex++)
+				// Check for non-coplanar polygons.
+				for( int32 x = 0 ; x < Brush->Polys->Element.Num() ; ++x )
 				{
-					if(FMath::Abs(FPlane((FVector)Poly->Vertices[0], (FVector)Poly->Normal).PlaneDot((FVector)Poly->Vertices[VertexIndex])) > UE_THRESH_POINT_ON_PLANE)
+					FPoly* Poly = &(Brush->Polys->Element[x]);
+					bool	Coplanar = 1;
+
+					for(int32 VertexIndex = 0;VertexIndex < Poly->Vertices.Num();VertexIndex++)
 					{
-						FFormatNamedArguments Arguments;
-						Arguments.Add(TEXT("ActorName"), FText::FromString(GetName()));
-						FMessageLog("MapCheck").Warning()
-							->AddToken(FUObjectToken::Create(this))
-							->AddToken(FTextToken::Create(FText::Format(LOCTEXT( "MapCheck_Message_NonCoPlanarPolys", "{ActorName} : Brush has non-coplanar polygons" ), Arguments) ))
-							->AddToken(FMapErrorToken::Create(FMapErrors::NonCoPlanarPolys));
-						Coplanar = 0;
-						break;
+						if(FMath::Abs(FPlane((FVector)Poly->Vertices[0], (FVector)Poly->Normal).PlaneDot((FVector)Poly->Vertices[VertexIndex])) > UE_THRESH_POINT_ON_PLANE)
+						{
+							FFormatNamedArguments Arguments;
+							Arguments.Add(TEXT("ActorName"), FText::FromString(GetName()));
+							FMessageLog("MapCheck").Warning()
+								->AddToken(FUObjectToken::Create(this))
+								->AddToken(FTextToken::Create(FText::Format(LOCTEXT( "MapCheck_Message_NonCoPlanarPolys", "{ActorName} : Brush has non-coplanar polygons" ), Arguments) ))
+								->AddToken(FMapErrorToken::Create(FMapErrors::NonCoPlanarPolys));
+							Coplanar = 0;
+							break;
+						}
 					}
-				}
 
-				if(!Coplanar)
-				{
-					break;
+					if(!Coplanar)
+					{
+						break;
+					}
 				}
 			}
+			else
+			{
+				// A brush without any polygons in it is invalid.  Should be deleted.
+				FFormatNamedArguments Arguments;
+                Arguments.Add(TEXT("ActorName"), FText::FromString(GetName()));
+                FMessageLog("MapCheck").Warning()
+                	->AddToken(FUObjectToken::Create(this))
+                	->AddToken(FTextToken::Create(FText::Format(LOCTEXT( "MapCheck_Message_BrushNullPolygons", "{ActorName} : Brush has a NULL polygon reference - please delete" ), Arguments) ))
+                	->AddToken(FMapErrorToken::Create(FMapErrors::BrushZeroPolygons));
+			}
+			
 
 			// check for planar brushes which might mess up collision
 			if(Brush->Bounds.BoxExtent.Z < UE_SMALL_NUMBER || Brush->Bounds.BoxExtent.Y < UE_SMALL_NUMBER || Brush->Bounds.BoxExtent.X < UE_SMALL_NUMBER)
