Fixed material texture parameter duplicates caused by reading unitialized material cache tag indices
- Bumped material translator DDC version

diff --git a/Engine/Source/Runtime/Core/Private/UObject/FortniteMainBranchObjectVersion.cpp b/Engine/Source/Runtime/Core/Private/UObject/FortniteMainBranchObjectVersion.cpp
--- a/Engine/Source/Runtime/Core/Private/UObject/FortniteMainBranchObjectVersion.cpp
+++ b/Engine/Source/Runtime/Core/Private/UObject/FortniteMainBranchObjectVersion.cpp
@@ -14,6 +14,6 @@ TMap<FGuid, FGuid> FFortniteMainBranchObjectVersion::GetSystemGuids()
 	SystemGuids.Add(DevGuids.Niagara_LatestScriptCompileVersion, FGuid("97F0FD7149D90D42B65D6557AF477D1A"));
 	SystemGuids.Add(DevGuids.SkeletalMeshDerivedDataVersion, FGuid("C998F2F3B2884A02B7E290B8C25835F0"));
 	SystemGuids.Add(DevGuids.STATICMESH_DERIVEDDATA_VER, FGuid("3ABF32149F9D4D838750368AB95573BE"));
-	SystemGuids.Add(DevGuids.MaterialTranslationDDCVersion, FGuid("AECF2BF62A8A49B8929C91FE8AA32F44"));
+	SystemGuids.Add(DevGuids.MaterialTranslationDDCVersion, FGuid("5A3429F7839C4DCF88F2CA3679C0D10E"));
 	return SystemGuids;
 }
diff --git a/Engine/Source/Runtime/Engine/Private/Materials/MaterialUniformExpressions.cpp b/Engine/Source/Runtime/Engine/Private/Materials/MaterialUniformExpressions.cpp
--- a/Engine/Source/Runtime/Engine/Private/Materials/MaterialUniformExpressions.cpp
+++ b/Engine/Source/Runtime/Engine/Private/Materials/MaterialUniformExpressions.cpp
@@ -1741,6 +1741,7 @@ uint32 FUniformExpressionSet::GetReferencedTexture2DRHIHash(const FMaterialRende
 FMaterialUniformExpressionTexture::FMaterialUniformExpressionTexture() :
 	TextureIndex(INDEX_NONE),
 	TextureLayerIndex(INDEX_NONE),
+	MaterialCacheTagIndex(INDEX_NONE),
 	PageTableLayerIndex(INDEX_NONE),
 #if WITH_EDITORONLY_DATA
 	SamplerType(SAMPLERTYPE_Color),
@@ -1752,6 +1753,7 @@ FMaterialUniformExpressionTexture::FMaterialUniformExpressionTexture() :
 FMaterialUniformExpressionTexture::FMaterialUniformExpressionTexture(int32 InTextureIndex, EMaterialSamplerType InSamplerType, ESamplerSourceMode InSamplerSource, bool InVirtualTexture) :
 	TextureIndex(InTextureIndex),
 	TextureLayerIndex(INDEX_NONE),
+	MaterialCacheTagIndex(INDEX_NONE),
 	PageTableLayerIndex(INDEX_NONE),
 #if WITH_EDITORONLY_DATA
 	SamplerType(InSamplerType),
@@ -1765,6 +1767,7 @@ FMaterialUniformExpressionTexture::FMaterialUniformExpressionTexture(int32 InTex
 FMaterialUniformExpressionTexture::FMaterialUniformExpressionTexture(int32 InTextureIndex, int16 InTextureLayerIndex, int16 InPageTableLayerIndex, EMaterialSamplerType InSamplerType)
 	: TextureIndex(InTextureIndex)
 	, TextureLayerIndex(InTextureLayerIndex)
+	, MaterialCacheTagIndex(INDEX_NONE)
 	, PageTableLayerIndex(InPageTableLayerIndex)
 #if WITH_EDITORONLY_DATA
 	, SamplerType(InSamplerType)
@@ -1778,6 +1781,7 @@ FMaterialUniformExpressionTexture::FMaterialUniformExpressionTexture(int32 InTex
 FMaterialUniformExpressionTexture::FMaterialUniformExpressionTexture(int32 InTextureIndex, EMaterialSamplerType InSamplerType)
 	: TextureIndex(InTextureIndex)
 	, TextureLayerIndex(INDEX_NONE)
+	, MaterialCacheTagIndex(INDEX_NONE)
 	, PageTableLayerIndex(INDEX_NONE)
 #if WITH_EDITORONLY_DATA
 	, SamplerType(InSamplerType)
