New Material Translator:  Shading Model expression support.  Also added and made shared a bunch of validation logic from the old translator that includes shading model specific validation, but also other miscellaneous validation checking for invalid combinations of material domain, blend mode, and translucency pass.

#rnx

diff --git a/Engine/Source/Runtime/Engine/Private/Materials/HLSLMaterialTranslator.cpp b/Engine/Source/Runtime/Engine/Private/Materials/HLSLMaterialTranslator.cpp
--- a/Engine/Source/Runtime/Engine/Private/Materials/HLSLMaterialTranslator.cpp
+++ b/Engine/Source/Runtime/Engine/Private/Materials/HLSLMaterialTranslator.cpp
@@ -2031,21 +2031,6 @@ void FHLSLMaterialTranslator::TranslateMaterial()
 
 	CHECK_DDC_QUERY_FINISHED_ELSE_RETURN();
 
-	// If Substrate is enabled or if this is a Substrate material cooked/used in non-Substrate mode, 
-	// we disable this warning as Substrate supports 'colored transmittance only' mode (i.e, modulate).
-	if (!bSubstrateEnabled && !FrontMaterialExpr && IsModulateBlendMode(BlendMode) && MaterialShadingModels.IsLit() && !Material->IsDeferredDecal())
-	{
-		Errorf(TEXT("Dynamically lit translucency is not supported for BLEND_Modulate materials."));
-	}
-
-	if (Domain == MD_Surface)
-	{
-		if (IsModulateBlendMode(BlendMode) && Material->IsTranslucencyAfterDOFEnabled() && !(RHISupportsDualSourceBlending(Platform) || IsMobilePlatform(Platform)))
-		{
-			Errorf(TEXT("Translucency after DOF with BLEND_Modulate is only allowed on mobile platforms or desktop platforms that support dual-blending. Consider using BLEND_Translucent with black emissive"));
-		}
-	}
-
 	// Don't allow opaque and masked materials to scene depth as the results are undefined
 	if (bUsesSceneDepth && Domain != MD_PostProcess && !IsTranslucentBlendMode(BlendMode))
 	{
@@ -2059,140 +2044,6 @@ void FHLSLMaterialTranslator::TranslateMaterial()
 
 	MaterialCompilationOutput.bUsesDistanceCullFade = bUsesDistanceCullFade;
 
-	if (MaterialCompilationOutput.RequiresSceneColorCopy())
-	{
-		if (Domain != MD_Surface)
-		{
-			Errorf(TEXT("Only 'surface' material domain can use the scene color node."));
-		}
-		else if (!IsTranslucentBlendMode(BlendMode))
-		{
-			Errorf(TEXT("Only translucent materials can use the scene color node."));
-		}
-	}
-
-	if (IsAlphaHoldoutBlendMode(BlendMode) && !Material->IsPostProcessMaterial())
-	{
-		if (!MaterialShadingModels.IsUnlit())
-		{
-			Errorf(TEXT("Alpha Holdout blend mode must use unlit shading model."));
-		}
-
-		const int32 FnMainVersion = Material->GetMaterialInterface()->GetLinkerCustomVersion(FFortniteMainBranchObjectVersion::GUID);
-
-		// Note: As a retroactive fix for content failing to cook, we use the closest fn main version change to constrain the error only to recently created materials.
-		if (FnMainVersion > FFortniteMainBranchObjectVersion::LandscapeTargetLayersInLandscapeActor)
-		{
-			if (Material->IsTranslucencyAfterDOFEnabled() || Material->IsTranslucencyAfterMotionBlurEnabled())
-			{
-				// We must write into the alpha channel of the SceneColor now for later translucents rendered after DOF to be able to be affected by the alpha.
-				Errorf(TEXT("Alpha Holdout blend mode must use BeforeDOF under Translucency / Advanced / Translucency Pass."));
-			}
-		}
-	}
-
-	if (Domain == MD_Volume && BlendMode != BLEND_Additive)
-	{
-		Errorf(TEXT("Volume materials must use an Additive blend mode."));
-	}
-	if (Domain == MD_Volume && Material->IsUsedWithSkeletalMesh())
-	{
-		Errorf(TEXT("Volume materials are not compatible with skinned meshes: they are voxelised as boxes anyway. Please disable UsedWithSkeletalMesh on the material."));
-	}
-
-	if (Material->IsLightFunction() && !IsOpaqueBlendMode(BlendMode))
-	{
-		Errorf(TEXT("Light function materials must be opaque."));
-	}
-
-	if (Material->IsLightFunction() && MaterialShadingModels.IsLit())
-	{
-		Errorf(TEXT("Light function materials must use unlit."));
-	}
-
-	if (Domain == MD_PostProcess && MaterialShadingModels.IsLit())
-	{
-		Errorf(TEXT("Post process materials must use unlit."));
-	}
-
-	if (Material->AllowNegativeEmissiveColor() && MaterialShadingModels.IsLit())
-	{
-		Errorf(TEXT("Only unlit materials can output negative emissive color."));
-	}
-
-	if (Material->IsSky() && (!MaterialShadingModels.IsUnlit() || !(IsOpaqueOrMaskedBlendMode(BlendMode))))
-	{
-		Errorf(TEXT("Sky materials must be opaque or masked, and unlit. They are expected to completely replace the background."));
-	}
-
-	if (MaterialShadingModels.HasShadingModel(MSM_SingleLayerWater))
-	{
-		if (!IsOpaqueOrMaskedBlendMode(BlendMode))
-		{
-			Errorf(TEXT("SingleLayerWater materials must be opaque or masked."));
-		}
-		if (!MaterialShadingModels.HasOnlyShadingModel(MSM_SingleLayerWater))
-		{
-			Errorf(TEXT("SingleLayerWater materials cannot be combined with other shading models.")); // Simply untested for now
-		}
-
-		if (Material->GetMaterialInterface() && !Material->GetMaterialInterface()->GetMaterial()->HasAnyExpressionsInMaterialAndFunctionsOfType<UMaterialExpressionSingleLayerWaterMaterialOutput>()
-			&& !bSubstrateEnabled
-			)
-		{
-			Errorf(TEXT("SingleLayerWater materials requires the use of SingleLayerWaterMaterial output node."));
-		}
-	}
-
-	if (MaterialShadingModels.HasShadingModel(MSM_ThinTranslucent))
-	{
-		if (!IsTranslucentOnlyBlendMode(BlendMode))
-		{
-			Errorf(TEXT("ThinTranslucent materials must be translucent."));
-		}
-
-		const ETranslucencyLightingMode TranslucencyLightingMode = Material->GetTranslucencyLightingMode();
-
-		if (TranslucencyLightingMode != ETranslucencyLightingMode::TLM_SurfacePerPixelLighting)
-		{
-			Errorf(TEXT("ThinTranslucent materials must use Surface Per Pixel Lighting (Translucency->LightingMode=Surface ForwardShading).\n"));
-		}
-		if (!MaterialShadingModels.HasOnlyShadingModel(MSM_ThinTranslucent))
-		{
-			Errorf(TEXT("ThinTranslucent materials cannot be combined with other shading models.")); // Simply untested for now
-		}
-		if (Material->GetMaterialInterface() && !Material->GetMaterialInterface()->GetMaterial()->HasAnyExpressionsInMaterialAndFunctionsOfType<UMaterialExpressionThinTranslucentMaterialOutput>())
-		{
-			Errorf(TEXT("ThinTranslucent materials requires the use of ThinTranslucentMaterial output node."));
-		}
-		if (Material->IsTranslucencyAfterMotionBlurEnabled())
-		{
-			// We don't currently have a separate translucency modulation pass for After Motion Blur
-			Errorf(TEXT("ThinTranslucent materials are not currently supported in the \"After Motion Blur\" translucency pass."));
-		}
-	}
-
-	if (Material->GetMaterialInterface() && Material->GetMaterialInterface()->GetMaterial()->HasAnyExpressionsInMaterialAndFunctionsOfType<UMaterialExpressionAbsorptionMediumMaterialOutput>())
-	{
-		// User placed an AbsorptionMedium node, make sure default lit is being used
-		if (!MaterialShadingModels.HasShadingModel(MSM_DefaultLit))
-		{
-			Errorf(TEXT("AbsorptionMedium custom output requires Default Lit shading model."));
-		}
-		if (Material->GetRefractionMode() != RM_IndexOfRefraction)
-		{
-			Errorf(TEXT("AbsorptionMedium custom output requires \"Index of Refraction\" as the \"Refraction Mode\"."));
-		}
-	}
-
-	if (Domain == MD_DeferredDecal && !(IsTranslucentOnlyBlendMode(BlendMode) || BlendMode == BLEND_AlphaComposite || IsModulateBlendMode(BlendMode)))
-	{
-		// We could make the change for the user but it would be confusing when going to DeferredDecal and back
-		// or we would have to pay a performance cost to make the change more transparently.
-		// The change saves performance as with translucency we don't need to test for MeshDecals in all opaque rendering passes
-		Errorf(TEXT("Material using the DeferredDecal domain can only use the BlendModes Translucent, AlphaComposite or Modulate"));
-	}
-
 	if (MaterialCompilationOutput.bNeedsSceneTextures)
 	{
 		if (Domain != MD_DeferredDecal && Domain != MD_PostProcess)
@@ -2209,16 +2060,20 @@ void FHLSLMaterialTranslator::TranslateMaterial()
 		}
 	}
 
-	int32 NumPostProcessInputs = MaterialCompilationOutput.GetNumPostProcessInputsUsed();
-	if (NumPostProcessInputs > kPostProcessMaterialInputCountMax)
-	{
-		Errorf(TEXT("Maximum Scene Texture post process inputs exceeded (%d > %d), between SceneTexture nodes with PostProcessInputs or UserSceneTexture nodes."), NumPostProcessInputs, kPostProcessMaterialInputCountMax);
-	}
+	// Final validation logic shared between old and new translator.
+	TArray<FString> ValidationErrors;
+	UE::MaterialTranslatorUtils::FinalCompileValidation(
+		Material->GetMaterialInterface()->GetMaterial(),
+		MaterialCompilationOutput,
+		MaterialShadingModels,
+		BlendMode,
+		FrontMaterialExpr != nullptr,
+		Platform,
+		ValidationErrors);
 
-	if (IsModulateBlendMode(BlendMode) && Material->IsTranslucencyAfterMotionBlurEnabled())
+	for (const FString& ValidationError : ValidationErrors)
 	{
-		// We don't currently have a separate translucency modulation pass for After Motion Blur
-		Errorf(TEXT("Blend Mode \"Modulate\" materials are not currently supported in the \"After Motion Blur\" translucency pass."));
+		Errorf(*ValidationError);
 	}
 
 	CHECK_DDC_QUERY_FINISHED_ELSE_RETURN();
diff --git a/Engine/Source/Runtime/Engine/Private/Materials/MaterialExpressionsIR.cpp b/Engine/Source/Runtime/Engine/Private/Materials/MaterialExpressionsIR.cpp
--- a/Engine/Source/Runtime/Engine/Private/Materials/MaterialExpressionsIR.cpp
+++ b/Engine/Source/Runtime/Engine/Private/Materials/MaterialExpressionsIR.cpp
@@ -1231,6 +1231,11 @@ void UMaterialExpressionSphericalParticleOpacity::Build(MIR::FEmitter& Em)
 	BuildInlineHLSLOutput(Em, *this, { DensityValue });
 }
 
+void UMaterialExpressionShadingModel::Build(MIR::FEmitter& Em)
+{
+	Em.Output(0, Em.ShadingModel(ShadingModel));
+}
+
 void UMaterialExpressionTextureObject::Build(MIR::FEmitter& Em)
 {
 	Em.Output(0, Em.TextureObject(Texture, SamplerType));
diff --git a/Engine/Source/Runtime/Engine/Private/Materials/MaterialIR.cpp b/Engine/Source/Runtime/Engine/Private/Materials/MaterialIR.cpp
--- a/Engine/Source/Runtime/Engine/Private/Materials/MaterialIR.cpp
+++ b/Engine/Source/Runtime/Engine/Private/Materials/MaterialIR.cpp
@@ -45,6 +45,7 @@ const TCHAR* LexToString(EValueKind Kind)
 		case VK_ExternalInput: return TEXT("ExternalInput");
 		case VK_MaterialParameterCollection: return TEXT("MaterialParameterCollection");
 		case VK_ScreenTexture: return TEXT("ScreenTexture");
+		case VK_ShadingModel: return TEXT("ShadingModel");
 		case VK_TextureObject: return TEXT("TextureObject");
 		case VK_UniformParameter: return TEXT("UniformParameter");
 		case VK_Composite: return TEXT("Composite");
@@ -108,6 +109,7 @@ uint32 FValue::GetSizeInBytes() const
 		case VK_ExternalInput: return sizeof(FExternalInput);
 		case VK_MaterialParameterCollection: return sizeof(FMaterialParameterCollection);
 		case VK_ScreenTexture: return sizeof(FScreenTexture);
+		case VK_ShadingModel: return sizeof(FShadingModel);
 		case VK_TextureObject: return sizeof(FTextureObject);
 		case VK_UniformParameter: return sizeof(FUniformParameter);
 		case VK_Composite: return sizeof(FComposite) + sizeof(FValue*) * static_cast<const FComposite*>(this)->GetComponents().Num();
@@ -392,6 +394,50 @@ bool FValue::Equals(const FValue* Other) const
 	return FMemory::Memcmp(this, Other, SizeInBytes) == 0;
 }
 
+bool FValue::EqualsConstant(float TestValue) const
+{
+	const MIR::FConstant* ValueConstant = this->As<MIR::FConstant>();
+
+	if (ValueConstant)
+	{
+		switch (ValueConstant->Type.GetPrimitive().ScalarKind)
+		{
+			case MIR::EScalarKind::Bool:	return ValueConstant->Boolean == !!TestValue;
+			case MIR::EScalarKind::Int:		return ValueConstant->Integer == (TInteger)TestValue;
+			case MIR::EScalarKind::Float:	return ValueConstant->Float == TestValue;
+			case MIR::EScalarKind::Double:	return ValueConstant->Double == TestValue;
+			default: UE_MIR_UNREACHABLE();
+		}
+	}
+	return false;
+}
+
+bool FValue::EqualsConstant(FVector4f TestValue) const
+{
+	if (this->Kind == MIR::VK_Constant)
+	{
+		return EqualsConstant(TestValue.X);
+	}
+	else if (this->Kind == MIR::VK_Composite)
+	{
+		TConstArrayView<FValue*> Components = static_cast<const FComposite*>(this)->GetComponents();
+		int32 NumComponents = FMath::Min(Components.Num(), 4);
+
+		for (int32 ComponentIndex = 0; ComponentIndex < NumComponents; ++ComponentIndex)
+		{
+			if (!Components[ComponentIndex]->EqualsConstant(TestValue[ComponentIndex]))
+			{
+				return false;
+			}
+		}
+		return true;
+	}
+	else
+	{
+		return false;
+	}
+}
+
 UTexture* FValue::GetTexture() const
 {
 	if (const FTextureObject* TextureObject = As<FTextureObject>())
diff --git a/Engine/Source/Runtime/Engine/Private/Materials/MaterialIREmitter.cpp b/Engine/Source/Runtime/Engine/Private/Materials/MaterialIREmitter.cpp
--- a/Engine/Source/Runtime/Engine/Private/Materials/MaterialIREmitter.cpp
+++ b/Engine/Source/Runtime/Engine/Private/Materials/MaterialIREmitter.cpp
@@ -13,6 +13,7 @@
 #include "MaterialExpressionIO.h"
 #include "Engine/Texture.h"
 #include "Engine/Texture2D.h"
+#include "RenderUtils.h"
 
 #if WITH_EDITOR
 
@@ -439,7 +440,7 @@ FType FEmitter::GetMaterialAggregateAttributeType(const UMaterialAggregate* Aggr
 		case EMaterialAggregateAttributeType::Bool2: return FType::MakeVector(EScalarKind::Bool, 2);
 		case EMaterialAggregateAttributeType::Bool3: return FType::MakeVector(EScalarKind::Bool, 3);
 		case EMaterialAggregateAttributeType::Bool4: return FType::MakeVector(EScalarKind::Bool, 4);
-		case EMaterialAggregateAttributeType::ShadingModel:
+		case EMaterialAggregateAttributeType::ShadingModel: return FType::MakeShadingModel();
 		case EMaterialAggregateAttributeType::UInt1: return FType::MakeVector(EScalarKind::Int, 1);
 		case EMaterialAggregateAttributeType::UInt2: return FType::MakeVector(EScalarKind::Int, 2);
 		case EMaterialAggregateAttributeType::UInt3: return FType::MakeVector(EScalarKind::Int, 3);
@@ -908,6 +909,20 @@ FValueRef FEmitter::MaterialParameterCollection(UMaterialParameterCollection* Co
 	return EmitPrototype(*this, Prototype);
 }
 
+FValueRef FEmitter::ShadingModel(EMaterialShadingModel Id)
+{
+	// If the shading model is masked out, fallback to default shading model
+	uint32 PlatformShadingModelsMask = GetPlatformShadingModelsMask(GetShaderPlatform());
+	if ((PlatformShadingModelsMask & (1u << (uint32)Id)) == 0)
+	{
+		Id = MSM_DefaultLit;
+	}
+
+	FShadingModel Prototype = MakePrototype<FShadingModel>(FType::MakeShadingModel());
+	Prototype.Id = Id;
+	return EmitPrototype(*this, Prototype);
+}
+
 FValueRef FEmitter::TextureObject(UTexture* Texture, EMaterialSamplerType SamplerType)
 {
 	FTextureObject Proto = MakePrototype<FTextureObject>(FType::MakeTexture2D());
@@ -1296,7 +1311,7 @@ static FValueRef EmitAttributeDefaultValue(FEmitter& Emitter, const UMaterial* M
 		case EMaterialAggregateAttributeType::Float2: return Emitter.ConstantFloat2({ (TFloat)Attribute.DefaultValue.X, (TFloat)Attribute.DefaultValue.Y });
 		case EMaterialAggregateAttributeType::Float3: return Emitter.ConstantFloat3({ (TFloat)Attribute.DefaultValue.X, (TFloat)Attribute.DefaultValue.Y, (TFloat)Attribute.DefaultValue.Z });
 		case EMaterialAggregateAttributeType::Float4: return Emitter.ConstantFloat4(UE::Math::TVector4<TFloat>{ Attribute.DefaultValue.X, Attribute.DefaultValue.Y, Attribute.DefaultValue.Z, Attribute.DefaultValue.W });
-		case EMaterialAggregateAttributeType::ShadingModel: return Emitter.ConstantInt(Material->GetShadingModels().GetFirstShadingModel());
+		case EMaterialAggregateAttributeType::ShadingModel: return Emitter.ShadingModel(Material->GetShadingModels().GetFirstShadingModel());
 		case EMaterialAggregateAttributeType::MaterialAttributes: return Emitter.Aggregate(UMaterialAggregate::GetMaterialAttributes());
 		case EMaterialAggregateAttributeType::Aggregate: return Emitter.Aggregate(Attribute.Aggregate);
 		default: UE_MIR_UNREACHABLE();
diff --git a/Engine/Source/Runtime/Engine/Private/Materials/MaterialIRModule.cpp b/Engine/Source/Runtime/Engine/Private/Materials/MaterialIRModule.cpp
--- a/Engine/Source/Runtime/Engine/Private/Materials/MaterialIRModule.cpp
+++ b/Engine/Source/Runtime/Engine/Private/Materials/MaterialIRModule.cpp
@@ -3,6 +3,9 @@
 #include "Materials/MaterialIRModule.h"
 #include "Materials/MaterialIR.h"
 #include "Materials/MaterialIRTypes.h"
+#include "Materials/MaterialIRInternal.h"
+#include "Materials/Material.h"
+#include "Materials/MaterialAttributeDefinitionMap.h"
 #include "ParameterCollection.h"
 
 #if WITH_EDITOR
@@ -77,4 +80,22 @@ int32 FMaterialIRModule::FindOrAddParameterCollection(UMaterialParameterCollecti
 	return CollectionIndex;
 }
 
+void FMaterialIRModule::AddShadingModel(EMaterialShadingModel InShadingModel)
+{
+	ShadingModelsFromCompilation.AddShadingModel(InShadingModel);
+}
+
+// Similar to FHLSLMaterialTranslator::IsMaterialPropertyUsed.
+bool FMaterialIRModule::IsMaterialPropertyUsed(MIR::EStage Stage, EMaterialProperty InProperty) const
+{
+	for (const MIR::FSetMaterialOutput* Output : GetOutputs(Stage))
+	{
+		if (Output->Property == InProperty)
+		{
+			return !Output->Arg->EqualsConstant(FMaterialAttributeDefinitionMap::GetDefaultValue(InProperty));
+		}
+	}
+	return false;
+}
+
 #endif // #if WITH_EDITOR
diff --git a/Engine/Source/Runtime/Engine/Private/Materials/MaterialIRModuleBuilder.cpp b/Engine/Source/Runtime/Engine/Private/Materials/MaterialIRModuleBuilder.cpp
--- a/Engine/Source/Runtime/Engine/Private/Materials/MaterialIRModuleBuilder.cpp
+++ b/Engine/Source/Runtime/Engine/Private/Materials/MaterialIRModuleBuilder.cpp
@@ -13,10 +13,14 @@
 #include "Materials/MaterialIRInternal.h"
 #include "Materials/MaterialAggregate.h"
 #include "Materials/MaterialAttributeDefinitionMap.h"
+#include "Materials/MaterialExpressionAbsorptionMediumMaterialOutput.h"
 #include "Materials/MaterialExpressionFunctionInput.h"
 #include "Materials/MaterialExpressionFunctionOutput.h"
 #include "Materials/MaterialExpressionMaterialFunctionCall.h"
 #include "Materials/MaterialExpressionNamedReroute.h"
+#include "Materials/MaterialExpressionSingleLayerWaterMaterialOutput.h"
+#include "Materials/MaterialExpressionThinTranslucentMaterialOutput.h"
+#include "Materials/MaterialSharedPrivate.h"
 #include "Materials/Material.h"
 #include "MaterialDomain.h"
 #include "MaterialExpressionIO.h"
@@ -420,7 +424,8 @@ struct FMaterialIRModuleBuilderImpl
 					else
 					{
 						// If a value is flowing in through the connection, cast it to this material attribute type and assign it.
-						MIR::FType OutputArgType = MIR::FType::FromShaderType(Input.Type);
+						// Special case for shading model, because UE::Shader::EValueType doesn't include a native shading model type, and uses int instead.
+						MIR::FType OutputArgType = Output->Property == MP_ShadingModel ? MIR::FType::MakeShadingModel() : MIR::FType::FromShaderType(Input.Type);
 						Output->Arg = Emitter->Cast(Value, OutputArgType);
 					}
 				}
@@ -615,6 +620,7 @@ struct FMaterialIRModuleBuilderImpl
 		// Match various defines against the material configuration
 		if (Module->EnvironmentDefines.Contains(TEXT("MIR.SceneDepth")))
 		{
+			// @todo-jason.hoerner - Support for material instance blend mode overrides needed.
 			if (ValueAnalyzer.Material->MaterialDomain != MD_PostProcess && !IsTranslucentBlendMode(ValueAnalyzer.Material->BlendMode))
 			{
 				Module->AddError(nullptr, TEXT("Only transparent or postprocess materials can read from scene depth."));
@@ -723,6 +729,30 @@ struct FMaterialIRModuleBuilderImpl
 			CompilationOutput.bUsesWorldPositionOffset ||
 			CompilationOutput.bUsesDisplacement ||
 			bHasFirstPersonOutput;
+
+		// If the material doesn't use expression shading models, or they aren't valid, initialize the shading models to the ones from the material.
+		// Logic adapted from FHLSLMaterialTranslator::GetCompiledShadingModels, but done once, rather than for each call.
+		if (!ValueAnalyzer.Material->IsShadingModelFromMaterialExpression() || !Module->ShadingModelsFromCompilation.IsValid())
+		{
+			Module->ShadingModelsFromCompilation = ValueAnalyzer.Material->GetShadingModels();
+			UMaterialInterface::FilterOutPlatformShadingModels(Module->GetShaderPlatform(), Module->ShadingModelsFromCompilation);
+		}
+
+		// Final validation logic shared between old and new translator.
+		TArray<FString> ValidationErrors;
+		UE::MaterialTranslatorUtils::FinalCompileValidation(
+			ValueAnalyzer.Material,
+			CompilationOutput,
+			Module->GetCompiledShadingModels(),
+			ValueAnalyzer.Material->GetBlendMode(),					// @todo-jason.hoerner - we need to handle material instances overriding the blend mode
+			Module->IsMaterialPropertyUsed(MIR::EStage::Stage_Pixel, MP_FrontMaterial),
+			Module->GetShaderPlatform(),
+			ValidationErrors);
+
+		for (const FString& ValidationError : ValidationErrors)
+		{
+			Module->AddError(nullptr, *ValidationError);
+		}
 	}
 	
 	/* Other functions */
diff --git a/Engine/Source/Runtime/Engine/Private/Materials/MaterialIRToHLSLTranslator.cpp b/Engine/Source/Runtime/Engine/Private/Materials/MaterialIRToHLSLTranslator.cpp
--- a/Engine/Source/Runtime/Engine/Private/Materials/MaterialIRToHLSLTranslator.cpp
+++ b/Engine/Source/Runtime/Engine/Private/Materials/MaterialIRToHLSLTranslator.cpp
@@ -406,6 +406,7 @@ struct FPrivate : FMaterialIRToHLSLTranslation
 			case MIR::VK_ExternalInput: LowerExternalInput(static_cast<const MIR::FExternalInput*>(Value)); break;
 			case MIR::VK_MaterialParameterCollection: LowerMaterialParameterCollection(static_cast<const MIR::FMaterialParameterCollection*>(Value)); break;
 			case MIR::VK_ScreenTexture: LowerScreenTexture(static_cast<const MIR::FScreenTexture*>(Value)); break;
+			case MIR::VK_ShadingModel: LowerShadingModel(static_cast<const MIR::FShadingModel*>(Value)); break;
 			case MIR::VK_TextureObject: LowerTextureObject(static_cast<const MIR::FTextureObject*>(Value)); break;
 			case MIR::VK_UniformParameter: LowerUniformParameter(static_cast<const MIR::FUniformParameter*>(Value)); break;
 			default: UE_MIR_UNREACHABLE();
@@ -697,6 +698,11 @@ struct FPrivate : FMaterialIRToHLSLTranslation
 		}
 	}
 
+	void LowerShadingModel(const MIR::FShadingModel* ShadingModel)
+	{
+		Printer << (int32)ShadingModel->Id;
+	}
+
 	void LowerTextureObject(const MIR::FTextureObject* TextureObject)
 	{
 		LowerTextureReference(TextureObject->Texture->GetMaterialType(), TextureObject->Analysis_UniformParameterIndex);
@@ -2161,9 +2167,16 @@ struct FPrivate : FMaterialIRToHLSLTranslation
 			OutEnvironment.SetDefine(EnvironmentDefine, true);
 		}
 
-		FMaterialShadingModelField ShadingModels = Material->GetShadingModels();
+		FMaterialShadingModelField ShadingModels = Module->GetCompiledShadingModels();
 		ensure(ShadingModels.IsValid());
 
+		// Logic from FHLSLMaterialTranslator::TranslateMaterial
+		bool bOpacityPropertyIsUsed = Module->IsMaterialPropertyUsed(MIR::EStage::Stage_Pixel, MP_Opacity);
+
+		bool bUsesCurvature = Module->GetFeatureLevel() == ERHIFeatureLevel::ES3_1 &&
+			((ShadingModels.HasShadingModel(MSM_SubsurfaceProfile) && Module->IsMaterialPropertyUsed(MIR::EStage::Stage_Pixel, MP_CustomData0))
+			|| (ShadingModels.HasShadingModel(MSM_Eye) && bOpacityPropertyIsUsed));
+
 		int32 NumActiveShadingModels = 0;
 		if (ShadingModels.IsLit())
 		{
@@ -2192,6 +2205,28 @@ struct FPrivate : FMaterialIRToHLSLTranslation
 				OutEnvironment.SetDefine(GetShadingModelParameterName(Model), true);
 				NumActiveShadingModels += 1;
 			}
+
+			if (ShadingModels.HasShadingModel(MSM_SubsurfaceProfile) && bUsesCurvature)
+			{
+				OutEnvironment.SetDefine(TEXT("MATERIAL_SUBSURFACE_PROFILE_USE_CURVATURE"), true);
+			}
+
+			if (ShadingModels.HasShadingModel(MSM_Eye) && bUsesCurvature)
+			{
+				OutEnvironment.SetDefine(TEXT("MATERIAL_SHADINGMODEL_EYE_USE_CURVATURE"), true);
+			}
+
+			if (ShadingModels.HasShadingModel(MSM_SingleLayerWater) && FDataDrivenShaderPlatformInfo::GetRequiresDisableForwardLocalLights(ShaderPlatform))
+			{
+				OutEnvironment.SetDefine(TEXT("DISABLE_FORWARD_LOCAL_LIGHTS"), true);
+			}
+
+			const bool bIsWaterDistanceFieldShadowEnabled = IsWaterDistanceFieldShadowEnabled(ShaderPlatform);
+			const bool bIsWaterVSMFilteringEnabled = IsWaterVirtualShadowMapFilteringEnabled(ShaderPlatform);
+			if (ShadingModels.HasShadingModel(MSM_SingleLayerWater) && (bIsWaterDistanceFieldShadowEnabled || bIsWaterVSMFilteringEnabled))
+			{
+				OutEnvironment.SetDefine(TEXT("SINGLE_LAYER_WATER_SEPARATED_MAIN_LIGHT"), TEXT("1"));
+			}
 		}
 		else
 		{
diff --git a/Engine/Source/Runtime/Engine/Private/Materials/MaterialIRTypes.cpp b/Engine/Source/Runtime/Engine/Private/Materials/MaterialIRTypes.cpp
--- a/Engine/Source/Runtime/Engine/Private/Materials/MaterialIRTypes.cpp
+++ b/Engine/Source/Runtime/Engine/Private/Materials/MaterialIRTypes.cpp
@@ -16,6 +16,7 @@ const TCHAR* TypeKindToString(ETypeKind Kind)
 		case ETypeKind::Void: return TEXT("Void");
 		case ETypeKind::Primitive: return TEXT("Primitive");
 		case ETypeKind::Aggregate: return TEXT("Aggregate");
+		case ETypeKind::ShadingModel: return TEXT("ShadingModel");
 		case ETypeKind::Texture2D: return TEXT("Texture2D");
 		case ETypeKind::ParameterCollection: return TEXT("ParameterCollection");
 		case ETypeKind::SubstrateData: return TEXT("SubstrateData");
@@ -175,6 +176,13 @@ FType FType::MakeParameterCollection()
 	return Type;
 }
 
+FType FType::MakeShadingModel()
+{
+	FType Type{};
+	Type.Kind = ETypeKind::ShadingModel;
+	return Type;
+}
+
 FType FType::MakeTexture2D()
 {
 	FType Type{};
diff --git a/Engine/Source/Runtime/Engine/Private/Materials/MaterialIRValueAnalyzer.cpp b/Engine/Source/Runtime/Engine/Private/Materials/MaterialIRValueAnalyzer.cpp
--- a/Engine/Source/Runtime/Engine/Private/Materials/MaterialIRValueAnalyzer.cpp
+++ b/Engine/Source/Runtime/Engine/Private/Materials/MaterialIRValueAnalyzer.cpp
@@ -463,6 +463,14 @@ static void AnalyzeInStageScreenTexture(FMaterialIRValueAnalyzer& Analyzer, MIR:
 	}
 }
 
+static void AnalyzeShadingModel(FMaterialIRValueAnalyzer& Analyzer, MIR::FShadingModel* ShadingModel)
+{
+	if (ShadingModel->Id < MSM_NUM)
+	{
+		Analyzer.Module->AddShadingModel(ShadingModel->Id);
+	}
+}
+
 static void AnalyzeTextureObject(FMaterialIRValueAnalyzer& Analyzer, MIR::FTextureObject* TextureObject)
 {
 	EMaterialTextureParameterType ParamType = MIR::Internal::TextureMaterialValueTypeToParameterType(TextureObject->Texture->GetMaterialType());
@@ -680,6 +688,7 @@ void FMaterialIRValueAnalyzer::Analyze(MIR::FValue* Value)
 		EXPAND_CASE(SetMaterialOutput);
 		EXPAND_CASE(InlineHLSL);
 		EXPAND_CASE(ScreenTexture);
+		EXPAND_CASE(ShadingModel);
 		default: break;
 	}
 
diff --git a/Engine/Source/Runtime/Engine/Private/Materials/MaterialShared.cpp b/Engine/Source/Runtime/Engine/Private/Materials/MaterialShared.cpp
--- a/Engine/Source/Runtime/Engine/Private/Materials/MaterialShared.cpp
+++ b/Engine/Source/Runtime/Engine/Private/Materials/MaterialShared.cpp
@@ -5950,6 +5950,196 @@ FString SceneTextureIdToHLSLString(ESceneTextureId TexId)
 	}
 }
 
+// FMaterialResource::IsTranslucencyAfterDOFEnabled
+static bool IsTranslucencyAfterDOFEnabled(UMaterial* Material)
+{
+	return Material->TranslucencyPass == MTP_AfterDOF
+		&& !Material->IsUIMaterial()
+		&& !Material->IsDeferredDecal();
+}
+
+// FMaterialResource::IsTranslucencyAfterMotionBlurEnabled
+static bool IsTranslucencyAfterMotionBlurEnabled(UMaterial* Material)
+{
+	return Material->TranslucencyPass == MTP_AfterMotionBlur
+		&& !Material->IsUIMaterial()
+		&& !Material->IsDeferredDecal();
+}
+
+void FinalCompileValidation(
+	UMaterial* Material,
+	const FMaterialCompilationOutput& MaterialCompilationOutput,
+	FMaterialShadingModelField MaterialShadingModels,
+	EBlendMode BlendMode,
+	bool bHasFrontMaterialExpr,
+	EShaderPlatform Platform,
+	TArray<FString>& OutErrors)
+{
+	const bool bSubstrateEnabled = Substrate::IsSubstrateEnabled();
+	const EMaterialDomain Domain = Material->MaterialDomain;
+
+	// If Substrate is enabled or if this is a Substrate material cooked/used in non-Substrate mode, 
+	// we disable this warning as Substrate supports 'colored transmittance only' mode (i.e, modulate).
+	if (!bSubstrateEnabled && !bHasFrontMaterialExpr && IsModulateBlendMode(BlendMode) && MaterialShadingModels.IsLit() && !Material->IsDeferredDecal())
+	{
+		OutErrors.Add(TEXT("Dynamically lit translucency is not supported for BLEND_Modulate materials."));
+	}
+
+	if (Domain == MD_Surface)
+	{
+		if (IsModulateBlendMode(BlendMode) && IsTranslucencyAfterDOFEnabled(Material) && !(RHISupportsDualSourceBlending(Platform) || IsMobilePlatform(Platform)))
+		{
+			OutErrors.Add(TEXT("Translucency after DOF with BLEND_Modulate is only allowed on mobile platforms or desktop platforms that support dual-blending. Consider using BLEND_Translucent with black emissive"));
+		}
+	}
+
+	if (MaterialCompilationOutput.RequiresSceneColorCopy())
+	{
+		if (Domain != MD_Surface)
+		{
+			OutErrors.Add(TEXT("Only 'surface' material domain can use the scene color node."));
+		}
+		else if (!IsTranslucentBlendMode(BlendMode))
+		{
+			OutErrors.Add(TEXT("Only translucent materials can use the scene color node."));
+		}
+	}
+
+	if (IsAlphaHoldoutBlendMode(BlendMode) && !Material->IsPostProcessMaterial())
+	{
+		if (!MaterialShadingModels.IsUnlit())
+		{
+			OutErrors.Add(TEXT("Alpha Holdout blend mode must use unlit shading model."));
+		}
+
+		const int32 FnMainVersion = Material->GetLinkerCustomVersion(FFortniteMainBranchObjectVersion::GUID);
+
+		// Note: As a retroactive fix for content failing to cook, we use the closest fn main version change to constrain the error only to recently created materials.
+		if (FnMainVersion > FFortniteMainBranchObjectVersion::LandscapeTargetLayersInLandscapeActor)
+		{
+			if (IsTranslucencyAfterDOFEnabled(Material) || IsTranslucencyAfterMotionBlurEnabled(Material))
+			{
+				// We must write into the alpha channel of the SceneColor now for later translucents rendered after DOF to be able to be affected by the alpha.
+				OutErrors.Add(TEXT("Alpha Holdout blend mode must use BeforeDOF under Translucency / Advanced / Translucency Pass."));
+			}
+		}
+	}
+
+	if (Domain == MD_Volume && BlendMode != BLEND_Additive)
+	{
+		OutErrors.Add(TEXT("Volume materials must use an Additive blend mode."));
+	}
+	if (Domain == MD_Volume && Material->bUsedWithSkeletalMesh)
+	{
+		OutErrors.Add(TEXT("Volume materials are not compatible with skinned meshes: they are voxelised as boxes anyway. Please disable UsedWithSkeletalMesh on the material."));
+	}
+
+	if (Domain == MD_LightFunction && !IsOpaqueBlendMode(BlendMode))
+	{
+		OutErrors.Add(TEXT("Light function materials must be opaque."));
+	}
+
+	if (Domain == MD_LightFunction && MaterialShadingModels.IsLit())
+	{
+		OutErrors.Add(TEXT("Light function materials must use unlit."));
+	}
+
+	if (Domain == MD_PostProcess && MaterialShadingModels.IsLit())
+	{
+		OutErrors.Add(TEXT("Post process materials must use unlit."));
+	}
+
+	if (Material->bAllowNegativeEmissiveColor && MaterialShadingModels.IsLit())
+	{
+		OutErrors.Add(TEXT("Only unlit materials can output negative emissive color."));
+	}
+
+	if (Material->bIsSky && (!MaterialShadingModels.IsUnlit() || !(IsOpaqueOrMaskedBlendMode(BlendMode))))
+	{
+		OutErrors.Add(TEXT("Sky materials must be opaque or masked, and unlit. They are expected to completely replace the background."));
+	}
+
+	if (MaterialShadingModels.HasShadingModel(MSM_SingleLayerWater))
+	{
+		if (!IsOpaqueOrMaskedBlendMode(BlendMode))
+		{
+			OutErrors.Add(TEXT("SingleLayerWater materials must be opaque or masked."));
+		}
+		if (!MaterialShadingModels.HasOnlyShadingModel(MSM_SingleLayerWater))
+		{
+			OutErrors.Add(TEXT("SingleLayerWater materials cannot be combined with other shading models.")); // Simply untested for now
+		}
+
+		if (!Material->HasAnyExpressionsInMaterialAndFunctionsOfType<UMaterialExpressionSingleLayerWaterMaterialOutput>()
+			&& !bSubstrateEnabled
+			)
+		{
+			OutErrors.Add(TEXT("SingleLayerWater materials requires the use of SingleLayerWaterMaterial output node."));
+		}
+	}
+
+	if (MaterialShadingModels.HasShadingModel(MSM_ThinTranslucent))
+	{
+		if (!IsTranslucentOnlyBlendMode(BlendMode))
+		{
+			OutErrors.Add(TEXT("ThinTranslucent materials must be translucent."));
+		}
+
+		const ETranslucencyLightingMode TranslucencyLightingMode = Material->TranslucencyLightingMode;
+
+		if (TranslucencyLightingMode != ETranslucencyLightingMode::TLM_SurfacePerPixelLighting)
+		{
+			OutErrors.Add(TEXT("ThinTranslucent materials must use Surface Per Pixel Lighting (Translucency->LightingMode=Surface ForwardShading).\n"));
+		}
+		if (!MaterialShadingModels.HasOnlyShadingModel(MSM_ThinTranslucent))
+		{
+			OutErrors.Add(TEXT("ThinTranslucent materials cannot be combined with other shading models.")); // Simply untested for now
+		}
+		if (!Material->HasAnyExpressionsInMaterialAndFunctionsOfType<UMaterialExpressionThinTranslucentMaterialOutput>())
+		{
+			OutErrors.Add(TEXT("ThinTranslucent materials requires the use of ThinTranslucentMaterial output node."));
+		}
+		if (IsTranslucencyAfterMotionBlurEnabled(Material))
+		{
+			// We don't currently have a separate translucency modulation pass for After Motion Blur
+			OutErrors.Add(TEXT("ThinTranslucent materials are not currently supported in the \"After Motion Blur\" translucency pass."));
+		}
+	}
+
+	if (Material->HasAnyExpressionsInMaterialAndFunctionsOfType<UMaterialExpressionAbsorptionMediumMaterialOutput>())
+	{
+		// User placed an AbsorptionMedium node, make sure default lit is being used
+		if (!MaterialShadingModels.HasShadingModel(MSM_DefaultLit))
+		{
+			OutErrors.Add(TEXT("AbsorptionMedium custom output requires Default Lit shading model."));
+		}
+		if (Material->RefractionMethod != RM_IndexOfRefraction)
+		{
+			OutErrors.Add(TEXT("AbsorptionMedium custom output requires \"Index of Refraction\" as the \"Refraction Mode\"."));
+		}
+	}
+
+	if (Domain == MD_DeferredDecal && !(IsTranslucentOnlyBlendMode(BlendMode) || BlendMode == BLEND_AlphaComposite || IsModulateBlendMode(BlendMode)))
+	{
+		// We could make the change for the user but it would be confusing when going to DeferredDecal and back
+		// or we would have to pay a performance cost to make the change more transparently.
+		// The change saves performance as with translucency we don't need to test for MeshDecals in all opaque rendering passes
+		OutErrors.Add(TEXT("Material using the DeferredDecal domain can only use the BlendModes Translucent, AlphaComposite or Modulate"));
+	}
+
+	int32 NumPostProcessInputs = MaterialCompilationOutput.GetNumPostProcessInputsUsed();
+	if (NumPostProcessInputs > kPostProcessMaterialInputCountMax)
+	{
+		OutErrors.Add(FString::Printf(TEXT("Maximum Scene Texture post process inputs exceeded (%d > %d), between SceneTexture nodes with PostProcessInputs or UserSceneTexture nodes."), NumPostProcessInputs, kPostProcessMaterialInputCountMax));
+	}
+
+	if (IsModulateBlendMode(BlendMode) && IsTranslucencyAfterMotionBlurEnabled(Material))
+	{
+		// We don't currently have a separate translucency modulation pass for After Motion Blur
+		OutErrors.Add(TEXT("Blend Mode \"Modulate\" materials are not currently supported in the \"After Motion Blur\" translucency pass."));
+	}
+}
+
 }  // namespace UE::MaterialTranslatorUtils
 
 #endif // WITH_EDITOR
diff --git a/Engine/Source/Runtime/Engine/Private/Materials/MaterialSharedPrivate.h b/Engine/Source/Runtime/Engine/Private/Materials/MaterialSharedPrivate.h
--- a/Engine/Source/Runtime/Engine/Private/Materials/MaterialSharedPrivate.h
+++ b/Engine/Source/Runtime/Engine/Private/Materials/MaterialSharedPrivate.h
@@ -6,6 +6,7 @@
 #include "Containers/Array.h"
 #include "Containers/ArrayView.h"
 #include "Containers/UnrealString.h"
+#include "Engine/EngineTypes.h"
 #include "Engine/TextureDefines.h"
 #include "MaterialShared.h"
 #include "Materials/MaterialExpressionTransform.h"
@@ -86,6 +87,16 @@ namespace UE::MaterialTranslatorUtils
 	// Converts a scene texture identifier to its corresponding HLSL string representation.
 	FString SceneTextureIdToHLSLString(ESceneTextureId TexId);
 
+	// Final compile validation logic, run at the end of translation to HLSL.
+	void FinalCompileValidation(
+		UMaterial* Material,
+		const FMaterialCompilationOutput& MaterialCompilationOutput,
+		FMaterialShadingModelField MaterialShadingModels,
+		EBlendMode BlendMode,
+		bool bHasFrontMaterialExpr,
+		EShaderPlatform Platform,
+		TArray<FString>& OutErrors);
+
 } // namespace UE::MaterialTranslatorUtils
 
 ///////////////////////////////////////////////////////
diff --git a/Engine/Source/Runtime/Engine/Public/Materials/MaterialExpressionShadingModel.h b/Engine/Source/Runtime/Engine/Public/Materials/MaterialExpressionShadingModel.h
--- a/Engine/Source/Runtime/Engine/Public/Materials/MaterialExpressionShadingModel.h
+++ b/Engine/Source/Runtime/Engine/Public/Materials/MaterialExpressionShadingModel.h
@@ -27,6 +27,7 @@ class UMaterialExpressionShadingModel : public UMaterialExpression
 
 	//~ Begin UMaterialExpression Interface
 #if WITH_EDITOR
+	virtual void Build(MIR::FEmitter& Emitter) override;
 	virtual int32 Compile(class FMaterialCompiler* Compiler, int32 OutputIndex) override;
 	virtual int32 CompilePreview(class FMaterialCompiler* Compiler, int32 OutputIndex) override;
 	virtual EMaterialValueType GetOutputValueType(int32 OutputIndex) override;
diff --git a/Engine/Source/Runtime/Engine/Public/Materials/MaterialIR.h b/Engine/Source/Runtime/Engine/Public/Materials/MaterialIR.h
--- a/Engine/Source/Runtime/Engine/Public/Materials/MaterialIR.h
+++ b/Engine/Source/Runtime/Engine/Public/Materials/MaterialIR.h
@@ -105,6 +105,7 @@ enum EValueKind
 	VK_ExternalInput,
 	VK_MaterialParameterCollection,
 	VK_ScreenTexture,
+	VK_ShadingModel,
 	VK_TextureObject,
 	VK_UniformParameter,
 
@@ -236,6 +237,12 @@ struct FValue
 	// Returns whether this value exactly equals Other.
 	bool Equals(const FValue* Other) const;
 
+	// Returns whether this value equals the given constant scalar.
+	bool EqualsConstant(float Value) const;
+
+	// Returns whether this value equals the given constant vector.
+	bool EqualsConstant(FVector4f Value) const;
+
 	// Returns the UTexture object if this value is of type FTextureObject. Returns null otherwise.
 	UTexture* GetTexture() const;
 
@@ -523,6 +530,11 @@ struct FScreenTexture : TValue<VK_ScreenTexture>
 	FName UserSceneTexture;
 };
 
+struct FShadingModel : TValue<VK_ShadingModel>
+{
+	EMaterialShadingModel Id;
+};
+
 /*------------------------------------ Instructions ------------------------------------*/
 
 // A block of instructions.
diff --git a/Engine/Source/Runtime/Engine/Public/Materials/MaterialIREmitter.h b/Engine/Source/Runtime/Engine/Public/Materials/MaterialIREmitter.h
--- a/Engine/Source/Runtime/Engine/Public/Materials/MaterialIREmitter.h
+++ b/Engine/Source/Runtime/Engine/Public/Materials/MaterialIREmitter.h
@@ -388,6 +388,9 @@ public:
 	// Emits a material parameter collection.
 	FValueRef MaterialParameterCollection(UMaterialParameterCollection* Collection);
 
+	// Emits a shading model.
+	FValueRef ShadingModel(EMaterialShadingModel Id);
+
 	// Emits a texture object value of fiven sampler type and returns it.
 	FValueRef TextureObject(UTexture* Texture, EMaterialSamplerType SamplerType);
 
diff --git a/Engine/Source/Runtime/Engine/Public/Materials/MaterialIRModule.h b/Engine/Source/Runtime/Engine/Public/Materials/MaterialIRModule.h
--- a/Engine/Source/Runtime/Engine/Public/Materials/MaterialIRModule.h
+++ b/Engine/Source/Runtime/Engine/Public/Materials/MaterialIRModule.h
@@ -92,6 +92,15 @@ public:
 	// Returns the array of HLSL-defined user-functions in the module.
 	TConstArrayView<const MIR::FFunctionHLSL*> GetFunctionHLSLs() const { return FunctionHLSLs; };
 
+	// Flags a shading model as used by this material.
+	void AddShadingModel(EMaterialShadingModel InShadingModel);
+
+	// Returns the shading models used by this material.
+	FMaterialShadingModelField GetCompiledShadingModels() const { return this->ShadingModelsFromCompilation; }
+
+	// Returns true if the given property has a non default value.  Similar to FHLSLMaterialTranslator::IsMaterialPropertyUsed.
+	bool IsMaterialPropertyUsed(MIR::EStage Stage, EMaterialProperty InProperty) const;
+
 	// Stores a user-defined string and returns a view of it.
 	FStringView InternString(FStringView InString);
 
@@ -150,6 +159,9 @@ private:
 	// Array of HLSL-defined user-functions in the module.
 	TArray<const MIR::FFunctionHLSL*> FunctionHLSLs;
 
+	/** Shading models used by this material. */
+	FMaterialShadingModelField ShadingModelsFromCompilation;
+
 	// List of compilation errors.
 	TArray<FError> Errors;
 
diff --git a/Engine/Source/Runtime/Engine/Public/Materials/MaterialIRTypes.h b/Engine/Source/Runtime/Engine/Public/Materials/MaterialIRTypes.h
--- a/Engine/Source/Runtime/Engine/Public/Materials/MaterialIRTypes.h
+++ b/Engine/Source/Runtime/Engine/Public/Materials/MaterialIRTypes.h
@@ -21,6 +21,7 @@ enum class ETypeKind : uint8
 
 	// Opaque Objects
 	
+	ShadingModel,			// Dynamically specified material shading model.
 	Texture2D,				// A 2D texture object.
 	ParameterCollection,	// Material parameter collection object.
 	SubstrateData,			// Substrate-specific data structure.
@@ -182,6 +183,9 @@ struct FType
 	// Returns the Material Parameter Collection object type.
 	static FType MakeParameterCollection();
 
+	// Returns the Material Shading Model object type.
+	static FType MakeShadingModel();
+
 	// Returns the Texture2D object type.
 	static FType MakeTexture2D();
 	
