Replaced a FMessageLog with a UE_LOG as the calling function can now be called outside of the Gamethread.

diff --git a/Engine/Source/Runtime/Engine/Private/Components/PrimitiveComponent.cpp b/Engine/Source/Runtime/Engine/Private/Components/PrimitiveComponent.cpp
--- a/Engine/Source/Runtime/Engine/Private/Components/PrimitiveComponent.cpp
+++ b/Engine/Source/Runtime/Engine/Private/Components/PrimitiveComponent.cpp
@@ -2015,8 +2015,7 @@ bool UPrimitiveComponent::ShouldCreatePhysicsState() const
 				const ECollisionEnabled::Type CollisionEnabled = GetCollisionEnabled();
 				if (CollisionEnabled == ECollisionEnabled::NoCollision || CollisionEnabled == ECollisionEnabled::QueryOnly)
 				{
-					FMessageLog("PIE").Warning(FText::Format(LOCTEXT("InvalidSimulateOptions", "Invalid Simulate Options: Body ({0}) is set to simulate physics but Collision Enabled is incompatible"),
-						FText::FromString(GetReadableName())));
+					UE_LOG(LogPhysics, Warning, TEXT("Invalid Simulate Options: Body (%s) is set to simulate physics but Collision Enabled is incompatible"), *GetReadableName());
 				}
 			}
 		}
