Add asset user data support to physics asset.

diff --git a/Engine/Source/Runtime/Engine/Classes/PhysicsEngine/PhysicsAsset.h b/Engine/Source/Runtime/Engine/Classes/PhysicsEngine/PhysicsAsset.h
--- a/Engine/Source/Runtime/Engine/Classes/PhysicsEngine/PhysicsAsset.h
+++ b/Engine/Source/Runtime/Engine/Classes/PhysicsEngine/PhysicsAsset.h
@@ -6,6 +6,7 @@
 #include "UObject/Object.h"
 #include "PhysicsEngine/RigidBodyIndexPair.h"
 #include "Interfaces/Interface_PreviewMeshProvider.h"
+#include "Interfaces/Interface_AssetUserData.h"
 #include "PhysicsAsset.generated.h"
 
 namespace ECollisionEnabled { enum Type : int; }
@@ -167,7 +168,7 @@ enum class EPhysicsAssetSolverType: uint8
  */
 
 UCLASS(hidecategories=Object, BlueprintType, MinimalAPI, Config=Game, PerObjectConfig, AutoCollapseCategories=(OldSolverSettings))
-class UPhysicsAsset : public UObject, public IInterface_PreviewMeshProvider
+class UPhysicsAsset : public UObject, public IInterface_PreviewMeshProvider, public IInterface_AssetUserData
 {
 	GENERATED_UCLASS_BODY()
 
@@ -420,6 +421,17 @@ private:
 
 	UPROPERTY(instanced)
 	TArray<TObjectPtr<class UBodySetup>> BodySetup_DEPRECATED;
+
+	// Asset User Data
+public:
+	ENGINE_API void AddAssetUserData(UAssetUserData* InUserData) override;
+	ENGINE_API void RemoveUserDataOfClass(TSubclassOf<UAssetUserData> InUserDataClass) override;
+	ENGINE_API UAssetUserData* GetAssetUserDataOfClass(TSubclassOf<UAssetUserData> InUserDataClass) override;
+	ENGINE_API const TArray<UAssetUserData*>* GetAssetUserDataArray() const override;
+
+private:
+	UPROPERTY(EditAnywhere, AdvancedDisplay, Instanced, Category = UserData)
+	TArray<TObjectPtr<UAssetUserData>> AssetUserData;
 };
 
 
diff --git a/Engine/Source/Runtime/Engine/Private/PhysicsEngine/PhysicsAsset.cpp b/Engine/Source/Runtime/Engine/Private/PhysicsEngine/PhysicsAsset.cpp
--- a/Engine/Source/Runtime/Engine/Private/PhysicsEngine/PhysicsAsset.cpp
+++ b/Engine/Source/Runtime/Engine/Private/PhysicsEngine/PhysicsAsset.cpp
@@ -740,6 +740,14 @@ void UPhysicsAsset::PostEditChangeProperty(FPropertyChangedEvent& PropertyChange
 		}
 	}
 
+	for(UAssetUserData* UserData : AssetUserData)
+	{
+		if(UserData)
+		{
+			UserData->PostEditChangeOwner(PropertyChangedEvent);
+		}
+	}
+
 	Super::PostEditChangeProperty(PropertyChangedEvent);
 
 	RefreshPhysicsAssetChange();
@@ -940,5 +948,47 @@ void UPhysicsAsset::GetResourceSizeEx(FResourceSizeEx& CumulativeResourceSize)
 	CumulativeResourceSize.AddDedicatedSystemMemoryBytes(CollisionDisableTable.GetAllocatedSize());
 }
 
+void UPhysicsAsset::AddAssetUserData(UAssetUserData* InUserData)
+{
+	if(!InUserData)
+	{
+		return;
+	}
+
+	UAssetUserData* ExistingData = GetAssetUserDataOfClass(InUserData->GetClass());
+	if(ExistingData)
+	{
+		AssetUserData.Remove(ExistingData);
+	}
+	AssetUserData.Add(InUserData);
+}
+
+UAssetUserData* UPhysicsAsset::GetAssetUserDataOfClass(TSubclassOf<UAssetUserData> InUserDataClass)
+{
+	TObjectPtr<UAssetUserData>* Found = AssetUserData.FindByPredicate([&InUserDataClass](TObjectPtr<UAssetUserData> Item)
+		{
+			return Item && Item->IsA(InUserDataClass);
+		});
+
+	if(Found)
+	{
+		return *Found;
+	}
+	return nullptr;
+}
+
+void UPhysicsAsset::RemoveUserDataOfClass(TSubclassOf<UAssetUserData> InUserDataClass)
+{
+	AssetUserData.RemoveAll([&InUserDataClass](TObjectPtr<UAssetUserData> Item)
+	{
+		return Item && Item->IsA(InUserDataClass);
+	});
+}
+
+const TArray<UAssetUserData*>* UPhysicsAsset::GetAssetUserDataArray() const
+{
+	return &ToRawPtrTArrayUnsafe(AssetUserData);
+}
+
 #undef LOCTEXT_NAMESPACE
 
