Mark the layer stack material expression as experimental for 5.7.

diff --git a/Engine/Source/Runtime/Engine/Public/Materials/MaterialExpressionLayerStack.h b/Engine/Source/Runtime/Engine/Public/Materials/MaterialExpressionLayerStack.h
--- a/Engine/Source/Runtime/Engine/Public/Materials/MaterialExpressionLayerStack.h
+++ b/Engine/Source/Runtime/Engine/Public/Materials/MaterialExpressionLayerStack.h
@@ -15,7 +15,7 @@
 class FMaterialCompiler;
 class UMaterialFunctionInterface;
 
-USTRUCT()
+USTRUCT(Experimental)
 struct FMaterialLayerInput : public FExpressionInput
 {
 	GENERATED_BODY()
@@ -27,7 +27,7 @@ struct FMaterialLayerInput : public FExpressionInput
 	FString GetInputName() const;
 };
 
-UCLASS(hidecategories = Object, MinimalAPI)
+UCLASS(hidecategories = Object, MinimalAPI, Experimental)
 class UMaterialExpressionLayerStack : public UMaterialExpressionMaterialAttributeLayers
 {
 	GENERATED_UCLASS_BODY()
