Fix ensure in UInstancedStaticMeshComponent::GetNavigationBounds() caused by invalid navigation bounds. CachedBounds where being invalidated and not being recomputed when instances where added individually.

diff --git a/Engine/Source/Runtime/Engine/Classes/Components/InstancedStaticMeshComponent.h b/Engine/Source/Runtime/Engine/Classes/Components/InstancedStaticMeshComponent.h
--- a/Engine/Source/Runtime/Engine/Classes/Components/InstancedStaticMeshComponent.h
+++ b/Engine/Source/Runtime/Engine/Classes/Components/InstancedStaticMeshComponent.h
@@ -796,6 +796,10 @@ private:
 
 	UPROPERTY()
 	mutable FBoundsCacheElement CachedBounds[(int)EBoundsCacheType::Count];
+
+	/** Temporary flag used to skip Navigation CachedBounds invalidation.
+	 * This member was added to avoid deprecating the virtual PartialNavigationUpdates() (which can't be deprecated nicely). */ 
+	bool bNavigationCachedBoundsUpdated = false;
 };
 
 /** InstancedStaticMeshInstance hit proxy */
diff --git a/Engine/Source/Runtime/Engine/Private/InstancedStaticMesh.cpp b/Engine/Source/Runtime/Engine/Private/InstancedStaticMesh.cpp
--- a/Engine/Source/Runtime/Engine/Private/InstancedStaticMesh.cpp
+++ b/Engine/Source/Runtime/Engine/Private/InstancedStaticMesh.cpp
@@ -3475,6 +3475,9 @@ int32 UInstancedStaticMeshComponent::AddInstanceInternal(int32 InstanceIndex, FI
 	// Update navigation relevancy
 	bNavigationRelevant = IsNavigationRelevant();
 
+	// Indicates if the cached bounds were updated, it's expected to occur in PartialNavigationUpdate().
+	bNavigationCachedBoundsUpdated = false;
+	
 	// Perform navigation update if the component is relevant to navigation and registered.
 	// Otherwise this will be handled in OnRegister (e.g. Editor manipulations).
 	if (bNavigationRelevant && IsRegistered())
@@ -3502,7 +3505,15 @@ int32 UInstancedStaticMeshComponent::AddInstanceInternal(int32 InstanceIndex, FI
 		FInstancedStaticMeshDelegates::OnInstanceIndexUpdated.Broadcast(this, MakeArrayView(&IndexUpdate, 1));
 	}
 
+	const FBoundsCacheElement NavigationCacheElement = CachedBounds[(int32)EBoundsCacheType::Navigation];
 	InvalidateCachedBounds();
+
+	// Restore Navigation CachedBounds (the invalidation is not needed if it was updated).
+	if (bNavigationCachedBoundsUpdated)
+	{
+		CachedBounds[(int32)EBoundsCacheType::Navigation] = NavigationCacheElement;
+	}
+	
 	return InstanceIndex;
 }
 
@@ -5247,6 +5258,7 @@ void UInstancedStaticMeshComponent::PartialNavigationUpdates(const TConstArrayVi
 
 		// Update cached navigation bounds
 		CalcAndCacheNavigationBounds();
+		bNavigationCachedBoundsUpdated = true;
 
 		// Dirty areas around the modified instances
 		TArray<FBox> DirtyAreas;
@@ -5413,7 +5425,6 @@ void UInstancedStaticMeshComponent::GetNavigationData(FNavigationRelevantData& D
 
 FBox UInstancedStaticMeshComponent::GetNavigationBounds() const
 {
-	const UWorld* const World = GetWorld();
 	// Navigation bounds must be calculated before being requested
 	constexpr bool bCheckValid = true;
 	return CachedBounds[(int)EBoundsCacheType::Navigation].GetValue(bCheckValid).GetBox();
