Added MaxUniformExpressions instead of manually repeating the value

diff --git a/Engine/Source/Runtime/Engine/Private/Materials/MaterialUniformExpressions.cpp b/Engine/Source/Runtime/Engine/Private/Materials/MaterialUniformExpressions.cpp
--- a/Engine/Source/Runtime/Engine/Private/Materials/MaterialUniformExpressions.cpp
+++ b/Engine/Source/Runtime/Engine/Private/Materials/MaterialUniformExpressions.cpp
@@ -433,39 +433,41 @@ FShaderParametersMetadata* FUniformExpressionSet::CreateBufferStruct()
 
 	check((NextMemberOffset % (2 * SHADER_PARAMETER_POINTER_ALIGNMENT)) == 0);
 
-	static FString Texture2DNames[128];
-	static FString Texture2DSamplerNames[128];
-	static FString TextureCubeNames[128];
-	static FString TextureCubeSamplerNames[128];
-	static FString Texture2DArrayNames[128];
-	static FString Texture2DArraySamplerNames[128];
-	static FString TextureCubeArrayNames[128];
-	static FString TextureCubeArraySamplerNames[128];
-	static FString VolumeTextureNames[128];
-	static FString VolumeTextureSamplerNames[128];
-	static FString ExternalTextureNames[128];
-	static FString TextureCollectionNames[128];
-	static FString MediaTextureSamplerNames[128];
-	static FString VirtualTexturePageTableNames0[128];
-	static FString VirtualTexturePageTableNames1[128];
-	static FString VirtualTexturePageTableIndirectionNames[128];
-	static FString VirtualTexturePhysicalNames[128];
-	static FString VirtualTexturePhysicalSamplerNames[128];
-	static FString SparseVolumeTexturePageTableNames[128];
-	static FString SparseVolumeTexturePhysicalANames[128];
-	static FString SparseVolumeTexturePhysicalBNames[128];
-	static FString SparseVolumeTexturePhysicalSamplerNames[128];
+	const uint32 MaxUniformExpressions = 128;
+
+	static FString Texture2DNames[MaxUniformExpressions];
+	static FString Texture2DSamplerNames[MaxUniformExpressions];
+	static FString TextureCubeNames[MaxUniformExpressions];
+	static FString TextureCubeSamplerNames[MaxUniformExpressions];
+	static FString Texture2DArrayNames[MaxUniformExpressions];
+	static FString Texture2DArraySamplerNames[MaxUniformExpressions];
+	static FString TextureCubeArrayNames[MaxUniformExpressions];
+	static FString TextureCubeArraySamplerNames[MaxUniformExpressions];
+	static FString VolumeTextureNames[MaxUniformExpressions];
+	static FString VolumeTextureSamplerNames[MaxUniformExpressions];
+	static FString ExternalTextureNames[MaxUniformExpressions];
+	static FString TextureCollectionNames[MaxUniformExpressions];
+	static FString MediaTextureSamplerNames[MaxUniformExpressions];
+	static FString VirtualTexturePageTableNames0[MaxUniformExpressions];
+	static FString VirtualTexturePageTableNames1[MaxUniformExpressions];
+	static FString VirtualTexturePageTableIndirectionNames[MaxUniformExpressions];
+	static FString VirtualTexturePhysicalNames[MaxUniformExpressions];
+	static FString VirtualTexturePhysicalSamplerNames[MaxUniformExpressions];
+	static FString SparseVolumeTexturePageTableNames[MaxUniformExpressions];
+	static FString SparseVolumeTexturePhysicalANames[MaxUniformExpressions];
+	static FString SparseVolumeTexturePhysicalBNames[MaxUniformExpressions];
+	static FString SparseVolumeTexturePhysicalSamplerNames[MaxUniformExpressions];
 	
-	static FString MaterialCacheTagBuffer[128];
-	static FString MaterialCacheTagPageTableNames0[128];
-	static FString MaterialCacheTagPhysicalNames[128];
-	static FString MaterialCacheTagPhysicalSamplerNames[128];
+	static FString MaterialCacheTagBuffer[MaxUniformExpressions];
+	static FString MaterialCacheTagPageTableNames0[MaxUniformExpressions];
+	static FString MaterialCacheTagPhysicalNames[MaxUniformExpressions];
+	static FString MaterialCacheTagPhysicalSamplerNames[MaxUniformExpressions];
 	
 	static bool bInitializedTextureNames = false;
 	if (!bInitializedTextureNames)
 	{
 		bInitializedTextureNames = true;
-		for (int32 i = 0; i < 128; ++i)
+		for (int32 i = 0; i < MaxUniformExpressions; ++i)
 		{
 			Texture2DNames[i] = FString::Printf(TEXT("Texture2D_%d"), i);
 			Texture2DSamplerNames[i] = FString::Printf(TEXT("Texture2D_%dSampler"), i);
@@ -499,10 +501,10 @@ FShaderParametersMetadata* FUniformExpressionSet::CreateBufferStruct()
 
 	for (uint32 TypeIndex = 0u; TypeIndex < NumMaterialTextureParameterTypes; ++TypeIndex)
 	{
-		check(UniformTextureParameters[TypeIndex].Num() <= 128);
+		check(UniformTextureParameters[TypeIndex].Num() <= MaxUniformExpressions);
 	}
-	check(VTStacks.Num() <= 128);
-	check(MaterialCacheTagStacks.Num() <= 128);
+	check(VTStacks.Num() <= MaxUniformExpressions);
+	check(MaterialCacheTagStacks.Num() <= MaxUniformExpressions);
 
 	for (int32 i = 0; i < UniformTextureParameters[(uint32)EMaterialTextureParameterType::Standard2D].Num(); ++i)
 	{
