Only add an actor to the NetGUID cache when it's mirgated to the server and isn't being destroyed.

diff --git a/Engine/Source/Runtime/Engine/Private/Components/ActorComponent.cpp b/Engine/Source/Runtime/Engine/Private/Components/ActorComponent.cpp
--- a/Engine/Source/Runtime/Engine/Private/Components/ActorComponent.cpp
+++ b/Engine/Source/Runtime/Engine/Private/Components/ActorComponent.cpp
@@ -764,7 +764,10 @@ void UActorComponent::PostMigrate(const struct FUObjectMigrationContext& Migrati
 
 		// Add the actor and it's sub-objects to the network id cache since these objects are no longer considered remote and can be derferenced
 		// through the network id cache (e.g. by RPC argument and object deserialization).
-		GetWorld()->GetNetDriver()->AddObjectToNetworkIdCache(this);
+		if (HasBegunPlay())
+		{
+			GetWorld()->GetNetDriver()->AddObjectToNetworkIdCache(this);
+		}
 
 		// We are expecting this flow from AActor, but now we must do this per-component in such a way
 		// that we are not mutating our previously deserialized data.
diff --git a/Engine/Source/Runtime/Engine/Private/NetDriver.cpp b/Engine/Source/Runtime/Engine/Private/NetDriver.cpp
--- a/Engine/Source/Runtime/Engine/Private/NetDriver.cpp
+++ b/Engine/Source/Runtime/Engine/Private/NetDriver.cpp
@@ -7925,13 +7925,14 @@ void UNetDriver::CreateInitialServerChannels(UNetConnection* ClientConnection)
 
 void UNetDriver::AddObjectToNetworkIdCache(UObject* Object)
 {
-	if (!Object)
+	if (!GuidCache)
 	{
 		return;
 	}
 
-	if (!GuidCache)
+	if (!ensure(IsValid(Object)))
 	{
+		UE_LOG(LogNet, Warning, TEXT("Cannot add an invalid object %s to the network id cache."), *GetNameSafe(Object));
 		return;
 	}
 
@@ -7943,11 +7944,15 @@ void UNetDriver::AddObjectToNetworkIdCache(UObject* Object)
 		for (const UE::Net::FSubObjectRegistry::FEntry& SubObjectInfo : SubObjects->GetRegistryList())
 		{
 			UObject* SubObject = SubObjectInfo.GetSubObject();
-			if (SubObject)
+			if (ensure(IsValid(SubObject)))
 			{
 				FNetworkGUID SubObjectNetGUID = GuidCache->GetOrAssignNetGUID(SubObject);
 				UE_LOG(LogNet, Verbose, TEXT("Adding sub-object %s (NetGUID=%s) to the NetGUID cache."), *SubObject->GetName(), *SubObjectNetGUID.ToString());
 			}
+			else
+			{
+				UE_LOG(LogNet, Warning, TEXT("Cannot add an invalid sub-object %s to the network id cache."), *GetNameSafe(SubObject));
+			}
 		}
 	};
 
