Improve Substrate GBuffer format project settings  tooltips.

diff --git a/Engine/Source/Runtime/Engine/Classes/Engine/RendererSettings.h b/Engine/Source/Runtime/Engine/Classes/Engine/RendererSettings.h
--- a/Engine/Source/Runtime/Engine/Classes/Engine/RendererSettings.h
+++ b/Engine/Source/Runtime/Engine/Classes/Engine/RendererSettings.h
@@ -1218,8 +1218,8 @@ class URendererSettings : public UDeveloperSettings
 	*/
 	UPROPERTY(config, EditAnywhere, Category = Substrate, meta = (
 		ConsoleVariable = "r.Substrate.ProjectGBufferFormat",
-		DisplayName = "Substrate GBuffer Format",
-		ToolTip = "Select the GBuffer format used for this project. Changing this setting requires restarting the editor.",
+		DisplayName = "Substrate GBuffer Format (Project)",
+		ToolTip = "Select the GBuffer format used for this project.\n - Blendable GBuffer is supported on all platforms. \n - Adaptive GBuffer is supported on certains platforms. Unsupported platforms will fallback onto 'Blendable GBuffer' format.\nChanging this setting requires restarting the editor.",
 		ConfigRestartRequired = true))
 	TEnumAsByte<ESubstrateStorageFormat::Type> SubstrateGBufferFormat;
 
