Fixed missing initialization on FMaterialExpressionMaterialCacheAttribute::ValueType

diff --git a/Engine/Source/Runtime/Engine/Public/Materials/MaterialExpressionMaterialCache.h b/Engine/Source/Runtime/Engine/Public/Materials/MaterialExpressionMaterialCache.h
--- a/Engine/Source/Runtime/Engine/Public/Materials/MaterialExpressionMaterialCache.h
+++ b/Engine/Source/Runtime/Engine/Public/Materials/MaterialExpressionMaterialCache.h
@@ -22,7 +22,7 @@ struct FMaterialExpressionMaterialCacheAttribute
 
 	/** EMaterialValueType, u64 for serialization */
 	UPROPERTY()
-	uint64 ValueType;
+	uint64 ValueType = 0;
 
 	/** Bindable input */
 	UPROPERTY()
