HLSL Translator, added missing CustomOutput code chunk check

diff --git a/Engine/Source/Runtime/Engine/Private/Materials/HLSLMaterialTranslator.cpp b/Engine/Source/Runtime/Engine/Private/Materials/HLSLMaterialTranslator.cpp
--- a/Engine/Source/Runtime/Engine/Private/Materials/HLSLMaterialTranslator.cpp
+++ b/Engine/Source/Runtime/Engine/Private/Materials/HLSLMaterialTranslator.cpp
@@ -1094,7 +1094,10 @@ uint64 FHLSLMaterialTranslator::CompileMaterialAttributesCustomOutputs(uint64 Cu
 			TArray<FShaderCodeChunk> CustomExpressionChunks;
 			AssignTempScope(CustomExpressionChunks);
 			int32 CodeChunk = Material->CompileCustomAttribute(CustomAttribute.AttributeID, this);
-			CustomOutput(CustomAttribute.AttributeName, HLSLTypeString(GetParameterType(CodeChunk)), CustomAttribute.FunctionName, CustomAttribute.OutputIndex, CodeChunk, EMaterialCustomOutputFlags::None);
+			if (CodeChunk != INDEX_NONE)
+			{
+				CustomOutput(CustomAttribute.AttributeName, HLSLTypeString(GetParameterType(CodeChunk)), CustomAttribute.FunctionName, CustomAttribute.OutputIndex, CodeChunk, EMaterialCustomOutputFlags::None);
+			}
 		}
 
 		bCompilingPreviousFrame = bCachedCompilingPreviousFrame; // Restore the cached value of bCompilingPreviousFrame
