[shaders] only prompt for retries for the first failed shadermap; avoid dialog spam when lots of materials are loaded.

diff --git a/Engine/Source/Runtime/Engine/Private/ShaderCompiler/ShaderCompiler.cpp b/Engine/Source/Runtime/Engine/Private/ShaderCompiler/ShaderCompiler.cpp
--- a/Engine/Source/Runtime/Engine/Private/ShaderCompiler/ShaderCompiler.cpp
+++ b/Engine/Source/Runtime/Engine/Private/ShaderCompiler/ShaderCompiler.cpp
@@ -2393,6 +2393,7 @@ bool FShaderCompilingManager::HandlePotentialRetry(TMap<int32, FShaderMapFinaliz
 	TArray<FShaderCommonCompileJobPtr> ReissueForDebugJobs;
 	TArray<FShaderCommonCompileJobPtr> ReissueJobs;
 	TArray<int32> ReissueMapIds;
+	bool bInteractiveRetry = false;
 
 	for (TMap<int32, FShaderMapFinalizeResults>::TIterator It(CompletedShaderMaps); It; ++It)
 	{
@@ -2419,7 +2420,6 @@ bool FShaderCompilingManager::HandlePotentialRetry(TMap<int32, FShaderMapFinaliz
 			}
 		}
 
-		bool bInteractiveRetry = false;
 
 		// interactive retries - prompt for global/default shaders that have errors, if we're not reissuing any jobs to dump debug info first
 		if (!FApp::IsUnattended() && !Results.bAllJobsSucceeded && ReissueForDebugJobs.IsEmpty())
@@ -2460,7 +2460,8 @@ bool FShaderCompilingManager::HandlePotentialRetry(TMap<int32, FShaderMapFinaliz
 
 				BuildErrorStringAndReport(ShaderDiagInfo, ErrorString);
 
-				if (UE_LOG_ACTIVE(LogShaders, Log) && (bAnyErrorLikelyToBeCodeError || bPromptToRetryFailedShaderCompiles || bSpecialEngineMaterial))
+				if ((bAnyErrorLikelyToBeCodeError || bPromptToRetryFailedShaderCompiles || bSpecialEngineMaterial) 
+					&& !bInteractiveRetry) // break/prompt only for the first shadermap which encountered an error, to avoid dialog spam
 				{
 					// Use debug break in debug with the debugger attached, otherwise message box
 					if (bDebugBreakOnPromptToRetryShaderCompile && FPlatformMisc::IsDebuggerPresent())
