[HWRT] Support GPU instanced static meshes using non-Nanite meshes in ray tracing scene.
- extended FInstancedStaticMeshSceneProxy with support for static ray tracing instance code path when not using WPO.
- fallback to static ray tracing instances when IsInstanceDataGPUOnly() since we can't run the logic that decides which instances to simulate on CPU.

diff --git a/Engine/Source/Runtime/Engine/Classes/Engine/InstancedStaticMesh.h b/Engine/Source/Runtime/Engine/Classes/Engine/InstancedStaticMesh.h
--- a/Engine/Source/Runtime/Engine/Classes/Engine/InstancedStaticMesh.h
+++ b/Engine/Source/Runtime/Engine/Classes/Engine/InstancedStaticMesh.h
@@ -476,15 +476,10 @@ public:
 	bool bAnySegmentUsesWorldPositionOffset = false;
 	bool bUseGpuLodSelection = false;
 
-#if RHI_RAYTRACING
-	virtual bool IsRayTracingStaticRelevant() const override
-	{
-		// TODO: support static ray tracing instances when not using WPO
-		return false;
-	}
+	ENGINE_API virtual void SetEvaluateWorldPositionOffsetInRayTracing(FRHICommandListBase& RHICmdList, bool bEvaluate) override;
 
+#if RHI_RAYTRACING
 	ENGINE_API virtual void GetDynamicRayTracingInstances(FRayTracingInstanceCollector& Collector) final override;
-
 #endif
 
 	ENGINE_API virtual void GetLightRelevance(const FLightSceneProxy* LightSceneProxy, bool& bDynamic, bool& bRelevant, bool& bLightMapped, bool& bShadowMapped) const override;
@@ -551,14 +546,16 @@ protected:
 
 #if RHI_RAYTRACING
 	ENGINE_API void SetupRayTracingDynamicInstances(FRHICommandListBase& RHICmdList, int32 NumDynamicInstances, int32 LODIndex);
+	ENGINE_API void ReleaseRayTracingDynamicInstances();
+	ENGINE_API void ApplyRayTracingWPOSettings();
 
-	struct FRayTracingDynamicData
+	struct FRayTracingDynamicInstanceData
 	{
-		FRayTracingGeometry DynamicGeometry;
-		FRWBuffer DynamicGeometryVertexBuffer;
+		FRayTracingGeometry Geometry;
+		FRWBuffer VertexBuffer;
 	};
 
-	TArray<FRayTracingDynamicData> RayTracingDynamicData;
+	TArray<FRayTracingDynamicInstanceData> RayTracingDynamicInstanceData;
 
 	int32 CachedRayTracingLODIndex;
 #endif
diff --git a/Engine/Source/Runtime/Engine/Private/InstancedStaticMesh.cpp b/Engine/Source/Runtime/Engine/Private/InstancedStaticMesh.cpp
--- a/Engine/Source/Runtime/Engine/Private/InstancedStaticMesh.cpp
+++ b/Engine/Source/Runtime/Engine/Private/InstancedStaticMesh.cpp
@@ -157,6 +157,28 @@ static TAutoConsoleVariable<int32> CVarRayTracingInstancesEvaluateWPO(
 	TEXT(" -1 - on only for meshes with evaluate WPO enabled")
 );
 
+FAutoConsoleVariableSink CVarRayTracingInstancesEvaluateWPOSink(FConsoleCommandDelegate::CreateLambda([]()
+	{
+		static int32 CachedValue = CVarRayTracingInstancesEvaluateWPO.GetValueOnGameThread();
+
+		int32 Value = CVarRayTracingInstancesEvaluateWPO.GetValueOnGameThread();
+
+		if (CachedValue != Value)
+		{
+			CachedValue = Value;
+
+			// NV-JIRA UE-668: Do this as a task on the game thread to break up a possible
+			// reentry call to USkeletalMeshComponent::PostAnimEvaluation if BP toggles this CVar
+			FFunctionGraphTask::CreateAndDispatchWhenReady([]()
+				{
+					// Easiest way to update all static scene proxies is to recreate them
+					FlushRenderingCommands();
+					FGlobalComponentReregisterContext ReregisterContext;
+
+				}, TStatId(), nullptr, ENamedThreads::GameThread);
+		}
+	}));
+
 static TAutoConsoleVariable<float> CVarRayTracingInstancesSimulationClusterRadius(
 	TEXT("r.RayTracing.Geometry.InstancedStaticMeshes.SimulationClusterRadius"),
 	500.0f, // 5 m
@@ -1458,16 +1480,23 @@ void FInstancedStaticMeshSceneProxy::SetupProxy(const FInstancedStaticMeshSceneP
 			bAnySegmentUsesWorldPositionOffset |= Section.Material->IsUsingWorldPositionOffset_Concurrent(GMaxRHIShaderPlatform);
 		}
 	}
-	
-	// TODO: when using WPO, the base class (FStaticMeshSceneProxy) currently creates DynamicRayTracingGeometries which is not actually used
-	// because FInstancedStaticMeshSceneProxy creates it's own ray tracing geometries for each simulated instance... 
+
+#if RHI_RAYTRACING
+	if (bSupportRayTracing)
+	{
+		ApplyRayTracingWPOSettings();
+
+		// TODO: when using WPO, the base class (FStaticMeshSceneProxy) currently creates DynamicRayTracingGeometries which is not actually used
+		// because FInstancedStaticMeshSceneProxy creates it's own ray tracing geometries for each simulated instance... 
+	}
+#endif // RHI_RAYTRACING
 
 	if (OverlayMaterial != nullptr)
 	{
 		if (!OverlayMaterial->CheckMaterialUsage_Concurrent(MATUSAGE_InstancedStaticMeshes))
 		{
 			OverlayMaterial = UMaterial::GetDefaultMaterial(MD_Surface);
-			UE_LOG(LogStaticMesh, Error, TEXT("Overlay material with missing usage flag was applied to instanced static mesh %s"),	*InProxyDesc.GetStaticMesh()->GetPathName());
+			UE_LOG(LogStaticMesh, Error, TEXT("Overlay material with missing usage flag was applied to instanced static mesh %s"), *InProxyDesc.GetStaticMesh()->GetPathName());
 		}
 	}
 
@@ -1540,13 +1569,7 @@ void FInstancedStaticMeshSceneProxy::DestroyRenderThreadResources()
 	FStaticMeshSceneProxy::DestroyRenderThreadResources();
 
 #if RHI_RAYTRACING
-	for (auto &DynamicRayTracingItem : RayTracingDynamicData)
-	{
-		DynamicRayTracingItem.DynamicGeometry.ReleaseResource();
-		DynamicRayTracingItem.DynamicGeometryVertexBuffer.Release();
-	}
-
-	RayTracingDynamicData.Empty();
+	ReleaseRayTracingDynamicInstances();
 #endif
 }
 
@@ -1683,9 +1706,40 @@ FInstanceDataUpdateTaskInfo *FInstancedStaticMeshSceneProxy::GetInstanceDataUpda
 	return InstanceDataSceneProxy->GetUpdateTaskInfo();
 }
 
+void FInstancedStaticMeshSceneProxy::SetEvaluateWorldPositionOffsetInRayTracing(FRHICommandListBase& RHICmdList, bool bEvaluate)
+{
+#if RHI_RAYTRACING
+	if (!bSupportRayTracing)
+	{
+		return;
+	}
+
+	bool bPrevDynamicRayTracingGeometry = bDynamicRayTracingGeometry;
+
+	FStaticMeshSceneProxy::SetEvaluateWorldPositionOffsetInRayTracing(RHICmdList, bEvaluate);
+
+	ApplyRayTracingWPOSettings();
+
+	if (bPrevDynamicRayTracingGeometry && !bDynamicRayTracingGeometry)
+	{
+		ReleaseRayTracingDynamicInstances();
+	}
+
+	GetScene().UpdateCachedRayTracingState(this);
+#endif // RHI_RAYTRACING
+}
+
 #if RHI_RAYTRACING
 void FInstancedStaticMeshSceneProxy::GetDynamicRayTracingInstances(FRayTracingInstanceCollector& Collector)
 {
+	// Workaround: SetEvaluateWorldPositionOffsetInRayTracing(...) calls UpdateCachedRayTracingState(...)
+	// however the update only happens after gathering relevant ray tracing primitives
+	// so ERayTracingPrimitiveFlags::Dynamic is set for one frame after the WPO evaluation is disabled.
+	if (!bDynamicRayTracingGeometry)
+	{
+		return;
+	}
+
 	if (!CVarRayTracingRenderInstances.GetValueOnRenderThread())
 	{
 		// TODO: Exclude proxy during ray tracing primitive gather instead of doing this early out here.
@@ -1699,6 +1753,10 @@ void FInstancedStaticMeshSceneProxy::GetDynamicRayTracingInstances(FRayTracingIn
 	}
 #endif
 
+	const FInstanceSceneDataBuffers* InstanceSceneDataBuffers = GetInstanceSceneDataBuffers();
+
+	checkf(!InstanceSceneDataBuffers->IsInstanceDataGPUOnly(), TEXT("FInstancedStaticMeshSceneProxy does not support dynamic ray tracing geometry when IsInstanceDataGPUOnly()."));
+
 	const uint32 InstanceCount = GetInstanceDataHeader().NumInstances;
 
 	if (InstanceCount == 0u)
@@ -1706,6 +1764,8 @@ void FInstancedStaticMeshSceneProxy::GetDynamicRayTracingInstances(FRayTracingIn
 		return;
 	}
 
+	check(InstanceSceneDataBuffers && InstanceSceneDataBuffers->GetNumInstances() == InstanceCount);
+
 	// TODO: Should only do this if any instance uses static geometry
 	Collector.AddReferencedGeometryGroup(RenderData->RayTracingGeometryGroupHandle);
 
@@ -1766,11 +1826,6 @@ void FInstancedStaticMeshSceneProxy::GetDynamicRayTracingInstances(FRayTracingIn
 	const uint32 RequestedSimulatedInstances = CVarRayTracingSimulatedInstanceCount.GetValueOnRenderThread();
 	const uint32 SimulatedInstances = FMath::Min(RequestedSimulatedInstances == -1 ? InstanceCount : FMath::Clamp(RequestedSimulatedInstances, 1u, InstanceCount), MaxSimulatedInstances);
 
-	const int32 WPOEvalMode = CVarRayTracingInstancesEvaluateWPO.GetValueOnRenderThread();
-	const bool bWantsWPOEvaluation = WPOEvalMode < 0 ? bDynamicRayTracingGeometry : WPOEvalMode != 0;
-	const bool bHasWorldPositionOffset = bWantsWPOEvaluation && bAnySegmentUsesWorldPositionOffset;
-
-	if (bHasWorldPositionOffset)
 	{
 		int32 SectionCount = InstancedRenderData.LODModels[LODIndex].Sections.Num();
 
@@ -1826,17 +1881,6 @@ void FInstancedStaticMeshSceneProxy::GetDynamicRayTracingInstances(FRayTracingIn
 	constexpr float LocalToWorldScale = 1.0f;
 	FVector ViewPosition = Collector.GetReferenceView()->ViewLocation;
 
-	const FInstanceSceneDataBuffers* InstanceSceneDataBuffers = GetInstanceSceneDataBuffers();
-
-	if (InstanceSceneDataBuffers->IsInstanceDataGPUOnly())
-	{
-		// TODO: Update FInstancedStaticMeshSceneProxy to make it compatible with GPU Scene when not using WPO materials, and support fallback to static geometry
-		// when IsInstanceDataGPUOnly() instead of simply being hidden from ray tracing.
-		return;
-	}
-
-	check(InstanceSceneDataBuffers && InstanceSceneDataBuffers->GetNumInstances() == InstanceCount);
-
 	auto GetDistanceToInstance = [&ViewPosition, InstanceSceneDataBuffers](int32 InstanceIndex, float& OutInstanceRadius, float& OutDistanceToInstanceCenter, float& OutDistanceToInstanceStart)
 	{
 		const FBoxSphereBounds& InstanceWorldBounds = InstanceSceneDataBuffers->GetInstanceWorldBounds(InstanceIndex);
@@ -1920,7 +1964,7 @@ void FInstancedStaticMeshSceneProxy::GetDynamicRayTracingInstances(FRayTracingIn
 	}
 
 	// Bucket which visible instances to simulate based on camera distance
-	if (bHasWorldPositionOffset && VisibleInstances.Num() > 0)
+	if (VisibleInstances.Num() > 0)
 	{
 		float SimulationClusterRadius = CVarRayTracingInstancesSimulationClusterRadius.GetValueOnRenderThread();
 
@@ -1977,7 +2021,7 @@ void FInstancedStaticMeshSceneProxy::GetDynamicRayTracingInstances(FRayTracingIn
 		FMatrix InstanceToWorld = InstanceSceneDataBuffers->GetInstanceToWorld(InstanceIndex);
 		const uint32 DynamicInstanceIdx = InstanceIndex % SimulatedInstances;
 
-		if (bHasWorldPositionOffset && InstancedRenderData.VertexFactories[LODIndex].GetType()->SupportsRayTracingDynamicGeometry())
+		if (InstancedRenderData.VertexFactories[LODIndex].GetType()->SupportsRayTracingDynamicGeometry())
 		{
 			FRayTracingInstance* DynamicInstance = nullptr;
 
@@ -1986,11 +2030,11 @@ void FInstancedStaticMeshSceneProxy::GetDynamicRayTracingInstances(FRayTracingIn
 				// first case of this dynamic instance, setup the material and add it
 				const FStaticMeshLODResources& LODModel = RenderData->LODResources[LODIndex];
 
-				FRayTracingDynamicData& DynamicData = RayTracingDynamicData[DynamicInstanceIdx];
+				FRayTracingDynamicInstanceData& DynamicData = RayTracingDynamicInstanceData[DynamicInstanceIdx];
 
 				ActiveInstances[DynamicInstanceIdx] = RayTracingInstances.Num();
 				FRayTracingInstance& RayTracingInstance = RayTracingInstances.Add_GetRef(RayTracingWPOInstanceTemplate);
-				RayTracingInstance.Geometry = &DynamicData.DynamicGeometry;
+				RayTracingInstance.Geometry = &DynamicData.Geometry;
 				RayTracingInstance.PrimitiveInstanceIndices.Reserve(InstanceCount);
 
 				DynamicInstance = &RayTracingInstance;
@@ -2002,8 +2046,8 @@ void FInstancedStaticMeshSceneProxy::GetDynamicRayTracingInstances(FRayTracingIn
 						false,
 						(uint32)LODModel.GetNumVertices(),
 						uint32((SIZE_T)LODModel.GetNumVertices() * sizeof(FVector3f)),
-						DynamicData.DynamicGeometry.Initializer.TotalPrimitiveCount,
-						&DynamicData.DynamicGeometry,
+						DynamicData.Geometry.Initializer.TotalPrimitiveCount,
+						&DynamicData.Geometry,
 						nullptr,
 						true,
 						InstanceIndex,
@@ -2060,7 +2104,7 @@ void FInstancedStaticMeshSceneProxy::GetDynamicRayTracingInstances(FRayTracingIn
 
 void FInstancedStaticMeshSceneProxy::SetupRayTracingDynamicInstances(FRHICommandListBase& RHICmdList, int32 NumDynamicInstances, int32 LODIndex)
 {
-	if (NumDynamicInstances == RayTracingDynamicData.Num() && LODIndex == CachedRayTracingLODIndex)
+	if (NumDynamicInstances == RayTracingDynamicInstanceData.Num() && LODIndex == CachedRayTracingLODIndex)
 	{
 		return;
 	}
@@ -2074,15 +2118,14 @@ void FInstancedStaticMeshSceneProxy::SetupRayTracingDynamicInstances(FRHICommand
 
 	// release entries
 
-	for (int32 Item = 0; Item < RayTracingDynamicData.Num(); Item++)
+	for (FRayTracingDynamicInstanceData& DynamicData : RayTracingDynamicInstanceData)
 	{
-		auto& DynamicRayTracingItem = RayTracingDynamicData[Item];
-		DynamicRayTracingItem.DynamicGeometry.ReleaseResource();
-		DynamicRayTracingItem.DynamicGeometryVertexBuffer.Release();
+		DynamicData.Geometry.ReleaseResource();
+		DynamicData.VertexBuffer.Release();
 	}
 
 	// clear and resize
-	RayTracingDynamicData.Empty(NumDynamicInstances);
+	RayTracingDynamicInstanceData.Empty(NumDynamicInstances);
 
 	// create new geometries
 
@@ -2090,7 +2133,7 @@ void FInstancedStaticMeshSceneProxy::SetupRayTracingDynamicInstances(FRHICommand
 
 	for (int32 Item = 0; Item < NumDynamicInstances; Item++)
 	{
-		FRayTracingDynamicData& DynamicData = RayTracingDynamicData.AddDefaulted_GetRef();
+		FRayTracingDynamicInstanceData& DynamicData = RayTracingDynamicInstanceData.AddDefaulted_GetRef();
 
 		FRayTracingGeometryInitializer Initializer = RayTracingLOD.RayTracingGeometry->Initializer;
 		for (FRayTracingGeometrySegment& Segment : Initializer.Segments)
@@ -2100,12 +2143,41 @@ void FInstancedStaticMeshSceneProxy::SetupRayTracingDynamicInstances(FRHICommand
 		Initializer.bAllowUpdate = true;
 		Initializer.bFastBuild = true;
 
-		DynamicData.DynamicGeometry.SetInitializer(MoveTemp(Initializer));
-		DynamicData.DynamicGeometry.InitResource(RHICmdList);
+		DynamicData.Geometry.SetInitializer(MoveTemp(Initializer));
+		DynamicData.Geometry.InitResource(RHICmdList);
 	}
 }
 
-#endif
+void FInstancedStaticMeshSceneProxy::ReleaseRayTracingDynamicInstances()
+{
+	for (auto& DynamicRayTracingItem : RayTracingDynamicInstanceData)
+	{
+		DynamicRayTracingItem.Geometry.ReleaseResource();
+		DynamicRayTracingItem.VertexBuffer.Release();
+	}
+
+	RayTracingDynamicInstanceData.Empty();
+}
+
+void FInstancedStaticMeshSceneProxy::ApplyRayTracingWPOSettings()
+{
+	const int32 WPOEvalMode = CVarRayTracingInstancesEvaluateWPO.GetValueOnAnyThread();
+
+	bDynamicRayTracingGeometry = WPOEvalMode < 0 ? bDynamicRayTracingGeometry : WPOEvalMode != 0;
+	bDynamicRayTracingGeometry &= bAnySegmentUsesWorldPositionOffset;
+
+	const FInstanceSceneDataBuffers* InstanceSceneDataBuffers = GetInstanceSceneDataBuffers();
+
+	if (InstanceSceneDataBuffers->IsInstanceDataGPUOnly())
+	{
+		UE_LOG(LogStaticMesh, Warning, TEXT("Dynamic ray tracing geometry not currently supported for GPU instanced static meshes (%s). Falling back to static ray tracing geometry."), *StaticMesh->GetPathName());
+
+		// TODO: Update FInstancedStaticMeshSceneProxy to make it compatible with GPU Scene when using WPO materials.
+		bDynamicRayTracingGeometry = false;
+	}
+}
+
+#endif // RHI_RAYTRACING
 
 void FInstancedStaticMeshSceneProxy::SetInstanceCullDistance_RenderThread(float StartCullDistance, float EndCullDistance)
 {
diff --git a/Engine/Source/Runtime/Renderer/Private/RayTracing/RayTracing.cpp b/Engine/Source/Runtime/Renderer/Private/RayTracing/RayTracing.cpp
--- a/Engine/Source/Runtime/Renderer/Private/RayTracing/RayTracing.cpp
+++ b/Engine/Source/Runtime/Renderer/Private/RayTracing/RayTracing.cpp
@@ -1513,11 +1513,13 @@ namespace RayTracing
 				continue;
 			}
 
+			const int32 NumInstances = SceneInfo->GetNumInstanceSceneDataEntries();
+
 			// location if this is a new entry
 			const uint64 InstanceKey = RelevantPrimitive.InstancingKey();
 
 			FAutoInstanceBatch DummyInstanceBatch = { };
-			FAutoInstanceBatch& InstanceBatch = bAutoInstance ? InstanceBatches.FindOrAdd(InstanceKey, DummyInstanceBatch) : DummyInstanceBatch;
+			FAutoInstanceBatch& InstanceBatch = bAutoInstance && (NumInstances == 1) ? InstanceBatches.FindOrAdd(InstanceKey, DummyInstanceBatch) : DummyInstanceBatch;
 
 			if (InstanceBatch.IsValid())
 			{
@@ -1555,16 +1557,29 @@ namespace RayTracing
 			{
 				// Starting new instance batch
 
-				InstanceBatch.Add(RayTracingScene, SceneInfo->GetInstanceSceneDataOffset());
+				const int32 InstanceSceneDataOffset = SceneInfo->GetInstanceSceneDataOffset();
+
+				InstanceBatch.Add(RayTracingScene, InstanceSceneDataOffset);
 
 				FRayTracingGeometryInstance RayTracingInstance;
 				RayTracingInstance.GeometryRHI = RelevantPrimitive.RayTracingGeometry->GetRHI();
 				checkf(RayTracingInstance.GeometryRHI, TEXT("Ray tracing instance must have a valid geometry."));
-				RayTracingInstance.InstanceSceneDataOffsets = InstanceBatch.InstanceSceneDataOffsets;
-				RayTracingInstance.UserData = InstanceBatch.InstanceSceneDataOffsets;
-				RayTracingInstance.NumTransforms = 1;
 				RayTracingInstance.bUsesLightingChannels = RelevantPrimitive.bUsesLightingChannels;
 
+				if (NumInstances == 1)
+				{
+					RayTracingInstance.NumTransforms = 1;
+					RayTracingInstance.InstanceSceneDataOffsets = InstanceBatch.InstanceSceneDataOffsets;
+					RayTracingInstance.UserData = InstanceBatch.InstanceSceneDataOffsets;
+				}
+				else
+				{
+					RayTracingInstance.NumTransforms = NumInstances;
+					RayTracingInstance.BaseInstanceSceneDataOffset = InstanceSceneDataOffset;
+					RayTracingInstance.DefaultUserData = InstanceSceneDataOffset;
+					RayTracingInstance.bIncrementUserDataPerInstance = true;
+				}
+
 				RayTracingInstance.Mask = RelevantPrimitive.CachedMeshCommandFlags.InstanceMask; // When no cached command is found, InstanceMask == 0 and the instance is effectively filtered out
 
 				// Run AHS for alpha masked and meshes with only some sections casting shadows, which require per mesh section filtering in AHS
