Rename the aliases for the Convert node from `Make Float4` to `Make Vector4` to match "Vector Parameter" and "Constant Vectors'

- This also avoids confusion with the existing Make Float 4 Engine Material functions

[RN] minor | Rendering - Materials and Shaders
#tests typing "make vector" in the expression search when adding a new node (it should find and place the convert node)

diff --git a/Engine/Source/Runtime/Engine/Private/Materials/MaterialExpressions.cpp b/Engine/Source/Runtime/Engine/Private/Materials/MaterialExpressions.cpp
--- a/Engine/Source/Runtime/Engine/Private/Materials/MaterialExpressions.cpp
+++ b/Engine/Source/Runtime/Engine/Private/Materials/MaterialExpressions.cpp
@@ -24251,8 +24251,8 @@ void UMaterialExpressionConvert::RegisterAdditionalMenuActions(FGraphActionMenuB
 	auto CreateMakeNode = [this, &ActionMenuBuilder, &CategoryName](EMaterialExpressionConvertType ConvertType)
 	{
 		const int32 NumComponents = MaterialExpressionConvertType::GetComponentCount(ConvertType);
-		const FText ActionName = FText::Format(LOCTEXT("Convert_MakeNode_Name", "Make Float{0}"), NumComponents);
-		const FText ActionTooltip = FText::Format(LOCTEXT("Convert_MakeNode_Tooltip", "Adds a Make Float{0} Node"), NumComponents);
+		const FText ActionName = FText::Format(LOCTEXT("Convert_MakeNode_Name", "Make Vector{0}"), NumComponents);
+		const FText ActionTooltip = FText::Format(LOCTEXT("Convert_MakeNode_Tooltip", "Adds a Make Vector{0} Node"), NumComponents);
 
 		TSharedPtr<FMaterialGraphSchemaAction_NewNode> Action_MakeNode(new FMaterialGraphSchemaAction_NewNode(
 			CategoryName,
