Removed false positive ensure in LevelStreaming when reusing an existing in-memory level package is disabled (LevelStreaming.ShouldReuseUnloadedButStillAroundLevels = false) and we find one in ULevelStreaming::RequestLevel.

diff --git a/Engine/Source/Runtime/Engine/Private/LevelStreaming.cpp b/Engine/Source/Runtime/Engine/Private/LevelStreaming.cpp
--- a/Engine/Source/Runtime/Engine/Private/LevelStreaming.cpp
+++ b/Engine/Source/Runtime/Engine/Private/LevelStreaming.cpp
@@ -1694,12 +1694,20 @@ bool ULevelStreaming::RequestLevel(UWorld* PersistentWorld, bool bAllowLevelLoad
 				return false;
 			}
 #endif
-			if (!ensure(ULevelStreaming::ShouldReuseUnloadedButStillAroundLevels(World->PersistentLevel)))
+			// Here, if the level should not be re-used, trash the package.
+			// This can happen in the case where a streaming level is first removed inside the World Tick,
+			// FLevelStreamingGCHelper::RequestUnload will not call PrepareStreamedOutLevelForGC right away 
+			// and will process the level in FLevelStreamingGCHelper::OnWorldTickEnd.
+			// If we are insde a BlockTillLevelStreamingCompleted (still in the World Tick), the target
+			// state of the Level can change to be loaded by a subsystem like UWorldPartitionSubsystem::OnUpdateStreamingState.
+			// In this scenario, a request to load the level will happen and the level package will be found in memory.
+			if (!ULevelStreaming::ShouldReuseUnloadedButStillAroundLevels(World->PersistentLevel))
 			{
-				UE_LOG(LogLevelStreaming, Log, TEXT("Trying to reuse level %s but ShouldReuseUnloadedButStillAroundLevels returns false. Trash the current package and reload"), *LevelPackage->GetName())
+				UE_LOG(LogLevelStreaming, Verbose, TEXT("Trying to reuse level %s but ShouldReuseUnloadedButStillAroundLevels returns false. Trash the current package and reload."), *LevelPackage->GetName());
 				FLevelStreamingGCHelper::TrashPackage(LevelPackage);
 				return false;
 			}
+
 			if (World->PersistentLevel != LoadedLevel)
 			{
 				// Level already exists but may have the wrong type due to being inactive before, so copy data over
