Add monospace font to GEngine for Canvas use in rendering debug code.

#rnx

diff --git a/Engine/Source/Runtime/Engine/Classes/Engine/Engine.h b/Engine/Source/Runtime/Engine/Classes/Engine/Engine.h
--- a/Engine/Source/Runtime/Engine/Classes/Engine/Engine.h
+++ b/Engine/Source/Runtime/Engine/Classes/Engine/Engine.h
@@ -756,6 +756,15 @@ public:
 	UPROPERTY(globalconfig, EditAnywhere, Category=Fonts, meta=(AllowedClasses="/Script/Engine.Font", DisplayName="Subtitle Font", ConfigRestartRequired=true), AdvancedDisplay)
 	FSoftObjectPath SubtitleFontName;
 
+private:
+	UPROPERTY()
+	TObjectPtr<class UFont> MonospaceFont;
+
+public:
+	/** Sets the font used for monospace engine text, e.g. in debug displays */
+	UPROPERTY(globalconfig, EditAnywhere, Category=Fonts, meta=(AllowedClasses="/Script/Engine.Font", DisplayName="Monospace Font", ConfigRestartRequired=true), AdvancedDisplay)
+	FSoftObjectPath MonospaceFontName;
+
 private:
 	UPROPERTY()
 	TArray<TObjectPtr<class UFont>> AdditionalFonts;
@@ -3033,6 +3042,13 @@ public:
 	 */
 	static ENGINE_API class UFont* GetSubtitleFont();
 
+	/**
+	 * Gets the engine's default monospace font.
+	 *
+	 * @return Subtitle font.
+	 */
+	static ENGINE_API class UFont* GetMonospaceFont();
+
 	/**
 	 * Gets the specified additional font.
 	 *
diff --git a/Engine/Source/Runtime/Engine/Private/UnrealEngine.cpp b/Engine/Source/Runtime/Engine/Private/UnrealEngine.cpp
--- a/Engine/Source/Runtime/Engine/Private/UnrealEngine.cpp
+++ b/Engine/Source/Runtime/Engine/Private/UnrealEngine.cpp
@@ -3412,6 +3412,16 @@ void UEngine::InitializeObjectReferences()
 			return DefaultTinyFont.Get();
 		};
 
+		TStrongObjectPtr<UFont> DefaultMonoFont;
+		auto GetDefaultMonoFont = [&DefaultMonoFont, &CreateFontObjectFromDefaultFont]() -> UFont*
+		{
+			if (!DefaultMonoFont)
+			{
+				DefaultMonoFont.Reset(CreateFontObjectFromDefaultFont(FName("DefaultMonoFont"), FName("Mono"), 10));
+			}
+			return DefaultMonoFont.Get();
+		};
+
 		TStrongObjectPtr<UFont> DefaultRegularFont;
 		auto GetDefaultRegularFont = [&DefaultRegularFont, &CreateFontObjectFromDefaultFont]() -> UFont*
 		{
@@ -3443,6 +3453,7 @@ void UEngine::InitializeObjectReferences()
 		ConditionalLoadEngineFont(static_cast<UFont*&>(MediumFont), MediumFontName.ToString(), GetDefaultRegularFont);
 		ConditionalLoadEngineFont(static_cast<UFont*&>(LargeFont), LargeFontName.ToString(), GetDefaultRegularFont);
 		ConditionalLoadEngineFont(static_cast<UFont*&>(SubtitleFont), SubtitleFontName.ToString(), GetDefaultRegularFont);
+		ConditionalLoadEngineFont(static_cast<UFont*&>(MonospaceFont), MonospaceFontName.ToString(), GetDefaultMonoFont);
 
 		// Additional fonts.
 		AdditionalFonts.Empty(AdditionalFontNames.Num());
@@ -3832,6 +3843,14 @@ UFont* UEngine::GetSubtitleFont()
 	return GEngine->SubtitleFont;
 }
 
+/**
+* Returns the engine's default monospace font
+*/
+UFont* UEngine::GetMonospaceFont()
+{
+	return GEngine->MonospaceFont;
+}
+
 /**
 * Returns the specified additional font.
 *
