Suppress warning C4996 in InputSettings.h

#rnx

diff --git a/Engine/Source/Runtime/Engine/Classes/GameFramework/InputSettings.h b/Engine/Source/Runtime/Engine/Classes/GameFramework/InputSettings.h
--- a/Engine/Source/Runtime/Engine/Classes/GameFramework/InputSettings.h
+++ b/Engine/Source/Runtime/Engine/Classes/GameFramework/InputSettings.h
@@ -174,8 +174,10 @@ private:
 
 	/** List of Speech Mappings */
 	UE_DEPRECATED(5.7, "SpeechMappings are deprecated. These were only used by a now deprecated platform.")
+	PRAGMA_DISABLE_DEPRECATION_WARNINGS
 	UPROPERTY(config, EditAnywhere, Category = "Bindings")
 	TArray<struct FInputActionSpeechMapping> SpeechMappings;
+	PRAGMA_ENABLE_DEPRECATION_WARNINGS
 
 	/** Default class type for player input object. May be overridden by player controller. */
 	UPROPERTY(config, EditAnywhere, NoClear, Category = DefaultClasses)
