Fix anisotropy rendering with Fromat=0 when using Substrate node.

diff --git a/Engine/Source/Runtime/Engine/Private/Materials/HLSLMaterialTranslator.cpp b/Engine/Source/Runtime/Engine/Private/Materials/HLSLMaterialTranslator.cpp
--- a/Engine/Source/Runtime/Engine/Private/Materials/HLSLMaterialTranslator.cpp
+++ b/Engine/Source/Runtime/Engine/Private/Materials/HLSLMaterialTranslator.cpp
@@ -16353,7 +16353,7 @@ void FHLSLMaterialTranslator::PrepareEnvironmentDefines()
 	EnvironmentDefines->bMaterialIsSky = Material->IsSky();
 	EnvironmentDefines->bMaterialComputeFogPerPixel = Material->ComputeFogPerPixel();
 	EnvironmentDefines->bMaterialFullyRough = bIsFullyRough || Material->IsFullyRough();
-	EnvironmentDefines->bMaterialUsesAnisotropy = bUsesAnisotropy && FDataDrivenShaderPlatformInfo::GetSupportsAnisotropicMaterials(InPlatform);
+	EnvironmentDefines->bMaterialUsesAnisotropy = (bUsesAnisotropy || EnumHasAnyFlags(MaterialCompilationOutput.SubstrateMaterialCompilationOutput.SubstrateMaterialBsdfFeatures, ESubstrateBsdfFeature::Anisotropy)) && FDataDrivenShaderPlatformInfo::GetSupportsAnisotropicMaterials(InPlatform);
 	EnvironmentDefines->MaterialDecalReadMask = MaterialCompilationOutput.UsedDBufferTextures;
 	EnvironmentDefines->bMaterialUsesDecalLookup = MaterialCompilationOutput.bUsesDBufferTextureLookup;
 	EnvironmentDefines->MaterialPathTracingBufferRead = MaterialCompilationOutput.UsedPathTracingBufferTextures;
