[MIR] Fix VT output in new material translator.

Handle MCT_TextureVirtual type in MIR-to-HLSL translator and also handle texture types in uniform parameters for sampler instructions.

#rnx

diff --git a/Engine/Source/Runtime/Engine/Private/Materials/MaterialIRToHLSLTranslator.cpp b/Engine/Source/Runtime/Engine/Private/Materials/MaterialIRToHLSLTranslator.cpp
--- a/Engine/Source/Runtime/Engine/Private/Materials/MaterialIRToHLSLTranslator.cpp
+++ b/Engine/Source/Runtime/Engine/Private/Materials/MaterialIRToHLSLTranslator.cpp
@@ -712,7 +712,9 @@ struct FPrivate : FMaterialIRToHLSLTranslation
 	{
 		if (UniformParameter->Type.IsTexture2D())
 		{
-		 	EMaterialValueType TextureType = Module->GetParameterMetadata(UniformParameter->ParameterIdInModule).Value.Texture->GetMaterialType();
+			UTexture* TextureObject = GetTextureFromUniformParameter(UniformParameter);
+			check(TextureObject);
+			EMaterialValueType TextureType = TextureObject->GetMaterialType();
 			LowerTextureReference(TextureType, UniformParameter->Analysis_UniformParameterIndex);
 		}
 		else
@@ -1563,7 +1565,9 @@ struct FPrivate : FMaterialIRToHLSLTranslation
 		case MCT_Texture2DArray: Printer << TEXTVIEW("Texture2DArray"); break;
 		case MCT_TextureCubeArray: Printer << TEXTVIEW("TextureCubeArray"); break;
 		case MCT_VolumeTexture: Printer << (bForResourceDeclarations ? TEXTVIEW("VolumeTexture") : TEXTVIEW("Texture3D")); break;
-		case MCT_TextureExternal: Printer << TEXTVIEW("Texture2D");
+		case MCT_TextureExternal: Printer << TEXTVIEW("Texture2D"); break;
+		case MCT_TextureVirtual: Printer << TEXTVIEW("VirtualTexturePhysical"); break;
+		default: UE_MIR_UNREACHABLE();
 		}
 		return NoOp;
 	}
@@ -1574,24 +1578,35 @@ struct FPrivate : FMaterialIRToHLSLTranslation
 		LowerTextureMaterialType(TextureObject->Texture->GetMaterialType());
 	}
 
-	void LowerStandardTextureRead(const MIR::FTextureRead* TextureRead)
+	UTexture* GetTextureFromUniformParameter(const MIR::FUniformParameter* UniformParameter)
 	{
-		switch (TextureRead->TextureObject->Type.GetKind())
+		const FMaterialParameterValue& TextureParameter = Module->GetParameterMetadata(UniformParameter->ParameterIdInModule).Value;
+		switch (TextureParameter.Type)
 		{
-			case MIR::ETypeKind::Texture2D: 
+			case EMaterialParameterType::Texture:
 			{
-				const MIR::FTextureObject* TextureObject = TextureRead->TextureObject->As<MIR::FTextureObject>();
-				
-				if (!TextureObject)
-				{
-					Printer << TEXTVIEW("Texture2D");
-				}
-				else
-				{
-					LowerTextureType(TextureObject);
-				}
-				break;
+				return TextureParameter.Texture;
+			}
+			default:
+			{
+				UE_MIR_UNREACHABLE();
 			}
+		}
+	}
+
+	void LowerTextureParameterType(const MIR::FUniformParameter* UniformParameter)
+	{
+		UTexture* TextureObject = GetTextureFromUniformParameter(UniformParameter);
+		check(TextureObject);
+		LowerTextureMaterialType(TextureObject->GetMaterialType());
+	}
+
+	void LowerStandardTextureRead(const MIR::FTextureRead* TextureRead)
+	{
+		switch (TextureRead->TextureObject->Kind)
+		{
+			case MIR::EValueKind::VK_TextureObject: LowerTextureType(TextureRead->TextureObject->As<MIR::FTextureObject>()); break;
+			case MIR::EValueKind::VK_UniformParameter: LowerTextureParameterType(TextureRead->TextureObject->As<MIR::FUniformParameter>()); break;
 			default: UE_MIR_UNREACHABLE();
 		}
 
@@ -1626,11 +1641,6 @@ struct FPrivate : FMaterialIRToHLSLTranslation
 
 	void LowerVirtualTextureRead(const MIR::FTextureRead* TextureRead)
 	{
-		checkf(
-			TextureRead->TextureObject->Type.IsTexture2D(),
-			TEXT("Texture object type must be a 2D texture for virtual texture sampling in new material translator")
-		);
-
 		checkf(TextureRead->VTPageTable, TEXT("Missing page table for virtual texture read instruction"));
 		const MIR::FVTPageTableRead* VTPageTableResult = TextureRead->VTPageTable->As<MIR::FVTPageTableRead>();
 
@@ -1677,6 +1687,8 @@ struct FPrivate : FMaterialIRToHLSLTranslation
 			Printer << TEXTVIEW("(");
 		}
 
+		checkf(TextureRead->TextureObject->Type.IsTexture2D(), TEXT("Invalid texture object type"));
+
 		if (IsVirtualSamplerType(TextureRead->SamplerType))
 		{
 			LowerVirtualTextureRead(TextureRead);
