Add mutators to UNetworkMetricsDatabase functionality.
Mutators can be added to UNetworkMetricsBaseListener and get a chance to process input metrics every frame. They report an output metric to their owning listener on the listener's reporting interval.
Use mutators to aggregate metrics that are set more frequently than a listener's reporting interval. Basic mutators for min, max, and average are provided.

Also add server-side network metrics for RawPing values across all client connections.

diff --git a/Engine/Source/Runtime/Engine/Private/Net/NetworkMetricsDatabase.cpp b/Engine/Source/Runtime/Engine/Private/Net/NetworkMetricsDatabase.cpp
--- a/Engine/Source/Runtime/Engine/Private/Net/NetworkMetricsDatabase.cpp
+++ b/Engine/Source/Runtime/Engine/Private/Net/NetworkMetricsDatabase.cpp
@@ -1,6 +1,7 @@
 // Copyright Epic Games, Inc. All Rights Reserved.
 
 #include "Net/NetworkMetricsDatabase.h"
+#include "Net/NetworkMetricsMutators.h"
 #include "EngineStats.h"
 #include "Logging/LogMacros.h"
 #include "ProfilingDebugging/CsvProfiler.h"
@@ -175,6 +176,26 @@ bool UNetworkMetricsDatabase::Contains(const FName MetricName) const
 	return MetricTypes.Contains(MetricName);
 }
 
+void UNetworkMetricsDatabase::AddMetricToSnapshot(UE::Net::FNetworkMetricSnapshot& Snapshot, const FName MetricName, const EMetricType MetricType)
+{
+	if (MetricType == EMetricType::Integer)
+	{
+		const UE::Net::FNetworkMetric<int64>* Metric = MetricInts.Find(MetricName);
+		if (ensureMsgf(Metric, TEXT("Unable to find metric %s when reporting to listeners."), *MetricName.ToString()))
+		{
+			Snapshot.MetricInts.Add(*Metric);
+		}
+	}
+	else if (MetricType == EMetricType::Float)
+	{
+		const UE::Net::FNetworkMetric<float>* Metric = MetricFloats.Find(MetricName);
+		if (ensureMsgf(Metric, TEXT("Unable to find metric %s when reporting to listeners."), *MetricName.ToString()))
+		{
+			Snapshot.MetricFloats.Add(*Metric);
+		}
+	}
+}
+
 void UNetworkMetricsDatabase::ProcessListeners()
 {
 	double CurrentTimeSeconds = FPlatformTime::Seconds();
@@ -190,6 +211,27 @@ void UNetworkMetricsDatabase::ProcessListeners()
 
 		if (ensure(Listener))
 		{
+			// For mutator listeners generate & process snapshots every frame. Do this before calling Report() on listeners
+			// so that they can take this frame's data into account.
+			for (TInstancedStruct<FNetworkMetricsMutator>& Mutator : Listener->GetMutators())
+			{
+				if (ensure(Mutator.IsValid()))
+				{
+					Snapshot.Reset();
+
+					for (const FName MetricName : Mutator.Get().GetAllMetricNames())
+					{
+						EMetricType* MetricType = MetricTypes.Find(MetricName);
+						if (MetricType)
+						{
+							AddMetricToSnapshot(Snapshot, MetricName, *MetricType);
+						}
+					}
+
+					Mutator.GetMutable().ProcessFrame(Snapshot);
+				}
+			}
+
 			if (!HasReportIntervalPassed(CurrentTimeSeconds, Listener))
 			{
 				continue;
@@ -202,20 +244,32 @@ void UNetworkMetricsDatabase::ProcessListeners()
 				const FName& MetricName = NameAndType.Key;
 				const EMetricType& MetricType = NameAndType.Value;
 
-				if (MetricType == EMetricType::Integer)
+				AddMetricToSnapshot(Snapshot, MetricName, MetricType);
+			}
+
+			// Add mutator values to the snapshot
+			for (TInstancedStruct<FNetworkMetricsMutator>& Mutator : Listener->GetMutators())
+			{
+				if (ensure(Mutator.IsValid()))
 				{
-					const UE::Net::FNetworkMetric<int64>* Metric = MetricInts.Find(MetricName);
-					if (ensureMsgf(Metric, TEXT("Unable to find metric %s when reporting to listeners."), *MetricName.ToString()))
+					FNetworkMetricsMutator::FValueVariant MutatorVariant = Mutator.Get().GetValue();
+					if (MutatorVariant.IsType<int64>())
 					{
-						Snapshot.MetricInts.Add(*Metric);
+						UE::Net::FNetworkMetric<int64> MutatorInt;
+						MutatorInt.Name = Mutator.Get().GetName();
+						MutatorInt.Value = MutatorVariant.Get<int64>();
+						Snapshot.MetricInts.Add(MutatorInt);
 					}
-				}
-				else if (MetricType == EMetricType::Float)
-				{
-					const UE::Net::FNetworkMetric<float>* Metric = MetricFloats.Find(MetricName);
-					if (ensureMsgf(Metric, TEXT("Unable to find metric %s when reporting to listeners."), *MetricName.ToString()))
+					else if (MutatorVariant.IsType<float>())
+					{
+						UE::Net::FNetworkMetric<float> MutatorFloat;
+						MutatorFloat.Name = Mutator.Get().GetName();
+						MutatorFloat.Value = MutatorVariant.Get<float>();
+						Snapshot.MetricFloats.Add(MutatorFloat);
+					}
+					else
 					{
-						Snapshot.MetricFloats.Add(*Metric);
+						ensureMsgf(false, TEXT("Mutator %s returned value of unknown type."), *Mutator.Get().GetName().ToString());
 					}
 				}
 			}
@@ -241,16 +295,91 @@ void UNetworkMetricsDatabase::ResetListeners()
 
 void UNetworkMetricsDatabase::Register(const FName MetricName, TWeakObjectPtr<UNetworkMetricsBaseListener> Reporter)
 {
+	if (!ensure(Reporter.IsValid()))
+	{
+		UE_LOG(LogNetworkMetrics, Warning, TEXT("UNetworkMetricsDatabase::Register: Cannot register metric %s with invalid listener."), *MetricName.ToString());
+		return;
+	}
+
 	const EMetricType* MetricType = MetricTypes.Find(MetricName);
 
 	if (ensureMsgf(MetricType, TEXT("Cannot find metric %s to register listener."), *MetricName.ToString()))
 	{
+		// Make sure the metric name doesn't collide with a mutator name that's already registered to the listener.
+		for (TInstancedStruct<FNetworkMetricsMutator>& Mutator : Reporter->GetMutators())
+		{
+			if (ensure(Mutator.IsValid()))
+			{
+				if (Mutator.Get().GetName() == MetricName)
+				{
+					UE_LOG(LogNetworkMetrics, Warning, TEXT("UNetworkMetricsDatabase::Register: Cannot register metric %s whose name collides with a registered mutator for listener %s."), *MetricName.ToString(), *Reporter->GetClass()->GetName());
+					ensureMsgf(false, TEXT("UNetworkMetricsDatabase::Register: Cannot register metric %s whose name collides with a registered mutator for listener %s."), *MetricName.ToString(), *Reporter->GetClass()->GetName());
+					return;
+				}
+			}
+		}
+
 		TSet<FNameAndType>& ListenerMetrics = ListenersToMetrics.FindOrAdd(Reporter);
 		ListenerMetrics.Add(FNameAndType(MetricName, *MetricType));
 		LastReportListener.Add(Reporter, 0.0);
 	}
 }
 
+void UNetworkMetricsDatabase::RegisterMutator(TWeakObjectPtr<UNetworkMetricsBaseListener> Listener, const TInstancedStruct<FNetworkMetricsMutator>& Mutator)
+{
+	if (!ensure(Listener.IsValid()))
+	{
+		UE_LOG(LogNetworkMetrics, Warning, TEXT("UNetworkMetricsDatabase::RegisterMutator: Cannot register mutator %s with invalid listener."), *Mutator.Get().GetName().ToString());
+		return;
+	}
+
+	for (const FName MutatorMetricName : Mutator.Get().GetAllMetricNames())
+	{
+		if (!Contains(MutatorMetricName))
+		{
+			UE_LOG(LogNetworkMetrics, Warning, TEXT("UNetworkMetricsDatabase::RegisterMutator: Cannot register mutator %s with an unknown metric '%s'."), *Mutator.Get().GetName().ToString(), *MutatorMetricName.ToString());
+			return;
+		}
+	}
+
+	// Make sure the mutator name doesn't collide with a mutator name that's already registered to the listener.
+	for (TInstancedStruct<FNetworkMetricsMutator>& ExistingMutator : Listener->GetMutators())
+	{
+		if (ensure(ExistingMutator.IsValid()))
+		{
+			if (ExistingMutator.Get().GetName() == Mutator.Get().GetName())
+			{
+				UE_LOG(LogNetworkMetrics, Warning, TEXT("UNetworkMetricsDatabase::RegisterMutator: Cannot register mutator %s whose name collides with an existing mutator for listener %s."), *Mutator.Get().GetName().ToString(), *Listener->GetClass()->GetName());
+				ensureMsgf(false, TEXT("UNetworkMetricsDatabase::RegisterMutator: Cannot register mutator %s whose name collides with an existing mutator for listener %s."), *Mutator.Get().GetName().ToString(), *Listener->GetClass()->GetName());
+				return;
+			}
+		}
+	}
+
+	// Make sure the mutator name doesn't collide with a metric name that's already registered to the listener.
+	TSet<FNameAndType>* ListenerMetrics = ListenersToMetrics.Find(Listener);
+	if (ListenerMetrics)
+	{
+		for (const FNameAndType& NameAndType : *ListenerMetrics)
+		{
+			if (Mutator.Get().GetName() == NameAndType.Key)
+			{
+				UE_LOG(LogNetworkMetrics, Warning, TEXT("UNetworkMetricsDatabase::RegisterMutator: Cannot register mutator %s whose name collides with a registered metric for listener %s."), *Mutator.Get().GetName().ToString(), *Listener->GetClass()->GetName());
+				ensureMsgf(false, TEXT("UNetworkMetricsDatabase::RegisterMutator: Cannot register mutator %s whose name collides with a registered metric for listener %s."), *Mutator.Get().GetName().ToString(), *Listener->GetClass()->GetName());
+				return;
+			}
+		}
+	}
+	else
+	{
+		// Add entries to the listener maps so this listener is processed even if it doesn't have any metrics directly registered.
+		ListenerMetrics = &ListenersToMetrics.Add(Listener);
+		LastReportListener.Add(Listener, 0.0);
+	}
+
+	Listener->GetMutators().Add(Mutator);
+}
+
 bool UNetworkMetricsDatabase::HasReportIntervalPassed(double CurrentTimeSeconds, UNetworkMetricsBaseListener* Listener)
 {
 	if (ensureMsgf(Listener->GetInterval() >= 0, TEXT("Listener has a negative reporting time interval.")))
@@ -293,7 +422,12 @@ void UNetworkMetricsLog::Report(const UE::Net::FNetworkMetricSnapshot& Snapshot)
 		MetricsStr += FString::Printf(TEXT("%s=%2.2f "), *Metric.Name.ToString(), Metric.Value);
 	}
 
-	UE_LOG(LogNetworkMetrics, Log, TEXT("Network Metrics: %s"), *MetricsStr);
+	UE_CLOG(!MetricsStr.IsEmpty(), LogNetworkMetrics, Log, TEXT("Network Metrics: %s"), *MetricsStr);
+
+	for (TInstancedStruct<FNetworkMetricsMutator>& Mutator : GetMutators())
+	{
+		Mutator.GetMutable().Reset();
+	}
 }
 
 UNetworkMetricsCSV::UNetworkMetricsCSV() :
diff --git a/Engine/Source/Runtime/Engine/Private/Net/NetworkMetricsDefs.cpp b/Engine/Source/Runtime/Engine/Private/Net/NetworkMetricsDefs.cpp
--- a/Engine/Source/Runtime/Engine/Private/Net/NetworkMetricsDefs.cpp
+++ b/Engine/Source/Runtime/Engine/Private/Net/NetworkMetricsDefs.cpp
@@ -101,6 +101,9 @@ const FName PingBucketInt7("PingBucketInt7");
 const FName OutgoingReliableMessageQueueMaxSize("OutgoingReliableMessageQueueMaxSize");
 const FName IncomingReliableMessageQueueMaxSize("IncomingReliableMessageQueueMaxSize");
 const FName RawPing("RawPing");
+const FName RawPingClientMin("RawPingClientMin");
+const FName RawPingClientMax("RawPingClientMax");
+const FName RawPingClientAvg("RawPingClientAvg");
 
 }
 
diff --git a/Engine/Source/Runtime/Engine/Private/Net/NetworkMetricsMutators.cpp b/Engine/Source/Runtime/Engine/Private/Net/NetworkMetricsMutators.cpp
new file mode 100644
--- /dev/null
+++ b/Engine/Source/Runtime/Engine/Private/Net/NetworkMetricsMutators.cpp
@@ -0,0 +1,272 @@
+// Copyright Epic Games, Inc. All Rights Reserved.
+
+#include "Net/NetworkMetricsMutators.h"
+#include "Net/NetworkMetricsDatabase.h"
+#include "Net/NetworkMetricsDefs.h"
+#include "EngineStats.h"
+
+#include UE_INLINE_GENERATED_CPP_BY_NAME(NetworkMetricsMutators)
+
+FName FNetworkMetricsMutator::GetName() const
+{
+	if (!NameOverride.IsNone())
+	{
+		return NameOverride;
+	}
+
+	return GetNameInternal();
+}
+
+FName FNetworkMetricsMutatorAvg::GetNameInternal() const
+{
+	// Lazily init the name because there's no good initialization hook after configs are loaded and MetricName is set.
+	if (MutatorName.IsNone())
+	{
+		MutatorName = FName(MetricName.GetPlainNameString() + TEXT("Avg"));
+	}
+
+	return MutatorName;
+}
+
+void FNetworkMetricsMutatorAvg::ProcessFrame(const UE::Net::FNetworkMetricSnapshot& Snapshot)
+{
+	for (const UE::Net::FNetworkMetric<int64>& Metric : Snapshot.MetricInts)
+	{
+		AddIntSample(Metric.Value);
+	}
+
+	for (const UE::Net::FNetworkMetric<float>& Metric : Snapshot.MetricFloats)
+	{
+		AddFloatSample(Metric.Value);
+	}
+}
+
+void FNetworkMetricsMutatorAvg::Reset()
+{
+	IntTotal = 0;
+	FloatTotal = 0.0f;
+	NumSamples = 0;
+}
+
+FNetworkMetricsMutator::FValueVariant FNetworkMetricsMutatorAvg::GetValue() const
+{
+	if (NumSamples != 0)
+	{
+		return FValueVariant(TInPlaceType<float>(), (static_cast<float>(IntTotal) + FloatTotal) / NumSamples);
+	}
+
+	return FValueVariant(TInPlaceType<float>(), 0.0f);
+}
+
+void FNetworkMetricsMutatorAvg::AddIntSample(int64 Sample)
+{
+	IntTotal += Sample;
+	++NumSamples;
+}
+
+void FNetworkMetricsMutatorAvg::AddFloatSample(float Sample)
+{
+	FloatTotal += Sample;
+	++NumSamples;
+}
+
+FName FNetworkMetricsMutatorMin::GetNameInternal() const
+{
+	// Lazily init the name because there's no good initialization hook after configs are loaded and MetricName is set.
+	if (MutatorName.IsNone())
+	{
+		MutatorName = FName(MetricName.GetPlainNameString() + TEXT("Min"));
+	}
+
+	return MutatorName;
+}
+
+void FNetworkMetricsMutatorMin::ProcessFrame(const UE::Net::FNetworkMetricSnapshot& Snapshot)
+{
+	// This mutator only supports a single input metric so we can assume it will be the only one in the snapshot.
+	ensureMsgf(Snapshot.MetricInts.Num() + Snapshot.MetricFloats.Num() == 1, TEXT("FNetworkMetricsMutatorMin::ProcessFrame expects a single metric"));
+
+	for (const UE::Net::FNetworkMetric<int64>& Metric : Snapshot.MetricInts)
+	{
+		if (MinValue.IsType<int64>())
+		{
+			MinValue.Set<int64>(FMath::Min(Metric.Value, MinValue.Get<int64>()));
+		}
+		else
+		{
+			MinValue.Set<int64>(FMath::Min(Metric.Value, std::numeric_limits<int64>::max()));
+		}
+	}
+
+	for (const UE::Net::FNetworkMetric<float>& Metric : Snapshot.MetricFloats)
+	{
+		if (MinValue.IsType<float>())
+		{
+			MinValue.Set<float>(FMath::Min(Metric.Value, MinValue.Get<float>()));
+		}
+		else
+		{
+			MinValue.Set<float>(FMath::Min(Metric.Value, std::numeric_limits<float>::max()));
+		}
+	}
+}
+
+void FNetworkMetricsMutatorMin::Reset()
+{
+	MinValue.Set<int64>(std::numeric_limits<int64>::max());
+}
+
+FNetworkMetricsMutator::FValueVariant FNetworkMetricsMutatorMin::GetValue() const
+{
+	return MinValue;
+}
+
+FName FNetworkMetricsMutatorMax::GetNameInternal() const
+{
+	// Lazily init the name because there's no good initialization hook after configs are loaded and MetricName is set.
+	if (MutatorName.IsNone())
+	{
+		MutatorName = FName(MetricName.GetPlainNameString() + TEXT("Max"));
+	}
+
+	return MutatorName;
+}
+
+void FNetworkMetricsMutatorMax::ProcessFrame(const UE::Net::FNetworkMetricSnapshot& Snapshot)
+{
+	// This mutator only supports a single input metric so we can assume it will be the only one in the snapshot.
+	ensureMsgf(Snapshot.MetricInts.Num() + Snapshot.MetricFloats.Num() == 1, TEXT("FNetworkMetricsMutatorMax::ProcessFrame expects a single metric"));
+
+	for (const UE::Net::FNetworkMetric<int64>& Metric : Snapshot.MetricInts)
+	{
+		if (MaxValue.IsType<int64>())
+		{
+			MaxValue.Set<int64>(FMath::Max(Metric.Value, MaxValue.Get<int64>()));
+		}
+		else
+		{
+			MaxValue.Set<int64>(FMath::Max(Metric.Value, std::numeric_limits<int64>::lowest()));
+		}
+	}
+
+	for (const UE::Net::FNetworkMetric<float>& Metric : Snapshot.MetricFloats)
+	{
+		if (MaxValue.IsType<float>())
+		{
+			MaxValue.Set<float>(FMath::Max(Metric.Value, MaxValue.Get<float>()));
+		}
+		else
+		{
+			MaxValue.Set<float>(FMath::Max(Metric.Value, std::numeric_limits<float>::lowest()));
+		}
+	}
+}
+
+void FNetworkMetricsMutatorMax::Reset()
+{
+	MaxValue.Set<int64>(std::numeric_limits<int64>::lowest());
+}
+
+FNetworkMetricsMutator::FValueVariant FNetworkMetricsMutatorMax::GetValue() const
+{
+	return MaxValue;
+}
+
+TArrayView<const FName> FNetworkMetricsMutatorPercent::GetAllMetricNames() const
+{
+	if (SourceMetricNames.IsEmpty())
+	{
+		SourceMetricNames.Add(NumeratorName);
+		SourceMetricNames.Add(DenominatorName);
+	}
+
+	return MakeArrayView(SourceMetricNames);
+}
+
+FName FNetworkMetricsMutatorPercent::GetNameInternal() const
+{
+	// Lazily init the name because there's no good initialization hook after configs are loaded and MetricName is set.
+	if (MutatorName.IsNone())
+	{
+		MutatorName = FName(NumeratorName.GetPlainNameString() + DenominatorName.GetPlainNameString() + TEXT("Pct"));
+	}
+
+	return MutatorName;
+}
+
+template<class T>
+void FNetworkMetricsMutatorPercent::UpdateFromSnapshotValue(FValueVariant& InOutValue, T MetricValue)
+{
+	if (InOutValue.IsType<T>())
+	{
+		InOutValue.Get<T>() += MetricValue;
+	}
+	else
+	{
+		InOutValue.Set<T>(MetricValue);
+	}
+}
+
+void FNetworkMetricsMutatorPercent::ProcessFrame(const UE::Net::FNetworkMetricSnapshot& Snapshot)
+{
+	// This function assumes metrics with the same name won't change types at runtime
+	for (const UE::Net::FNetworkMetric<int64>& Metric : Snapshot.MetricInts)
+	{
+		if (Metric.Name == NumeratorName)
+		{
+			UpdateFromSnapshotValue<int64>(Numerator, Metric.Value);
+		}
+
+		if (Metric.Name == DenominatorName)
+		{
+			UpdateFromSnapshotValue<int64>(Denominator, Metric.Value);
+		}
+	}
+
+	for (const UE::Net::FNetworkMetric<float>& Metric : Snapshot.MetricFloats)
+	{
+		if (Metric.Name == NumeratorName)
+		{
+			UpdateFromSnapshotValue<float>(Numerator, Metric.Value);
+		}
+
+		if (Metric.Name == DenominatorName)
+		{
+			UpdateFromSnapshotValue<float>(Denominator, Metric.Value);
+		}
+	}
+}
+
+void FNetworkMetricsMutatorPercent::Reset()
+{
+	Numerator = FValueVariant(TInPlaceType<int64>(), 0);
+	Denominator = FValueVariant(TInPlaceType<int64>(), 0);
+}
+
+FNetworkMetricsMutator::FValueVariant FNetworkMetricsMutatorPercent::GetValue() const
+{
+	const float FloatDenominator = Denominator.IsType<float>() ? Denominator.Get<float>() : static_cast<float>(Denominator.Get<int64>());
+
+	if (FMath::IsNearlyZero(FloatDenominator))
+	{
+		return FValueVariant(TInPlaceType<float>(), 0.0f);
+	}
+
+	const float FloatNumerator = Numerator.IsType<float>() ? Numerator.Get<float>() : static_cast<float>(Numerator.Get<int64>());
+
+	return FValueVariant(TInPlaceType<float>(), (FloatNumerator / FloatDenominator) * 100.0f);
+}
+
+FNetworkMetricsMutatorOutPacketLoss::FNetworkMetricsMutatorOutPacketLoss()
+	: Super()
+{
+	SetNumeratorMetricName(UE::Net::Metric::OutLostPacketsFoundPerFrame);
+	SetDenominatorMetricName(UE::Net::Metric::OutPacketsPerFrame);
+}
+
+FNetworkMetricsMutatorInPacketLoss::FNetworkMetricsMutatorInPacketLoss()
+	: Super()
+{
+	SetNumeratorMetricName(UE::Net::Metric::InLostPacketsFoundPerFrame);
+	SetDenominatorMetricName(UE::Net::Metric::InPacketsPerFrame);
+}
diff --git a/Engine/Source/Runtime/Engine/Private/NetDriver.cpp b/Engine/Source/Runtime/Engine/Private/NetDriver.cpp
--- a/Engine/Source/Runtime/Engine/Private/NetDriver.cpp
+++ b/Engine/Source/Runtime/Engine/Private/NetDriver.cpp
@@ -50,6 +50,7 @@
 #include "Net/NetworkMetricsDatabase.h"
 #include "Net/NetworkMetricsDefs.h"
 #include "Net/NetworkMetricsConfig.h"
+#include "Net/NetworkMetricsMutators.h"
 #include "Net/NetSubObjectRegistryGetter.h"
 #include "Net/NetworkGranularMemoryLogging.h"
 #include "Templates/Greater.h"
@@ -2014,6 +2015,15 @@ void UNetDriver::SetupNetworkMetrics()
 
 	// (Client only) The highest maximum ping processed every frame. If no pings are processed in the frame it keeps the last value set.
 	GetMetrics()->CreateInt(UE::Net::Metric::RawPing, 0);
+
+	// (Server only) The minimum RawPing from among all client connections each frame.
+	GetMetrics()->CreateInt(UE::Net::Metric::RawPingClientMin, 0);
+
+	// (Server only) The maximum RawPing from among all client connections each frame.
+	GetMetrics()->CreateInt(UE::Net::Metric::RawPingClientMax, 0);
+
+	// (Server only) The average RawPing from among all client connections each frame.
+	GetMetrics()->CreateFloat(UE::Net::Metric::RawPingClientAvg, 0.0f);
 }
 
 void UNetDriver::SetupNetworkMetricsListeners(bool bIsClient)
@@ -2082,14 +2092,25 @@ void UNetDriver::SetupNetworkMetricsListeners(bool bIsClient)
 			Listener = &NetworkMetricsListeners.Add(ListenerClassName, NewObject<UNetworkMetricsBaseListener>(this, ListenerClass));
 		}
 
-		if (!GetMetrics()->Contains(MetricConfig.MetricName))
+		if (MetricConfig.MetricName != NAME_None)
 		{
-			UE_LOG(LogNet, Warning, TEXT("NetDriver::SetupNetworkMetricListeners: Cannot register metric listener %s to an unknown metric '%s'."), *MetricConfig.Class.ToString(), *MetricConfig.MetricName.ToString());
-			continue;
+			if (!GetMetrics()->Contains(MetricConfig.MetricName))
+			{
+				UE_LOG(LogNet, Warning, TEXT("NetDriver::SetupNetworkMetricListeners: Cannot register metric listener %s to an unknown metric '%s'."), *MetricConfig.Class.ToString(), *MetricConfig.MetricName.ToString());
+				continue;
+			}
+
+			GetMetrics()->Register(MetricConfig.MetricName, Listener->Get());
+
+			UE_LOG(LogNet, Verbose, TEXT("Registering network metrics listener %s for metric %s"), *ListenerClassName.ToString(), *MetricConfig.MetricName.ToString());
 		}
 
-		GetMetrics()->Register(MetricConfig.MetricName, Listener->Get());
-		UE_LOG(LogNet, Verbose, TEXT("Registering network metrics listener %s for metric %s"), *ListenerClassName.ToString(), *MetricConfig.MetricName.ToString());
+		if (MetricConfig.Mutator.IsValid())
+		{
+			GetMetrics()->RegisterMutator(Listener->Get(), MetricConfig.Mutator);
+
+			UE_LOG(LogNet, Verbose, TEXT("%s: registering network metrics mutator %s with listener %s"), *NetDriverName.ToString(), *MetricConfig.Mutator.Get().GetName().ToString(), *ListenerClassName.ToString());
+		}
 	}
 
 	// The listeners for the Stats system are hardcoded here because Stats is always expected to be running and some/all of the metrics should
@@ -8731,7 +8752,7 @@ void UNetDriver::UpdateNetworkStats()
 		}
 		else
 		{
-#if CSV_PROFILER_STATS
+
 			// CSV stats need to be collected every frame
 
 			int32 ClientsInPacketsThisFrameAvg = 0;
@@ -8739,6 +8760,11 @@ void UNetDriver::UpdateNetworkStats()
 			int32 ClientsOutPacketsThisFrameAvg = 0;
 			int32 ClientsOutPacketsThisFrameMax = 0;
 
+			// Collect raw pings from all clients
+			int32 ClientsRawPingThisFrameMin = std::numeric_limits<int32>::max();
+			int32 ClientsRawPingThisFrameMax = 0;
+			float ClientsRawPingThisFrameAvg = 0;
+
 			int32 NumClients = 0;
 
 			for (const UNetConnection* Client : ClientConnections)
@@ -8747,10 +8773,17 @@ void UNetDriver::UpdateNetworkStats()
 				{
 					++NumClients;
 
+#if CSV_PROFILER_STATS
 					ClientsInPacketsThisFrameAvg += Client->InPacketsThisFrame;
 					ClientsOutPacketsThisFrameAvg += Client->OutPacketsThisFrame;
 					ClientsInPacketsThisFrameMax = FMath::Max(ClientsInPacketsThisFrameMax, Client->InPacketsThisFrame);
 					ClientsOutPacketsThisFrameMax = FMath::Max(ClientsOutPacketsThisFrameMax, Client->OutPacketsThisFrame);
+#endif  //#if CSV_PROFILER_STATS
+
+					const float RawPingValueInMS = Client->RawPingInSeconds * 1000.f;
+					ClientsRawPingThisFrameMin = FMath::Min(ClientsRawPingThisFrameMin, static_cast<int32>(RawPingValueInMS));
+					ClientsRawPingThisFrameMax = FMath::Max(ClientsRawPingThisFrameMax, static_cast<int32>(RawPingValueInMS));
+					ClientsRawPingThisFrameAvg += RawPingValueInMS;
 				}
 			}
 
@@ -8758,8 +8791,11 @@ void UNetDriver::UpdateNetworkStats()
 			{
 				ClientsInPacketsThisFrameAvg /= NumClients;
 				ClientsOutPacketsThisFrameAvg /= NumClients;
+
+				ClientsRawPingThisFrameAvg /= NumClients;
 			}
 
+#if CSV_PROFILER_STATS
 			if (ServerConnection)
 			{
 				const int32 RawPingValueInMS = (int32)(ServerConnection->RawPingInSeconds * 1000.f);
@@ -8771,6 +8807,13 @@ void UNetDriver::UpdateNetworkStats()
 			GetMetrics()->SetInt(UE::Net::Metric::OutPacketsClientAvg, ClientsOutPacketsThisFrameAvg);
 			GetMetrics()->SetInt(UE::Net::Metric::OutPacketsClientMax, ClientsOutPacketsThisFrameMax);
 #endif  //#if CSV_PROFILER_STATS
+
+			if (!ClientConnections.IsEmpty())
+			{
+				GetMetrics()->SetInt(UE::Net::Metric::RawPingClientMin, ClientsRawPingThisFrameMin);
+				GetMetrics()->SetInt(UE::Net::Metric::RawPingClientMax, ClientsRawPingThisFrameMax);
+				GetMetrics()->SetFloat(UE::Net::Metric::RawPingClientAvg, ClientsRawPingThisFrameAvg);
+			}
 		}
 	} // bCollectNetStats ||(USE_SERVER_PERF_COUNTERS) || STATS
 
diff --git a/Engine/Source/Runtime/Engine/Public/Net/NetworkMetricsConfig.h b/Engine/Source/Runtime/Engine/Public/Net/NetworkMetricsConfig.h
--- a/Engine/Source/Runtime/Engine/Public/Net/NetworkMetricsConfig.h
+++ b/Engine/Source/Runtime/Engine/Public/Net/NetworkMetricsConfig.h
@@ -2,8 +2,11 @@
 #pragma once
 
 #include "CoreMinimal.h"
+#include "StructUtils/InstancedStruct.h"
 #include "NetworkMetricsConfig.generated.h"
 
+struct FNetworkMetricsMutator;
+
 UENUM()
 enum class ENetworkMetricEnableMode : uint8
 {
@@ -20,10 +23,14 @@ struct FNetworkMetricConfig
 	GENERATED_BODY()
 
 public:
-	/** The name of the metric to register the listener. */
+	/** The name of the metric to register the listener. Optional if a Mutator is specified. */
 	UPROPERTY()
 	FName MetricName;
 
+	/** Mutator to add to the listener. Optional if a MetricName is specified. */
+	UPROPERTY()
+	TInstancedStruct<FNetworkMetricsMutator> Mutator;
+
 	/** A sub-class of UNetworkMetricBaseListener. */
 	UPROPERTY()
 	TSoftClassPtr<class UNetworkMetricsBaseListener> Class;
diff --git a/Engine/Source/Runtime/Engine/Public/Net/NetworkMetricsDatabase.h b/Engine/Source/Runtime/Engine/Public/Net/NetworkMetricsDatabase.h
--- a/Engine/Source/Runtime/Engine/Public/Net/NetworkMetricsDatabase.h
+++ b/Engine/Source/Runtime/Engine/Public/Net/NetworkMetricsDatabase.h
@@ -3,10 +3,13 @@
 #pragma once
 
 #include "CoreMinimal.h"
+#include "StructUtils/InstancedStruct.h"
 #include "NetworkMetricsDatabase.generated.h"
 
 #define UE_API ENGINE_API
 
+struct FNetworkMetricsMutator;
+
 namespace UE::Net
 {
 
@@ -72,12 +75,17 @@ public:
 	/* Register a listener to be called for a given metric. */
 	UE_API void Register(const FName MetricName, TWeakObjectPtr<UNetworkMetricsBaseListener> Reporter);
 
+	/* Register a mutator with a listener. */
+	UE_API void RegisterMutator(TWeakObjectPtr<UNetworkMetricsBaseListener> Listener, const TInstancedStruct<FNetworkMetricsMutator>& Mutator);
+
 private:
-	/* Return true if the report interval time for a listener has elapsed. */
+	/* Return true if the report interval time for a listener has elapsed, and update its last reported time if it has. */
 	bool HasReportIntervalPassed(double CurrentTimeSeconds, UNetworkMetricsBaseListener* Listener);
 
 	enum class EMetricType { Integer, Float };
 
+	void AddMetricToSnapshot(UE::Net::FNetworkMetricSnapshot& Snapshot, const FName MetricName, const EMetricType MetricType);
+
 	TMap<FName, EMetricType> MetricTypes;
 
 	TMap<FName, UE::Net::FNetworkMetric<int64>> MetricInts;
@@ -88,6 +96,7 @@ private:
 	/* The time, in seconds, metrics were reported to a listener. */
 	TMap<TWeakObjectPtr<UNetworkMetricsBaseListener>, double> LastReportListener;
 
+	// Map of listeners to metrics registered with them. A listener may have an entry with an empty set if it has no metrics but some mutators registered. */
 	TMap<TWeakObjectPtr<UNetworkMetricsBaseListener>, TSet<FNameAndType>> ListenersToMetrics;
 };
 
@@ -164,11 +173,25 @@ public:
 		return IntervalSeconds;
 	}
 
+	const TArray<TInstancedStruct<FNetworkMetricsMutator>>& GetMutators() const
+	{
+		return Mutators;
+	}
+
+	TArray<TInstancedStruct<FNetworkMetricsMutator>>& GetMutators()
+	{
+		return Mutators;
+	}
+
 	virtual void Report(const UE::Net::FNetworkMetricSnapshot& Snapshot) {};
 
 private:
 	UPROPERTY(Config)
 	double IntervalSeconds;
+
+	/** List of optional mutators for this listener. Configure them via UNetworkMetricsConfig. */
+	UPROPERTY()
+	TArray<TInstancedStruct<FNetworkMetricsMutator>> Mutators;
 };
 
 /** 
diff --git a/Engine/Source/Runtime/Engine/Public/Net/NetworkMetricsDefs.h b/Engine/Source/Runtime/Engine/Public/Net/NetworkMetricsDefs.h
--- a/Engine/Source/Runtime/Engine/Public/Net/NetworkMetricsDefs.h
+++ b/Engine/Source/Runtime/Engine/Public/Net/NetworkMetricsDefs.h
@@ -103,4 +103,7 @@ extern ENGINE_API const FName PingBucketInt7;
 extern ENGINE_API const FName OutgoingReliableMessageQueueMaxSize;
 extern ENGINE_API const FName IncomingReliableMessageQueueMaxSize;
 extern ENGINE_API const FName RawPing;
+extern ENGINE_API const FName RawPingClientMin;
+extern ENGINE_API const FName RawPingClientMax;
+extern ENGINE_API const FName RawPingClientAvg;
 }
diff --git a/Engine/Source/Runtime/Engine/Public/Net/NetworkMetricsMutators.h b/Engine/Source/Runtime/Engine/Public/Net/NetworkMetricsMutators.h
new file mode 100644
--- /dev/null
+++ b/Engine/Source/Runtime/Engine/Public/Net/NetworkMetricsMutators.h
@@ -0,0 +1,249 @@
+// Copyright Epic Games, Inc. All Rights Reserved.
+
+#pragma once
+
+#include "HAL/Platform.h"
+#include "Containers/ArrayView.h"
+#include "Misc/TVariant.h"
+#include "StructUtils/InstancedStruct.h"
+
+#include <limits>
+
+#include "NetworkMetricsMutators.generated.h"
+
+namespace UE::Net
+{
+	struct FNetworkMetricSnapshot;
+}
+
+/**
+ * A metrics mutator is a component of a UNetworkMetricsBaseListener that receices a metrics snapshot
+ * every frame. This enables a mutator to accumulate and/or transform (possibly multiple) input metrics
+ * into a single output value over the reporting interval of its owning listener.
+ * For example, FNetworkMetricsMutatorAvg allows a listener to report the average of a metric over
+ * a whole reporting interval.
+ */
+USTRUCT()
+struct FNetworkMetricsMutator
+{
+	GENERATED_BODY()
+
+public:
+	using FValueVariant = TVariant<float, int64>;
+
+	virtual ~FNetworkMetricsMutator() = default;
+
+	/**
+	 * Returns the name of this mutator. Independent of metric names in the database.
+	 * Either NameOverride if set, otherwise returns GetNameInternal().
+	 */
+	FName GetName() const;
+
+	/** Called every frame with a snapshot containing this mutator's metrics. */
+	virtual void ProcessFrame(const UE::Net::FNetworkMetricSnapshot& Snapshot) PURE_VIRTUAL(FNetworkMetricsMutator::ProcessFrame, );
+
+	/** Clears internal data, typically used after a listener reports its metrics. */
+	virtual void Reset() PURE_VIRTUAL(FNetworkMetricsMutator::Reset, );
+
+	/** Returns this mutator's current computed value */
+	virtual FValueVariant GetValue() const PURE_VIRTUAL(FNetworkMetricsMutator::GetValue, return FValueVariant(TInPlaceType<float>(), 0.0f););
+
+	/** Returns a list of all metrics used by this mutator. */
+	virtual TArrayView<const FName> GetAllMetricNames() const PURE_VIRTUAL(FNetworkMetricsMutator::GetAllMetricNames, return TArrayView<const FName>(););
+
+private:
+	/** Returns the name of this mutator. Independent of metric names in the database. */
+	virtual FName GetNameInternal() const PURE_VIRTUAL(FNetworkMetricsMutator::GetNameInternal, return FName(););
+
+	/**
+	 * User-configurable name for this mutator that if set, will override the internal name returned by GetNameInternal().
+	 * Use it to disambiguate naming when the internal name would collide with another on the same listener,
+	 * or to clarify an internal name that doesn't work very well (for example, RawPingClientMaxMax).
+	 */
+	UPROPERTY(Config)
+	FName NameOverride;
+};
+
+template<>
+struct TStructOpsTypeTraits<FNetworkMetricsMutator> : public TStructOpsTypeTraitsBase2<FNetworkMetricsMutator>
+{
+	enum
+	{
+		WithPureVirtual = true
+	};
+};
+
+/** Metrics mutator that accumulates an average of a metric. */
+USTRUCT()
+struct FNetworkMetricsMutatorAvg : public FNetworkMetricsMutator
+{
+	GENERATED_BODY()
+
+public:
+	virtual void ProcessFrame(const UE::Net::FNetworkMetricSnapshot& Snapshot) override;
+	virtual void Reset() override;
+	virtual FValueVariant GetValue() const override;
+
+	virtual TArrayView<const FName> GetAllMetricNames() const override
+	{
+		return MakeConstArrayView(&MetricName, 1);
+	}
+
+private:
+	virtual FName GetNameInternal() const override;
+
+	void AddIntSample(int64 Sample);
+	void AddFloatSample(float Sample);
+
+	int64 IntTotal = 0;
+	float FloatTotal = 0.0f;
+	uint32 NumSamples = 0;
+
+	/** The network metric to average. */
+	UPROPERTY(Config)
+	FName MetricName;
+
+	/** The default name of this mutator, derived from MetricName */
+	mutable FName MutatorName;
+};
+
+/** Metrics mutator that tracks the minimum of a metric. */
+USTRUCT()
+struct FNetworkMetricsMutatorMin : public FNetworkMetricsMutator
+{
+	GENERATED_BODY()
+
+public:
+	virtual void ProcessFrame(const UE::Net::FNetworkMetricSnapshot& Snapshot) override;
+	virtual void Reset() override;
+	virtual FValueVariant GetValue() const override;
+
+	virtual TArrayView<const FName> GetAllMetricNames() const override
+	{
+		return MakeConstArrayView(&MetricName, 1);
+	}
+
+private:
+	virtual FName GetNameInternal() const override;
+
+	FValueVariant MinValue = FValueVariant(TInPlaceType<int64>(), std::numeric_limits<int64>::max());
+
+	/** The network metric to track the min of. */
+	UPROPERTY(Config)
+	FName MetricName;
+
+	/** The default name of this mutator, derived from MetricName */
+	mutable FName MutatorName;
+};
+
+/** Metrics mutator that tracks the maximum of a metric. */
+USTRUCT()
+struct FNetworkMetricsMutatorMax : public FNetworkMetricsMutator
+{
+	GENERATED_BODY()
+
+public:
+	virtual void ProcessFrame(const UE::Net::FNetworkMetricSnapshot& Snapshot) override;
+	virtual void Reset() override;
+	virtual FValueVariant GetValue() const override;
+
+	virtual TArrayView<const FName> GetAllMetricNames() const override
+	{
+		return MakeConstArrayView(&MetricName, 1);
+	}
+
+private:
+	virtual FName GetNameInternal() const override;
+
+	FValueVariant MaxValue = FValueVariant(TInPlaceType<int64>(), std::numeric_limits<int64>::lowest());
+
+	/** The network metric to track the max of. */
+	UPROPERTY(Config)
+	FName MetricName;
+
+	/** The default name of this mutator, derived from MetricName */
+	mutable FName MutatorName;
+};
+
+/**
+ * Metrics mutator that reports a 0-100 percentage using two source metrics.
+ */
+USTRUCT()
+struct FNetworkMetricsMutatorPercent : public FNetworkMetricsMutator
+{
+	GENERATED_BODY()
+
+public:
+	virtual void ProcessFrame(const UE::Net::FNetworkMetricSnapshot& Snapshot) override;
+	virtual void Reset() override;
+
+	/** Returns the packet loss percentage as a float from 0-100. */
+	virtual FValueVariant GetValue() const override;
+
+	virtual TArrayView<const FName> GetAllMetricNames() const override;
+
+protected:
+	void SetNumeratorMetricName(FName MetricName)
+	{
+		NumeratorName = MetricName;
+	}
+
+	void SetDenominatorMetricName(FName MetricName)
+	{
+		DenominatorName = MetricName;
+	}
+
+private:
+	virtual FName GetNameInternal() const override;
+
+	// Private template so it can be defined in the .cpp file.
+	template<class T>
+	static void UpdateFromSnapshotValue(FNetworkMetricsMutator::FValueVariant& InOutValue, T MetricValue);
+
+	// Source metric names stored in an array so we can return a view to it from GetAllMetricNames.
+	mutable TArray<FName, TInlineAllocator<2>> SourceMetricNames;
+
+	FValueVariant Numerator = FValueVariant(TInPlaceType<int64>(), 0);
+	FValueVariant Denominator = FValueVariant(TInPlaceType<int64>(), 0);
+
+	UPROPERTY(Config)
+	FName NumeratorName;
+
+	UPROPERTY(Config)
+	FName DenominatorName;
+
+	// Auto-generated mutator name, derived from source metrics.
+	mutable FName MutatorName;
+};
+
+/** Metrics mutator that tracks out packet loss, based on the OutLostPacketsFoundPerFrame and OutPacketsPerFrame metrics. */
+USTRUCT()
+struct FNetworkMetricsMutatorOutPacketLoss : public FNetworkMetricsMutatorPercent
+{
+	GENERATED_BODY()
+
+public:
+	FNetworkMetricsMutatorOutPacketLoss();
+
+private:
+	virtual FName GetNameInternal() const override
+	{
+		return FName("OutPacketLoss");
+	}
+};
+
+/** Metrics mutator that tracks out packet loss, based on the InLostPacketsFoundPerFrame and InPacketsPerFrame metrics. */
+USTRUCT()
+struct FNetworkMetricsMutatorInPacketLoss : public FNetworkMetricsMutatorPercent
+{
+	GENERATED_BODY()
+
+public:
+	FNetworkMetricsMutatorInPacketLoss();
+
+private:
+	virtual FName GetNameInternal() const override
+	{
+		return FName("InPacketLoss");
+	}
+};
