Add cvar in non-shipping to override the skeletal mesh LOD index for debugging.

diff --git a/Engine/Source/Runtime/Engine/Private/Components/SkinnedMeshComponent.cpp b/Engine/Source/Runtime/Engine/Private/Components/SkinnedMeshComponent.cpp
--- a/Engine/Source/Runtime/Engine/Private/Components/SkinnedMeshComponent.cpp
+++ b/Engine/Source/Runtime/Engine/Private/Components/SkinnedMeshComponent.cpp
@@ -61,6 +61,16 @@
 DEFINE_LOG_CATEGORY_STATIC(LogSkinnedMeshComp, Log, All);
 #define LOCTEXT_NAMESPACE "SkinnedMeshComponent"
 
+#if !UE_BUILD_SHIPPING
+int32 GSkeletalMeshForceLOD = -1;
+FAutoConsoleVariableRef CVarSkeletalMeshForceLOD(
+	TEXT("r.SkeletalMeshForceLOD"),
+	GSkeletalMeshForceLOD,
+	TEXT("Forces LOD"),
+	ECVF_Default
+	);
+#endif
+
 int32 GSkeletalMeshLODBias = 0;
 FAutoConsoleVariableRef CVarSkeletalMeshLODBias(
 	TEXT("r.SkeletalMeshLODBias"),
@@ -4203,7 +4213,12 @@ void USkinnedMeshComponent::SetForcedLOD(int32 InNewForcedLOD)
 int32 USkinnedMeshComponent::GetForcedLOD() const
 {
 PRAGMA_DISABLE_DEPRECATION_WARNINGS
+#if !UE_BUILD_SHIPPING
+	return GSkeletalMeshForceLOD < 0 ? ForcedLodModel : GSkeletalMeshForceLOD + 1; // Add 1 to CVar because requires the return value is 1-based for legacy reasons.
+#else
 	return ForcedLodModel;
+#endif
+
 PRAGMA_ENABLE_DEPRECATION_WARNINGS
 }
 
