Forward declare TInstancedStruct in NetworkMetricsDatabase.h instead of including it.

diff --git a/Engine/Source/Runtime/Engine/Private/Net/NetworkMetricsDatabase.cpp b/Engine/Source/Runtime/Engine/Private/Net/NetworkMetricsDatabase.cpp
--- a/Engine/Source/Runtime/Engine/Private/Net/NetworkMetricsDatabase.cpp
+++ b/Engine/Source/Runtime/Engine/Private/Net/NetworkMetricsDatabase.cpp
@@ -5,6 +5,7 @@
 #include "EngineStats.h"
 #include "Logging/LogMacros.h"
 #include "ProfilingDebugging/CsvProfiler.h"
+#include "StructUtils/InstancedStruct.h"
 
 #if USE_SERVER_PERF_COUNTERS
 #include "PerfCountersModule.h"
@@ -408,6 +409,8 @@ UNetworkMetricsBaseListener::UNetworkMetricsBaseListener() :
 
 }
 
+UNetworkMetricsBaseListener::~UNetworkMetricsBaseListener() = default;
+
 void UNetworkMetricsLog::Report(const UE::Net::FNetworkMetricSnapshot& Snapshot)
 {
 	FString MetricsStr;
diff --git a/Engine/Source/Runtime/Engine/Public/Net/NetworkMetricsDatabase.h b/Engine/Source/Runtime/Engine/Public/Net/NetworkMetricsDatabase.h
--- a/Engine/Source/Runtime/Engine/Public/Net/NetworkMetricsDatabase.h
+++ b/Engine/Source/Runtime/Engine/Public/Net/NetworkMetricsDatabase.h
@@ -3,13 +3,15 @@
 #pragma once
 
 #include "CoreMinimal.h"
-#include "StructUtils/InstancedStruct.h"
 #include "NetworkMetricsDatabase.generated.h"
 
 #define UE_API ENGINE_API
 
 struct FNetworkMetricsMutator;
 
+template<typename BaseStructT>
+struct [[nodiscard]] TInstancedStruct;
+
 namespace UE::Net
 {
 
@@ -156,7 +158,7 @@ class UNetworkMetricsBaseListener : public UObject
 
 public:
 	UE_API UNetworkMetricsBaseListener();
-	virtual ~UNetworkMetricsBaseListener() = default;
+	virtual ~UNetworkMetricsBaseListener();
 
 	/* Set the interval, in seconds, between calling Report(). */
 	void SetInterval(double Seconds)
