Iris
* Made it possible to derive from UNetSubObjectFactory and use it as a base for other custom factories
* Disabled "Failed to find Outer" log warnings by default and moved them behind a cvar "net.Iris.WarnWhenOuterIsReplaced"
* Removed "net.Iris.SubObjectSupportArchetype" cvar since we don't need to rollback this feature anymore

#rnx

diff --git a/Engine/Source/Runtime/Engine/Private/Net/Iris/ReplicationSystem/NetSubObjectFactory.cpp b/Engine/Source/Runtime/Engine/Private/Net/Iris/ReplicationSystem/NetSubObjectFactory.cpp
--- a/Engine/Source/Runtime/Engine/Private/Net/Iris/ReplicationSystem/NetSubObjectFactory.cpp
+++ b/Engine/Source/Runtime/Engine/Private/Net/Iris/ReplicationSystem/NetSubObjectFactory.cpp
@@ -28,13 +28,22 @@
 
 #if UE_WITH_IRIS
 
-//$IRIS TODO: Temp cvar can be removed later
-static bool bUseArchetypeInSubObjectCreation = true;
-static FAutoConsoleVariableRef CVarUseArchetype(
-	TEXT("net.Iris.SubObjectSupportArchetype"),
-	bUseArchetypeInSubObjectCreation,
-	TEXT("Support dynamic subobjects that have an archetype instead of just the class CDO.")
-);
+namespace UE::Net::Private
+{
+	static bool bWarnWhenOuterIsReplaced = false;
+	static FAutoConsoleVariableRef CVarWarnWhenOuterIsReplaced(
+		TEXT("net.Iris.WarnWhenOuterIsReplaced"),
+		bWarnWhenOuterIsReplaced,
+		TEXT("Log a warning when a subobject's original outer is not found on the client and replaced with the RootObject.")
+	);
+
+	static FName SubObjectFactoryName(TEXT("NetSubObjectFactory"));
+}
+
+FName UNetSubObjectFactory::GetFactoryName()
+{
+	return UE::Net::Private::SubObjectFactoryName;
+}
 
 TUniquePtr<UE::Net::FNetObjectCreationHeader> UNetSubObjectFactory::CreateAndFillHeader(UE::Net::FNetRefHandle Handle)
 {
@@ -60,32 +69,42 @@ TUniquePtr<UE::Net::FNetObjectCreationHeader> UNetSubObjectFactory::CreateAndFil
 	TUniquePtr<FNetDynamicSubObjectCreationHeader> Header(new FNetDynamicSubObjectCreationHeader);
 	FNetDynamicSubObjectCreationHeader* SubObjectHeader = static_cast<FNetDynamicSubObjectCreationHeader*>(Header.Get());
 
+	const bool bSuccess = FillDynamicHeader(SubObjectHeader, SubObject, Handle);
+	if (UNLIKELY(!bSuccess))
+	{
+		return nullptr;
+	}
+
+	return Header;
+}
+
+bool UNetSubObjectFactory::FillDynamicHeader(UE::Net::FNetDynamicSubObjectCreationHeader* DynamicHeader, const UObject* SubObject, UE::Net::FNetRefHandle Handle)
+{
+	using namespace UE::Net;
+	using namespace UE::Net::Private;
+
 	check(SubObject->NeedsLoadForClient() ); // We have no business sending this unless the client can load
 	check(SubObject->GetClass()->NeedsLoadForClient());	// We have no business sending this unless the client can load
 
 	// Set the template to clone the object from
 	{
 		UObject* Archetype = nullptr;
-		if (bUseArchetypeInSubObjectCreation)
-		{
-			Archetype = SubObject->GetArchetype();
-		}
-		else
-		{
-			Archetype = SubObject->GetClass();
-		}
+		Archetype = SubObject->GetArchetype();
 
-		SubObjectHeader->TemplateReference = Bridge->GetOrCreateObjectReference(Archetype);
+		DynamicHeader->TemplateReference = Bridge->GetOrCreateObjectReference(Archetype);
 
-		ensureMsgf(SubObjectHeader->TemplateReference.GetRefHandle().IsStatic(), TEXT("Archetype %s for subobject %s (%s) should be a stable name otherwise the client can't use it"),
-			*GetNameSafe(Archetype), *Bridge->PrintObjectFromNetRefHandle(Handle), *GetNameSafe(SubObject->GetClass()->GetAuthoritativeClass()));
+		if (!DynamicHeader->TemplateReference.GetRefHandle().IsStatic())
+		{
+			UE_LOG(LogIris, Error, TEXT("Archetype %s for subobject %s (%s) should be a stable name otherwise the client can't use it"), *GetNameSafe(Archetype), *Bridge->PrintObjectFromNetRefHandle(Handle), *GetNameSafe(SubObject->GetClass()->GetAuthoritativeClass()));
+			return false;
+		}
 	}
 
 	// Find the right Outer
 	UObject* OuterObject = SubObject->GetOuter();	
 	if (OuterObject == GetTransientPackage())
 	{
-		SubObjectHeader->bOuterIsTransientLevel = true;
+		DynamicHeader->bOuterIsTransientLevel = true;
 	}
 	else 
 	{
@@ -95,26 +114,26 @@ TUniquePtr<UE::Net::FNetObjectCreationHeader> UNetSubObjectFactory::CreateAndFil
 
 		if (OuterObject == RootObject)
 		{
-			SubObjectHeader->bOuterIsRootObject = true;
+			DynamicHeader->bOuterIsRootObject = true;
 		}
 		else
 		{
-			SubObjectHeader->OuterReference = Bridge->GetOrCreateObjectReference(OuterObject);
+			DynamicHeader->OuterReference = Bridge->GetOrCreateObjectReference(OuterObject);
 
 			// If the Outer is not net-referenceable, use the RootObject instead
-			if (!SubObjectHeader->OuterReference.IsValid())
+			if (!DynamicHeader->OuterReference.IsValid())
 			{
-				UE_LOG(LogIris, Warning, TEXT("UNetSubObjectFactory::CreateAndFillHeader subobject: %s has an Outer: %s that is not stable or replicated. Clients will use RootObject: %s as the Outer instead"), 
+				UE_CLOG(bWarnWhenOuterIsReplaced, LogIris, Warning, TEXT("UNetSubObjectFactory::CreateAndFillHeader subobject: %s has an Outer: %s that is not stable or replicated. Clients will use RootObject: %s as the Outer instead"), 
 					*Bridge->PrintObjectFromNetRefHandle(Handle), 
 					*GetNameSafe(OuterObject),
 					*GetNameSafe(RootObject));
 
-				SubObjectHeader->bOuterIsRootObject = true;
+				DynamicHeader->bOuterIsRootObject = true;
 			}
 		}
 	}
 
-	return Header;
+	return true;
 }
 
 TUniquePtr<UE::Net::FNetObjectCreationHeader> UNetSubObjectFactory::CreateAndDeserializeHeader(const UE::Net::FCreationHeaderContext& Context)
@@ -141,6 +160,7 @@ TUniquePtr<UE::Net::FNetObjectCreationHeader> UNetSubObjectFactory::CreateAndDes
 UNetSubObjectFactory::FInstantiateResult UNetSubObjectFactory::InstantiateReplicatedObjectFromHeader(const FInstantiateContext& Context, const UE::Net::FNetObjectCreationHeader* Header)
 {
 	using namespace UE::Net;
+	using namespace UE::Net::Private;
 
 	const FNetBaseSubObjectCreationHeader* BaseHeader = static_cast<const FNetBaseSubObjectCreationHeader*>(Header);
 
@@ -168,6 +188,15 @@ UNetSubObjectFactory::FInstantiateResult UNetSubObjectFactory::InstantiateReplic
 	const FNetDynamicSubObjectCreationHeader* SubObjectHeader = static_cast<const FNetDynamicSubObjectCreationHeader*>(Header);
 		
 	UObject* RootObject = Bridge->GetReplicatedObject(Context.RootObjectOfSubObject);
+
+	// Find the archetype of the subobject
+	UObject* Template = Bridge->ResolveObjectReference(SubObjectHeader->TemplateReference, Context.ResolveContext);
+	if (!Template)
+	{
+		UE_LOG(LogIris, Error, TEXT("UNetSubObjectFactory::InstantiateReplicatedObjectFromHeader Unable to resolve template: %s, cannot instantiate dynamic subobject %s"), *Bridge->DescribeObjectReference(SubObjectHeader->TemplateReference, Context.ResolveContext), *Context.Handle.ToString());
+		ensure(Template);
+		return FInstantiateResult();
+	}
 			
 	// Find the proper Outer
 	UObject* OuterObject = nullptr;
@@ -185,37 +214,19 @@ UNetSubObjectFactory::FInstantiateResult UNetSubObjectFactory::InstantiateReplic
 
 		if (!OuterObject)
 		{
-			UE_LOG(LogIris, Warning, TEXT("UNetSubObjectFactory::InstantiateReplicatedObjectFromHeader Failed to find Outer %s for dynamic subobject %s"), *Bridge->DescribeObjectReference(SubObjectHeader->OuterReference, Context.ResolveContext), *Context.Handle.ToString())
-
 			// Fallback to the rootobject instead
 			OuterObject = RootObject;
-		}
-	}
 
-	// Find the archetype of the subobject
-	UObject* Template = Bridge->ResolveObjectReference(SubObjectHeader->TemplateReference, Context.ResolveContext);
-	if (!Template)
-	{
-		UE_LOG(LogIris, Error, TEXT("UNetSubObjectFactory::InstantiateReplicatedObjectFromHeader Unable to resolve template: %s, cannot instantiate dynamic subobject %s"), *Bridge->DescribeObjectReference(SubObjectHeader->TemplateReference, Context.ResolveContext), *Context.Handle.ToString());
-		ensure(Template);
-		return FInstantiateResult();
+			UE_CLOG(bWarnWhenOuterIsReplaced, LogIris, Warning, TEXT("UNetSubObjectFactory::InstantiateReplicatedObjectFromHeader Failed to find Outer %s for dynamic subobject %s (template: %s). Original outer will be replaced with %s instead."), *Bridge->DescribeObjectReference(SubObjectHeader->OuterReference, Context.ResolveContext), *Context.Handle.ToString(), *GetNameSafe(Template), *GetNameSafe(RootObject));
+		}
 	}
 
 	// Instantiate the SubObject
-	 UObject* SubObj = nullptr;
-	 UClass* SubObjClass = nullptr;
+	UObject* SubObj = nullptr;
+	UClass* SubObjClass = nullptr;
 
-	// Only if bUseArchetypeInSubObjectCreation is false on the server. Remove along with the cvar later
-	if (Template->IsA<UClass>())
-	{
-		SubObjClass = CastChecked<UClass>(Template);
-		SubObj = NewObject<UObject>(OuterObject, SubObjClass);
-	}
-	else
-	{
-		SubObjClass = Template->GetClass();
-		SubObj = NewObject<UObject>(OuterObject, SubObjClass, NAME_None, RF_NoFlags, Template);
-	}
+	SubObjClass = Template->GetClass();
+	SubObj = NewObject<UObject>(OuterObject, SubObjClass, NAME_None, RF_NoFlags, Template);
 
 	// Sanity check some things
 	checkf(SubObj != nullptr, TEXT("UNetSubObjectFactory::InstantiateReplicatedObjectFromHeader: Subobject is NULL after instantiating. Class: %s, Outer %s, Actor %s"), *GetNameSafe(SubObjClass), *GetNameSafe(OuterObject), *GetNameSafe(RootObject));
diff --git a/Engine/Source/Runtime/Engine/Public/Net/Iris/ReplicationSystem/NetSubObjectFactory.h b/Engine/Source/Runtime/Engine/Public/Net/Iris/ReplicationSystem/NetSubObjectFactory.h
--- a/Engine/Source/Runtime/Engine/Public/Net/Iris/ReplicationSystem/NetSubObjectFactory.h
+++ b/Engine/Source/Runtime/Engine/Public/Net/Iris/ReplicationSystem/NetSubObjectFactory.h
@@ -10,66 +10,28 @@
 
 #include "NetSubObjectFactory.generated.h"
 
-/**
- * Responsible for creating headers allowing remote factories to spawn replicated actors
- */
-UCLASS()
-class UNetSubObjectFactory : public UNetObjectFactory
-{
-	GENERATED_BODY()
-
-#if UE_WITH_IRIS
-
-public:
-
-	static FName GetFactoryName()
-	{ 
-		return TEXT("NetSubObjectFactory");
-	}
-
-	virtual FInstantiateResult InstantiateReplicatedObjectFromHeader(const FInstantiateContext& Context, const UE::Net::FNetObjectCreationHeader* Header) override;
-
-	virtual void SubObjectCreatedFromReplication(UE::Net::FNetRefHandle RootObject, UE::Net::FNetRefHandle SubObjectCreated) override;
-
-	virtual void DestroyReplicatedObject(const FDestroyedContext& Context) override;
-
-	virtual TOptional<FWorldInfoData> GetWorldInfo(const FWorldInfoContext& Context) const override;
-
-	virtual float GetPollFrequency(UE::Net::FNetRefHandle RootObjectHandle, UObject* RootObjectInstance) override;
-
-protected:
-
-	virtual TUniquePtr<UE::Net::FNetObjectCreationHeader> CreateAndFillHeader(UE::Net::FNetRefHandle Handle) override;
-	virtual TUniquePtr<UE::Net::FNetObjectCreationHeader> CreateAndDeserializeHeader(const UE::Net::FCreationHeaderContext& Context) override;
-
-	virtual bool SerializeHeader(const UE::Net::FCreationHeaderContext& Context, const UE::Net::FNetObjectCreationHeader* Header)  override;
-#endif // UE_WITH_IRIS
-};
-
 
 namespace UE::Net
 {
 #if UE_WITH_IRIS
 
 /**
- * Header information to be able to tell if its a dynamic or static header
- */
+* Header information to to tell if it is a dynamic or static header
+*/
 class FNetBaseSubObjectCreationHeader : public FNetObjectCreationHeader
 {
 public:
 	
 	virtual bool IsDynamic() const = 0;
 
-	virtual bool Serialize(const FCreationHeaderContext& Context) const
-	{ 
-		return false;
-	}
+	virtual bool Serialize(const FCreationHeaderContext& Context) const { return false; }
+	virtual bool Deserialize(const FCreationHeaderContext& Context) { return false; }
 };
 
 
 /**
- * Header information representing subobjects that can be found via their pathname (aka: static or stable name)
- */
+* Header information representing subobjects that can be found via their pathname (aka: static or stable name)
+*/
 class FNetStaticSubObjectCreationHeader : public FNetBaseSubObjectCreationHeader
 {
 public:
@@ -79,30 +41,30 @@ public:
 		return false;
 	}
 
-	virtual bool Serialize(const FCreationHeaderContext& Context) const override;
-	bool Deserialize(const FCreationHeaderContext& Context);
+	ENGINE_API virtual bool Serialize(const FCreationHeaderContext& Context) const override;
+	ENGINE_API virtual bool Deserialize(const FCreationHeaderContext& Context) override;
 
-	virtual FString ToString() const override;
+	ENGINE_API virtual FString ToString() const override;
 
 	FNetObjectReference ObjectReference; // Only for static objects
 };
 
 /**
- * Header information representing subobjects that must be instantiated
- */
+* Header information representing subobjects that must be instantiated
+*/
 class FNetDynamicSubObjectCreationHeader : public FNetBaseSubObjectCreationHeader
 {
 public:
 
-	virtual bool Serialize(const FCreationHeaderContext& Context) const override;
-	bool Deserialize(const FCreationHeaderContext& Context);
-
 	virtual bool IsDynamic() const override
 	{ 
 		return true;
 	}
 
-	virtual FString ToString() const override;
+	ENGINE_API virtual bool Serialize(const FCreationHeaderContext& Context) const override;
+	ENGINE_API virtual bool Deserialize(const FCreationHeaderContext& Context) override;
+
+	ENGINE_API virtual FString ToString() const override;
 
 	FNetObjectReference TemplateReference;
 	FNetObjectReference OuterReference;
@@ -111,5 +73,45 @@ public:
 	uint8 bOuterIsRootObject : 1 = false;		// When set the OuterReference was not sent because the Outer is the known RootObject.
 
 };
+
 #endif // UE_WITH_IRIS
-} // end namespace UE::Net
\ No newline at end of file
+
+} // end namespace UE::Net
+
+/**
+ * Responsible for creating subobjects only
+ */
+UCLASS(MinimalAPI)
+class UNetSubObjectFactory : public UNetObjectFactory
+{
+	GENERATED_BODY()
+
+#if UE_WITH_IRIS
+
+public:
+
+	static FName GetFactoryName();
+	
+	ENGINE_API virtual FInstantiateResult InstantiateReplicatedObjectFromHeader(const FInstantiateContext& Context, const UE::Net::FNetObjectCreationHeader* Header) override;
+
+	ENGINE_API virtual void SubObjectCreatedFromReplication(UE::Net::FNetRefHandle RootObject, UE::Net::FNetRefHandle SubObjectCreated) override;
+
+	ENGINE_API virtual void DestroyReplicatedObject(const FDestroyedContext& Context) override;
+
+	ENGINE_API virtual TOptional<FWorldInfoData> GetWorldInfo(const FWorldInfoContext& Context) const override;
+
+	ENGINE_API virtual float GetPollFrequency(UE::Net::FNetRefHandle RootObjectHandle, UObject* RootObjectInstance) override;
+
+protected:
+
+	ENGINE_API virtual TUniquePtr<UE::Net::FNetObjectCreationHeader> CreateAndFillHeader(UE::Net::FNetRefHandle Handle) override;
+	ENGINE_API virtual TUniquePtr<UE::Net::FNetObjectCreationHeader> CreateAndDeserializeHeader(const UE::Net::FCreationHeaderContext& Context) override;
+
+	ENGINE_API virtual bool SerializeHeader(const UE::Net::FCreationHeaderContext& Context, const UE::Net::FNetObjectCreationHeader* Header)  override;
+
+	/** Write the needed header information for a dynamic subobject */
+	ENGINE_API bool FillDynamicHeader(UE::Net::FNetDynamicSubObjectCreationHeader* DynamicHeader, const UObject* SubObject, UE::Net::FNetRefHandle Handle);
+
+#endif // UE_WITH_IRIS
+};
+
