Fix SpecialComplex tiles being dispatch with Substrate Format=0.

Materials with glints were requesting 'complex special' complexity, despite glint being disabled.

diff --git a/Engine/Source/Runtime/Engine/Private/Materials/MaterialExpressionSubstrate.cpp b/Engine/Source/Runtime/Engine/Private/Materials/MaterialExpressionSubstrate.cpp
--- a/Engine/Source/Runtime/Engine/Private/Materials/MaterialExpressionSubstrate.cpp
+++ b/Engine/Source/Runtime/Engine/Private/Materials/MaterialExpressionSubstrate.cpp
@@ -1211,11 +1211,11 @@ FName UMaterialExpressionSubstrateSlabBSDF::GetInputName(int32 InputIndex) const
 	}
 	else if (InputIndex == 16)
 	{
-		return Substrate::IsGlintEnabled() ? TEXT("Glint Density") : TEXT("Glint Density (Disabled)");
+		return Substrate::IsGlintEnabled(GMaxRHIShaderPlatform) ? TEXT("Glint Density") : TEXT("Glint Density (Disabled)");
 	}
 	else if (InputIndex == 17)
 	{
-		return Substrate::IsGlintEnabled() ? TEXT("Glint UVs") : TEXT("Glint UVs (Disabled)");
+		return Substrate::IsGlintEnabled(GMaxRHIShaderPlatform) ? TEXT("Glint UVs") : TEXT("Glint UVs (Disabled)");
 	}
 
 	return TEXT("Unknown");
@@ -1345,8 +1345,7 @@ bool UMaterialExpressionSubstrateSlabBSDF::HasGlint() const
 
 bool UMaterialExpressionSubstrateSlabBSDF::HasSpecularProfile() const
 {
-	// Use IsSpecularLUTEnabled to not promote to ComplexSpecial render path if glint rendering is not enabled and avoid registering such possibility at runtime.
-	return SpecularProfile != nullptr && Substrate::IsSpecularProfileEnabled();
+	return SpecularProfile != nullptr;
 }
 
 void UMaterialExpressionSubstrateSlabBSDF::PostEditChangeProperty(FPropertyChangedEvent& PropertyChangedEvent)
diff --git a/Engine/Source/Runtime/Engine/Private/Rendering/SpecularProfile.cpp b/Engine/Source/Runtime/Engine/Private/Rendering/SpecularProfile.cpp
--- a/Engine/Source/Runtime/Engine/Private/Rendering/SpecularProfile.cpp
+++ b/Engine/Source/Runtime/Engine/Private/Rendering/SpecularProfile.cpp
@@ -50,7 +50,7 @@ static bool ForceUpdateSpecularProfile()
 
 static bool IsSpecularProfileSupport(EShaderPlatform InShaderPlatform)
 {
-	return IsFeatureLevelSupported(InShaderPlatform, ERHIFeatureLevel::SM5);
+	return IsFeatureLevelSupported(InShaderPlatform, ERHIFeatureLevel::SM5) && Substrate::IsSpecularProfileEnabled(InShaderPlatform);
 }
 
 class FSpecularProfileCopyCS : public FGlobalShader
@@ -323,7 +323,7 @@ IPooledRenderTarget* FSpecularProfileTextureManager::GetAtlasTexture()
 
 IPooledRenderTarget* FSpecularProfileTextureManager::GetAtlasTexture(FRDGBuilder& GraphBuilder, EShaderPlatform ShaderPlatform)
 {
-	if (!Substrate::IsSubstrateEnabled() || !Substrate::IsSpecularProfileEnabled() || !IsSpecularProfileSupport(ShaderPlatform))
+	if (!Substrate::IsSubstrateEnabled() || !IsSpecularProfileSupport(ShaderPlatform))
 	{
 		return nullptr;
 	}
diff --git a/Engine/Source/Runtime/RenderCore/Private/RenderUtils.cpp b/Engine/Source/Runtime/RenderCore/Private/RenderUtils.cpp
--- a/Engine/Source/Runtime/RenderCore/Private/RenderUtils.cpp
+++ b/Engine/Source/Runtime/RenderCore/Private/RenderUtils.cpp
@@ -2182,18 +2182,14 @@ namespace Substrate
 		return CVar.Get(InPlatform) > 0;
 	}
 
-	bool IsGlintEnabled()
-	{
-		return IsSubstrateEnabled() && CVarSubstrateGlints.GetValueOnAnyThread() > 0;
-	}
-
 	bool IsGlintEnabled(EShaderPlatform InPlatform)
 	{
 		static FShaderPlatformCachedIniValue<int32> CVar(TEXT("r.Substrate.Glints"));
 		// Force disable glint on ES31/SM4 platforms as it causes compilation error (e.g. too large local light unroll in forward shaders).
 		// Usually this effect is explicitly disabled in XXXPlatform.ini files. However when running preview (e.g. -featureleveles31), 
 		// these settings are not taken into account causing compilation issue.
-		const bool bSupported = FDataDrivenShaderPlatformInfo::GetMaxFeatureLevel(InPlatform) >= ERHIFeatureLevel::SM5;
+		// Also disable glint when blendable gbuffer is used as the data can't be stored
+		const bool bSupported = FDataDrivenShaderPlatformInfo::GetMaxFeatureLevel(InPlatform) >= ERHIFeatureLevel::SM5 && !Substrate::IsSubstrateBlendableGBufferEnabled(InPlatform);
 		return IsSubstrateEnabled() && CVar.Get(InPlatform) && bSupported;
 	}
 
diff --git a/Engine/Source/Runtime/RenderCore/Public/RenderUtils.h b/Engine/Source/Runtime/RenderCore/Public/RenderUtils.h
--- a/Engine/Source/Runtime/RenderCore/Public/RenderUtils.h
+++ b/Engine/Source/Runtime/RenderCore/Public/RenderUtils.h
@@ -617,7 +617,6 @@ namespace Substrate
 	RENDERCORE_API bool IsStochasticLightingEnabled(EShaderPlatform InPlatform);
 	RENDERCORE_API bool IsRoughDiffuseEnabled();
 	RENDERCORE_API bool IsRoughDiffuseEnabled(EShaderPlatform InPlatform);
-	RENDERCORE_API bool IsGlintEnabled();
 	RENDERCORE_API bool IsGlintEnabled(EShaderPlatform InPlatform);
 	RENDERCORE_API uint32 GlintLUTIndex();
 	RENDERCORE_API float GlintLevelBias();
