Fixed UAnimBank not loading correctly from DDC. Wait to build derived data until the anim sequence's derived data build will succeed.

diff --git a/Engine/Source/Runtime/Engine/Private/Animation/AnimBank.cpp b/Engine/Source/Runtime/Engine/Private/Animation/AnimBank.cpp
--- a/Engine/Source/Runtime/Engine/Private/Animation/AnimBank.cpp
+++ b/Engine/Source/Runtime/Engine/Private/Animation/AnimBank.cpp
@@ -288,7 +288,7 @@ bool FAnimBankBuildAsyncCacheTask::ShouldWaitForCompilation() const
 			}
 
 			// If the sequence is still waiting for a post load call, let it build its stuff first to avoid blocking the Game Thread
-			if (BankSequence.Sequence->HasAnyFlags(RF_NeedPostLoad) || BankSequence.Sequence->IsCompiling())
+			if (BankSequence.Sequence->HasAnyFlags(RF_NeedPostLoad) || BankSequence.Sequence->IsCompiling() || !BankSequence.Sequence->CanBeCompressed())
 			{
 				return true;
 			}
@@ -372,6 +372,17 @@ void FAnimBankBuildAsyncCacheTask::BeginCache(const FIoHash& InKeyHash)
 
 	if (UAnimBank* Bank = WeakAnimBank.Get())
 	{
+		for (FAnimBankSequence& BankSequence : Bank->Sequences)
+		{
+			UAnimSequence* Sequence = BankSequence.Sequence;
+			if (!IsValid(Sequence))
+			{
+				continue;
+			}
+
+			Sequence->BeginCacheDerivedData(TargetPlatform);
+		}
+
 		// Queue this launch through the thread pool so that we benefit from fair scheduling and memory throttling
 		FQueuedThreadPool* ThreadPool = FAnimBankCompilingManager::Get().GetThreadPool();
 		EQueuedWorkPriority BasePriority = FAnimBankCompilingManager::Get().GetBasePriority(Bank);
@@ -1122,17 +1133,6 @@ FIoHash UAnimBank::BeginCacheDerivedData(const ITargetPlatform* TargetPlatform)
 		TargetData = DataByPlatformKeyHash.Emplace(KeyHash, MakeUnique<FAnimBankData>()).Get();
 	}
 
-	for (FAnimBankSequence& BankSequence : Sequences)
-	{
-		UAnimSequence* Sequence = BankSequence.Sequence;
-		if (!IsValid(Sequence))
-		{
-			continue;
-		}
-
-		Sequence->BeginCacheDerivedData(TargetPlatform);
-	}
-
 	CacheTasksByKeyHash.Emplace(KeyHash, MakePimpl<FAnimBankBuildAsyncCacheTask>(KeyHash, TargetData, *this, TargetPlatform));
 
 	// The compiling manager provides throttling, notification manager, etc... for the asset being built.
