Mover: Fixing an issue in networked physics component where input wouldn't be produced for locally controlled pawns not using a player controller

diff --git a/Engine/Source/Runtime/Engine/Private/PhysicsEngine/NetworkPhysicsComponent.cpp b/Engine/Source/Runtime/Engine/Private/PhysicsEngine/NetworkPhysicsComponent.cpp
--- a/Engine/Source/Runtime/Engine/Private/PhysicsEngine/NetworkPhysicsComponent.cpp
+++ b/Engine/Source/Runtime/Engine/Private/PhysicsEngine/NetworkPhysicsComponent.cpp
@@ -1462,11 +1462,12 @@ bool UNetworkPhysicsComponent::IsLocallyControlled() const
 	{
 		return true;
 	}
-
-	if (APlayerController* PlayerController = GetPlayerController())
+	
+	if (const AController* Controller = GetController())
 	{
-		return PlayerController->IsLocalController();
+		return Controller->IsLocalController();
 	}
+	
 	return false;
 }
 
@@ -1507,24 +1508,29 @@ void UNetworkPhysicsComponent::SetCompareInputToTriggerRewind(const bool bInComp
 }
 
 APlayerController* UNetworkPhysicsComponent::GetPlayerController() const
+{
+	return Cast<APlayerController>(GetController());
+}
+
+AController* UNetworkPhysicsComponent::GetController() const
 {
 	Chaos::EnsureIsInGameThreadContext();
-	if (APlayerController* PC = Cast<APlayerController>(GetOwner()))
+	if (AController* Controller = Cast<AController>(GetOwner()))
 	{
-		return PC;
+		return Controller;
 	}
 
-	if (APawn* Pawn = Cast<APawn>(GetOwner()))
+	if (const APawn* Pawn = Cast<APawn>(GetOwner()))
 	{
-		if (APlayerController * PC = Pawn->GetController<APlayerController>())
+		if (AController* Controller = Pawn->GetController())
 		{
-			return PC;
+			return Controller;
 		}
 
-		// In this case the APlayerController can be found as the owner of the pawn
-		if (APlayerController* PC = Cast<APlayerController>(Pawn->GetOwner()))
+		// In this case the AController can be found as the owner of the pawn
+		if (AController* Controller = Cast<AController>(Pawn->GetOwner()))
 		{
-			return PC;
+			return Controller;
 		}
 	}
 
diff --git a/Engine/Source/Runtime/Engine/Public/Physics/NetworkPhysicsComponent.h b/Engine/Source/Runtime/Engine/Public/Physics/NetworkPhysicsComponent.h
--- a/Engine/Source/Runtime/Engine/Public/Physics/NetworkPhysicsComponent.h
+++ b/Engine/Source/Runtime/Engine/Public/Physics/NetworkPhysicsComponent.h
@@ -1246,6 +1246,9 @@ public:
 	// Get the player controller
 	ENGINE_API virtual APlayerController* GetPlayerController() const;
 
+	// Get the controller for this pawn
+	ENGINE_API virtual AController* GetController() const;
+	
 	// Init the network physics component 
 	ENGINE_API void InitPhysics();
 
