Fixed a bug in APlayerController::IsLocalController() where it could return true for remote clients during teardown.

Previously, after a listen server transitioned back to a standalone map and the NetDriver was destroyed, GetNetMode() would fall back to NM_Standalone. This caused IsLocalController() to incorrectly return true for remote PlayerControllers that had a ULocalPlayer, leading to missed cleanup logic during controller destruction (e.g., StopGuard not being called in EndPlay).

This change avoids setting bIsLocalPlayerController during teardown and instead adds a conservative fallback path: if NetDriver is null, we only return true if GetLocalRole() == ROLE_Authority and bIsLocalPlayerController is already true. This ensures that only truly local controllers report as local during teardown, fixing the original bug without causing a regression.

diff --git a/Engine/Source/Runtime/Engine/Private/PlayerController.cpp b/Engine/Source/Runtime/Engine/Private/PlayerController.cpp
--- a/Engine/Source/Runtime/Engine/Private/PlayerController.cpp
+++ b/Engine/Source/Runtime/Engine/Private/PlayerController.cpp
@@ -323,6 +323,17 @@ bool APlayerController::IsLocalController() const
 		return true;
 	}
 
+	// Fallback logic for edge cases like PIE teardown where NetDriver is null.
+	if (GetNetDriver() == nullptr)
+	{
+		// Only treat as local if we’re Authority and were already marked as local before NetDriver was torn down
+		const bool bIsServer = GetLocalRole() == ROLE_Authority;
+		const bool bResult = bIsServer && bIsLocalPlayerController;
+
+
+		return bResult;
+	}
+
 	return bIsLocalPlayerController;
 }
 
@@ -1574,6 +1585,17 @@ void APlayerController::BeginPlay()
 	{
 		SetShowMouseCursor(true);
 	}
+	
+	// Initialization for bIsLocalPlayerController in case NetDriver is set to null before IsLocalController is called
+	if (!bIsLocalPlayerController)
+	{
+		ENetMode NetMode = GetNetMode();
+		if (NetMode == NM_Client || NetMode == NM_Standalone)
+		{
+			// Clients or Standalone only receive their own PC.
+			bIsLocalPlayerController = true;
+		}
+	}
 }
 
 void APlayerController::EndPlay(const EEndPlayReason::Type EndPlayReason)
