- New class FIgnoreScope to disable the asset registry tracker.
- Disable the asset registry tracker in FWorldPartitionClassDescRegistry::PrefetchClassDescs. This function makes asset registry request. But the request is cached so depending on the package cook order, the dependency is not always added to the same package. The dependency is not needed anyway since the map and generated packages already create the correct dependencies.

#rnx

diff --git a/Engine/Source/Editor/UnrealEd/Private/Cooker/CookAssetRegistryAccessTracker.cpp b/Engine/Source/Editor/UnrealEd/Private/Cooker/CookAssetRegistryAccessTracker.cpp
--- a/Engine/Source/Editor/UnrealEd/Private/Cooker/CookAssetRegistryAccessTracker.cpp
+++ b/Engine/Source/Editor/UnrealEd/Private/Cooker/CookAssetRegistryAccessTracker.cpp
@@ -6,6 +6,7 @@
 #include "AssetRegistry/AssetData.h"
 #include "AssetRegistry/AssetRegistryModule.h"
 #include "Cooker/CookAssetRegistryAccessTracker_AddPackageDep.h"
+#include "Cooker/CookAssetRegistryAccessTracker_IgnoreScope.h"
 #include "CookOnTheSide/CookLog.h"
 #include "Interfaces/ITargetPlatform.h"
 #include "Misc/PackageAccessTracking.h"
@@ -148,6 +149,11 @@ namespace UE::CookAssetRegistryAccessTracker
 			return;
 		}
 
+		if (FIgnoreScope::ShouldIgnoreAccessTracker())
+		{
+			return;
+		}
+
 		PackageAccessTracking_Private::FTrackedData* AccumulatedScopeData = PackageAccessTracking_Private::FPackageAccessRefScope::GetCurrentThreadAccumulatedData();
 		if (!AccumulatedScopeData || AccumulatedScopeData->BuildOpName.IsNone())
 		{
diff --git a/Engine/Source/Runtime/CoreUObject/Private/Cooker/CookAssetRegistryAccessTracker_IgnoreScope.cpp b/Engine/Source/Runtime/CoreUObject/Private/Cooker/CookAssetRegistryAccessTracker_IgnoreScope.cpp
new file mode 100644
--- /dev/null
+++ b/Engine/Source/Runtime/CoreUObject/Private/Cooker/CookAssetRegistryAccessTracker_IgnoreScope.cpp
@@ -0,0 +1,24 @@
+// Copyright Epic Games, Inc. All Rights Reserved.
+
+#include "Cooker/CookAssetRegistryAccessTracker_IgnoreScope.h"
+
+namespace UE::CookAssetRegistryAccessTracker
+{
+	static thread_local bool bIgnoreAccessTracker = false;
+
+	FIgnoreScope::FIgnoreScope()
+		: bPreviousValue(bIgnoreAccessTracker)
+	{
+		bIgnoreAccessTracker = true;
+	}
+
+	FIgnoreScope::~FIgnoreScope()
+	{
+		bIgnoreAccessTracker = bPreviousValue;
+	}
+
+	bool FIgnoreScope::ShouldIgnoreAccessTracker()
+	{
+		return bIgnoreAccessTracker;
+	}
+}
diff --git a/Engine/Source/Runtime/CoreUObject/Public/Cooker/CookAssetRegistryAccessTracker_IgnoreScope.h b/Engine/Source/Runtime/CoreUObject/Public/Cooker/CookAssetRegistryAccessTracker_IgnoreScope.h
new file mode 100644
--- /dev/null
+++ b/Engine/Source/Runtime/CoreUObject/Public/Cooker/CookAssetRegistryAccessTracker_IgnoreScope.h
@@ -0,0 +1,19 @@
+// Copyright Epic Games, Inc. All Rights Reserved.
+
+#include "CoreTypes.h"
+
+namespace UE::CookAssetRegistryAccessTracker
+{
+	/** RAII class used to enable the creation of dependencies for the packages returned by an asset registry request. */
+	class FIgnoreScope
+	{
+	public:
+		COREUOBJECT_API FIgnoreScope();
+		COREUOBJECT_API ~FIgnoreScope();
+
+		COREUOBJECT_API static bool ShouldIgnoreAccessTracker();
+
+	private:
+		bool bPreviousValue;
+	};
+}
diff --git a/Engine/Source/Runtime/Engine/Private/WorldPartition/WorldPartitionClassDescRegistry.cpp b/Engine/Source/Runtime/Engine/Private/WorldPartition/WorldPartitionClassDescRegistry.cpp
--- a/Engine/Source/Runtime/Engine/Private/WorldPartition/WorldPartitionClassDescRegistry.cpp
+++ b/Engine/Source/Runtime/Engine/Private/WorldPartition/WorldPartitionClassDescRegistry.cpp
@@ -7,6 +7,7 @@
 #include "Algo/ForEach.h"
 #include "Algo/RemoveIf.h"
 #include "Algo/Transform.h"
+#include "Cooker/CookAssetRegistryAccessTracker_IgnoreScope.h"
 #include "HAL/FileManager.h"
 #include "Misc/PackageName.h"
 #include "UObject/CoreRedirects.h"
@@ -180,6 +181,12 @@ void FWorldPartitionClassDescRegistry::PrefetchClassDescs(const TArray<FTopLevel
 
 	check(IsInitialized());
 
+	// This function makes asset registry request during the cook. The tracker will register those requests as dependencies. The issue
+	// is the requests are not always made by the same package so I end up with non deterministic dependencies.
+	// Those dependencies are useless. This function makes a list of BP and native classes used by the actor descriptors. Whatever dependencies
+	// map/generated packages and actors need should already be tracked by the actors/map themselves.
+	UE::CookAssetRegistryAccessTracker::FIgnoreScope IgnoreAssetRegistryTracker;
+
 	IAssetRegistry& AssetRegistry = FModuleManager::LoadModuleChecked<FAssetRegistryModule>(TEXT("AssetRegistry")).Get();
 
 	// First, remap classes and remove already known ones.
