Deleting some unused variables. No functional changes.

diff --git a/Engine/Source/Runtime/Engine/Private/Collision/WorldCollision.cpp b/Engine/Source/Runtime/Engine/Private/Collision/WorldCollision.cpp
--- a/Engine/Source/Runtime/Engine/Private/Collision/WorldCollision.cpp
+++ b/Engine/Source/Runtime/Engine/Private/Collision/WorldCollision.cpp
@@ -506,11 +506,6 @@ bool UWorld::ComponentSweepMultiByChannel(TArray<struct FHitResult>& OutHits, cl
 	SCOPE_CYCLE_COUNTER(STAT_Collision_GeomSweepMultiple);
 	bool bHaveBlockingHit = false;
 
-	const FTransform ComponentToStart{ Rot, Start };
-	const FTransform ComponentToEnd{ Rot, End };
-
-	const FTransform WorldToComponent = PrimComp->GetComponentToWorld().Inverse();
-
 	TArray<FHitResult> TempHits;
 	for (Chaos::FPhysicsObjectHandle Object : PhysicsObjects)
 	{
