Add SmoothWalkingMode to Mover2.0
Add base class SimpleWalkingMode to help make writing walking modes easier
Rename AnimationMath -> SpringMath and make public
Rename SpringWalkingMode -> SimpleSpringWalkingMode
Add Velocity Spring functions to SpringMath
Add character prediction functions to SpringMath
Add documentation to SpringMath

diff --git a/Engine/Plugins/Animation/AnimationWarping/Source/Runtime/Private/BoneControllers/AnimNode_Steering.cpp b/Engine/Plugins/Animation/AnimationWarping/Source/Runtime/Private/BoneControllers/AnimNode_Steering.cpp
--- a/Engine/Plugins/Animation/AnimationWarping/Source/Runtime/Private/BoneControllers/AnimNode_Steering.cpp
+++ b/Engine/Plugins/Animation/AnimationWarping/Source/Runtime/Private/BoneControllers/AnimNode_Steering.cpp
@@ -5,7 +5,7 @@
 #include "Animation/AnimNodeFunctionRef.h"
 #include "Animation/AnimRootMotionProvider.h"
 #include "Animation/AnimSequenceHelpers.h"
-#include "Animation/AnimationMath.h"
+#include "Animation/SpringMath.h"
 #include "HAL/IConsoleManager.h"
 #include "Animation/AnimTrace.h"
 #include "BoneControllers/AnimNode_OffsetRootBone.h"
@@ -148,7 +148,7 @@ void FAnimNode_Steering::EvaluateSkeletalControl_AnyThread(FComponentSpacePoseCo
 						{
 							// Apply linear correction
 							FQuat LinearCorrection = FQuat::Identity;
-							AnimationMath::SimpleSpringDamperQuat(IN OUT LinearCorrection, IN OUT AngularVelocity, DeltaToTargetOrientation, AnimationMath::SpringHalfLifeToSmoothingTime(ProceduralTargetTime), DeltaSeconds);
+							SpringMath::SimpleSpringDamperQuat(IN OUT LinearCorrection, IN OUT AngularVelocity, DeltaToTargetOrientation, SpringMath::SpringHalfLifeToSmoothingTime(ProceduralTargetTime), DeltaSeconds);
 
 							UE_VLOG_ARROW(Output.AnimInstanceProxy->GetAnimInstanceObject(), "Steering", Display,
 								RootBoneTransform.GetLocation(),
diff --git a/Engine/Plugins/Experimental/Mover/Source/Mover/Private/DefaultMovementSet/Modes/SimpleSpringState.cpp b/Engine/Plugins/Experimental/Mover/Source/Mover/Private/DefaultMovementSet/Modes/SimpleSpringState.cpp
new file mode 100644
--- /dev/null
+++ b/Engine/Plugins/Experimental/Mover/Source/Mover/Private/DefaultMovementSet/Modes/SimpleSpringState.cpp
@@ -0,0 +1,52 @@
+// Copyright Epic Games, Inc. All Rights Reserved.
+
+#include "SimpleSpringState.h"
+
+UScriptStruct* FSimpleSpringState::GetScriptStruct() const
+{ 
+	return StaticStruct(); 
+}
+
+FMoverDataStructBase* FSimpleSpringState::Clone() const
+{ 
+	return new FSimpleSpringState(*this); 
+}
+
+bool FSimpleSpringState::NetSerialize(FArchive& Ar, UPackageMap* Map, bool& bOutSuccess)
+{
+	bool bSuccess = Super::NetSerialize(Ar, Map, bOutSuccess);
+
+	// Could be quantized to save bandwidth
+	Ar << CurrentVelocity;
+	Ar << CurrentAccel;
+	Ar << CurrentAngularVelocity;
+
+	return bSuccess;
+}
+
+void FSimpleSpringState::ToString(FAnsiStringBuilderBase& Out) const
+{
+	Super::ToString(Out);
+
+	Out.Appendf("Vel=%s Accel=%s AngVel=%s\n", *CurrentVelocity.ToCompactString(), *CurrentAccel.ToCompactString(), *CurrentAngularVelocity.ToCompactString());
+}
+
+bool FSimpleSpringState::ShouldReconcile(const FMoverDataStructBase& AuthorityState) const
+{
+	const FSimpleSpringState* AuthoritySpringState = static_cast<const FSimpleSpringState*>(&AuthorityState);
+
+	return (!(CurrentVelocity - AuthoritySpringState->CurrentVelocity).IsNearlyZero() ||
+			!(CurrentAccel - AuthoritySpringState->CurrentAccel).IsNearlyZero() ||
+			!(CurrentAngularVelocity - AuthoritySpringState->CurrentAngularVelocity).IsNearlyZero());
+}
+
+
+void FSimpleSpringState::Interpolate(const FMoverDataStructBase& From, const FMoverDataStructBase& To, float Pct)
+{
+	const FSimpleSpringState* FromState = static_cast<const FSimpleSpringState*>(&From);
+	const FSimpleSpringState* ToState   = static_cast<const FSimpleSpringState*>(&To);
+
+	CurrentVelocity        = FMath::Lerp(FromState->CurrentVelocity, ToState->CurrentVelocity, Pct);
+	CurrentAccel           = FMath::Lerp(FromState->CurrentAccel, ToState->CurrentAccel, Pct);
+	CurrentAngularVelocity = FMath::Lerp(FromState->CurrentAngularVelocity, ToState->CurrentAngularVelocity, Pct);
+}
diff --git a/Engine/Plugins/Experimental/Mover/Source/Mover/Private/DefaultMovementSet/Modes/SimpleSpringState.h b/Engine/Plugins/Experimental/Mover/Source/Mover/Private/DefaultMovementSet/Modes/SimpleSpringState.h
new file mode 100644
--- /dev/null
+++ b/Engine/Plugins/Experimental/Mover/Source/Mover/Private/DefaultMovementSet/Modes/SimpleSpringState.h
@@ -0,0 +1,37 @@
+// Copyright Epic Games, Inc. All Rights Reserved.
+
+#pragma once
+
+#include "CoreMinimal.h"
+#include "MovementMode.h"
+#include "MoverTypes.h"
+#include "SimpleSpringState.generated.h"
+
+/**
+* Internal state data for the SpringWalkingState
+*/
+USTRUCT()
+struct FSimpleSpringState : public FMoverDataStructBase
+{
+	GENERATED_BODY()
+	
+	virtual UScriptStruct* GetScriptStruct() const override;
+	virtual FMoverDataStructBase* Clone() const override;
+	virtual bool NetSerialize(FArchive& Ar, UPackageMap* Map, bool& bOutSuccess) override;
+	virtual void ToString(FAnsiStringBuilderBase& Out) const override;
+	virtual bool ShouldReconcile(const FMoverDataStructBase& AuthorityState) const override;
+	virtual void Interpolate(const FMoverDataStructBase& From, const FMoverDataStructBase& To, float Pct) override;
+
+
+	// Velocity of internal spring model. Not necessarily the velocity of the character
+	UPROPERTY(BlueprintReadOnly, Category = "Mover|Experimental")
+	FVector CurrentVelocity = FVector::ZeroVector;
+
+	// Acceleration of internal spring model
+	UPROPERTY(BlueprintReadOnly, Category = "Mover|Experimental")
+	FVector CurrentAccel = FVector::ZeroVector;
+
+	// Angular velocity of internal spring model in radians per second
+	UPROPERTY(BlueprintReadOnly, Category = "Mover|Experimental")
+	FVector CurrentAngularVelocity = FVector::ZeroVector;
+};
diff --git a/Engine/Plugins/Experimental/Mover/Source/Mover/Private/DefaultMovementSet/Modes/SimpleSpringWalkingMode.cpp b/Engine/Plugins/Experimental/Mover/Source/Mover/Private/DefaultMovementSet/Modes/SimpleSpringWalkingMode.cpp
new file mode 100644
--- /dev/null
+++ b/Engine/Plugins/Experimental/Mover/Source/Mover/Private/DefaultMovementSet/Modes/SimpleSpringWalkingMode.cpp
@@ -0,0 +1,42 @@
+// Copyright Epic Games, Inc. All Rights Reserved.
+
+#include "DefaultMovementSet/Modes/SimpleSpringWalkingMode.h"
+#include "DefaultMovementSet/Modes/SimpleSpringState.h"
+
+#include "MoverComponent.h"
+#include "DefaultMovementSet/Settings/CommonLegacyMovementSettings.h"
+#include "Animation/SpringMath.h"
+
+#include UE_INLINE_GENERATED_CPP_BY_NAME(SimpleSpringWalkingMode)
+
+void USimpleSpringWalkingMode::SimulationTick_Implementation(const FSimulationTickParams& Params, FMoverTickEndData& OutputState)
+{
+	Super::SimulationTick_Implementation(Params, OutputState);
+
+	// We've already updated the spring state during GenerateMove, and just need to copy it into the output simulation state
+	if (const FSimpleSpringState* InSpringState = Params.StartState.SyncState.SyncStateCollection.FindDataByType<FSimpleSpringState>())
+	{
+		FSimpleSpringState& OutputSpringState = OutputState.SyncState.SyncStateCollection.FindOrAddMutableDataByType<FSimpleSpringState>();
+		OutputSpringState = *InSpringState;
+	}
+}
+
+void USimpleSpringWalkingMode::GenerateWalkMove(const FMoverTickStartData& StartState, float DeltaSeconds, const FVector& DesiredVelocity,
+	const FQuat& DesiredFacing, FVector& OutVelocity, FQuat& OutFacing) const
+{
+	const FMoverDefaultSyncState* StartingSyncState = StartState.SyncState.SyncStateCollection.FindDataByType<FMoverDefaultSyncState>();
+	FVector CurrentPosition = StartingSyncState->GetLocation_WorldSpace();
+	FVector PrevPosition = CurrentPosition;
+
+	FMoverTickStartData* MutableStartData = const_cast<FMoverTickStartData*>(&StartState);
+	FSimpleSpringState& SpringState = MutableStartData->SyncState.SyncStateCollection.FindOrAddMutableDataByType<FSimpleSpringState>();
+
+	// Linear //
+	SpringMath::SpringCharacterUpdate(CurrentPosition, SpringState.CurrentVelocity, SpringState.CurrentAccel, DesiredVelocity, VelocitySmoothingTime, DeltaSeconds, VelocityDeadzoneThreshold);
+	OutVelocity = (CurrentPosition - PrevPosition) / DeltaSeconds;
+
+	// Angular //
+	OutFacing = StartingSyncState->GetOrientation_WorldSpace().Quaternion();
+	SpringMath::SimpleSpringDamperQuat(OutFacing, SpringState.CurrentAngularVelocity, DesiredFacing, FacingSmoothingTime, DeltaSeconds);
+}
+
diff --git a/Engine/Plugins/Experimental/Mover/Source/Mover/Private/DefaultMovementSet/Modes/SimpleWalkingMode.cpp b/Engine/Plugins/Experimental/Mover/Source/Mover/Private/DefaultMovementSet/Modes/SimpleWalkingMode.cpp
new file mode 100644
--- /dev/null
+++ b/Engine/Plugins/Experimental/Mover/Source/Mover/Private/DefaultMovementSet/Modes/SimpleWalkingMode.cpp
@@ -0,0 +1,110 @@
+// Copyright Epic Games, Inc. All Rights Reserved.
+
+#include "DefaultMovementSet/Modes/SimpleWalkingMode.h"
+
+#include "MoverComponent.h"
+#include "DefaultMovementSet/Settings/CommonLegacyMovementSettings.h"
+
+void USimpleWalkingMode::GenerateMove_Implementation(const FMoverTickStartData& StartState, const FMoverTimeStep& TimeStep,
+                                                   FProposedMove& OutProposedMove) const
+{
+	const FMoverDefaultSyncState* StartingSyncState = StartState.SyncState.SyncStateCollection.FindDataByType<FMoverDefaultSyncState>();
+
+	if (CommonLegacySettings.Get() == nullptr || !StartingSyncState)
+	{
+		return;
+	}
+
+	const float DeltaSeconds = TimeStep.StepMs * 0.001f;
+	if (DeltaSeconds <= FLT_EPSILON)
+	{
+		return;
+	}
+
+	// Get input
+	FVector DesiredVelocity;
+	EMoveInputType MoveInputType;
+	FVector DesiredFacingDir;
+	 
+	if (const FCharacterDefaultInputs* CharacterInputs = StartState.InputCmd.InputCollection.FindDataByType<FCharacterDefaultInputs>())
+	{
+		DesiredVelocity = CharacterInputs->GetMoveInput_WorldSpace();
+		MoveInputType = CharacterInputs->GetMoveInputType();
+		DesiredFacingDir = CharacterInputs->GetOrientationIntentDir_WorldSpace();
+	}
+	else
+	{
+		DesiredVelocity = FVector::ZeroVector;
+		MoveInputType = EMoveInputType::DirectionalIntent;
+		DesiredFacingDir = FVector::ZeroVector;
+	}
+	
+	
+	float MaxMoveSpeed = CommonLegacySettings->MaxSpeed;
+
+	// Subtract vertical component but keep same magnitude
+	float DesiredVelMag = DesiredVelocity.Length();
+	DesiredVelocity -= DesiredVelocity.ProjectOnTo(GetMoverComponent()->GetUpDirection());
+	float DesiredVel2DSquaredLength = DesiredVelocity.SquaredLength(); 
+	if (DesiredVel2DSquaredLength > 0.0f)
+	{
+		DesiredVelocity *= DesiredVelMag / FMath::Sqrt(DesiredVel2DSquaredLength); 
+	}
+
+	switch (MoveInputType)
+	{
+	case EMoveInputType::DirectionalIntent:
+		{
+			OutProposedMove.DirectionIntent = DesiredVelocity;	// here, DesiredVelocity is already in "intent space" (unit length for "max intent") so we can use it directly
+			DesiredVelocity *= MaxMoveSpeed;
+		}
+		break;
+	case EMoveInputType::Velocity:
+		{
+			// Clamp to max move speed
+			DesiredVelocity = DesiredVelocity.GetClampedToMaxSize(MaxMoveSpeed);
+			OutProposedMove.DirectionIntent = DesiredVelocity / MaxMoveSpeed; // here, DesiredVelocity is converted to "intent space"
+		}
+		break;
+	case EMoveInputType::None:
+	case EMoveInputType::Invalid:
+	default:
+		{
+			UE_LOG(LogMover, Warning, TEXT("Unhandled MoveInputType %i in USpringWalkingMode"), MoveInputType);
+			DesiredVelocity = FVector::ZeroVector;
+		}
+	break;
+	}
+
+	OutProposedMove.bHasDirIntent = !OutProposedMove.DirectionIntent.IsNearlyZero();
+	DesiredFacingDir -= DesiredFacingDir.ProjectOnTo(GetMoverComponent()->GetUpDirection());
+	FQuat CurrentFacing = StartingSyncState->GetOrientation_WorldSpace().Quaternion();
+	FQuat DesiredFacing = CurrentFacing;
+	FQuat PrevFacing = CurrentFacing;
+	
+	if (DesiredFacingDir.Normalize())
+	{
+		DesiredFacing = FQuat::FindBetween(FVector::ForwardVector, DesiredFacingDir);
+	}
+
+	FQuat OutputFacing = CurrentFacing;
+	OutProposedMove.LinearVelocity = StartingSyncState->GetVelocity_WorldSpace();
+	GenerateWalkMove(StartState, DeltaSeconds, DesiredVelocity, DesiredFacing, OutProposedMove.LinearVelocity, OutputFacing);
+
+	// Calc angular velocity
+	FQuat Diff = (PrevFacing.Inverse() * OutputFacing);
+	Diff.Normalize();
+	Diff.EnforceShortestArcWith(FQuat::Identity);
+
+	FRotator AngularVelocity = Diff.Rotator();
+	AngularVelocity *= (1.0f / DeltaSeconds);
+	OutProposedMove.AngularVelocity = AngularVelocity;
+}
+
+void USimpleWalkingMode::GenerateWalkMove(const FMoverTickStartData& StartState, float DeltaSeconds, const FVector& DesiredVelocity,
+	const FQuat& DesiredFacing, FVector& OutVelocity, FQuat& OutFacing) const
+{
+	OutVelocity = DesiredVelocity;
+	OutFacing = DesiredFacing;
+}
+
diff --git a/Engine/Plugins/Experimental/Mover/Source/Mover/Private/DefaultMovementSet/Modes/SmoothWalkingMode.cpp b/Engine/Plugins/Experimental/Mover/Source/Mover/Private/DefaultMovementSet/Modes/SmoothWalkingMode.cpp
new file mode 100644
--- /dev/null
+++ b/Engine/Plugins/Experimental/Mover/Source/Mover/Private/DefaultMovementSet/Modes/SmoothWalkingMode.cpp
@@ -0,0 +1,95 @@
+// Copyright Epic Games, Inc. All Rights Reserved.
+
+#include "DefaultMovementSet/Modes/SmoothWalkingMode.h"
+#include "DefaultMovementSet/Modes/SmoothWalkingState.h"
+
+#include "MoverComponent.h"
+#include "DefaultMovementSet/Settings/CommonLegacyMovementSettings.h"
+#include "Animation/SpringMath.h"
+
+#include UE_INLINE_GENERATED_CPP_BY_NAME(SmoothWalkingMode)
+
+void USmoothWalkingMode::SimulationTick_Implementation(const FSimulationTickParams& Params, FMoverTickEndData& OutputState)
+{
+	Super::SimulationTick_Implementation(Params, OutputState);
+
+	// We've already updated the spring state during GenerateMove, and just need to copy it into the output simulation state
+	if (const FSmoothWalkingState* InSpringState = Params.StartState.SyncState.SyncStateCollection.FindDataByType<FSmoothWalkingState>())
+	{
+		FSmoothWalkingState& OutputSpringState = OutputState.SyncState.SyncStateCollection.FindOrAddMutableDataByType<FSmoothWalkingState>();
+		OutputSpringState = *InSpringState;
+	}
+}
+
+void USmoothWalkingMode::GenerateWalkMove(const FMoverTickStartData& StartState, float DeltaSeconds, const FVector& DesiredVelocity,
+	const FQuat& DesiredFacing, FVector& InOutVelocity, FQuat& InOutFacing) const
+{
+	if (DeltaSeconds <= FLT_EPSILON)
+	{
+		return;
+	}
+	
+	const FMoverDefaultSyncState* StartingSyncState = StartState.SyncState.SyncStateCollection.FindDataByType<FMoverDefaultSyncState>();
+
+	FMoverTickStartData* MutableStartData = const_cast<FMoverTickStartData*>(&StartState);
+	FSmoothWalkingState& SpringState = MutableStartData->SyncState.SyncStateCollection.FindOrAddMutableDataByType<FSmoothWalkingState>();
+
+	bool bIsAccelerating = (1.1f * DesiredVelocity.SquaredLength()) > InOutVelocity.SquaredLength();
+	float MaxAccelerationToUse =  bIsAccelerating ? TargetAcceleration : TargetDeceleration;
+
+	float DesiredVelocityLengthSq = DesiredVelocity.SquaredLength();
+	check(CommonLegacySettings);	// Should be already validated by this point by base class
+	float MaxSpeed = CommonLegacySettings->MaxSpeed;
+
+	if (TurningStrength > FLT_EPSILON)
+	{
+		// 'Friction' or Turning spring
+		FVector TargetMoveDir = FVector::ZeroVector;
+		if(DesiredVelocityLengthSq > UE_KINDA_SMALL_NUMBER)
+		{
+			TargetMoveDir = DesiredVelocity.GetSafeNormal();
+		}
+		else if(!InOutVelocity.IsNearlyZero())
+		{
+			TargetMoveDir = InOutVelocity.GetSafeNormal();
+		}
+
+		if (!TargetMoveDir.IsNearlyZero())
+		{
+			SpringMath::SimpleSpringDamper(InOutVelocity, SpringState.TurningAccel, TargetMoveDir * InOutVelocity.Length(), SpringMath::SpringStrengthToSmoothingTime(TurningStrength), DeltaSeconds);
+			InOutVelocity = InOutVelocity.GetClampedToMaxSize(MaxSpeed);
+		}
+	}
+
+	// Spring towards desired velocity using a velocity spring (which gives control over acceleration)
+
+	// First we rotate the intermediate velocity towards target so that it doesn't lag in direction
+	if(DesiredVelocityLengthSq > UE_KINDA_SMALL_NUMBER)
+	{
+		float VelIntLength = SpringState.IntermediateVelocity.Length();
+		SpringState.IntermediateVelocity = (VelIntLength / FMath::Sqrt(DesiredVelocityLengthSq)) * DesiredVelocity;
+	}
+
+	// Now update the velocity spring
+	FVector PreviousVelocity = InOutVelocity;
+	SpringMath::VelocitySpringDamper(InOutVelocity, SpringState.Accel, SpringState.IntermediateVelocity, DesiredVelocity, MaxAccelerationToUse, VelocitySmoothingTime, DeltaSeconds);
+
+	// Clamp accelerations and velocities
+	FVector ActualAccel = (InOutVelocity - PreviousVelocity) / DeltaSeconds;
+
+	float AccelClamp = MaxAccelerationToUse;
+	if(VelocitySmoothingTime > 0.0f)
+	{
+		AccelClamp *= 1.2f; // add a small buffer to allow spring to catch up to target
+	}
+
+	ActualAccel = ActualAccel.GetClampedToMaxSize(AccelClamp);
+	InOutVelocity = PreviousVelocity + (ActualAccel * DeltaSeconds);
+	InOutVelocity = InOutVelocity.GetClampedToMaxSize(MaxSpeed);
+
+	// Angular / Facing //
+	// We use a simple spring
+	SpringMath::SimpleSpringDamperQuat(InOutFacing, SpringState.CurrentAngularVelocity, DesiredFacing, FacingSmoothingTime, DeltaSeconds);
+}
+
+
diff --git a/Engine/Plugins/Experimental/Mover/Source/Mover/Private/DefaultMovementSet/Modes/SmoothWalkingState.cpp b/Engine/Plugins/Experimental/Mover/Source/Mover/Private/DefaultMovementSet/Modes/SmoothWalkingState.cpp
new file mode 100644
--- /dev/null
+++ b/Engine/Plugins/Experimental/Mover/Source/Mover/Private/DefaultMovementSet/Modes/SmoothWalkingState.cpp
@@ -0,0 +1,55 @@
+// Copyright Epic Games, Inc. All Rights Reserved.
+
+#include "SmoothWalkingState.h"
+
+UScriptStruct* FSmoothWalkingState::GetScriptStruct() const
+{ 
+	return StaticStruct(); 
+}
+
+FMoverDataStructBase* FSmoothWalkingState::Clone() const
+{ 
+	return new FSmoothWalkingState(*this); 
+}
+
+bool FSmoothWalkingState::NetSerialize(FArchive& Ar, UPackageMap* Map, bool& bOutSuccess)
+{
+	bool bSuccess = Super::NetSerialize(Ar, Map, bOutSuccess);
+
+	// Could be quantized to save bandwidth
+	Ar << Accel;
+	Ar << IntermediateVelocity;
+	Ar << TurningAccel;
+	Ar << CurrentAngularVelocity;
+
+	return bSuccess;
+}
+
+void FSmoothWalkingState::ToString(FAnsiStringBuilderBase& Out) const
+{
+	Super::ToString(Out);
+
+	Out.Appendf("Accel=%s IntVel=%s TurningAccel=% AngVel=%s\n", *Accel.ToCompactString(), *IntermediateVelocity.ToCompactString(), *TurningAccel.ToCompactString(), *CurrentAngularVelocity.ToCompactString());
+}
+
+bool FSmoothWalkingState::ShouldReconcile(const FMoverDataStructBase& AuthorityState) const
+{
+	const FSmoothWalkingState* AuthoritySpringState = static_cast<const FSmoothWalkingState*>(&AuthorityState);
+
+	return (!(Accel - AuthoritySpringState->Accel).IsNearlyZero() ||
+			!(IntermediateVelocity - AuthoritySpringState->IntermediateVelocity).IsNearlyZero() ||
+			!(TurningAccel - AuthoritySpringState->TurningAccel).IsNearlyZero() ||
+			!(CurrentAngularVelocity - AuthoritySpringState->CurrentAngularVelocity).IsNearlyZero());
+}
+
+
+void FSmoothWalkingState::Interpolate(const FMoverDataStructBase& From, const FMoverDataStructBase& To, float Pct)
+{
+	const FSmoothWalkingState* FromState = static_cast<const FSmoothWalkingState*>(&From);
+	const FSmoothWalkingState* ToState   = static_cast<const FSmoothWalkingState*>(&To);
+
+	Accel = FMath::Lerp(FromState->Accel, ToState->Accel, Pct);
+	IntermediateVelocity = FMath::Lerp(FromState->IntermediateVelocity, ToState->IntermediateVelocity, Pct);
+	TurningAccel = FMath::Lerp(FromState->TurningAccel, ToState->TurningAccel, Pct);
+	CurrentAngularVelocity = FMath::Lerp(FromState->CurrentAngularVelocity, ToState->CurrentAngularVelocity, Pct);
+}
diff --git a/Engine/Plugins/Experimental/Mover/Source/Mover/Private/DefaultMovementSet/Modes/SmoothWalkingState.h b/Engine/Plugins/Experimental/Mover/Source/Mover/Private/DefaultMovementSet/Modes/SmoothWalkingState.h
new file mode 100644
--- /dev/null
+++ b/Engine/Plugins/Experimental/Mover/Source/Mover/Private/DefaultMovementSet/Modes/SmoothWalkingState.h
@@ -0,0 +1,39 @@
+// Copyright Epic Games, Inc. All Rights Reserved.
+
+#pragma once
+
+#include "CoreMinimal.h"
+#include "MovementMode.h"
+#include "MoverTypes.h"
+#include "SmoothWalkingState.generated.h"
+
+/**
+* Internal state data for the SpringWalkingState
+*/
+USTRUCT()
+struct FSmoothWalkingState : public FMoverDataStructBase
+{
+	GENERATED_BODY()
+	
+	virtual UScriptStruct* GetScriptStruct() const override;
+	virtual FMoverDataStructBase* Clone() const override;
+	virtual bool NetSerialize(FArchive& Ar, UPackageMap* Map, bool& bOutSuccess) override;
+	virtual void ToString(FAnsiStringBuilderBase& Out) const override;
+	virtual bool ShouldReconcile(const FMoverDataStructBase& AuthorityState) const override;
+	virtual void Interpolate(const FMoverDataStructBase& From, const FMoverDataStructBase& To, float Pct) override;
+
+
+	// Acceleration of internal spring model
+	UPROPERTY(BlueprintReadOnly, Category = "Mover|Experimental")
+	FVector Accel = FVector::ZeroVector;
+
+	UPROPERTY(BlueprintReadOnly, Category = "Mover|Experimental")
+	FVector IntermediateVelocity = FVector::ZeroVector;
+
+	UPROPERTY(BlueprintReadOnly, Category = "Mover|Experimental")
+	FVector TurningAccel = FVector::ZeroVector;
+
+	// Angular velocity of internal spring model in radians per second
+	UPROPERTY(BlueprintReadOnly, Category = "Mover|Experimental")
+	FVector CurrentAngularVelocity = FVector::ZeroVector;
+};
diff --git a/Engine/Plugins/Experimental/Mover/Source/Mover/Private/DefaultMovementSet/Modes/SpringWalkingMode.cpp b/Engine/Plugins/Experimental/Mover/Source/Mover/Private/DefaultMovementSet/Modes/SpringWalkingMode.cpp
deleted file mode 100644
diff --git a/Engine/Plugins/Experimental/Mover/Source/Mover/Private/DefaultMovementSet/Modes/SpringWalkingState.cpp b/Engine/Plugins/Experimental/Mover/Source/Mover/Private/DefaultMovementSet/Modes/SpringWalkingState.cpp
deleted file mode 100644
diff --git a/Engine/Plugins/Experimental/Mover/Source/Mover/Private/DefaultMovementSet/Modes/SpringWalkingState.h b/Engine/Plugins/Experimental/Mover/Source/Mover/Private/DefaultMovementSet/Modes/SpringWalkingState.h
deleted file mode 100644
diff --git a/Engine/Plugins/Experimental/Mover/Source/Mover/Public/DefaultMovementSet/Modes/SimpleSpringWalkingMode.h b/Engine/Plugins/Experimental/Mover/Source/Mover/Public/DefaultMovementSet/Modes/SimpleSpringWalkingMode.h
new file mode 100644
--- /dev/null
+++ b/Engine/Plugins/Experimental/Mover/Source/Mover/Public/DefaultMovementSet/Modes/SimpleSpringWalkingMode.h
@@ -0,0 +1,37 @@
+// Copyright Epic Games, Inc. All Rights Reserved.
+
+#pragma once
+
+#include "CoreMinimal.h"
+#include "SimpleWalkingMode.h"
+#include "SimpleSpringWalkingMode.generated.h"
+
+#define UE_API MOVER_API
+
+/**
+ * A walking mode that uses a critically damped spring for translation and rotation.
+ * The strength of the critically damped spring is set via smoothing times (separate for translation and rotation)
+ */
+UCLASS(BlueprintType)
+class USimpleSpringWalkingMode : public USimpleWalkingMode
+{
+	GENERATED_BODY()
+
+public:
+	UE_API virtual void SimulationTick_Implementation(const FSimulationTickParams& Params, FMoverTickEndData& OutputState) override;
+	UE_API virtual void GenerateWalkMove(const FMoverTickStartData& StartState, float DeltaSeconds, const FVector& DesiredVelocity,
+									 const FQuat& DesiredFacing, FVector& OutVelocity, FQuat& OutFacing) const override;
+
+protected:
+	UPROPERTY(EditAnywhere, BlueprintReadWrite, Category = "Mover|Spring Settings", meta = (ClampMin = "0", UIMin = "0", ForceUnits = "s"))
+	float VelocitySmoothingTime = 0.2f;
+
+	UPROPERTY(EditAnywhere, BlueprintReadWrite, Category = "Mover|Spring Settings", meta = (ClampMin = "0", UIMin = "0", ForceUnits = "s"))
+	float FacingSmoothingTime = 0.2f;
+
+	UPROPERTY(EditAnywhere, BlueprintReadWrite, Category = "Mover|Spring Settings", meta = (ClampMin = "0", UIMin = "0", ForceUnits = "cm/s"))
+	// Below this speed we set velocity to 0
+	float VelocityDeadzoneThreshold = 0.1f;
+};
+
+#undef UE_API
diff --git a/Engine/Plugins/Experimental/Mover/Source/Mover/Public/DefaultMovementSet/Modes/SimpleWalkingMode.h b/Engine/Plugins/Experimental/Mover/Source/Mover/Public/DefaultMovementSet/Modes/SimpleWalkingMode.h
new file mode 100644
--- /dev/null
+++ b/Engine/Plugins/Experimental/Mover/Source/Mover/Public/DefaultMovementSet/Modes/SimpleWalkingMode.h
@@ -0,0 +1,28 @@
+// Copyright Epic Games, Inc. All Rights Reserved.
+
+#pragma once
+
+#include "CoreMinimal.h"
+#include "DefaultMovementSet/Modes/WalkingMode.h"
+#include "SimpleWalkingMode.generated.h"
+
+#define UE_API MOVER_API
+
+/**
+ * Basic walking mode that implements the ground based walking
+ */
+UCLASS(BlueprintType, Abstract)
+class USimpleWalkingMode : public UWalkingMode
+{
+	GENERATED_BODY()
+
+public:
+
+	UE_API virtual void GenerateMove_Implementation(const FMoverTickStartData& StartState, const FMoverTimeStep& TimeStep, FProposedMove& OutProposedMove) const override;
+
+	// Override this to make a simple walking mode
+	UE_API virtual void GenerateWalkMove(const FMoverTickStartData& StartState, float DeltaSeconds, const FVector& DesiredVelocity,
+	                                     const FQuat& DesiredFacing, FVector& InOutVelocity, FQuat& InOutFacing) const;
+};
+
+#undef UE_API
diff --git a/Engine/Plugins/Experimental/Mover/Source/Mover/Public/DefaultMovementSet/Modes/SmoothWalkingMode.h b/Engine/Plugins/Experimental/Mover/Source/Mover/Public/DefaultMovementSet/Modes/SmoothWalkingMode.h
new file mode 100644
--- /dev/null
+++ b/Engine/Plugins/Experimental/Mover/Source/Mover/Public/DefaultMovementSet/Modes/SmoothWalkingMode.h
@@ -0,0 +1,47 @@
+// Copyright Epic Games, Inc. All Rights Reserved.
+
+#pragma once
+
+#include "CoreMinimal.h"
+#include "SimpleWalkingMode.h"
+#include "DefaultMovementSet/Modes/WalkingMode.h"
+#include "SmoothWalkingMode.generated.h"
+
+#define UE_API MOVER_API
+
+/**
+ * A walking mode that uses a critically damped spring for translation and rotation.
+ * The strength of the critically damped spring is set via smoothing times (separate for translation and rotation)
+ */
+UCLASS(BlueprintType, Experimental)
+class USmoothWalkingMode : public USimpleWalkingMode
+{
+	GENERATED_BODY()
+
+public:
+	UE_API virtual void SimulationTick_Implementation(const FSimulationTickParams& Params, FMoverTickEndData& OutputState) override;
+	UE_API virtual void GenerateWalkMove(const FMoverTickStartData& StartState, float DeltaSeconds, const FVector& DesiredVelocity,
+									 const FQuat& DesiredFacing, FVector& InOutVelocity, FQuat& InOutFacing) const override;
+
+protected:
+	UPROPERTY(EditAnywhere, BlueprintReadWrite, Category = "Mover|Spring Settings", meta = (ClampMin = "0", UIMin = "0", ForceUnits = "s"))
+	float VelocitySmoothingTime = 0.2f;
+
+	UPROPERTY(EditAnywhere, BlueprintReadWrite, Category = "Mover|Spring Settings", meta = (ClampMin = "0", UIMin = "0", ForceUnits = "s"))
+	float FacingSmoothingTime = 0.2f;
+
+	UPROPERTY(EditAnywhere, BlueprintReadWrite, Category = "Mover|Spring Settings", meta = (ClampMin = "0", UIMin = "0", ForceUnits = "cm/s^2"))
+	float TargetAcceleration = 700.0f;
+
+	UPROPERTY(EditAnywhere, BlueprintReadWrite, Category = "Mover|Spring Settings", meta = (ClampMin = "0", UIMin = "0", ForceUnits = "cm/s^2"))
+	float TargetDeceleration = 700.0f;
+
+	UPROPERTY(EditAnywhere, BlueprintReadWrite, Category = "Mover|Spring Settings", meta = (ClampMin = "0", UIMin = "0"))
+	float TurningStrength = 10.0f;
+
+	UPROPERTY(EditAnywhere, BlueprintReadWrite, Category = "Mover|Spring Settings", meta = (ClampMin = "0", UIMin = "0", ForceUnits = "cm/s"))
+	// Below this speed we set velocity to 0
+	float VelocityDeadzoneThreshold = 0.1f;
+};
+
+#undef UE_API
diff --git a/Engine/Plugins/Experimental/Mover/Source/Mover/Public/DefaultMovementSet/Modes/SpringWalkingMode.h b/Engine/Plugins/Experimental/Mover/Source/Mover/Public/DefaultMovementSet/Modes/SpringWalkingMode.h
deleted file mode 100644
diff --git a/Engine/Plugins/Experimental/UAF/UAFWarping/Source/UAFWarping/Private/SteeringTrait.cpp b/Engine/Plugins/Experimental/UAF/UAFWarping/Source/UAFWarping/Private/SteeringTrait.cpp
--- a/Engine/Plugins/Experimental/UAF/UAFWarping/Source/UAFWarping/Private/SteeringTrait.cpp
+++ b/Engine/Plugins/Experimental/UAF/UAFWarping/Source/UAFWarping/Private/SteeringTrait.cpp
@@ -12,7 +12,7 @@
 #include "TraitInterfaces/IGraphFactory.h"
 #include "TraitInterfaces/ITimeline.h"
 #include "VisualLogger/VisualLogger.h"
-#include "Animation/AnimationMath.h"
+#include "Animation/SpringMath.h"
 
 bool bAnimNextSteeringTraitEnabled = true;
 static FAutoConsoleVariableRef CVarAnimNextSteeringTraitEnabled(
@@ -198,7 +198,7 @@ void FAnimNextSteeringTask::Execute(UE::UAF::FEvaluationVM& VM) const
 						{
 							// Apply linear correction
 							FQuat LinearCorrection = FQuat::Identity;
-							AnimationMath::SimpleSpringDamperQuat(LinearCorrection, InstanceData->AngularVelocity, DeltaToTargetOrientation, AnimationMath::SpringHalfLifeToSmoothingTime(SharedData->ProceduralTargetTime), InstanceData->DeltaTime);
+							SpringMath::SimpleSpringDamperQuat(LinearCorrection, InstanceData->AngularVelocity, DeltaToTargetOrientation, SpringMath::SpringHalfLifeToSmoothingTime(SharedData->ProceduralTargetTime), InstanceData->DeltaTime);
 
 #if ENABLE_ANIM_DEBUG
 							UE_VLOG_ARROW(InstanceData->HostObject, "Steering", Display,
diff --git a/Engine/Source/Runtime/Engine/Internal/Animation/AnimationMath.h b/Engine/Source/Runtime/Engine/Internal/Animation/AnimationMath.h
deleted file mode 100644
diff --git a/Engine/Source/Runtime/Engine/Public/Animation/SpringMath.h b/Engine/Source/Runtime/Engine/Public/Animation/SpringMath.h
new file mode 100644
--- /dev/null
+++ b/Engine/Source/Runtime/Engine/Public/Animation/SpringMath.h
@@ -0,0 +1,369 @@
+// Copyright Epic Games, Inc. All Rights Reserved.
+
+#pragma once
+
+#include "Math/MathFwd.h"
+#include "Math/Vector.h"
+#include "Math/Quat.h"
+
+namespace SpringMath
+{
+	// Reference https://theorangeduck.com/page/spring-roll-call 
+
+	constexpr float SpringHalfLifeToDamping(float HalfLife, float eps = 1e-5f)
+	{
+		return (4.0f * UE_LN2) / (HalfLife + eps);
+	}
+
+	constexpr float SpringDampingToHalfLife(float Damping)
+	{
+		return Damping / (4.0f * UE_LN2);
+	}
+
+	constexpr float SpringSmoothingTimeToHalfLife(float SmoothingTime)
+	{
+		return SmoothingTime / UE_LN2;
+	}
+
+	constexpr float SpringHalfLifeToSmoothingTime(float HalfLife)
+	{
+		return HalfLife * UE_LN2;
+	}
+
+	constexpr float SpringSmoothingTimeToDamping(float SmoothingTime, float eps = 1e-5f)
+	{
+		return 4.0f / FMath::Max(SmoothingTime, eps);
+	}
+
+	constexpr float SpringSmoothingTimeToStrength(float SmoothingTime, float eps = 1e-5f)
+	{
+		return 2.0f / FMath::Max(SmoothingTime, eps);
+	}
+
+	constexpr float SpringStrengthToSmoothingTime(float Strength, float eps = 1e-5f)
+	{
+		return 2.0f / FMath::Max(Strength, eps);
+	}
+
+	constexpr float HalfLifeToLag(float HalfLife)
+	{
+		return HalfLife / UE_LN2;
+	}
+
+	constexpr float LagToHalfLife(float Lag)
+	{
+		return Lag * UE_LN2;
+	}
+
+	inline FVector QuatToScaledAngleAxis(const FQuat& Q)
+	{
+		FQuat QLog = Q.Log();
+		return 2.0f * FVector(QLog.X, QLog.Y, QLog.Z);
+	}
+
+	inline FQuat QuatFromScaledAngleAxis(FVector V)
+	{
+		V /= 2.0f;
+		FQuat QLog(V.X, V.Y, V.Z, 0.0f);
+		return QLog.Exp();
+	}
+
+	/** Simplified version of FMath::CriticallyDampedSmoothing where v_goal is assumed to be 0
+	 * 
+	 * @tparam T The type to be damped 
+	 * @param InOutX The value to be damped
+	 * @param InOutV The speed of the value to be damped
+	 * @param TargetX The goal to damp towards
+	 * @param SmoothingTime The smoothing time to use for the spring. Longer times lead to more damped behaviour. A time of 0 will snap X to XGoal
+	 * @param DeltaTime Timestep in seconds
+	 */
+	template<typename T>
+	void SimpleSpringDamper(
+		T& InOutX,
+		T& InOutV,
+		T TargetX,
+		float SmoothingTime,
+		float DeltaTime)
+	{
+		float y = SpringSmoothingTimeToDamping(SmoothingTime) / 2.0f;
+		T j0 = InOutX - TargetX;
+		T j1 = InOutV + j0 * y;
+		float eydt = FMath::InvExpApprox(y * DeltaTime);
+
+		InOutX = eydt * (j0 + j1 * DeltaTime) + TargetX;
+		InOutV = eydt * (InOutV - j1 * y * DeltaTime);
+	}
+
+	/** Specialized quaternion version, uses FVector for angular velocity
+	 * 
+	 * @param InOutRotation The value to be damped
+	 * @param InOutAngularVelocityRadians The angular velocity of the rotation in radians
+	 * @param TargetRotation The target rotation to damp towards
+	 * @param SmoothingTime he smoothing time to use for the spring. Longer times lead to more damped behaviour. A time of 0 will snap InOutRotation to TargetRotation
+	 * @param DeltaTime 
+	 */
+	inline void SimpleSpringDamperQuat(
+		FQuat& InOutRotation,
+		FVector& InOutAngularVelocityRadians,
+		const FQuat& TargetRotation,
+		float SmoothingTime,
+		float DeltaTime)
+	{
+		float y = SpringSmoothingTimeToDamping(SmoothingTime) / 2.0f;
+
+		FQuat diff = InOutRotation * TargetRotation.Inverse();
+		diff.EnforceShortestArcWith(FQuat::Identity);
+		FVector j0 = QuatToScaledAngleAxis(diff);
+		FVector j1 = InOutAngularVelocityRadians + j0 * y;
+
+		float eydt = FMath::InvExpApprox(y * DeltaTime);
+
+		InOutRotation = QuatFromScaledAngleAxis(eydt * (j0 + j1 * DeltaTime)) * TargetRotation;
+		InOutAngularVelocityRadians = eydt * (InOutAngularVelocityRadians - j1 * y * DeltaTime);
+	}
+
+	/** A velocity spring will damp towards a target that follows a fixed linear target velocity, allowing control of the interpolation speed
+	 * while still giving a smoothed beheaviour. A SmoothingTime of 0 will give a linear interpolation between X and TargetX
+	 * 
+	 * @tparam TFloat The type of floating point to damp (e.g. float or double)
+	 * @param InOutX The value to be damped
+	 * @param InOutV The velocity of the value to be damped
+	 * @param InOutXi The intermediate target of the value to be damped
+	 * @param TargetX The target value of X to damp towards
+	 * @param TargetV The desired linear velocity of X to achieve while damping towards X
+	 * @param SmoothingTime The smoothing time to use while damping towards X. Higher values will give more smoothed behaviour. A value of 0 will give a linear interpolation of X to Target
+	 * @param DeltaTime The timestep in seconds
+	 */
+	template <typename TFloat>
+	void VelocitySpringDamper(
+		TFloat& InOutX,
+		TFloat& InOutV,
+		TFloat& InOutXi,
+		TFloat TargetX,
+		TFloat TargetV,
+		float SmoothingTime,
+		float DeltaTime)
+	{
+		float halflife = SpringSmoothingTimeToHalfLife(SmoothingTime);
+		static_assert(std::is_floating_point_v<TFloat>, "TFloat must be floating point");
+
+		TFloat x_diff = ((TargetX - InOutXi) > 0.0f ? 1.0f : -1.0f) * TargetV;
+
+		float t_goal_future = HalfLifeToLag(halflife);
+		TFloat x_goal_future = fabs(TargetX - InOutXi) > t_goal_future * TargetV ? InOutXi + x_diff * t_goal_future : TargetX;
+
+		SimpleSpringDamper(InOutX, InOutV, x_goal_future, halflife, DeltaTime);
+
+		InOutXi = fabs(TargetX - InOutXi) > DeltaTime * TargetV ? InOutXi + x_diff * DeltaTime : TargetX;
+	}
+
+	/** A velocity spring will damp towards a target that follows a fixed linear target velocity, allowing control of the interpolation speed
+	* while still giving a smoothed beheaviour. A SmoothingTime of 0 will give a linear interpolation between X and TargetX
+	 * 
+	 * @tparam TVector 
+	 * @param InOutX The value to be damped
+	 * @param InOutV The velocity of the value to be damped
+	 * @param InOutXi The intermediate target of the value to be damped
+	 * @param TargetX The target value of X to damp towards
+	 * @param MaxV The max velocity to use for the intermediate target interpolation
+	 * @param SmoothingTime The smoothing time to use while damping towards X. Higher values will give more smoothed behaviour. A value of 0 will give a linear interpolation of X to Target
+	 * @param DeltaTime The timestep in seconds
+	 */
+	template <typename TVector>
+	void VelocitySpringDamper(
+		TVector& InOutX,
+		TVector& InOutV,
+		TVector& InOutXi,
+		TVector TargetX,
+		float MaxV,
+		float SmoothingTime,
+		float DeltaTime)
+	{
+		float halflife = SpringSmoothingTimeToHalfLife(SmoothingTime);
+		TVector x_diff = TargetX - InOutXi;
+		float x_diff_length = x_diff.Length();
+		TVector x_diff_dir = x_diff_length > FLT_EPSILON ? x_diff / x_diff_length : TVector::ZeroVector;
+
+		float t_goal_future = HalfLifeToLag(halflife);
+		TVector x_goal_future = x_diff_length > t_goal_future * MaxV ? InOutXi + (x_diff_dir * MaxV) * t_goal_future : TargetX;
+
+		SimpleSpringDamper(InOutX, InOutV, x_goal_future, halflife, DeltaTime);
+
+		InOutXi = x_diff_length > DeltaTime * MaxV ? InOutXi + x_diff_dir * MaxV * DeltaTime : TargetX;
+	}
+
+	/** Update a position representing a character given a target velocity using a simple damped spring
+	 * 
+	 * @param InOutPosition 
+	 * @param InOutVelocity 
+	 * @param InOutAcceleration 
+	 * @param TargetVelocity 
+	 * @param SmoothingTime 
+	 * @param DeltaTime 
+	 * @param VDeadzone 
+	 * @param ADeadzone 
+	 */
+	inline void SpringCharacterUpdate(
+		FVector& InOutPosition,
+		FVector& InOutVelocity,
+		FVector& InOutAcceleration,
+		const FVector& TargetVelocity,
+		float SmoothingTime,
+		float DeltaTime,
+		float VDeadzone = 1e-2,
+		float ADeadzone = 1e-4)
+	{
+		float y = SpringSmoothingTimeToDamping(SmoothingTime) / 2.0f;
+		FVector j0 = InOutVelocity - TargetVelocity;
+		FVector j1 = InOutAcceleration + j0 * y;
+		float eydt = FMath::InvExpApprox(y * DeltaTime);
+
+		InOutPosition = eydt * (((-j1) / (y * y)) + ((-j0 - j1 * DeltaTime) / y)) +
+			(j1 / (y * y)) + j0 / y + TargetVelocity * DeltaTime + InOutPosition;
+		InOutVelocity = eydt * (j0 + j1 * DeltaTime) + TargetVelocity;
+		InOutAcceleration = eydt * (InOutAcceleration - j1 * y * DeltaTime);
+
+		if ((TargetVelocity - InOutVelocity).SquaredLength() < FMath::Square(VDeadzone))
+		{
+			// We reached our target
+			InOutVelocity = TargetVelocity;
+
+			if (InOutAcceleration.SquaredLength() < FMath::Square(ADeadzone))
+			{
+				InOutAcceleration = FVector::ZeroVector;
+			}
+		}
+	}
+
+	/** Gives predicted positions, velocities and accelerations for SpringCharacterUpdate
+	 * 
+	 * @param OutPredictedPositions 
+	 * @param OutPredictedVelocities 
+	 * @param OutPredictedAccelerations 
+	 * @param PredictCount 
+	 * @param CurrentPosition 
+	 * @param CurrentVelocity 
+	 * @param CurrentAcceleration 
+	 * @param VGoal 
+	 * @param SmoothingTime 
+	 * @param SecondsPerPredictionStep 
+	 * @param VDeadzone 
+	 */
+	inline void SpringCharacterPredict(
+		TArrayView<FVector> OutPredictedPositions,
+		TArrayView<FVector> OutPredictedVelocities,
+		TArrayView<FVector> OutPredictedAccelerations,
+		int PredictCount,
+		const FVector& CurrentPosition,
+		const FVector& CurrentVelocity,
+		const FVector& CurrentAcceleration,
+		const FVector& VGoal,
+		float SmoothingTime,
+		float SecondsPerPredictionStep,
+		float VDeadzone)
+	{
+		for (int i = 0; i < PredictCount; i++)
+		{
+			OutPredictedPositions[i] = CurrentPosition;
+			OutPredictedVelocities[i] = CurrentVelocity;
+			OutPredictedAccelerations[i] = CurrentAcceleration;
+		}
+
+		for (int i = 0; i < PredictCount; i++)
+		{
+			const float PredictTime = (i + 1) * SecondsPerPredictionStep; // Note i+1 since we want index 0 to be the first prediction step
+			SpringCharacterUpdate(OutPredictedPositions[i], OutPredictedVelocities[i], OutPredictedAccelerations[i], VGoal, SmoothingTime,
+				PredictTime, VDeadzone);
+		}
+	}
+
+	/** Update a position representing a character given a target velocity using a velocity spring
+	 * 
+	 * @tparam TVector 
+	 * @param InOutPosition 
+	 * @param InOutVelocity 
+	 * @param InOutVelocityIntermediate 
+	 * @param InOutAccel 
+	 * @param TargetVelocity 
+	 * @param SmoothingTime 
+	 * @param MacAccel 
+	 * @param DeltaTime 
+	 */
+	template <typename TVector>
+	void VelocitySpringCharacterUpdate(
+		TVector& InOutPosition,
+		TVector& InOutVelocity,
+		TVector& InOutVelocityIntermediate,
+		TVector& InOutAccel,
+		TVector TargetVelocity,
+		float SmoothingTime,
+		float MacAccel,
+		float DeltaTime)
+	{
+		float halflife = SpringSmoothingTimeToHalfLife(SmoothingTime);
+		TVector v_diff = TargetVelocity - InOutVelocityIntermediate;
+		float v_diff_length = v_diff.Length();
+		TVector v_diff_dir = v_diff_length > 0.0001f ? v_diff / v_diff_length : TVector::ZeroVector;
+
+		float t_goal_future = HalfLifeToLag(halflife);
+		TVector max_v_future = v_diff_length > t_goal_future * MacAccel ? InOutVelocityIntermediate + (v_diff_dir * MacAccel) * t_goal_future : TargetVelocity;
+
+		float y = SpringHalfLifeToDamping(halflife) / 2.0f;
+		TVector j0 = InOutVelocity - max_v_future;
+		TVector j1 = InOutAccel + j0 * y;
+		float eydt = FMath::InvExpApprox(y * DeltaTime);
+
+		InOutPosition = eydt * (((-j1) / (y * y)) + ((-j0 - j1 * DeltaTime) / y)) +
+			(j1 / (y * y)) + j0 / y + max_v_future * DeltaTime + InOutPosition;
+		InOutVelocity = eydt * (j0 + j1 * DeltaTime) + max_v_future;
+		InOutAccel = eydt * (InOutAccel - j1 * y * DeltaTime);
+		InOutVelocityIntermediate = v_diff_length > DeltaTime * MacAccel ? InOutVelocityIntermediate + v_diff_dir * MacAccel * DeltaTime : TargetVelocity;
+	}
+
+
+
+	/** Specialized quaternion damper, similar to FMath::ExponentialSmoothingApprox but for quats
+	 * 
+	 * @param CurrentRotation 
+	 * @param TargetRotation 
+	 * @param SmoothingTime 
+	 * @param DeltaTime 
+	 * @param Eps 
+	 * @return 
+	 */
+	FORCEINLINE FQuat SpringDamperQuat(const FQuat& CurrentRotation, const FQuat& TargetRotation, float SmoothingTime, float DeltaTime, float Eps = 1e-5f)
+	{
+		if (SmoothingTime < Eps)
+			return TargetRotation;
+
+		return FQuat::Slerp(CurrentRotation, TargetRotation, 1.0f - FMath::InvExpApprox((DeltaTime) / (SmoothingTime)));
+	}
+
+	/** Prediction of SpringDamperQuat
+	 * 
+	 * @param OutPredictedRotations 
+	 * @param OutPredictedAngularVelocities 
+	 * @param PredictorCount 
+	 * @param CurrentRotation 
+	 * @param CurrentAngularVelocity 
+	 * @param TargetRotation 
+	 * @param SmoothingTime 
+	 * @param SecondsPerPredictionStep 
+	 */
+	inline void SimpleSpringDamperQuatPredict(TArrayView<FQuat> OutPredictedRotations, TArrayView<FVector> OutPredictedAngularVelocities,
+	                                          int PredictorCount, const FQuat& CurrentRotation, const FVector& CurrentAngularVelocity,
+	                                          const FQuat& TargetRotation, float SmoothingTime, float SecondsPerPredictionStep)
+	{
+		for (int i = 0; i < PredictorCount; i++)
+		{
+			OutPredictedRotations[i] = CurrentRotation;
+			OutPredictedAngularVelocities[i] = CurrentAngularVelocity;
+		}
+
+		for (int i = 0; i < PredictorCount; i++)
+		{
+			const float PredictTime = (i + 1) * SecondsPerPredictionStep; // Note i+1 since we want index 0 to be the first prediction step
+			SimpleSpringDamperQuat(OutPredictedRotations[i], OutPredictedAngularVelocities[i], TargetRotation, SmoothingTime, PredictTime);
+		}
+	}
+}
\ No newline at end of file
