Add MATERIAL_USES_SPECULAR_PROFILE to enable SpecularProfile baking only for material having a specular profile.

This offers better parity between Format=0 and Format=1 at no extra cost for material not having a specular profile.

diff --git a/Engine/Shaders/Private/Substrate/SubstrateExport.ush b/Engine/Shaders/Private/Substrate/SubstrateExport.ush
--- a/Engine/Shaders/Private/Substrate/SubstrateExport.ush
+++ b/Engine/Shaders/Private/Substrate/SubstrateExport.ush
@@ -963,8 +963,7 @@ FExportResult SubstrateMaterialExportOut(
 		#define CurrentBSDF SubstratePixelHeader.SubstrateTree.BSDFs[BSDFIdx]
 
 			// Bake the Specular Profiles into F0/F90 of BSDFs to blend in order to retain some of the original look when using parameter blending or fully simplified materials.
-			// Disable for now until we have mean to enable this per material
-		#if 0 //SUBSTRATE_SPECPROFILE_ENABLED
+		#if SUBSTRATE_SPECPROFILE_ENABLED && MATERIAL_USES_SPECULAR_PROFILE
 			const float NoV = saturate(dot(SurfaceWorldNormal, V));
 			CurrentBSDF.SubstrateBakeSpecularProfileToReflectivity(NoV);
 			BSDF_SETHASSPECPROFILE(CurrentBSDF, 0);
diff --git a/Engine/Source/Runtime/Core/Private/UObject/FortniteMainBranchObjectVersion.cpp b/Engine/Source/Runtime/Core/Private/UObject/FortniteMainBranchObjectVersion.cpp
--- a/Engine/Source/Runtime/Core/Private/UObject/FortniteMainBranchObjectVersion.cpp
+++ b/Engine/Source/Runtime/Core/Private/UObject/FortniteMainBranchObjectVersion.cpp
@@ -14,6 +14,6 @@ TMap<FGuid, FGuid> FFortniteMainBranchObjectVersion::GetSystemGuids()
 	SystemGuids.Add(DevGuids.Niagara_LatestScriptCompileVersion, FGuid("97F0FD7149D90D42B65D6557AF477D1A"));
 	SystemGuids.Add(DevGuids.SkeletalMeshDerivedDataVersion, FGuid("C998F2F3B2884A02B7E290B8C25835F0"));
 	SystemGuids.Add(DevGuids.STATICMESH_DERIVEDDATA_VER, FGuid("3ABF32149F9D4D838750368AB95573BE"));
-	SystemGuids.Add(DevGuids.MaterialTranslationDDCVersion, FGuid("352803E079F54E7AB52040F59068EF4E"));
+	SystemGuids.Add(DevGuids.MaterialTranslationDDCVersion, FGuid("C6B481BEACC04585B05635E3EC491F99"));
 	return SystemGuids;
 }
diff --git a/Engine/Source/Runtime/Engine/Private/Materials/HLSLMaterialTranslator.cpp b/Engine/Source/Runtime/Engine/Private/Materials/HLSLMaterialTranslator.cpp
--- a/Engine/Source/Runtime/Engine/Private/Materials/HLSLMaterialTranslator.cpp
+++ b/Engine/Source/Runtime/Engine/Private/Materials/HLSLMaterialTranslator.cpp
@@ -265,6 +265,7 @@ struct FHLSLMaterialTranslator::FEnvironmentDefines
 	bool bMaterialComputeFogPerPixel;
 	bool bMaterialFullyRough;
 	bool bMaterialUsesAnisotropy;
+	bool bMaterialUsesSpecularProfile;
 	uint8 MaterialDecalReadMask;
 	int8 PixelDepthOffsetMode;
 	bool bMaterialUsesDecalLookup;
@@ -371,6 +372,7 @@ struct FHLSLMaterialTranslator::FEnvironmentDefines
 		Ar << bMaterialComputeFogPerPixel;
 		Ar << bMaterialFullyRough;
 		Ar << bMaterialUsesAnisotropy;
+		Ar << bMaterialUsesSpecularProfile;
 		Ar << MaterialDecalReadMask;
 		Ar << PixelDepthOffsetMode;
 		Ar << bMaterialUsesDecalLookup;
@@ -2780,6 +2782,7 @@ void FHLSLMaterialTranslator::GetMaterialEnvironment(EShaderPlatform InPlatform,
 	OutEnvironment.SetDefine(TEXT("MATERIAL_FULLY_ROUGH"), EnvironmentDefines->bMaterialFullyRough);
 
 	OutEnvironment.SetDefine(TEXT("MATERIAL_USES_ANISOTROPY"), EnvironmentDefines->bMaterialUsesAnisotropy);
+	OutEnvironment.SetDefine(TEXT("MATERIAL_USES_SPECULAR_PROFILE"), EnvironmentDefines->bMaterialUsesSpecularProfile);
 
 	OutEnvironment.SetDefine(TEXT("MATERIAL_DECAL_READ_MASK"), EnvironmentDefines->MaterialDecalReadMask);
 	OutEnvironment.SetDefine(TEXT("MATERIAL_USES_DECAL_LOOKUP"), EnvironmentDefines->bMaterialUsesDecalLookup);
@@ -16354,6 +16357,7 @@ void FHLSLMaterialTranslator::PrepareEnvironmentDefines()
 	EnvironmentDefines->bMaterialComputeFogPerPixel = Material->ComputeFogPerPixel();
 	EnvironmentDefines->bMaterialFullyRough = bIsFullyRough || Material->IsFullyRough();
 	EnvironmentDefines->bMaterialUsesAnisotropy = (bUsesAnisotropy || EnumHasAnyFlags(MaterialCompilationOutput.SubstrateMaterialCompilationOutput.SubstrateMaterialBsdfFeatures, ESubstrateBsdfFeature::Anisotropy)) && FDataDrivenShaderPlatformInfo::GetSupportsAnisotropicMaterials(InPlatform);
+	EnvironmentDefines->bMaterialUsesSpecularProfile = EnumHasAnyFlags(MaterialCompilationOutput.SubstrateMaterialCompilationOutput.SubstrateMaterialBsdfFeatures, ESubstrateBsdfFeature::SpecularProfile) && Substrate::IsSpecularProfileEnabled(InPlatform);
 	EnvironmentDefines->MaterialDecalReadMask = MaterialCompilationOutput.UsedDBufferTextures;
 	EnvironmentDefines->bMaterialUsesDecalLookup = MaterialCompilationOutput.bUsesDBufferTextureLookup;
 	EnvironmentDefines->MaterialPathTracingBufferRead = MaterialCompilationOutput.UsedPathTracingBufferTextures;
