Fixed UEngineReplicationBridge::AddRootObjectToLevelGroup to generate a unique level group package name.
When LevelStreaming.ShouldReuseUnloadedButStillAroundLevels is set to false, we can end up re-loading a level that was previously unloaded. In this case, the old level package can be trashed and pending to be fully removed by FLevelStreamingGCHelper::PrepareStreamedOutLevelsForGC.
This is why we could end up with a new Level pointer that has a package name that is still valid in the list of level groups.

diff --git a/Engine/Source/Runtime/Engine/Private/Net/Iris/ReplicationSystem/EngineReplicationBridge.cpp b/Engine/Source/Runtime/Engine/Private/Net/Iris/ReplicationSystem/EngineReplicationBridge.cpp
--- a/Engine/Source/Runtime/Engine/Private/Net/Iris/ReplicationSystem/EngineReplicationBridge.cpp
+++ b/Engine/Source/Runtime/Engine/Private/Net/Iris/ReplicationSystem/EngineReplicationBridge.cpp
@@ -892,14 +892,25 @@ void UEngineReplicationBridge::AddRootObjectToLevelGroup(const UObject* RootObje
 	if (Level && (UE::Net::Private::bAlwaysCreateLevelFilteringGroupsForPersistentLevels || !Level->IsPersistentLevel() || (Level != NetDriver->GetWorld()->PersistentLevel)))
 	{
 		const UPackage* const LevelPackage = Level->GetOutermost();
-		const FName PackageName = LevelPackage->GetFName();
+		const FString PackageName = LevelPackage->GetName();
 
 		UE::Net::FNetObjectGroupHandle LevelGroup = GetLevelGroup(Level);
 		if (!LevelGroup.IsValid())
 		{
-			LevelGroup = CreateLevelGroup(Level, PackageName);
+			// When a level is unloaded and reloaded immediately and when LevelStreaming.ShouldReuseUnloadedButStillAroundLevels is false, 
+			// we can end up with 2 different sets of level pointers that are have the same original package name.
+			// Make the package name unique so a new group gets assigned to it.
+			FString LevelGroupPackageUniqueName = PackageName;
+			int32 Index = 0;
+			while (GetReplicationSystem()->FindGroup(*LevelGroupPackageUniqueName).IsValid())
+			{
+				++Index;
+				LevelGroupPackageUniqueName = FString::Printf(TEXT("%s_%d"), *PackageName, Index);
+			}
+
+			LevelGroup = CreateLevelGroup(Level, *LevelGroupPackageUniqueName);
 
-			UE_LOG_BRIDGEID(LogIris, Verbose, TEXT("Created new GroupIndex: %u for Level: %s"), LevelGroup.GetGroupIndex(), ToCStr(PackageName.ToString()));
+			UE_LOG_BRIDGEID(LogIris, Verbose, TEXT("Created new GroupIndex: %u for Level: %s using Package Unique Name: %s"), LevelGroup.GetGroupIndex(), ToCStr(PackageName), ToCStr(LevelGroupPackageUniqueName));
 
 			// Update the filtering status of the group based on current level visibility for all connections
 			NetDriver->UpdateGroupFilterStatusForLevel(Level, LevelGroup);
@@ -908,7 +919,7 @@ void UEngineReplicationBridge::AddRootObjectToLevelGroup(const UObject* RootObje
 		const UE::Net::FNetRefHandle RootObjectRefHandle = GetReplicatedRefHandle(RootObject);
 
 		// Add object to group
-		UE_LOG_BRIDGEID(LogIris, Verbose, TEXT("Added %s to GroupIndex: %u Level: %s"), *RootObjectRefHandle.ToString(), LevelGroup.GetGroupIndex(), ToCStr(PackageName.ToString()));
+		UE_LOG_BRIDGEID(LogIris, Verbose, TEXT("Added %s to GroupIndex: %u Level: %s"), *RootObjectRefHandle.ToString(), LevelGroup.GetGroupIndex(), ToCStr(PackageName));
 		GetReplicationSystem()->AddToGroup(LevelGroup, RootObjectRefHandle);
 	}
 }
