Spline: Use FRWLock for accessing the reparam table and swap dirty bool for monotonically increasing version number

diff --git a/Engine/Source/Runtime/Engine/Private/Curves/Spline.cpp b/Engine/Source/Runtime/Engine/Private/Curves/Spline.cpp
--- a/Engine/Source/Runtime/Engine/Private/Curves/Spline.cpp
+++ b/Engine/Source/Runtime/Engine/Private/Curves/Spline.cpp
@@ -3,6 +3,7 @@
 #include "Curves/Spline.h"
 
 #include "BoxTypes.h"
+#include "Async/TransactionallySafeMutex.h"
 #include "Components/SplineComponent.h"	// for FSplineCurves, ideally removed
 #include "Misc/Base64.h"
 #include "Serialization/MemoryReader.h"
@@ -231,10 +232,49 @@ public:
 	const FInterpCurveVector& GetSplinePointsScale() const { return ScaleCurve; }
 
 private:
+
+	/** Inspired by FDynamicMesh3::FChangeStamp. */
+	struct FChangeStamp
+	{
+		FChangeStamp()
+		{}
+		
+		UE_NONCOPYABLE(FChangeStamp);
+
+		/** Updates the change stamp in an thread-safe, transactionally-safe way. */
+		void Increment()
+		{
+			Mutex.Lock();
+			++Value;
+			Mutex.Unlock();
+		}
+
+		/** Returns the current change value in a thread-safe, transactionally-safe way. */
+		uint32 GetValue() const
+		{
+			Mutex.Lock();
+			uint32 Result = Value;
+			Mutex.Unlock();
+
+			return Result;
+		}
+
+	private:
+		
+		/** Guards `Value`. */
+		mutable UE::FTransactionallySafeMutex Mutex;
+
+		/** The change stamp is incremented when modifications occur. It's guarded by `Mutex`. */
+		uint32 Value = 1;
+	};
+
+	FInterpCurveFloat ReparamTable;
+	FChangeStamp ReparamTableNextVersion;
+	uint32 ReparamTableVersion = 0;
+	FRWLock ReparamTableRWLock;
 	
 	static inline const FName RotationAttrName = FName("Rotation");
 	static inline const FName ScaleAttrName = FName("Scale");
-	FInterpCurveFloat ReparamTable;
 
     /**
      * We keep these up to date with Points.
@@ -264,8 +304,7 @@ private:
 	/** Convert from [0, NumSegments] to internal spline space */
 	float ToInternalSplineSpace(float Parameter) const;
 
-	bool bReparamTableDirty = true;
-	void MarkReparamTableDirty() { bReparamTableDirty = true; }
+	void MarkReparamTableDirty();
 	
 	/** Clears and repopulates reparameterization attribute channels. */
 	void UpdateReparamTable();
@@ -2278,6 +2317,8 @@ float FNewSpline::GetParameterAtDistance(float Distance)
 {
 	UpdateReparamTable();
 	
+	FReadScopeLock Lock(ReparamTableRWLock);
+	
 	return ReparamTable.Eval(Distance);
 }
 
@@ -2289,6 +2330,8 @@ float FNewSpline::GetDistanceAtParameter(float Parameter)
 	}
 	
 	UpdateReparamTable();
+	
+	FReadScopeLock Lock(ReparamTableRWLock);
 
 	const float ParameterMax = static_cast<float>(GetNumberOfSegments());
 	const float Key = (Parameter / ParameterMax) * (ReparamTable.Points.Num() - 1);
@@ -2297,9 +2340,7 @@ float FNewSpline::GetDistanceAtParameter(float Parameter)
 	const int32 UpperKey = FMath::CeilToInt32(Key);
 	ensureAlways(UpperKey >= 0 && UpperKey < ReparamTable.Points.Num());
 	const float Alpha = FMath::Frac(Key);
-	const float Distance = FMath::Lerp(ReparamTable.Points[LowerKey].InVal, ReparamTable.Points[UpperKey].InVal, Alpha);
-
-	return Distance;
+	return FMath::Lerp(ReparamTable.Points[LowerKey].InVal, ReparamTable.Points[UpperKey].InVal, Alpha);
 }
 
 float FNewSpline::FindNearest(const FVector& InLocation, float& OutSquaredDistance) const
@@ -2395,32 +2436,51 @@ float FNewSpline::ToInternalSplineSpace(float Parameter) const
 	return GetSegmentParameterRange(SegmentIndex).Interpolate(LocalT);
 }
 
+void FNewSpline::MarkReparamTableDirty()
+{
+	ReparamTableNextVersion.Increment();
+}
+
 void FNewSpline::UpdateReparamTable()
 {
-	if (!bReparamTableDirty)
+	const int32 NumSegments = GetNumberOfSegments();
+
+	auto IsReparamTableDirty = [this]()
+	{
+		return ReparamTableNextVersion.GetValue() != ReparamTableVersion;
+	};
+	
+	if (!IsReparamTableDirty() || ReparamStepsPerSegment == 0 || NumSegments == 0)
 	{
 		return;
 	}
 
-	const int32 NumSegments = GetNumberOfSegments();
+	FWriteScopeLock Lock(ReparamTableRWLock);
+
+	if (!IsReparamTableDirty())
+	{
+		return;
+	}
+
+	// We can't rely on the next version to not change during the update.
+	const uint32 CachedNextVersion = ReparamTableNextVersion.GetValue();
 
-	// Start by clearing it
 	ReparamTable.Points.Reset(NumSegments * ReparamStepsPerSegment + 1);
 	float AccumulatedLength = 0.0f;
 	for (int32 SegmentIndex = 0; SegmentIndex < NumSegments; ++SegmentIndex)
 	{
-		for (int32 Step = 0; Step < ReparamStepsPerSegment; ++Step)
+		ReparamTable.Points.Emplace(AccumulatedLength, SegmentIndex, 0.0f, 0.0f, CIM_Linear);
+		for (int32 Step = 1; Step < ReparamStepsPerSegment; ++Step)
 		{
 			const float Param = static_cast<float>(Step) / ReparamStepsPerSegment;
-			const float SegmentLength = (Step == 0) ? 0.0f : GetSegmentLength(SegmentIndex, Param);
-
-			ReparamTable.Points.Emplace(SegmentLength + AccumulatedLength, SegmentIndex + Param, 0.0f, 0.0f, CIM_Linear);
+			ReparamTable.Points.Emplace(GetSegmentLength(SegmentIndex, Param) + AccumulatedLength, SegmentIndex + Param, 0.0f, 0.0f, CIM_Linear);
 		}
+
 		AccumulatedLength += GetSegmentLength(SegmentIndex, 1.0f);
 	}
 	ReparamTable.Points.Emplace(AccumulatedLength, static_cast<float>(NumSegments), 0.0f, 0.0f, CIM_Linear);
 
-	bReparamTableDirty = false;
+	ReparamTableVersion = CachedNextVersion;
 }
 
 void FNewSpline::RebuildLegacyCurves()
