Fix r.dynamicres.TestScreenPercentage to take into account dynamic res upper bound if configured. Previously we would always render to full res targets, regardless of the screen percentage, so clears and resolves were not representative.
Note: this will affect performance of certain GPU performance tests
#tests tested locally via replayrun
#DisabledValidations Preflight, checkCopyright, CheckBalancedMacros

diff --git a/Engine/Source/Runtime/Engine/Private/UnrealEngine.cpp b/Engine/Source/Runtime/Engine/Private/UnrealEngine.cpp
--- a/Engine/Source/Runtime/Engine/Private/UnrealEngine.cpp
+++ b/Engine/Source/Runtime/Engine/Private/UnrealEngine.cpp
@@ -13779,6 +13779,13 @@ void UEngine::GetDynamicResolutionCurrentStateInfos(FDynamicResolutionStateInfos
 			OutInfos.ResolutionFractionApproximations = DynamicResolutionState->GetResolutionFractionsApproximation();
 			OutInfos.ResolutionFractionApproximations[GDynamicPrimaryResolutionFraction] = ForceResolutionFraction;
 			OutInfos.ResolutionFractionUpperBounds = DynamicResolutionState->GetResolutionFractionsUpperBoundBudgetValue();
+
+			// If dynamic upper bound is enabled, set the upper bound to the forced resolution to properly simulate that resolution (otherwise clears/resolves may not be representative)
+			static auto CVarUpperBoundQuantization = IConsoleManager::Get().FindConsoleVariable(TEXT("r.DynamicRes.UpperBoundQuantization"));
+			if (CVarUpperBoundQuantization && CVarUpperBoundQuantization->GetFloat() > 0.0f)
+			{
+				OutInfos.ResolutionFractionUpperBounds[GDynamicPrimaryResolutionFraction] = ForceResolutionFraction;
+			}
 		}
 		#endif
 		else if (!DynamicResolutionState->IsSupported())
