[XR] Adding extra 180 rotation to stereo layer object so that UE objects match Open XR coordinates and layers are seen when facing the camera

UE creates them with forward away from camera (no rotation == facing away from camera) where as OpenXr expects them to be facing the camera (no rotation == facing camera) so unless flipped only objects that look like they are facing away from the camera in UE are rendered.

diff --git a/Engine/Source/Runtime/Engine/Private/Components/StereoLayerComponent.cpp b/Engine/Source/Runtime/Engine/Private/Components/StereoLayerComponent.cpp
--- a/Engine/Source/Runtime/Engine/Private/Components/StereoLayerComponent.cpp
+++ b/Engine/Source/Runtime/Engine/Private/Components/StereoLayerComponent.cpp
@@ -14,6 +14,17 @@
 #include "StereoRendering.h"
 #endif
 
+namespace UE::XRStereoLayerCVars
+{
+	static bool bUseLegacyFacing = false;
+	static FAutoConsoleVariableRef UseUECoordinateSystemCVar(
+		TEXT("XR.StereoLayers.UseLegacyFacing"),
+		bUseLegacyFacing,
+		TEXT("If true the old UE coordinate system will be used where layers facing away from the camera"
+			 "are visible instead of the ones seemingly facing it inside of UE"),
+		ECVF_Default);
+}
+
 UStereoLayerComponent::UStereoLayerComponent(const FObjectInitializer& ObjectInitializer)
 	: Super(ObjectInitializer)
 	, bLiveTexture(false)
@@ -108,6 +119,12 @@ void UStereoLayerComponent::TickComponent(float DeltaTime, enum ELevelTick TickT
 		LayerDesc.QuadSize = QuadSize;
 		LayerDesc.UVRect = UVRect;
 		LayerDesc.Transform = Transform;
+		
+		if (!UE::XRStereoLayerCVars::bUseLegacyFacing)
+		{
+			// The default Open Xr Stereo layer object faces the camera at no rotation where as the unreal object faces away from it at zero rotation
+			LayerDesc.Transform.SetRotation(Transform.GetRotation() *= FQuat(0., 0., 1., 0.));
+		}
 
 		if (Texture)
 		{
