[Audio Insights] Send OwnerActorName trace data in standalone version

diff --git a/Engine/Plugins/AudioInsights/Source/AudioInsights/Private/Messages/AudioEventLogTraceMessages.cpp b/Engine/Plugins/AudioInsights/Source/AudioInsights/Private/Messages/AudioEventLogTraceMessages.cpp
--- a/Engine/Plugins/AudioInsights/Source/AudioInsights/Private/Messages/AudioEventLogTraceMessages.cpp
+++ b/Engine/Plugins/AudioInsights/Source/AudioInsights/Private/Messages/AudioEventLogTraceMessages.cpp
@@ -16,7 +16,7 @@ namespace UE::Audio::Insights
 
 		EventData.GetString("EventLogName", EventName);
 		EventData.GetString("AssetPath", AssetPath);
-		EventData.GetString("ActorName", ActorName);
+		EventData.GetString("ActorLabel", ActorLabel);
 		EventData.GetString("ActorIconName", ActorIconName);
 		EventData.GetString("SoundClassName", CategoryName);
 
diff --git a/Engine/Plugins/AudioInsights/Source/AudioInsights/Private/Messages/SoundTraceMessages.cpp b/Engine/Plugins/AudioInsights/Source/AudioInsights/Private/Messages/SoundTraceMessages.cpp
--- a/Engine/Plugins/AudioInsights/Source/AudioInsights/Private/Messages/SoundTraceMessages.cpp
+++ b/Engine/Plugins/AudioInsights/Source/AudioInsights/Private/Messages/SoundTraceMessages.cpp
@@ -47,9 +47,7 @@ namespace UE::Audio::Insights
 			EntryType = ESoundDashboardEntryType::None;
 		}
 
-#if WITH_EDITOR
 		EventData.GetString("ActorLabel", ActorLabel);
-#endif // WITH_EDITOR
 	}
 
 	FSoundDashboardEntry::FSoundDashboardEntry()
diff --git a/Engine/Plugins/AudioInsights/Source/AudioInsights/Private/Providers/AudioEventLogTraceProvider.cpp b/Engine/Plugins/AudioInsights/Source/AudioInsights/Private/Providers/AudioEventLogTraceProvider.cpp
--- a/Engine/Plugins/AudioInsights/Source/AudioInsights/Private/Providers/AudioEventLogTraceProvider.cpp
+++ b/Engine/Plugins/AudioInsights/Source/AudioInsights/Private/Providers/AudioEventLogTraceProvider.cpp
@@ -138,11 +138,11 @@ namespace UE::Audio::Insights
 
 		// If this message did not provide an actor name, the object may have been destroying when the message was sent
 		// Try to locate the actor name using the play order ID
-		if (Msg.ActorName.IsEmpty())
+		if (Msg.ActorLabel.IsEmpty())
 		{
 			if (const FActorInformation* FoundActorInfo = PlayOrderToActorInfoMap.Find(Entry->PlayOrder))
 			{
-				Entry->ActorName = *FoundActorInfo->ActorName;
+				Entry->ActorLabel = *FoundActorInfo->ActorLabel;
 #if WITH_EDITOR
 				Entry->ActorIconName = FoundActorInfo->ActorIconName;
 #else
@@ -151,13 +151,13 @@ namespace UE::Audio::Insights
 			}
 			else
 			{
-				Entry->ActorName = LOCTEXT("AudioEventLog_None", "None").ToString();
+				Entry->ActorLabel = LOCTEXT("AudioEventLog_None", "None").ToString();
 				Entry->ActorIconName = NAME_None;
 			}
 		}
 		else
 		{
-			Entry->ActorName = Msg.ActorName;
+			Entry->ActorLabel = Msg.ActorLabel;
 
 #if WITH_EDITOR
 			Entry->ActorIconName = FName(Msg.ActorIconName);
@@ -166,7 +166,7 @@ namespace UE::Audio::Insights
 #endif // WITH_EDITOR
 			
 			// Cache the actor information associated with this PlayOrder ID so we can look it up later
-			PlayOrderToActorInfoMap.Add(Entry->PlayOrder, FActorInformation{ Msg.ActorName, FName(Msg.ActorIconName) });
+			PlayOrderToActorInfoMap.Add(Entry->PlayOrder, FActorInformation{ Msg.ActorLabel, FName(Msg.ActorIconName) });
 		}
 
 		return Entry;
diff --git a/Engine/Plugins/AudioInsights/Source/AudioInsights/Private/Providers/AudioEventLogTraceProvider.h b/Engine/Plugins/AudioInsights/Source/AudioInsights/Private/Providers/AudioEventLogTraceProvider.h
--- a/Engine/Plugins/AudioInsights/Source/AudioInsights/Private/Providers/AudioEventLogTraceProvider.h
+++ b/Engine/Plugins/AudioInsights/Source/AudioInsights/Private/Providers/AudioEventLogTraceProvider.h
@@ -30,7 +30,7 @@ namespace UE::Audio::Insights
 
 		struct FActorInformation
 		{
-			const FString ActorName;
+			const FString ActorLabel;
 			const FName ActorIconName;
 		};
 
diff --git a/Engine/Plugins/AudioInsights/Source/AudioInsights/Private/Providers/SoundTraceProvider.cpp b/Engine/Plugins/AudioInsights/Source/AudioInsights/Private/Providers/SoundTraceProvider.cpp
--- a/Engine/Plugins/AudioInsights/Source/AudioInsights/Private/Providers/SoundTraceProvider.cpp
+++ b/Engine/Plugins/AudioInsights/Source/AudioInsights/Private/Providers/SoundTraceProvider.cpp
@@ -423,9 +423,7 @@ namespace UE::Audio::Insights
 				ActiveSoundEntryRef.bIsPlotActive = true;
 			}
 
-#if WITH_EDITOR
 			ActiveSoundEntryRef.ActorLabel = FText::FromString(*Msg.ActorLabel);
-#endif // WITH_EDITOR
 		};
 
 
@@ -460,10 +458,7 @@ namespace UE::Audio::Insights
 			SoundWaveEntry->EntryType = Msg.EntryType;
 			SoundWaveEntry->CategoryName = FText::FromString(FSoundTraceProviderPrivate::GetEntryTypeDisplayName(Msg.EntryType));
 			SoundWaveEntry->PinnedEntryType = (*OutActiveSoundEntry)->PinnedEntryType;
-
-#if WITH_EDITOR
 			SoundWaveEntry->ActorLabel = FText::FromString(*Msg.ActorLabel);
-#endif // WITH_EDITOR
 
 			(*OutActiveSoundEntry)->Children.Add(SoundWaveEntry);
 		};
@@ -1075,6 +1070,7 @@ namespace UE::Audio::Insights
 			OutActiveSoundEntry.SetName(*Msg.Name);
 			OutActiveSoundEntry.EntryType = Msg.EntryType;
 			OutActiveSoundEntry.bForceKeepEntryAlive = false;
+			OutActiveSoundEntry.ActorLabel = FText::FromString(*Msg.ActorLabel);
 
 			if (!OutActiveSoundEntry.bIsPlotActive)
 			{
@@ -1089,6 +1085,7 @@ namespace UE::Audio::Insights
 			ChildEntry->PlayOrder = Msg.PlayOrder;
 			ChildEntry->Timestamp = Msg.Timestamp;
 			ChildEntry->SetName(*Msg.Name);
+			ChildEntry->ActorLabel = FText::FromString(*Msg.ActorLabel);
 			ChildEntry->EntryType = Msg.EntryType;
 			ChildEntry->bForceKeepEntryAlive = false;
 
diff --git a/Engine/Plugins/AudioInsights/Source/AudioInsights/Private/Views/AudioEventLogDashboardViewFactory.cpp b/Engine/Plugins/AudioInsights/Source/AudioInsights/Private/Views/AudioEventLogDashboardViewFactory.cpp
--- a/Engine/Plugins/AudioInsights/Source/AudioInsights/Private/Views/AudioEventLogDashboardViewFactory.cpp
+++ b/Engine/Plugins/AudioInsights/Source/AudioInsights/Private/Views/AudioEventLogDashboardViewFactory.cpp
@@ -241,7 +241,7 @@ namespace UE::Audio::Insights
 
 			if (bPassesEventTypeFilter)
 			{
-				if (EventLogEntry.GetDisplayName().ToString().Contains(FilterString) || EventLogEntry.ActorName.Contains(FilterString))
+				if (EventLogEntry.GetDisplayName().ToString().Contains(FilterString) || EventLogEntry.ActorLabel.Contains(FilterString))
 				{
 					return true;
 				}
@@ -357,7 +357,7 @@ namespace UE::Audio::Insights
 					ActorColumnID,
 					{
 						LOCTEXT("AudioEventLog_ActorColumnDisplayName", "Actor"),
-						[](const IDashboardDataViewEntry& InData) { return FText::FromString(AudioEventLogPrivate::CastEntry(InData).ActorName); },
+						[](const IDashboardDataViewEntry& InData) { return FText::FromString(AudioEventLogPrivate::CastEntry(InData).ActorLabel); },
 						[this](const IDashboardDataViewEntry& InData) -> FName
 						{
 							return CastEntry(InData).ActorIconName;
diff --git a/Engine/Plugins/AudioInsights/Source/AudioInsights/Private/Views/SoundDashboardViewFactory.cpp b/Engine/Plugins/AudioInsights/Source/AudioInsights/Private/Views/SoundDashboardViewFactory.cpp
--- a/Engine/Plugins/AudioInsights/Source/AudioInsights/Private/Views/SoundDashboardViewFactory.cpp
+++ b/Engine/Plugins/AudioInsights/Source/AudioInsights/Private/Views/SoundDashboardViewFactory.cpp
@@ -1894,7 +1894,6 @@ namespace UE::Audio::Insights
 					.Alignment		  = EHorizontalAlignment::HAlign_Left
 				}
 			},
-#if WITH_EDITOR
 			{
 				ActorLabelColumnName,
 				{
@@ -1907,7 +1906,6 @@ namespace UE::Audio::Insights
 					.Alignment		  = EHorizontalAlignment::HAlign_Left
 				}
 			},
-#endif // WITH_EDITOR
 			{
 				CategoryColumnName,
 				{
@@ -2165,7 +2163,6 @@ namespace UE::Audio::Insights
 					}
 				}
 			},
-#if WITH_EDITOR
 			{
 				ActorLabelColumnName,
 				{
@@ -2177,7 +2174,6 @@ namespace UE::Audio::Insights
 					}
 				}
 			},
-#endif // WITH_EDITOR
 			{
 				CategoryColumnName,
 				{
@@ -2317,7 +2313,6 @@ namespace UE::Audio::Insights
 			return ComparisonDiff < 0.0f;
 		};
 
-#if WITH_EDITOR
 		auto SortByActor = [&SortByPlayOrder](const FSoundDashboardEntry& First, const FSoundDashboardEntry& Second)
 		{
 			const int32 Comparison = First.ActorLabel.CompareToCaseIgnored(Second.ActorLabel);
@@ -2328,7 +2323,6 @@ namespace UE::Audio::Insights
 			
 			return Comparison < 0;
 		};
-#endif // WITH_EDITOR
 
 		auto SortByCategory = [&SortByPlayOrder](const FSoundDashboardEntry& First, const FSoundDashboardEntry& Second)
 		{
@@ -2385,12 +2379,10 @@ namespace UE::Audio::Insights
 		{
 			SortByPredicate(SortByRelativeRenderCost);
 		}
-#if WITH_EDITOR
 		else if (SortByColumn == ActorLabelColumnName)
 		{
 			SortByPredicate(SortByActor);
 		}
-#endif // WITH_EDITOR
 		else if (SortByColumn == CategoryColumnName)
 		{
 			SortByPredicate(SortByCategory);
diff --git a/Engine/Plugins/AudioInsights/Source/AudioInsights/Public/Messages/AudioEventLogTraceMessages.h b/Engine/Plugins/AudioInsights/Source/AudioInsights/Public/Messages/AudioEventLogTraceMessages.h
--- a/Engine/Plugins/AudioInsights/Source/AudioInsights/Public/Messages/AudioEventLogTraceMessages.h
+++ b/Engine/Plugins/AudioInsights/Source/AudioInsights/Public/Messages/AudioEventLogTraceMessages.h
@@ -35,7 +35,7 @@ namespace UE::Audio::Insights
 
 		FString EventName;
 		FString AssetPath;
-		FString ActorName;
+		FString ActorLabel;
 		FString ActorIconName;
 		FString CategoryName;
 
@@ -58,7 +58,7 @@ namespace UE::Audio::Insights
 		uint32 MessageID = INDEX_NONE;
 
 		FString EventName;
-		FString ActorName;
+		FString ActorLabel;
 		FName ActorIconName;
 		FString CategoryName;
 
diff --git a/Engine/Plugins/AudioInsights/Source/AudioInsights/Public/Messages/SoundTraceMessages.h b/Engine/Plugins/AudioInsights/Source/AudioInsights/Public/Messages/SoundTraceMessages.h
--- a/Engine/Plugins/AudioInsights/Source/AudioInsights/Public/Messages/SoundTraceMessages.h
+++ b/Engine/Plugins/AudioInsights/Source/AudioInsights/Public/Messages/SoundTraceMessages.h
@@ -68,10 +68,7 @@ namespace UE::Audio::Insights
 
 		FString Name;
 		ESoundDashboardEntryType EntryType;
-
-#if WITH_EDITOR
 		FString ActorLabel;
-#endif // WITH_EDITOR
 	};
 
 	struct FSoundWaveStartMessage : public FSoundStartMessage
@@ -261,9 +258,7 @@ namespace UE::Audio::Insights
 		::Audio::TCircularAudioBuffer<FDataPoint> PitchDataPoints;
 		::Audio::TCircularAudioBuffer<FDataPoint> RelativeRenderCostDataPoints;
 
-#if WITH_EDITOR
 		FText ActorLabel;
-#endif // WITH_EDITOR
 
 	private:
 		FString FullName;
diff --git a/Engine/Source/Runtime/Engine/Private/Audio/AudioTraceUtil.cpp b/Engine/Source/Runtime/Engine/Private/Audio/AudioTraceUtil.cpp
--- a/Engine/Source/Runtime/Engine/Private/Audio/AudioTraceUtil.cpp
+++ b/Engine/Source/Runtime/Engine/Private/Audio/AudioTraceUtil.cpp
@@ -14,7 +14,7 @@ UE_TRACE_EVENT_BEGIN(Audio, EventLog)
 	UE_TRACE_EVENT_FIELD(uint32, PlayOrder)
 	UE_TRACE_EVENT_FIELD(UE::Trace::WideString, AssetPath)
 	UE_TRACE_EVENT_FIELD(UE::Trace::WideString, EventLogName)
-	UE_TRACE_EVENT_FIELD(UE::Trace::WideString, ActorName)
+	UE_TRACE_EVENT_FIELD(UE::Trace::WideString, ActorLabel)
 	UE_TRACE_EVENT_FIELD(UE::Trace::WideString, ActorIconName)
 	UE_TRACE_EVENT_FIELD(UE::Trace::WideString, SoundClassName)
 UE_TRACE_EVENT_END()
@@ -59,30 +59,15 @@ namespace Audio::Trace
 			return SoundClassName;
 		}
 
-		FString GetOwnerActorName(const FActiveSound& InActiveSound)
-		{
-			if (const TObjectPtr<UAudioComponent> AudioComponent = UAudioComponent::GetAudioComponentFromID(InActiveSound.GetAudioComponentID()))
-			{
-				if (const AActor* OwnerActor = AudioComponent->GetOwner())
-				{
-					return OwnerActor->GetName();
-				}
-			}
-
-			return FString();
-		}
-
 		FString GetOwnerActorLabel(const FActiveSound& InActiveSound)
 		{
-#if WITH_EDITOR
 			if (const TObjectPtr<UAudioComponent> AudioComponent = UAudioComponent::GetAudioComponentFromID(InActiveSound.GetAudioComponentID()))
 			{
 				if (const AActor* OwnerActor = AudioComponent->GetOwner())
 				{
-					return OwnerActor->GetActorLabel(false);
+					return OwnerActor->GetActorNameOrLabel();
 				}
 			}
-#endif // WITH_EDITOR
 
 			return FString();
 		}
@@ -131,13 +116,13 @@ namespace Audio::Trace
 						, Event
 						, InActiveSound.GetPlayOrder()
 						, SoundObjectPtr->GetPathName()
-						, Util::GetOwnerActorName(InActiveSound)
+						, Util::GetOwnerActorLabel(InActiveSound)
 						, Util::GetOwnerActorIconName(InActiveSound).ToString()
 						, Util::GetSoundBaseAssetName(SoundObjectPtr->GetClass()));
 			}
 		}
 
-		void SendEvent(const Audio::FDeviceId AudioDeviceID, const FString& Event, const uint32 AudioObjectID, const FString& AssetPath, const FString& ActorName, const FString& ActorIconName, const FString& SoundClassName)
+		void SendEvent(const Audio::FDeviceId AudioDeviceID, const FString& Event, const uint32 AudioObjectID, const FString& AssetPath, const FString& ActorLabel, const FString& ActorIconName, const FString& SoundClassName)
 		{
 			const bool bChannelEnabled = UE_TRACE_CHANNELEXPR_IS_ENABLED(AudioMixerChannel);
 
@@ -149,7 +134,7 @@ namespace Audio::Trace
 					<< EventLog.PlayOrder(AudioObjectID)
 					<< EventLog.AssetPath(*AssetPath)
 					<< EventLog.EventLogName(*Event)
-					<< EventLog.ActorName(*ActorName)
+					<< EventLog.ActorLabel(*ActorLabel)
 					<< EventLog.ActorIconName(*ActorIconName)
 					<< EventLog.SoundClassName(*SoundClassName);
 			}
diff --git a/Engine/Source/Runtime/Engine/Private/AudioDevice.cpp b/Engine/Source/Runtime/Engine/Private/AudioDevice.cpp
--- a/Engine/Source/Runtime/Engine/Private/AudioDevice.cpp
+++ b/Engine/Source/Runtime/Engine/Private/AudioDevice.cpp
@@ -247,7 +247,6 @@ UE_TRACE_EVENT_BEGIN(Audio, SoundStart)
 	UE_TRACE_EVENT_FIELD(UE::Trace::WideString, AssetPath)
 	UE_TRACE_EVENT_FIELD(UE::Trace::WideString, EventLogName)
 	UE_TRACE_EVENT_FIELD(UE::Trace::WideString, SoundClassName)
-	UE_TRACE_EVENT_FIELD(UE::Trace::WideString, ActorName)
 	UE_TRACE_EVENT_FIELD(UE::Trace::WideString, ActorLabel)
 	UE_TRACE_EVENT_FIELD(UE::Trace::WideString, ActorIconName)
 UE_TRACE_EVENT_END()
@@ -268,7 +267,7 @@ UE_TRACE_EVENT_BEGIN(Audio, SoundStop)
 	UE_TRACE_EVENT_FIELD(UE::Trace::WideString, AssetPath)
 	UE_TRACE_EVENT_FIELD(UE::Trace::WideString, EventLogName)
 	UE_TRACE_EVENT_FIELD(UE::Trace::WideString, SoundClassName)
-	UE_TRACE_EVENT_FIELD(UE::Trace::WideString, ActorName)
+	UE_TRACE_EVENT_FIELD(UE::Trace::WideString, ActorLabel)
 	UE_TRACE_EVENT_FIELD(UE::Trace::WideString, ActorIconName)
 UE_TRACE_EVENT_END()
 
@@ -5117,7 +5116,6 @@ void FAudioDevice::SendSoundStartedTraceMessage(const FActiveSound& InActiveSoun
 			<< SoundStart.AssetPath(*SoundObjectPtr->GetPathName())
 			<< SoundStart.EventLogName(*Audio::Trace::EventLog::ID::SoundStart)
 			<< SoundStart.SoundClassName(*Audio::Trace::Util::GetSoundBaseAssetName(SoundObjectPtr->GetClass()))
-			<< SoundStart.ActorName(*Audio::Trace::Util::GetOwnerActorName(InActiveSound))
 			<< SoundStart.ActorLabel(*Audio::Trace::Util::GetOwnerActorLabel(InActiveSound))
 			<< SoundStart.ActorIconName(*Audio::Trace::Util::GetOwnerActorIconName(InActiveSound).ToString());
 	}
@@ -6093,7 +6091,7 @@ void FAudioDevice::RemoveActiveSound(FActiveSound* ActiveSound)
 				<< SoundStop.AssetPath(*SoundObjectPtr->GetPathName())
 				<< SoundStop.EventLogName(*Audio::Trace::EventLog::ID::SoundStop)
 				<< SoundStop.SoundClassName(*Audio::Trace::Util::GetSoundBaseAssetName(SoundObjectPtr->GetClass()))
-				<< SoundStop.ActorName(*Audio::Trace::Util::GetOwnerActorName(*ActiveSound))
+				<< SoundStop.ActorLabel(*Audio::Trace::Util::GetOwnerActorLabel(*ActiveSound))
 				<< SoundStop.ActorIconName(*Audio::Trace::Util::GetOwnerActorIconName(*ActiveSound).ToString());
 		}
 	}
diff --git a/Engine/Source/Runtime/Engine/Public/Audio/AudioTraceUtil.h b/Engine/Source/Runtime/Engine/Public/Audio/AudioTraceUtil.h
--- a/Engine/Source/Runtime/Engine/Public/Audio/AudioTraceUtil.h
+++ b/Engine/Source/Runtime/Engine/Public/Audio/AudioTraceUtil.h
@@ -14,7 +14,6 @@ namespace Audio::Trace
 	namespace Util
 	{
 		ENGINE_API FString GetSoundBaseAssetName(const TObjectPtr<UClass> SoundClass);
-		ENGINE_API FString GetOwnerActorName(const FActiveSound& InActiveSound);
 		ENGINE_API FString GetOwnerActorLabel(const FActiveSound& InActiveSound);
 		ENGINE_API FName GetOwnerActorIconName(const FActiveSound& InActiveSound);
 		ENGINE_API TObjectPtr<UObject> GetSoundObjectPointer(const FActiveSound& InActiveSound);
@@ -31,7 +30,7 @@ namespace Audio::Trace
 								, const FString& Event
 								, const uint32 AudioObjectID = INDEX_NONE
 								, const FString& AssetPath = FString()
-								, const FString& ActorName = FString()
+								, const FString& ActorLabel = FString()
 								, const FString& ActorIconName = FString()
 								, const FString& SoundClassName = FString());
 
