[Shaders] Dump artifacts before compiling shaders in-process after distributed jobs failed.

When distributed shader compile jobs failed, indicated by a missing remote worker output file, the distributed shader compiler executes the jobs again in-process.
If they crash the process, the only indication are the log entries for the last executed job but no artifacts are dumped.
This change also dumps the artifacts to allow diagnosing a crash.

#rnx

diff --git a/Engine/Source/Runtime/Engine/Internal/ShaderCompiler/ShaderCompilerInternal.cpp b/Engine/Source/Runtime/Engine/Internal/ShaderCompiler/ShaderCompilerInternal.cpp
--- a/Engine/Source/Runtime/Engine/Internal/ShaderCompiler/ShaderCompilerInternal.cpp
+++ b/Engine/Source/Runtime/Engine/Internal/ShaderCompiler/ShaderCompilerInternal.cpp
@@ -12,3 +12,18 @@ void FShaderCompileInternalUtilities::ExecuteShaderCompileJob(FShaderCommonCompi
 {
 	FShaderCompileUtilities::ExecuteShaderCompileJob(Job);
 }
+
+void FShaderCompileInternalUtilities::EnableDumpDebugInfoForRetry(FShaderCommonCompileJob& Job)
+{
+	Job.ForEachSingleShaderJob(
+		[](FShaderCompileJob& SingleJob)
+		{
+			if (SingleJob.Input.DumpDebugInfoPath.IsEmpty())
+			{
+				SingleJob.Input.DumpDebugInfoPath = GShaderCompilingManager->CreateShaderDebugInfoPath(SingleJob.Input);
+				// Any reissued jobs due to this condition will dump debug information, so increment the dump count here
+				GShaderCompilingManager->IncrementNumDumpedShaderSources();
+			}
+		}
+	);
+}
diff --git a/Engine/Source/Runtime/Engine/Internal/ShaderCompiler/ShaderCompilerInternal.h b/Engine/Source/Runtime/Engine/Internal/ShaderCompiler/ShaderCompilerInternal.h
--- a/Engine/Source/Runtime/Engine/Internal/ShaderCompiler/ShaderCompilerInternal.h
+++ b/Engine/Source/Runtime/Engine/Internal/ShaderCompiler/ShaderCompilerInternal.h
@@ -16,5 +16,8 @@ class FShaderCompileInternalUtilities
 public:
 	/** Execute the specified (single or pipeline) shader compile job. */
 	static ENGINE_API void ExecuteShaderCompileJob(FShaderCommonCompileJob& Job);
+
+	/** Ensures DumpDebugInfoPath is assigned and tracked by the shader compiling manager. */
+	static ENGINE_API void EnableDumpDebugInfoForRetry(FShaderCommonCompileJob& Job);
 };
 
diff --git a/Engine/Source/Runtime/Engine/Private/ShaderCompiler/ShaderCompiler.cpp b/Engine/Source/Runtime/Engine/Private/ShaderCompiler/ShaderCompiler.cpp
--- a/Engine/Source/Runtime/Engine/Private/ShaderCompiler/ShaderCompiler.cpp
+++ b/Engine/Source/Runtime/Engine/Private/ShaderCompiler/ShaderCompiler.cpp
@@ -5,6 +5,7 @@
 =============================================================================*/
 
 #include "ShaderCompiler.h"
+#include "ShaderCompiler/ShaderCompilerInternal.h"
 #include "ShaderCompilerPrivate.h"
 #include "AsyncCompilationHelpers.h"
 #include "AssetCompilingManager.h"
@@ -2512,12 +2513,7 @@ bool FShaderCompilingManager::HandlePotentialRetry(TMap<int32, FShaderMapFinaliz
 		{
 			// if any stage of a pipeline job is reissued to dump debug info, we need to enable debug info on all stages of the job
 			// since the debug information relies on debug source files existing for every stage
-			DebugJob->ForEachSingleShaderJob([](FShaderCompileJob& SingleJob)
-				{
-					SingleJob.Input.DumpDebugInfoPath = GShaderCompilingManager->CreateShaderDebugInfoPath(SingleJob.Input);
-					// Any reissued jobs due to this condition will dump debug information, so increment the dump count here
-					GShaderCompilingManager->IncrementNumDumpedShaderSources();
-				});
+			FShaderCompileInternalUtilities::EnableDumpDebugInfoForRetry(*DebugJob);
 		}
 
 		ReissueJobs = MoveTemp(ReissueForDebugJobs);
diff --git a/Engine/Source/Runtime/Engine/Private/ShaderCompiler/ShaderCompilerDistributed.cpp b/Engine/Source/Runtime/Engine/Private/ShaderCompiler/ShaderCompilerDistributed.cpp
--- a/Engine/Source/Runtime/Engine/Private/ShaderCompiler/ShaderCompilerDistributed.cpp
+++ b/Engine/Source/Runtime/Engine/Private/ShaderCompiler/ShaderCompilerDistributed.cpp
@@ -7,6 +7,7 @@
 #include "Misc/StringBuilder.h"
 #include "ShaderCompiler.h"
 #include "ShaderCompilerJobTypes.h"
+#include "ShaderCompiler/ShaderCompilerInternal.h"
 #include "ShaderCompileWorkerUtil.h"
 
 namespace DistributedShaderCompilerVariables
@@ -551,13 +552,18 @@ int32 FShaderCompileDistributedThreadRunnable_Interface::CompilingLoop()
 				FString JobDiagnostics;
 				for (int32 JobIndex = 0; JobIndex < Task->ShaderJobs.Num(); ++JobIndex)
 				{
-					const FShaderCommonCompileJobPtr& Job = Task->ShaderJobs[JobIndex];
+					FShaderCommonCompileJobPtr& Job = Task->ShaderJobs[JobIndex];
 
 					// Rescheduling jobs after distributed readback failed should be rare, so display all job details with default verbosity
 					JobDiagnostics.Empty();
 					Job->AppendDiagnostics(JobDiagnostics, JobIndex, Task->ShaderJobs.Num());
 					UE_LOG(LogShaderCompilers, Display, TEXT("Executing %s"), *JobDiagnostics);
 
+					// Dumping debug info in case the compile job crashes the cooker
+					FShaderCompileInternalUtilities::EnableDumpDebugInfoForRetry(*Job);
+					FShaderDebugDataContext Ctx;
+					Job->OnComplete(Ctx);
+
 					FShaderCompileUtilities::ExecuteShaderCompileJob(*Job);
 				}
 			}
