Fixed use after free from GSpawnActorDeferredTransformCache.

diff --git a/Engine/Source/Runtime/Engine/Private/Actor.cpp b/Engine/Source/Runtime/Engine/Private/Actor.cpp
--- a/Engine/Source/Runtime/Engine/Private/Actor.cpp
+++ b/Engine/Source/Runtime/Engine/Private/Actor.cpp
@@ -4384,27 +4384,24 @@ void AActor::FinishSpawning(const FTransform& UserTransform, bool bIsDefaultTran
 		// than was passed in during the original SpawnActor call)
 		if (RootComponent && !bIsDefaultTransform)
 		{
-			FTransform const* const OriginalSpawnTransform = GSpawnActorDeferredTransformCache.Find(this);
-			if (OriginalSpawnTransform)
-			{
-				GSpawnActorDeferredTransformCache.Remove(this);
+			FTransform OriginalSpawnTransform(NoInit);
+			const bool bFound = GSpawnActorDeferredTransformCache.RemoveAndCopyValue(this, OriginalSpawnTransform);
 
-				if (OriginalSpawnTransform->Equals(UserTransform) == false)
-				{
-					UserTransform.GetLocation().DiagnosticCheckNaN(TEXT("AActor::FinishSpawning: UserTransform.GetLocation()"));
-					UserTransform.GetRotation().DiagnosticCheckNaN(TEXT("AActor::FinishSpawning: UserTransform.GetRotation()"));
-
-					// caller passed a different transform!
-					// undo the original spawn transform to get back to the template transform, so we can recompute a good
-					// final transform that takes into account the template's transform
-					FTransform const TemplateTransform = RootComponent->GetComponentTransform() * OriginalSpawnTransform->Inverse();
-					FinalRootComponentTransform = TemplateTransform * UserTransform;
-				}
+			if (bFound && !OriginalSpawnTransform.Equals(UserTransform))
+			{
+				UserTransform.GetLocation().DiagnosticCheckNaN(TEXT("AActor::FinishSpawning: UserTransform.GetLocation()"));
+				UserTransform.GetRotation().DiagnosticCheckNaN(TEXT("AActor::FinishSpawning: UserTransform.GetRotation()"));
+
+				// caller passed a different transform!
+				// undo the original spawn transform to get back to the template transform, so we can recompute a good
+				// final transform that takes into account the template's transform
+				FTransform const TemplateTransform = RootComponent->GetComponentTransform() * OriginalSpawnTransform.Inverse();
+				FinalRootComponentTransform = TemplateTransform * UserTransform;
 			}
 
-				// should be fast and relatively rare
-				ValidateDeferredTransformCache();
-			}
+			// should be fast and relatively rare
+			ValidateDeferredTransformCache();
+		}
 
 		FinalRootComponentTransform.GetLocation().DiagnosticCheckNaN(TEXT("AActor::FinishSpawning: FinalRootComponentTransform.GetLocation()"));
 		FinalRootComponentTransform.GetRotation().DiagnosticCheckNaN(TEXT("AActor::FinishSpawning: FinalRootComponentTransform.GetRotation()"));
