Fix a unity build issue with a cvar declaration

#DisabledValidations Preflight, CheckCopyright, CheckBalancedMacros

diff --git a/Engine/Source/Runtime/Engine/Private/UnrealEngine.cpp b/Engine/Source/Runtime/Engine/Private/UnrealEngine.cpp
--- a/Engine/Source/Runtime/Engine/Private/UnrealEngine.cpp
+++ b/Engine/Source/Runtime/Engine/Private/UnrealEngine.cpp
@@ -13781,8 +13781,8 @@ void UEngine::GetDynamicResolutionCurrentStateInfos(FDynamicResolutionStateInfos
 			OutInfos.ResolutionFractionUpperBounds = DynamicResolutionState->GetResolutionFractionsUpperBoundBudgetValue();
 
 			// If dynamic upper bound is enabled, set the upper bound to the forced resolution to properly simulate that resolution (otherwise clears/resolves may not be representative)
-			static auto CVarUpperBoundQuantization = IConsoleManager::Get().FindConsoleVariable(TEXT("r.DynamicRes.UpperBoundQuantization"));
-			if (CVarUpperBoundQuantization && CVarUpperBoundQuantization->GetFloat() > 0.0f)
+			static auto LocalCVarUpperBoundQuantization = IConsoleManager::Get().FindConsoleVariable(TEXT("r.DynamicRes.UpperBoundQuantization"));
+			if (LocalCVarUpperBoundQuantization && LocalCVarUpperBoundQuantization->GetFloat() > 0.0f)
 			{
 				OutInfos.ResolutionFractionUpperBounds[GDynamicPrimaryResolutionFraction] = ForceResolutionFraction;
 			}
