[MIR] Implement FontSample and FontSampleParameter material expressions in new translator.

This also adds Outputs() to the IR emitter that accepts multiple expression outputs and automatically applies their component swizzling.
This simplifies expressions like the implemented FontSample from this:
> MaterialFloat _1 = (lerp((MaterialFloat4)-90.000000f, (MaterialFloat4)90.000000f, ProcessMaterialAlphaTextureLookup(Texture2DSampleGrad(Material.Texture2D_0, Material.Texture2D_0Sampler, Parameters.TexCoords[0], Parameters.TexCoords_DDX[0], Parameters.TexCoords_DDY[0]))) * _0).x;
... to this:
> MaterialFloat _1 = (lerp(-90.000000f, 90.000000f, ProcessMaterialAlphaTextureLookup(Texture2DSampleGrad(Material.Texture2D_0, Material.Texture2D_0Sampler, Parameters.TexCoords[0], Parameters.TexCoords_DDX[0], Parameters.TexCoords_DDY[0])).x) * _0.w);

#test EngineTest VKTextMaterial works with the new translator now.
#rnx

diff --git a/Engine/Source/Runtime/Engine/Private/Materials/MaterialExpressions.cpp b/Engine/Source/Runtime/Engine/Private/Materials/MaterialExpressions.cpp
--- a/Engine/Source/Runtime/Engine/Private/Materials/MaterialExpressions.cpp
+++ b/Engine/Source/Runtime/Engine/Private/Materials/MaterialExpressions.cpp
@@ -12127,28 +12127,10 @@ int32 UMaterialExpressionFontSample::Compile(class FMaterialCompiler* Compiler,
 	}
 	else
 	{
-		UTexture* Texture = Font->Textures[FontTexturePage];
-		if( !Texture )
+		auto [bSuccess, Texture, ExpectedSamplerType, ErrorOutput] = ValidateAndGetTextureSampler(Compiler->GetShaderPlatform(), Compiler->GetTargetPlatform());
+		if (!bSuccess)
 		{
-			UE_LOG(LogMaterial, Log, TEXT("Invalid font texture. Using default texture"));
-			Texture = GEngine->DefaultTexture;
-		}
-		check(Texture);
-
-		EMaterialSamplerType ExpectedSamplerType;
-		if (Texture->CompressionSettings == TC_DistanceFieldFont)
-		{
-			ExpectedSamplerType = SAMPLERTYPE_DistanceFieldFont;
-		}
-		else
-		{
-			ExpectedSamplerType = Texture->SRGB ? SAMPLERTYPE_Color : SAMPLERTYPE_LinearColor;
-		}
-
-		FString SamplerTypeError;
-		if (!UMaterialExpressionTextureBase::VerifySamplerType(Compiler->GetShaderPlatform(), Compiler->GetTargetPlatform(), Texture, ExpectedSamplerType, SamplerTypeError))
-		{
-			return Compiler->Errorf(TEXT("%s"), *SamplerTypeError);
+			return Compiler->Error(*ErrorOutput);
 		}
 
 		int32 TextureCodeIndex = Compiler->Texture(Texture, ExpectedSamplerType);
@@ -12161,6 +12143,35 @@ int32 UMaterialExpressionFontSample::Compile(class FMaterialCompiler* Compiler,
 	return Result;
 }
 
+TTuple<bool, UTexture*, EMaterialSamplerType, FString> UMaterialExpressionFontSample::ValidateAndGetTextureSampler(EShaderPlatform ShaderPlatform, const ITargetPlatform* TargetPlatform)
+{
+	UTexture* Texture = Font->Textures[FontTexturePage];
+	if (!Texture)
+	{
+		UE_LOG(LogMaterial, Log, TEXT("Invalid font texture. Using default texture"));
+		Texture = GEngine->DefaultTexture;
+	}
+	check(Texture);
+
+	EMaterialSamplerType ExpectedSamplerType;
+	if (Texture->CompressionSettings == TC_DistanceFieldFont)
+	{
+		ExpectedSamplerType = SAMPLERTYPE_DistanceFieldFont;
+	}
+	else
+	{
+		ExpectedSamplerType = Texture->SRGB ? SAMPLERTYPE_Color : SAMPLERTYPE_LinearColor;
+	}
+
+	FString SamplerTypeError;
+	if (!UMaterialExpressionTextureBase::VerifySamplerType(ShaderPlatform, TargetPlatform, Texture, ExpectedSamplerType, SamplerTypeError))
+	{
+		return { false, nullptr, SAMPLERTYPE_MAX, SamplerTypeError };
+	}
+
+	return { true, Texture, ExpectedSamplerType, FString() };
+}
+
 int32 UMaterialExpressionFontSample::GetWidth() const
 {
 	return ME_STD_THUMBNAIL_SZ+(ME_STD_BORDER*2);
@@ -12215,31 +12226,13 @@ int32 UMaterialExpressionFontSampleParameter::Compile(class FMaterialCompiler* C
 	}
 	else 
 	{
-		UTexture* Texture = Font->Textures[FontTexturePage];
-		if( !Texture )
-		{
-			UE_LOG(LogMaterial, Log, TEXT("Invalid font texture. Using default texture"));
-			Texture = GEngine->DefaultTexture;
-		}
-		check(Texture);
-
-		EMaterialSamplerType ExpectedSamplerType;
-		if (Texture->CompressionSettings == TC_DistanceFieldFont)
-		{
-			ExpectedSamplerType = SAMPLERTYPE_DistanceFieldFont;
-		}
-		else
+		auto [bSuccess, Texture, ExpectedSamplerType, ErrorOutput] = ValidateAndGetTextureSampler(Compiler->GetShaderPlatform(), Compiler->GetTargetPlatform());
+		if (!bSuccess)
 		{
-			ExpectedSamplerType = Texture->SRGB ? SAMPLERTYPE_Color : SAMPLERTYPE_LinearColor;
-		}
-
-		FString SamplerTypeError;
-		if (!UMaterialExpressionTextureBase::VerifySamplerType(Compiler->GetShaderPlatform(), Compiler->GetTargetPlatform(), Texture, ExpectedSamplerType, SamplerTypeError))
-		{
-			return Compiler->Errorf(TEXT("%s"), *SamplerTypeError);
+			return Compiler->Error(*ErrorOutput);
 		}
 
-		int32 TextureCodeIndex = Compiler->TextureParameter(ParameterName,Texture, ExpectedSamplerType);
+		int32 TextureCodeIndex = Compiler->TextureParameter(ParameterName, Texture, ExpectedSamplerType);
 		Result = Compiler->TextureSample(
 			TextureCodeIndex,
 			Compiler->TextureCoordinate(0, false, false),
diff --git a/Engine/Source/Runtime/Engine/Private/Materials/MaterialExpressionsIR.cpp b/Engine/Source/Runtime/Engine/Private/Materials/MaterialExpressionsIR.cpp
--- a/Engine/Source/Runtime/Engine/Private/Materials/MaterialExpressionsIR.cpp
+++ b/Engine/Source/Runtime/Engine/Private/Materials/MaterialExpressionsIR.cpp
@@ -85,6 +85,7 @@
 #include "Materials/MaterialExpressionFmod.h"
 #include "Materials/MaterialExpressionAtmosphericFogColor.h"
 #include "Materials/MaterialExpressionFontSample.h"
+#include "Materials/MaterialExpressionFontSampleParameter.h"
 #include "Materials/MaterialExpressionFrac.h"
 #include "Materials/MaterialExpressionFresnel.h"
 #include "Materials/MaterialExpressionFunctionInput.h"
@@ -925,9 +926,18 @@ static FValueRef EmitInlineHLSL(MIR::FEmitter& Em, const UMaterialExpressionExte
 
 static void BuildInlineHLSLOutput(MIR::FEmitter& Em, const UMaterialExpressionExternalCodeBase& InExternalCodeExpression, TConstArrayView<FValueRef> InArguments)
 {
-	for (int32 OutputIndex = 0; OutputIndex < InExternalCodeExpression.ExternalCodeIdentifiers.Num(); ++OutputIndex)
+	// If there are multiple output pins but only one external code identifier, use it for all outputs and let the emitter handle the swizzling.
+	// This is used for output pins that map to component swizzling, e.g. R, G, B, RGB, RGBA
+	if (InExternalCodeExpression.Outputs.Num() > 1 && InExternalCodeExpression.ExternalCodeIdentifiers.Num() == 1)
 	{
-		Em.Output(OutputIndex, EmitInlineHLSL(Em, InExternalCodeExpression, OutputIndex, InArguments));
+		Em.Outputs(InExternalCodeExpression.Outputs, EmitInlineHLSL(Em, InExternalCodeExpression, 0, InArguments));
+	}
+	else
+	{
+		for (int32 OutputIndex = 0; OutputIndex < InExternalCodeExpression.ExternalCodeIdentifiers.Num(); ++OutputIndex)
+		{
+			Em.Output(OutputIndex, EmitInlineHLSL(Em, InExternalCodeExpression, OutputIndex, InArguments));
+		}
 	}
 }
 
@@ -1417,6 +1427,99 @@ void UMaterialExpressionTextureCoordinate::Build(MIR::FEmitter& Em)
 	Em.Output(GetOutput(0), OutputValue);
 }
 
+#if WITH_EDITOR
+void UMaterialExpressionFontSample::Build(MIR::FEmitter& Em)
+{
+#if PLATFORM_EXCEPTIONS_DISABLED
+	// if we can't throw the error below, attempt to thwart the error by using the default font
+	if (!Font)
+	{
+		UE_LOG(LogMaterial, Log, TEXT("Using default font instead of real font!"));
+		Font = GEngine->GetMediumFont();
+		FontTexturePage = 0;
+	}
+	else if (!Font->Textures.IsValidIndex(FontTexturePage))
+	{
+		UE_LOG(LogMaterial, Log, TEXT("Invalid font page %d. Max allowed is %d"), FontTexturePage, Font->Textures.Num());
+		FontTexturePage = 0;
+	}
+#endif
+	if (!Font)
+	{
+		Em.Error(TEXT("Missing input Font"));
+	}
+	else if (Font->FontCacheType == EFontCacheType::Runtime)
+	{
+		Em.Errorf(TEXT("Font '%s' is runtime cached, but only offline cached fonts can be sampled"), *Font->GetName());
+	}
+	else if (!Font->Textures.IsValidIndex(FontTexturePage))
+	{
+		Em.Errorf(TEXT("Invalid font page %d. Max allowed is %d"), FontTexturePage, Font->Textures.Num());
+	}
+	else
+	{
+		auto [bSuccess, Texture, ExpectedSamplerType, ErrorOutput] = ValidateAndGetTextureSampler(Em.GetShaderPlatform(), Em.GetTargetPlatform());
+		if (!bSuccess)
+		{
+			Em.Error(*ErrorOutput);
+			return;
+		}
+
+		Em.Outputs(Outputs,
+			Em.TextureSample(
+				Em.TextureObject(Texture, ExpectedSamplerType),
+				Em.ExternalInput(MIR::TexCoordIndexToExternalInput(0)),
+				false,
+				SSM_FromTextureAsset,
+				ExpectedSamplerType
+			)
+		);
+	}
+}
+#endif
+
+#if WITH_EDITOR
+void UMaterialExpressionFontSampleParameter::Build(MIR::FEmitter& Em)
+{
+	if (!ParameterName.IsValid() ||
+		ParameterName.IsNone() ||
+		!Font ||
+		!Font->Textures.IsValidIndex(FontTexturePage))
+	{
+		UMaterialExpressionFontSample::Build(Em);
+	}
+	else
+	{
+		auto [bSuccess, Texture, ExpectedSamplerType, ErrorOutput] = ValidateAndGetTextureSampler(Em.GetShaderPlatform(), Em.GetTargetPlatform());
+		if (!bSuccess)
+		{
+			Em.Error(*ErrorOutput);
+			return;
+		}
+
+		FMaterialParameterMetadata ParameterMetaData;
+		GetParameterValue(ParameterMetaData);
+
+		FValueRef TextureParameter = Em.Parameter(ParameterName, ParameterMetaData, ExpectedSamplerType);
+		if (!TextureParameter->Type.IsTexture2D())
+		{
+			Em.Error(TEXT("Parameter is not a texture"));
+			return;
+		}
+
+		Em.Outputs(Outputs,
+			Em.TextureSample(
+				TextureParameter,
+				Em.ExternalInput(MIR::TexCoordIndexToExternalInput(0)),
+				false,
+				SSM_FromTextureAsset,
+				ExpectedSamplerType
+			)
+		);
+	}
+}
+#endif
+
 void UMaterialExpressionTime::Build(MIR::FEmitter& Em)
 {
 	MIR::FType ScalarFloatType = MIR::FType::MakeScalar(MIR::EScalarKind::Float);
diff --git a/Engine/Source/Runtime/Engine/Private/Materials/MaterialIREmitter.cpp b/Engine/Source/Runtime/Engine/Private/Materials/MaterialIREmitter.cpp
--- a/Engine/Source/Runtime/Engine/Private/Materials/MaterialIREmitter.cpp
+++ b/Engine/Source/Runtime/Engine/Private/Materials/MaterialIREmitter.cpp
@@ -75,6 +75,17 @@ void FSwizzleMask::Append(EVectorComponent Component)
 	Components[NumComponents++] = Component;
 }
 
+bool FSwizzleMask::IsXYZW() const
+{
+	return
+		NumComponents == 4 &&
+		Components[0] == EVectorComponent::X &&
+		Components[1] == EVectorComponent::Y &&
+		Components[2] == EVectorComponent::Z &&
+		Components[3] == EVectorComponent::W
+		;
+}
+
 struct FEmitter::FPrivate
 {
 	// Searches the emitter’s value set for an existing FValue matching Prototype, or returns nullptr.
@@ -674,11 +685,44 @@ FEmitter& FEmitter::Output(int32 OutputIndex, FValueRef Value)
 	return *this;
 }
 
-FEmitter& FEmitter::Output(const FExpressionOutput* Output, FValueRef Value)
+FEmitter& FEmitter::Output(const FExpressionOutput* ExpressionOutput, FValueRef Value)
 {
 	if (Value)
 	{
-		Internal::BindValueToExpressionOutput(BuilderImpl, Output, Value);
+		Internal::BindValueToExpressionOutput(BuilderImpl, ExpressionOutput, Value);
+	}
+	return *this;
+}
+
+static MIR::FSwizzleMask GetSwizzleMaskFromExpressionOutputMask(const FExpressionOutput& ExpressionOutput)
+{
+	MIR::FSwizzleMask SwizzleMask;
+	if (ExpressionOutput.MaskR)
+	{
+		SwizzleMask.Append(MIR::EVectorComponent::X);
+	}
+	if (ExpressionOutput.MaskG)
+	{
+		SwizzleMask.Append(MIR::EVectorComponent::Y);
+	}
+	if (ExpressionOutput.MaskB)
+	{
+		SwizzleMask.Append(MIR::EVectorComponent::Z);
+	}
+	if (ExpressionOutput.MaskA)
+	{
+		SwizzleMask.Append(MIR::EVectorComponent::W);
+	}
+	return SwizzleMask;
+}
+
+FEmitter& FEmitter::Outputs(const TConstArrayView<FExpressionOutput>& ExpressionOutputs, FValueRef Value)
+{
+	for (const FExpressionOutput& CurrentOutput : ExpressionOutputs)
+	{
+		// Apply component swizzling for each expression output but use unmodified input value if swizzling is redundant (i.e. XYZW vector)
+		const MIR::FSwizzleMask SwizzleMask = GetSwizzleMaskFromExpressionOutputMask(CurrentOutput);
+		Output(&CurrentOutput, SwizzleMask.IsXYZW() ? Value : Swizzle(Value, SwizzleMask));
 	}
 	return *this;
 }
@@ -984,26 +1028,12 @@ FValueRef FEmitter::Parameter(FName Name, FMaterialParameterMetadata& Metadata,
 		}
 
 		case EMaterialParameterType::Texture:
+		case EMaterialParameterType::Font:
 		{
-			switch (Metadata.Value.Texture->GetMaterialType())
-			{
-				case MCT_Texture2D:
-				case MCT_TextureVirtual:
-				{
-					FUniformParameter Prototype = MakePrototype<FUniformParameter>(FType::MakeTexture2D());
-					Prototype.ParameterIdInModule = RegisterParameter(Module, Info, Metadata);
-					Prototype.SamplerType = SamplerType;
-					return EmitPrototype(*this, Prototype);
-				}
-				break;
-
-				default:
-				{
-					Errorf(TEXT("Unsupported texture type"));
-				}
-				break;
-			}
-			return FPoison::Get();
+			FUniformParameter Prototype = MakePrototype<FUniformParameter>(FType::MakeTexture2D());
+			Prototype.ParameterIdInModule = RegisterParameter(Module, Info, Metadata);
+			Prototype.SamplerType = SamplerType;
+			return EmitPrototype(*this, Prototype);
 		}
 
 		case EMaterialParameterType::StaticSwitch:
@@ -3894,6 +3924,11 @@ EShaderPlatform FEmitter::GetShaderPlatform() const
 	return Module->GetShaderPlatform();
 }
 
+const ITargetPlatform* FEmitter::GetTargetPlatform() const
+{
+	return Module->GetTargetPlatform();
+}
+
 ERHIFeatureLevel::Type FEmitter::GetFeatureLevel() const
 {
 	return Module->GetFeatureLevel();
diff --git a/Engine/Source/Runtime/Engine/Private/Materials/MaterialIRModuleBuilder.cpp b/Engine/Source/Runtime/Engine/Private/Materials/MaterialIRModuleBuilder.cpp
--- a/Engine/Source/Runtime/Engine/Private/Materials/MaterialIRModuleBuilder.cpp
+++ b/Engine/Source/Runtime/Engine/Private/Materials/MaterialIRModuleBuilder.cpp
@@ -104,6 +104,7 @@ struct FMaterialIRModuleBuilderImpl
 	{
 		Module->Empty();
 		Module->ShaderPlatform = Builder->ShaderPlatform;
+		Module->TargetPlatform = Builder->TargetPlatform;
 		Module->FeatureLevel = Builder->FeatureLevel;
 		Module->QualityLevel = Builder->QualityLevel;
 		
diff --git a/Engine/Source/Runtime/Engine/Private/Materials/MaterialIRToHLSLTranslator.cpp b/Engine/Source/Runtime/Engine/Private/Materials/MaterialIRToHLSLTranslator.cpp
--- a/Engine/Source/Runtime/Engine/Private/Materials/MaterialIRToHLSLTranslator.cpp
+++ b/Engine/Source/Runtime/Engine/Private/Materials/MaterialIRToHLSLTranslator.cpp
@@ -2,7 +2,9 @@
 
 #include "Materials/MaterialIRToHLSLTranslator.h"
 #include "DataDrivenShaderPlatformInfo.h"
+#include "Engine/Font.h"
 #include "Engine/Texture.h"
+#include "Engine/Texture2D.h"
 #include "MaterialDomain.h"
 #include "MaterialIRInternal.h"
 #include "Materials/Material.h"
@@ -1587,6 +1589,12 @@ struct FPrivate : FMaterialIRToHLSLTranslation
 			{
 				return TextureParameter.Texture;
 			}
+			case EMaterialParameterType::Font:
+			{
+				UTexture2D* FontTexture = TextureParameter.Font.Value->Textures[TextureParameter.Font.Page].Get();
+				checkf(FontTexture, TEXT("Missing texture for font '%s' [page %d]"), *TextureParameter.Font.Value->GetName(), TextureParameter.Font.Page);
+				return FontTexture;
+			}
 			default:
 			{
 				UE_MIR_UNREACHABLE();
diff --git a/Engine/Source/Runtime/Engine/Private/Materials/MaterialIRValueAnalyzer.cpp b/Engine/Source/Runtime/Engine/Private/Materials/MaterialIRValueAnalyzer.cpp
--- a/Engine/Source/Runtime/Engine/Private/Materials/MaterialIRValueAnalyzer.cpp
+++ b/Engine/Source/Runtime/Engine/Private/Materials/MaterialIRValueAnalyzer.cpp
@@ -18,7 +18,9 @@
 #include "MaterialDomain.h"
 #include "MaterialShared.h"
 #include "ParameterCollection.h"
+#include "Engine/Font.h"
 #include "Engine/Texture.h"
+#include "Engine/Texture2D.h"
 #include "UObject/Package.h"
 #include "PostProcess/PostProcessMaterialInputs.h"
 
@@ -488,7 +490,13 @@ static void AnalyzeTextureObject(FMaterialIRValueAnalyzer& Analyzer, MIR::FTextu
 
 static void AnalyzeTextureUniformParameter(FMaterialIRValueAnalyzer& Analyzer, MIR::FUniformParameter* Parameter)
 {
-	UTexture* Texture = Analyzer.Module->GetParameterMetadata(Parameter->ParameterIdInModule).Value.Texture;
+	const FMaterialParameterValue& ParameterValue = Analyzer.Module->GetParameterMetadata(Parameter->ParameterIdInModule).Value;
+
+	UTexture* Texture =
+		ParameterValue.Type == EMaterialParameterType::Texture ? ParameterValue.Texture :
+		ParameterValue.Type == EMaterialParameterType::Font ? ParameterValue.Font.Value->Textures[ParameterValue.Font.Page].Get() :
+		nullptr;
+	check(Texture);
 
 	FMaterialTextureParameterInfo ParamInfo {
 		.ParameterInfo = Analyzer.Module->GetParameterInfo(Parameter->ParameterIdInModule),
diff --git a/Engine/Source/Runtime/Engine/Public/Materials/MaterialExpressionFontSample.h b/Engine/Source/Runtime/Engine/Public/Materials/MaterialExpressionFontSample.h
--- a/Engine/Source/Runtime/Engine/Public/Materials/MaterialExpressionFontSample.h
+++ b/Engine/Source/Runtime/Engine/Public/Materials/MaterialExpressionFontSample.h
@@ -4,11 +4,16 @@
 #pragma once
 
 #include "CoreMinimal.h"
-#include "UObject/ObjectMacros.h"
+#include "Engine/EngineTypes.h"
 #include "Materials/MaterialExpression.h"
+#include "RHIShaderPlatform.h"
+#include "Templates/Tuple.h"
+#include "UObject/ObjectMacros.h"
+
 #include "MaterialExpressionFontSample.generated.h"
 
 class UTexture;
+class ITargetPlatform;
 
 namespace UE
 {
@@ -34,6 +39,7 @@ class UMaterialExpressionFontSample : public UMaterialExpression
 
 	//~ Begin UMaterialExpression Interface
 #if WITH_EDITOR
+	virtual void Build(MIR::FEmitter& Emitter) override;
 	virtual int32 Compile(class FMaterialCompiler* Compiler, int32 OutputIndex) override;
 	virtual void GetCaption(TArray<FString>& OutCaptions) const override;
 	virtual bool MatchesSearchQuery(const TCHAR* SearchQuery) override;
@@ -50,6 +56,9 @@ class UMaterialExpressionFontSample : public UMaterialExpression
 
 	virtual bool CanReferenceTexture() const override { return true; }
 	//~ End UMaterialExpression Interface
+
+protected:
+	TTuple<bool, UTexture*, EMaterialSamplerType, FString> ValidateAndGetTextureSampler(EShaderPlatform ShaderPlatform, const ITargetPlatform* TargetPlatform);
 };
 
 
diff --git a/Engine/Source/Runtime/Engine/Public/Materials/MaterialExpressionFontSampleParameter.h b/Engine/Source/Runtime/Engine/Public/Materials/MaterialExpressionFontSampleParameter.h
--- a/Engine/Source/Runtime/Engine/Public/Materials/MaterialExpressionFontSampleParameter.h
+++ b/Engine/Source/Runtime/Engine/Public/Materials/MaterialExpressionFontSampleParameter.h
@@ -36,6 +36,7 @@ class UMaterialExpressionFontSampleParameter : public UMaterialExpressionFontSam
 
 	//~ Begin UMaterialExpression Interface
 #if WITH_EDITOR
+	virtual void Build(MIR::FEmitter& Emitter) override;
 	virtual int32 Compile(class FMaterialCompiler* Compiler, int32 OutputIndex) override;
 	virtual void GetCaption(TArray<FString>& OutCaptions) const override;
 	virtual bool MatchesSearchQuery( const TCHAR* SearchQuery ) override;
diff --git a/Engine/Source/Runtime/Engine/Public/Materials/MaterialIREmitter.h b/Engine/Source/Runtime/Engine/Public/Materials/MaterialIREmitter.h
--- a/Engine/Source/Runtime/Engine/Public/Materials/MaterialIREmitter.h
+++ b/Engine/Source/Runtime/Engine/Public/Materials/MaterialIREmitter.h
@@ -40,6 +40,9 @@ struct FSwizzleMask
 	// Pushes a component to this swizzle mask.
 	void Append(EVectorComponent Component);
 
+	// Returns true if this is an identity swizzle mask XYZW, which implies i.e. it's a redundant swizzle mask.
+	bool IsXYZW() const;
+
 	const EVectorComponent* begin() const { return Components; }
 	const EVectorComponent* end() const { return Components + NumComponents; }
 };
@@ -312,8 +315,11 @@ public:
 	// Flows given `Value` out of the expression output with given `OutputIndex`.
 	FEmitter& Output(int32 OutputIndex, FValueRef Value);
 
-	// Flows given `Value` out of given expression `Output`.
-	FEmitter& Output(const FExpressionOutput* Output, FValueRef Value);
+	// Flows given `Value` out of given expression `ExpressionOutput`.
+	FEmitter& Output(const FExpressionOutput* ExpressionOutput, FValueRef Value);
+
+	// Flows given `Value` out of given expression output list. Each output applies its respective mask.
+	FEmitter& Outputs(const TConstArrayView<FExpressionOutput>& ExpressionOutputs, FValueRef Value);
 	
 	/*------------------------------- Constants emission -------------------------------*/
 
@@ -684,6 +690,7 @@ public:
 
 	/*-------------------------- Pass through accessors ---------------------------*/
 	EShaderPlatform GetShaderPlatform() const;
+	const ITargetPlatform* GetTargetPlatform() const;
 	ERHIFeatureLevel::Type GetFeatureLevel() const;
 	EMaterialQualityLevel::Type GetQualityLevel() const;
 
diff --git a/Engine/Source/Runtime/Engine/Public/Materials/MaterialIRModule.h b/Engine/Source/Runtime/Engine/Public/Materials/MaterialIRModule.h
--- a/Engine/Source/Runtime/Engine/Public/Materials/MaterialIRModule.h
+++ b/Engine/Source/Runtime/Engine/Public/Materials/MaterialIRModule.h
@@ -53,6 +53,9 @@ public:
 	// Returns the shader platform associated with this module.
 	EShaderPlatform GetShaderPlatform() const { return ShaderPlatform; }
 
+	// Returns the target platform interface associated with this module.
+	const ITargetPlatform* GetTargetPlatform() const { return TargetPlatform; };
+
 	// Returns the feature level associated with this module.
 	ERHIFeatureLevel::Type GetFeatureLevel() const { return FeatureLevel; }
 
@@ -120,6 +123,9 @@ private:
 	// Target shader platform.
 	EShaderPlatform ShaderPlatform;
 
+	// Target platform interface.
+	const ITargetPlatform* TargetPlatform;
+
 	// Target feature level.
 	ERHIFeatureLevel::Type FeatureLevel;
 
