Remove notify guids from Pose Search Database hash generation, otherwise the hash can be non-deterministic

diff --git a/Engine/Source/Runtime/Engine/Public/Animation/AnimTypes.h b/Engine/Source/Runtime/Engine/Public/Animation/AnimTypes.h
--- a/Engine/Source/Runtime/Engine/Public/Animation/AnimTypes.h
+++ b/Engine/Source/Runtime/Engine/Public/Animation/AnimTypes.h
@@ -359,7 +359,7 @@ struct FAnimNotifyEvent : public FAnimLinkableElement
 	FColor NotifyColor;
 
 	/** Guid for tracking notifies in editor */
-	UPROPERTY()
+	UPROPERTY(meta = (ExcludeFromHash))
 	FGuid Guid;
 #endif // WITH_EDITORONLY_DATA
 
