Spline: Increment version on copy assignment

diff --git a/Engine/Source/Runtime/Engine/Private/Curves/Spline.cpp b/Engine/Source/Runtime/Engine/Private/Curves/Spline.cpp
--- a/Engine/Source/Runtime/Engine/Private/Curves/Spline.cpp
+++ b/Engine/Source/Runtime/Engine/Private/Curves/Spline.cpp
@@ -850,6 +850,8 @@ FSpline& FSpline::operator=(const FSpline& Other)
 		break;
 	}
 	
+	Version++;
+	
 	return *this;
 }
 
@@ -867,6 +869,8 @@ FSpline& FSpline::operator=(const FSplineCurves& Other)
 		break;
 	}
 	
+	Version++;
+	
 	return *this;
 }
 
