Adding DelayUntilNextFrame latent action that guarantees it will not finish until the engine frame has advanced.
Clarified behavior with DelayUntilNextTick where it waits for the next tick, not necessarily the next FRAME.
Also fixes issue where duplicate delay-until actions were not ignored.

diff --git a/Engine/Source/Runtime/Engine/Classes/Kismet/KismetSystemLibrary.h b/Engine/Source/Runtime/Engine/Classes/Kismet/KismetSystemLibrary.h
--- a/Engine/Source/Runtime/Engine/Classes/Kismet/KismetSystemLibrary.h
+++ b/Engine/Source/Runtime/Engine/Classes/Kismet/KismetSystemLibrary.h
@@ -656,13 +656,25 @@ class UKismetSystemLibrary : public UBlueprintFunctionLibrary
 	static ENGINE_API void	Delay(const UObject* WorldContextObject, float Duration, struct FLatentActionInfo LatentInfo );
 
 	/**
-	 * Perform a latent action with a delay of one tick.  Calling again while it is counting down will be ignored.
+	 * Perform a latent action with a delay of one tick.  Calling again while it is pending will be ignored.
+	 * Note that this delay is until the object's next tick, which may occur subsequently during the current engine frame.
+	 * @see DelayUntilNextFrame for a guaranteed wait until next engine frame's tick.
+	 * Enabling the CVar 'LatentActions.GuaranteeNextTickDelay' will guarantee this delay waits until next frame's tick.
 	 *
 	 * @param WorldContext	World context.
 	 * @param LatentInfo 	The latent action.
 	 */
 	UFUNCTION(BlueprintCallable, Category = "Utilities|FlowControl", meta = (Latent, WorldContext = "WorldContextObject", LatentInfo = "LatentInfo", Keywords = "sleep"))
-	static ENGINE_API void	DelayUntilNextTick(const UObject* WorldContextObject, struct FLatentActionInfo LatentInfo);
+	static ENGINE_API void	DelayUntilNextTick(const UObject* WorldContextObject, FLatentActionInfo LatentInfo);
+
+	/**
+	 * Perform a latent action with a delay until the next engine frame.  Calling again while it is pending will be ignored.
+	 * @param WorldContext	World context.
+	 * @param LatentInfo 	The latent action.
+	 */
+	UFUNCTION(BlueprintCallable, Category = "Utilities|FlowControl", meta = (Latent, WorldContext = "WorldContextObject", LatentInfo = "LatentInfo", Keywords = "sleep"))
+	static ENGINE_API void	DelayUntilNextFrame(const UObject* WorldContextObject, FLatentActionInfo LatentInfo);
+
 
 	/** 
 	 * Perform a latent action with a retriggerable delay (specified in seconds).  Calling again while it is counting down will reset the countdown to Duration.
diff --git a/Engine/Source/Runtime/Engine/Private/KismetSystemLibrary.cpp b/Engine/Source/Runtime/Engine/Private/KismetSystemLibrary.cpp
--- a/Engine/Source/Runtime/Engine/Private/KismetSystemLibrary.cpp
+++ b/Engine/Source/Runtime/Engine/Private/KismetSystemLibrary.cpp
@@ -2621,13 +2621,26 @@ void UKismetSystemLibrary::DelayUntilNextTick(const UObject* WorldContextObject,
 	if (UWorld* World = GEngine->GetWorldFromContextObject(WorldContextObject, EGetWorldErrorMode::LogAndReturnNull))
 	{
 		FLatentActionManager& LatentActionManager = World->GetLatentActionManager();
-		if (LatentActionManager.FindExistingAction<FDelayAction>(LatentInfo.CallbackTarget, LatentInfo.UUID) == NULL)
+		if (LatentActionManager.FindExistingAction<FDelayUntilNextTickAction>(LatentInfo.CallbackTarget, LatentInfo.UUID) == NULL)
 		{
 			LatentActionManager.AddNewAction(LatentInfo.CallbackTarget, LatentInfo.UUID, new FDelayUntilNextTickAction(LatentInfo));
 		}
 	}
 }
 
+void UKismetSystemLibrary::DelayUntilNextFrame(const UObject* WorldContextObject, FLatentActionInfo LatentInfo)
+{
+	if (UWorld* World = GEngine->GetWorldFromContextObject(WorldContextObject, EGetWorldErrorMode::LogAndReturnNull))
+	{
+		FLatentActionManager& LatentActionManager = World->GetLatentActionManager();
+		if (LatentActionManager.FindExistingAction<FDelayUntilNextFrameAction>(LatentInfo.CallbackTarget, LatentInfo.UUID) == NULL)
+		{
+			LatentActionManager.AddNewAction(LatentInfo.CallbackTarget, LatentInfo.UUID, new FDelayUntilNextFrameAction(LatentInfo));
+		}
+	}
+}
+
+
 // Delay execution by Duration seconds; Calling again before the delay has expired will reset the countdown to Duration.
 void UKismetSystemLibrary::RetriggerableDelay(const UObject* WorldContextObject, float Duration, FLatentActionInfo LatentInfo)
 {
diff --git a/Engine/Source/Runtime/Engine/Public/DelayAction.h b/Engine/Source/Runtime/Engine/Public/DelayAction.h
--- a/Engine/Source/Runtime/Engine/Public/DelayAction.h
+++ b/Engine/Source/Runtime/Engine/Public/DelayAction.h
@@ -43,8 +43,10 @@ public:
 };
 
 
-// FDelayUntilNextTickAction
-// A simple delay action; triggers on the next engine tick. See also @CVarLatentActionGuaranteeEngineTickDelay
+/** FDelayUntilNextTickAction
+ * A simple delay action; triggers on the object's next tick. This tick *may* occur in the current engine frame, depending on order of operations.
+ * See also @CVarLatentActionGuaranteeEngineTickDelay and FDelayUntilNextFrameTickAction for guaranteed delays until the next frame.
+ */
 class FDelayUntilNextTickAction : public FPendingLatentAction
 {
 public:
@@ -76,5 +78,27 @@ public:
 #endif
 };
 
+// A variant of the simple delay action that guarantees a delay until the next engine frame
+class FDelayUntilNextFrameAction : public FDelayUntilNextTickAction
+{
+public:
+	FDelayUntilNextFrameAction(const FLatentActionInfo& LatentInfo)
+		: FDelayUntilNextTickAction(LatentInfo)
+	{
+	}
+
+	virtual void UpdateOperation(FLatentResponse& Response) override
+	{
+		const bool bHasFrameAdvanced = ((uint8)GFrameCounter & 1) != InitialFrameParity;
+		Response.FinishAndTriggerIf(bHasFrameAdvanced, ExecutionFunction, OutputLink, CallbackTarget);
+	}
 
+#if WITH_EDITOR
+	// Returns a human readable description of the latent operation's current state
+	virtual FString GetDescription() const override
+	{
+		return NSLOCTEXT("DelayUntilNextFrameAction", "DelayUntilNextFrameActionFmt", "Delay for one engine frame").ToString();
+	}
+#endif
+};
 
