Fix object trace performance issue when tracing is disabled

diff --git a/Engine/Source/Runtime/Engine/Private/ObjectTrace.cpp b/Engine/Source/Runtime/Engine/Private/ObjectTrace.cpp
--- a/Engine/Source/Runtime/Engine/Private/ObjectTrace.cpp
+++ b/Engine/Source/Runtime/Engine/Private/ObjectTrace.cpp
@@ -433,24 +433,14 @@ void FObjectTrace::OutputView(const UObject* InPlayer, const FSceneView* InView)
 		<< View.AspectRatio(AspectRatio);
 }
 
-void FObjectTrace::OutputInstance(const UObject* InOuterObject, uint64 InInstanceId, uint64 InOuterId, UStruct* InType, const FString& InName, const FString& InPathName)
+void OutputInstanceInternal(const UObject* InOuterObject, uint64 InInstanceId, uint64 InOuterId, UStruct* InType, const FString& InName, const FString& InPathName)
 {
-	if (UE::ObjectTrace::IsTracingDisabledForObject(InOuterObject))
-	{
-		return;
-	}
-	
-	if (!AddObjectToTracedSet(UE::ObjectTrace::FObjectId{InInstanceId}))
-	{
-		return;
-	}
-
 	// Trace the object's class first
 	TRACE_TYPE(InType);
 
 	UE_TRACE_LOG(Object, Object, ObjectChannel)
 		<< Object.Id(InInstanceId)
-		<< Object.ClassId(GetObjectId(InType))
+		<< Object.ClassId(FObjectTrace::GetObjectId(InType))
 		<< Object.OuterId(InOuterId)
 		<< Object.Name(*InName, InName.Len())
 		<< Object.Path(*InPathName, InPathName.Len());
@@ -461,13 +451,38 @@ void FObjectTrace::OutputInstance(const UObject* InOuterObject, uint64 InInstanc
 		<< ObjectLifetimeBegin2.Id(InInstanceId);
 }
 
+void FObjectTrace::OutputInstance(const UObject* InOuterObject, uint64 InInstanceId, uint64 InOuterId, UStruct* InType, const FString& InName, const FString& InPathName)
+{
+	if (UE::ObjectTrace::IsTracingDisabledForObject(InOuterObject))
+	{
+		return;
+	}
+	
+	if (!AddObjectToTracedSet(UE::ObjectTrace::FObjectId{InInstanceId}))
+	{
+		return;
+	}
+	
+	OutputInstanceInternal(InOuterObject, InInstanceId, InOuterId, InType, InName, InPathName);
+}
+
+
 void FObjectTrace::OutputObject(const UObject* InObject)
 {
-	if (InObject)
+	if (UE::ObjectTrace::IsTracingDisabledForObject(InObject))
+ 	{
+ 		return;
+ 	}
+
+	uint64 InstanceId = GetObjectId(InObject);
+	
+	if (!AddObjectToTracedSet(UE::ObjectTrace::FObjectId{InstanceId}))
 	{
-		OutputObject(InObject->GetOuter());
-		OutputInstance(InObject, GetObjectId(InObject), GetObjectId(InObject->GetOuter()), InObject->GetClass(), InObject->GetFName().ToString(), InObject->GetPathName());
+		return;
 	}
+ 	
+	OutputObject(InObject->GetOuter());
+	OutputInstanceInternal(InObject, InstanceId, GetObjectId(InObject->GetOuter()), InObject->GetClass(), InObject->GetFName().ToString(), InObject->GetPathName());
 }
 
 void FObjectTrace::OutputObjectTransform(const UObject* InObject, const FTransform& InTransform, ETeleportType InTeleportType)
