Fixing incorrect term in description. No functional changes.

diff --git a/Engine/Source/Runtime/Engine/Public/DelayAction.h b/Engine/Source/Runtime/Engine/Public/DelayAction.h
--- a/Engine/Source/Runtime/Engine/Public/DelayAction.h
+++ b/Engine/Source/Runtime/Engine/Public/DelayAction.h
@@ -45,7 +45,7 @@ public:
 
 /** FDelayUntilNextTickAction
  * A simple delay action; triggers on the object's next tick. This tick *may* occur in the current engine frame, depending on order of operations.
- * See also @CVarLatentActionGuaranteeEngineTickDelay and FDelayUntilNextFrameTickAction for guaranteed delays until the next frame.
+ * See also @CVarLatentActionGuaranteeEngineTickDelay and FDelayUntilNextFrameAction for guaranteed delays until the next frame.
  */
 class FDelayUntilNextTickAction : public FPendingLatentAction
 {
