FStaticMeshSceneProxyDesc cleanup
- Moving all members together instead of being mixed between methods.
- Moving all FStaticMeshSceneProxyDesc method definitions to its own cpp file.
- Moving non-trivial inline methods to cpp.

diff --git a/Engine/Plugins/Experimental/FastGeoStreaming/Source/FastGeoStreaming/Private/FastGeoStaticMeshComponent.cpp b/Engine/Plugins/Experimental/FastGeoStreaming/Source/FastGeoStreaming/Private/FastGeoStaticMeshComponent.cpp
--- a/Engine/Plugins/Experimental/FastGeoStreaming/Source/FastGeoStreaming/Private/FastGeoStaticMeshComponent.cpp
+++ b/Engine/Plugins/Experimental/FastGeoStreaming/Source/FastGeoStreaming/Private/FastGeoStaticMeshComponent.cpp
@@ -1,6 +1,8 @@
 // Copyright Epic Games, Inc. All Rights Reserved.
 
 #include "FastGeoStaticMeshComponent.h"
+
+#include "Engine/Engine.h"
 #include "FastGeoContainer.h"
 #include "FastGeoLog.h"
 #include "AI/Navigation/NavCollisionBase.h"
@@ -321,6 +323,11 @@ UStaticMesh* FFastGeoStaticMeshComponentBase::GetStaticMesh() const
 	return GetStaticMeshSceneProxyDesc().StaticMesh;
 }
 
+UObject const* FFastGeoStaticMeshComponentBase::AdditionalStatObject() const
+{
+	return GetStaticMesh();
+}
+
 UMaterialInterface* FFastGeoStaticMeshComponentBase::GetOverlayMaterial() const
 {
 	return GetStaticMeshSceneProxyDesc().OverlayMaterial;
diff --git a/Engine/Plugins/Experimental/FastGeoStreaming/Source/FastGeoStreaming/Private/FastGeoStaticMeshComponent.h b/Engine/Plugins/Experimental/FastGeoStreaming/Source/FastGeoStreaming/Private/FastGeoStaticMeshComponent.h
--- a/Engine/Plugins/Experimental/FastGeoStreaming/Source/FastGeoStreaming/Private/FastGeoStaticMeshComponent.h
+++ b/Engine/Plugins/Experimental/FastGeoStreaming/Source/FastGeoStreaming/Private/FastGeoStaticMeshComponent.h
@@ -38,7 +38,7 @@ public:
 	bool HasValidNaniteData() const;
 	bool UseNaniteOverrideMaterials(bool bDoingMaterialAudit) const;
 	bool ShouldCreateNaniteProxy(Nanite::FMaterialAudit* OutNaniteMaterials) const;
-	virtual UObject const* AdditionalStatObject() const override { return GetStaticMesh(); }
+	virtual UObject const* AdditionalStatObject() const override;
 	bool IsReverseCulling() const;
 	bool IsDisallowNanite() const;
 	bool IsForceDisableNanite() const;
diff --git a/Engine/Source/Runtime/Engine/Private/Components/StaticMeshComponent.cpp b/Engine/Source/Runtime/Engine/Private/Components/StaticMeshComponent.cpp
--- a/Engine/Source/Runtime/Engine/Private/Components/StaticMeshComponent.cpp
+++ b/Engine/Source/Runtime/Engine/Private/Components/StaticMeshComponent.cpp
@@ -2438,11 +2438,6 @@ bool UStaticMeshComponent::HasValidNaniteData() const
 	return Nanite::FNaniteResourcesHelper::HasValidNaniteData(*this);
 }
 
-bool FStaticMeshSceneProxyDesc::HasValidNaniteData() const
-{
-	return Nanite::FNaniteResourcesHelper::HasValidNaniteData(*this);
-}
-
 bool UStaticMeshComponent::UseNaniteOverrideMaterials(bool bDoingMaterialAudit) const
 {
 	return Nanite::FNaniteResourcesHelper::UseNaniteOverrideMaterials(*this, bDoingMaterialAudit);
@@ -2453,11 +2448,6 @@ bool UStaticMeshComponent::UseNaniteOverrideMaterials() const
 	return UseNaniteOverrideMaterials(false);
 }
 
-bool FStaticMeshSceneProxyDesc::UseNaniteOverrideMaterials(bool bDoingMaterialAudit) const
-{
-	return Nanite::FNaniteResourcesHelper::UseNaniteOverrideMaterials(*this, bDoingMaterialAudit);	
-}
-
 void UStaticMeshComponent::SetForcedLodModel(int32 NewForcedLodModel)
 {
 	if (ForcedLodModel != NewForcedLodModel)
@@ -2911,27 +2901,6 @@ void UStaticMeshComponent::GetUsedMaterials(TArray<UMaterialInterface*>& OutMate
 	FStaticMeshComponentHelper::GetUsedMaterials(*this, OutMaterials, bGetDebugMaterials);
 }
 
-void FStaticMeshSceneProxyDesc::GetUsedMaterials(TArray<UMaterialInterface*>& OutMaterials, bool bGetDebugMaterials) const
-{
-	FStaticMeshComponentHelper::GetUsedMaterials(*this, OutMaterials, bGetDebugMaterials);
-}
-
-// Deprecated in 5.7
-FMaterialRelevance FStaticMeshSceneProxyDesc::GetMaterialRelevance(ERHIFeatureLevel::Type InFeatureLevel) const
-{
-	return GetMaterialRelevance(GetFeatureLevelShaderPlatform_Checked(InFeatureLevel));
-}
-
-FMaterialRelevance FStaticMeshSceneProxyDesc::GetMaterialRelevance(EShaderPlatform InShaderPlatform) const
-{
-	if (bUseProvidedMaterialRelevance)
-	{
-		return MaterialRelevance;
-	}
-
-	return FMeshComponentHelper::GetMaterialRelevance(*this, InShaderPlatform);
-}
-
 #if WITH_EDITOR
 bool UStaticMeshComponent::GetMaterialPropertyPath(int32 ElementIndex, UObject*& OutOwner, FString& OutPropertyPath, FProperty*& OutProperty)
 {
diff --git a/Engine/Source/Runtime/Engine/Private/SplineMeshSceneProxy.cpp b/Engine/Source/Runtime/Engine/Private/SplineMeshSceneProxy.cpp
--- a/Engine/Source/Runtime/Engine/Private/SplineMeshSceneProxy.cpp
+++ b/Engine/Source/Runtime/Engine/Private/SplineMeshSceneProxy.cpp
@@ -1,6 +1,7 @@
 // Copyright Epic Games, Inc. All Rights Reserved.
 
 #include "SplineMeshSceneProxy.h"
+#include "Engine/StaticMesh.h"
 #include "Materials/Material.h"
 #include "Materials/MaterialRenderProxy.h"
 #include "MaterialDomain.h"
diff --git a/Engine/Source/Runtime/Engine/Private/SplineMeshSceneProxyDesc.cpp b/Engine/Source/Runtime/Engine/Private/SplineMeshSceneProxyDesc.cpp
--- a/Engine/Source/Runtime/Engine/Private/SplineMeshSceneProxyDesc.cpp
+++ b/Engine/Source/Runtime/Engine/Private/SplineMeshSceneProxyDesc.cpp
@@ -5,6 +5,7 @@
 #include "SplineMeshSceneProxy.h"
 #include "StaticMeshResources.h"
 #include "Components/SplineMeshComponent.h"
+#include "Engine/StaticMesh.h"
 #include "HAL/ConsoleManager.h"
 
 static TAutoConsoleVariable<bool> CVarSplineMeshFitToSourceMeshBounds(
diff --git a/Engine/Source/Runtime/Engine/Private/StaticMeshSceneProxy.cpp b/Engine/Source/Runtime/Engine/Private/StaticMeshSceneProxy.cpp
--- a/Engine/Source/Runtime/Engine/Private/StaticMeshSceneProxy.cpp
+++ b/Engine/Source/Runtime/Engine/Private/StaticMeshSceneProxy.cpp
@@ -2759,110 +2759,11 @@ FLODMask FStaticMeshSceneProxy::GetLODMask(const FSceneView* View) const
 	return Result;
 }
 
-UMaterialInterface* FStaticMeshSceneProxyDesc::GetMaterial(int32 MaterialIndex, bool bDoingNaniteMaterialAudit, bool bIgnoreNaniteOverrideMaterials) const
-{
-	return FStaticMeshComponentHelper::GetMaterial(*this, MaterialIndex, bDoingNaniteMaterialAudit, bIgnoreNaniteOverrideMaterials);
-}
-
 bool UStaticMeshComponent::ShouldCreateNaniteProxy(Nanite::FMaterialAudit* OutNaniteMaterials) const
 {
 	return Nanite::FNaniteResourcesHelper::ShouldCreateNaniteProxy(*this, OutNaniteMaterials);
 }
 
-bool FStaticMeshSceneProxyDesc::ShouldCreateNaniteProxy(Nanite::FMaterialAudit* OutNaniteMaterials /*= nullptr*/) const
-{
-	return Nanite::FNaniteResourcesHelper::ShouldCreateNaniteProxy(*this, OutNaniteMaterials);
-}
-
-FTextureResource* FStaticMeshSceneProxyDesc::GetMeshPaintTextureResource() const
-{
-	if (MeshPaintTexture && MeshPaintTexture->IsCurrentlyVirtualTextured())
-	{
-		return MeshPaintTexture->GetResource();
-	}
-	return nullptr;
-}
-
-bool FStaticMeshSceneProxyDesc::ShouldCreateMaterialCacheProxy() const
-{
-	return !MaterialCacheTextures.IsEmpty() && MaterialCacheTextures[0]->IsCurrentlyVirtualTextured();
-}
-
-FStaticMeshSceneProxyDesc::FStaticMeshSceneProxyDesc(const UStaticMeshComponent* InComponent)
-	: FStaticMeshSceneProxyDesc()
-{	
-	InitializeFromStaticMeshComponent(InComponent);
-}
-
-void FStaticMeshSceneProxyDesc::InitializeFromStaticMeshComponent(const UStaticMeshComponent* InComponent)
-{
-	InitializeFromPrimitiveComponent(InComponent);	
-
-	StaticMesh = InComponent->GetStaticMesh();
-	OverrideMaterials = const_cast<UStaticMeshComponent*>(InComponent)->OverrideMaterials;	
-	OverlayMaterial = InComponent->GetOverlayMaterial();
-	OverlayMaterialMaxDrawDistance = InComponent->GetOverlayMaterialMaxDrawDistance();
-	InComponent->GetMaterialSlotsOverlayMaterial(MaterialSlotsOverlayMaterial);
-
-	ForcedLodModel = InComponent->ForcedLodModel ;
-	MinLOD = InComponent->MinLOD ;
-	WorldPositionOffsetDisableDistance = InComponent->WorldPositionOffsetDisableDistance ;
-	NanitePixelProgrammableDistance = InComponent->NanitePixelProgrammableDistance;
-	bReverseCulling = InComponent->bReverseCulling ;
-	bEvaluateWorldPositionOffset = InComponent->bEvaluateWorldPositionOffset ;
-	bOverrideMinLOD = InComponent->bOverrideMinLOD ;
-	bCastDistanceFieldIndirectShadow = InComponent->bCastDistanceFieldIndirectShadow ;
-	bOverrideDistanceFieldSelfShadowBias = InComponent->bOverrideDistanceFieldSelfShadowBias ;
-	bEvaluateWorldPositionOffsetInRayTracing = InComponent->bEvaluateWorldPositionOffsetInRayTracing ;	
-	bSortTriangles = InComponent->bSortTriangles ;
-#if WITH_EDITOR
-	bDisplayNaniteFallbackMesh = InComponent->bDisplayNaniteFallbackMesh ;
-#endif
-	bDisallowNanite = InComponent->bDisallowNanite ;
-	bForceDisableNanite = InComponent->bForceDisableNanite ;
-	bForceNaniteForMasked = InComponent->bForceNaniteForMasked ;
-	DistanceFieldSelfShadowBias = InComponent->DistanceFieldSelfShadowBias ;
-	DistanceFieldIndirectShadowMinVisibility = InComponent->DistanceFieldIndirectShadowMinVisibility ;
-	StaticLightMapResolution = InComponent->GetStaticLightMapResolution();
-	LightmapType = InComponent->GetLightmapType();
-
-#if WITH_EDITORONLY_DATA
-	StreamingDistanceMultiplier = InComponent->StreamingDistanceMultiplier;
-	MaterialStreamingRelativeBoxes = const_cast<UStaticMeshComponent*>(InComponent)->MaterialStreamingRelativeBoxes;
-	SectionIndexPreview = InComponent->SectionIndexPreview;
-	MaterialIndexPreview = InComponent->MaterialIndexPreview;
-	SelectedEditorMaterial = InComponent->SelectedEditorMaterial;
-	SelectedEditorSection = InComponent->SelectedEditorSection;
-
-	TextureStreamingTransformScale = InComponent->GetTextureStreamingTransformScale();	
-#endif
-
-	NaniteResources = InComponent->GetNaniteResources();
-	BodySetup = const_cast<UStaticMeshComponent*>(InComponent)->GetBodySetup();
-
-#if STATICMESH_ENABLE_DEBUG_RENDERING
-	const bool bHasCollisionState = BodySetup && !BodySetup->bNeverNeedsCookedCollisionData;
-	bDrawMeshCollisionIfComplex = InComponent->bDrawMeshCollisionIfComplex && bHasCollisionState;
-	bDrawMeshCollisionIfSimple = InComponent->bDrawMeshCollisionIfSimple && bHasCollisionState;
-#endif
-
-	LODData = const_cast<UStaticMeshComponent*>(InComponent)->LODData;
-
-	WireframeColor = InComponent->GetWireframeColor();
-	LODParentPrimitive = InComponent->GetLODParentPrimitive();	
-
-	if (GetScene())
-	{
-		SetMaterialRelevance(InComponent->GetMaterialRelevance(GetScene()->GetShaderPlatform()));
-	}
-	SetCollisionResponseToChannels(InComponent->GetCollisionResponseToChannels());
-
-	MeshPaintTexture = InComponent->MeshPaintTextureOverride ? InComponent->MeshPaintTextureOverride.Get() : InComponent->GetMeshPaintTexture();
-	MeshPaintTextureCoordinateIndex = InComponent->GetMeshPaintTextureCoordinateIndex();
-
-	MaterialCacheTextures = InComponent->MaterialCacheTextures;
-}
-
 FPrimitiveSceneProxy* UStaticMeshComponent::CreateStaticMeshSceneProxy(Nanite::FMaterialAudit& NaniteMaterials, bool bCreateNanite)
 {
 	// Default implementation: Nanite::FSceneProxy or FStaticMeshSceneProxy
diff --git a/Engine/Source/Runtime/Engine/Private/StaticMeshSceneProxyDesc.cpp b/Engine/Source/Runtime/Engine/Private/StaticMeshSceneProxyDesc.cpp
new file mode 100644
--- /dev/null
+++ b/Engine/Source/Runtime/Engine/Private/StaticMeshSceneProxyDesc.cpp
@@ -0,0 +1,226 @@
+// Copyright Epic Games, Inc. All Rights Reserved.
+
+#include "StaticMeshSceneProxyDesc.h"
+
+#include "MaterialCache/MaterialCacheVirtualTexture.h"
+#include "PhysicsEngine/BodySetup.h"
+#include "SceneInterface.h"
+#include "StaticMeshComponentHelper.h"
+
+FStaticMeshSceneProxyDesc::FStaticMeshSceneProxyDesc()
+{
+	CastShadow = true;
+	bUseAsOccluder = true;
+
+	// bitfields init (until we switch to c++20)
+	bReverseCulling = false;
+#if STATICMESH_ENABLE_DEBUG_RENDERING
+	bDrawMeshCollisionIfComplex = false;
+	bDrawMeshCollisionIfSimple = false;
+#endif
+	bEvaluateWorldPositionOffset = true;
+	bOverrideMinLOD = false;
+
+	bCastDistanceFieldIndirectShadow = false;
+	bOverrideDistanceFieldSelfShadowBias = false;
+	bEvaluateWorldPositionOffsetInRayTracing = false;
+	bSortTriangles = false;
+
+	bDisplayNaniteFallbackMesh = false;
+	bDisallowNanite = false;
+	bForceDisableNanite = false;
+	bForceNaniteForMasked = false;
+
+	bUseProvidedMaterialRelevance = false;
+}
+
+FStaticMeshSceneProxyDesc::FStaticMeshSceneProxyDesc(const UStaticMeshComponent* InComponent)
+	: FStaticMeshSceneProxyDesc()
+{
+	InitializeFromStaticMeshComponent(InComponent);
+}
+
+FStaticMeshSceneProxyDesc::~FStaticMeshSceneProxyDesc() = default;
+
+void FStaticMeshSceneProxyDesc::InitializeFromStaticMeshComponent(const UStaticMeshComponent* InComponent)
+{
+	InitializeFromPrimitiveComponent(InComponent);
+
+	StaticMesh = InComponent->GetStaticMesh();
+	OverrideMaterials = const_cast<UStaticMeshComponent*>(InComponent)->OverrideMaterials;
+	OverlayMaterial = InComponent->GetOverlayMaterial();
+	OverlayMaterialMaxDrawDistance = InComponent->GetOverlayMaterialMaxDrawDistance();
+	InComponent->GetMaterialSlotsOverlayMaterial(MaterialSlotsOverlayMaterial);
+
+	ForcedLodModel = InComponent->ForcedLodModel;
+	MinLOD = InComponent->MinLOD;
+	WorldPositionOffsetDisableDistance = InComponent->WorldPositionOffsetDisableDistance;
+	NanitePixelProgrammableDistance = InComponent->NanitePixelProgrammableDistance;
+	bReverseCulling = InComponent->bReverseCulling;
+	bEvaluateWorldPositionOffset = InComponent->bEvaluateWorldPositionOffset;
+	bOverrideMinLOD = InComponent->bOverrideMinLOD;
+	bCastDistanceFieldIndirectShadow = InComponent->bCastDistanceFieldIndirectShadow;
+	bOverrideDistanceFieldSelfShadowBias = InComponent->bOverrideDistanceFieldSelfShadowBias;
+	bEvaluateWorldPositionOffsetInRayTracing = InComponent->bEvaluateWorldPositionOffsetInRayTracing;
+	bSortTriangles = InComponent->bSortTriangles;
+#if WITH_EDITOR
+	bDisplayNaniteFallbackMesh = InComponent->bDisplayNaniteFallbackMesh;
+#endif
+	bDisallowNanite = InComponent->bDisallowNanite;
+	bForceDisableNanite = InComponent->bForceDisableNanite;
+	bForceNaniteForMasked = InComponent->bForceNaniteForMasked;
+	DistanceFieldSelfShadowBias = InComponent->DistanceFieldSelfShadowBias;
+	DistanceFieldIndirectShadowMinVisibility = InComponent->DistanceFieldIndirectShadowMinVisibility;
+	StaticLightMapResolution = InComponent->GetStaticLightMapResolution();
+	LightmapType = InComponent->GetLightmapType();
+
+#if WITH_EDITORONLY_DATA
+	StreamingDistanceMultiplier = InComponent->StreamingDistanceMultiplier;
+	MaterialStreamingRelativeBoxes = const_cast<UStaticMeshComponent*>(InComponent)->MaterialStreamingRelativeBoxes;
+	SectionIndexPreview = InComponent->SectionIndexPreview;
+	MaterialIndexPreview = InComponent->MaterialIndexPreview;
+	SelectedEditorMaterial = InComponent->SelectedEditorMaterial;
+	SelectedEditorSection = InComponent->SelectedEditorSection;
+
+	TextureStreamingTransformScale = InComponent->GetTextureStreamingTransformScale();
+#endif
+
+	NaniteResources = InComponent->GetNaniteResources();
+	BodySetup = const_cast<UStaticMeshComponent*>(InComponent)->GetBodySetup();
+
+#if STATICMESH_ENABLE_DEBUG_RENDERING
+	const bool bHasCollisionState = BodySetup && !BodySetup->bNeverNeedsCookedCollisionData;
+	bDrawMeshCollisionIfComplex = InComponent->bDrawMeshCollisionIfComplex && bHasCollisionState;
+	bDrawMeshCollisionIfSimple = InComponent->bDrawMeshCollisionIfSimple && bHasCollisionState;
+#endif
+
+	LODData = const_cast<UStaticMeshComponent*>(InComponent)->LODData;
+
+	WireframeColor = InComponent->GetWireframeColor();
+	LODParentPrimitive = InComponent->GetLODParentPrimitive();
+
+	if (GetScene())
+	{
+		SetMaterialRelevance(InComponent->GetMaterialRelevance(GetScene()->GetShaderPlatform()));
+	}
+	SetCollisionResponseToChannels(InComponent->GetCollisionResponseToChannels());
+
+	MeshPaintTexture = InComponent->MeshPaintTextureOverride ? InComponent->MeshPaintTextureOverride.Get() : InComponent->GetMeshPaintTexture();
+	MeshPaintTextureCoordinateIndex = InComponent->GetMeshPaintTextureCoordinateIndex();
+
+	MaterialCacheTextures = InComponent->MaterialCacheTextures;
+}
+
+void FStaticMeshSceneProxyDesc::GetUsedMaterials(TArray<UMaterialInterface*>& OutMaterials, bool bGetDebugMaterials) const
+{
+	FStaticMeshComponentHelper::GetUsedMaterials(*this, OutMaterials, bGetDebugMaterials);
+}
+
+UMaterialInterface* FStaticMeshSceneProxyDesc::GetMaterial(int32 MaterialIndex, bool bDoingNaniteMaterialAudit, bool bIgnoreNaniteOverrideMaterials) const
+{
+	return FStaticMeshComponentHelper::GetMaterial(*this, MaterialIndex, bDoingNaniteMaterialAudit, bIgnoreNaniteOverrideMaterials);
+}
+
+int32 FStaticMeshSceneProxyDesc::GetNumMaterials() const
+{
+	return GetStaticMesh() ? GetStaticMesh()->GetStaticMaterials().Num() : 0;
+}
+
+bool FStaticMeshSceneProxyDesc::ShouldCreateMaterialCacheProxy() const
+{
+	return !MaterialCacheTextures.IsEmpty() && MaterialCacheTextures[0]->IsCurrentlyVirtualTextured();
+}
+
+const UStaticMeshComponent* FStaticMeshSceneProxyDesc::GetUStaticMeshComponent() const
+{
+	return Cast<UStaticMeshComponent>(Component);
+}
+
+void FStaticMeshSceneProxyDesc::GetMaterialSlotsOverlayMaterial(TArray<TObjectPtr<UMaterialInterface>>& OutMaterialSlotsOverlayMaterial) const
+{
+	OutMaterialSlotsOverlayMaterial = MaterialSlotsOverlayMaterial;
+}
+
+void FStaticMeshSceneProxyDesc::SetMaterialRelevance(const FMaterialRelevance& InRelevance)
+{
+	MaterialRelevance = InRelevance;
+	bUseProvidedMaterialRelevance = true;
+}
+
+FMaterialRelevance FStaticMeshSceneProxyDesc::GetMaterialRelevance(EShaderPlatform InShaderPlatform) const
+{
+	if (bUseProvidedMaterialRelevance)
+	{
+		return MaterialRelevance;
+	}
+
+	return FMeshComponentHelper::GetMaterialRelevance(*this, InShaderPlatform);
+}
+
+// Deprecated in 5.7
+FMaterialRelevance FStaticMeshSceneProxyDesc::GetMaterialRelevance(ERHIFeatureLevel::Type InFeatureLevel) const
+{
+	return GetMaterialRelevance(GetFeatureLevelShaderPlatform_Checked(InFeatureLevel));
+}
+
+FTextureResource* FStaticMeshSceneProxyDesc::GetMeshPaintTextureResource() const
+{
+	if (MeshPaintTexture && MeshPaintTexture->IsCurrentlyVirtualTextured())
+	{
+		return MeshPaintTexture->GetResource();
+	}
+	return nullptr;
+}
+
+UBodySetup* FStaticMeshSceneProxyDesc::GetBodySetup() const
+{
+	if (BodySetup)
+	{
+		return BodySetup;
+	}
+
+	if (GetStaticMesh())
+	{
+		return GetStaticMesh()->GetBodySetup();
+	}
+
+	return nullptr;
+}
+
+bool FStaticMeshSceneProxyDesc::HasValidNaniteData() const
+{
+	return Nanite::FNaniteResourcesHelper::HasValidNaniteData(*this);
+}
+
+bool FStaticMeshSceneProxyDesc::ShouldCreateNaniteProxy(Nanite::FMaterialAudit* OutNaniteMaterials /*= nullptr*/) const
+{
+	return Nanite::FNaniteResourcesHelper::ShouldCreateNaniteProxy(*this, OutNaniteMaterials);
+}
+
+bool FStaticMeshSceneProxyDesc::UseNaniteOverrideMaterials(bool bDoingMaterialAudit) const
+{
+	return Nanite::FNaniteResourcesHelper::UseNaniteOverrideMaterials(*this, bDoingMaterialAudit);
+}
+
+UMaterialInterface* FStaticMeshSceneProxyDesc::GetNaniteAuditMaterial(int32 MaterialIndex) const
+{
+	return GetMaterial(MaterialIndex, true);
+}
+
+void FStaticMeshSceneProxyDesc::SetCollisionResponseToChannels(const FCollisionResponseContainer& InContainer)
+{
+	if (&InContainer != &FCollisionResponseContainer::GetDefaultResponseContainer())
+	{
+		CollisionResponseContainer = InContainer;
+	}
+}
+
+const FCollisionResponseContainer& FStaticMeshSceneProxyDesc::GetCollisionResponseToChannels() const
+{
+	if (CollisionResponseContainer.IsSet())
+	{
+		return *CollisionResponseContainer;
+	}
+
+	return FCollisionResponseContainer::GetDefaultResponseContainer();
+}
diff --git a/Engine/Source/Runtime/Engine/Public/StaticMeshSceneProxyDesc.h b/Engine/Source/Runtime/Engine/Public/StaticMeshSceneProxyDesc.h
--- a/Engine/Source/Runtime/Engine/Public/StaticMeshSceneProxyDesc.h
+++ b/Engine/Source/Runtime/Engine/Public/StaticMeshSceneProxyDesc.h
@@ -3,60 +3,20 @@
 #pragma once
 
 #include "PrimitiveSceneProxyDesc.h"
-
-#include "StaticMeshSceneProxy.h"
-#include "Components/StaticMeshComponent.h"
-#include "Engine/StaticMesh.h"
-#include "Materials/Material.h"
-#include "Engine/Engine.h"
-#include "StaticMeshComponentLODInfo.h"
-#include "Rendering/NaniteResources.h"
+#include "Components/StaticMeshComponent.h" // STATICMESH_ENABLE_DEBUG_RENDERING
 
 class UBodySetup;
+class FTextureResource;
+struct FStaticMeshComponentLODInfo;
 
 struct FStaticMeshSceneProxyDesc : public FPrimitiveSceneProxyDesc
 {
-	FStaticMeshSceneProxyDesc()
-	{
-		CastShadow = true;
-		bUseAsOccluder = true;
-
-		// bitfields init (until we switch to c++20)
-		bReverseCulling = false;
-#if STATICMESH_ENABLE_DEBUG_RENDERING
-		bDrawMeshCollisionIfComplex = false;
-		bDrawMeshCollisionIfSimple = false;
-#endif
-		bEvaluateWorldPositionOffset = true;
-		bOverrideMinLOD = false;
-
-		bCastDistanceFieldIndirectShadow = false;
-		bOverrideDistanceFieldSelfShadowBias = false;
-		bEvaluateWorldPositionOffsetInRayTracing = false;
-		bSortTriangles = false;
-
-		bDisplayNaniteFallbackMesh = false;
-		bDisallowNanite = false;
-		bForceDisableNanite = false;
-		bForceNaniteForMasked = false;
-
-		bUseProvidedMaterialRelevance = false;
-	}
-
+	ENGINE_API FStaticMeshSceneProxyDesc();
 	ENGINE_API FStaticMeshSceneProxyDesc(const UStaticMeshComponent*);
+	ENGINE_API ~FStaticMeshSceneProxyDesc();
 
 	ENGINE_API void InitializeFromStaticMeshComponent(const UStaticMeshComponent*);
 
-	UE_DEPRECATED(5.5, "Use InitializeFromStaticMeshComponent instead.")
-	void InitializeFrom(const UStaticMeshComponent* InComponent) { InitializeFromStaticMeshComponent(InComponent); }
-
-	inline bool IsReverseCulling() const { return bReverseCulling; }
-	inline bool IsDisallowNanite() const { return bDisallowNanite; }
-	inline bool IsForceDisableNanite() const { return bForceDisableNanite; }
-	inline bool IsForceNaniteForMasked() const { return bForceNaniteForMasked; }
-	inline int32 GetForcedLodModel() const { return ForcedLodModel; }
-	inline bool IsDisplayNaniteFallbackMesh() const { return bDisplayNaniteFallbackMesh; }
-
 	UStaticMesh* StaticMesh = nullptr;
 	TArrayView<TObjectPtr<UMaterialInterface>>	OverrideMaterials;
 	TObjectPtr<class UMaterialInterface> OverlayMaterial;
@@ -106,9 +66,10 @@ struct FStaticMeshSceneProxyDesc : public FPrimitiveSceneProxyDesc
 	float TextureStreamingTransformScale = 1.0f;
 #endif
 
-	const Nanite::FResources*	NaniteResources = nullptr;
+	const Nanite::FResources* NaniteResources = nullptr;
+
+	TArrayView<FStaticMeshComponentLODInfo> LODData;
 
-	TArrayView<struct FStaticMeshComponentLODInfo> LODData;
 	FMaterialRelevance MaterialRelevance;
 
 	UTexture* MeshPaintTexture = nullptr;
@@ -117,84 +78,121 @@ struct FStaticMeshSceneProxyDesc : public FPrimitiveSceneProxyDesc
 	TArray<UMaterialCacheVirtualTexture*> MaterialCacheTextures;
 	int32 MaterialCacheTextureCoordinateIndex = 0;
 
-	UStaticMesh* GetStaticMesh() const { return StaticMesh; }
-
 	UBodySetup* BodySetup = nullptr;
-	UBodySetup* GetBodySetup() const
+	UObject* LODParentPrimitive = nullptr;
+
+	TOptional<FCollisionResponseContainer> CollisionResponseContainer;
+
+	//@todo: share color selection logic according to mobility with USMC?
+	FColor	WireframeColor = FColor(0, 255, 255, 255);
+
+	bool IsReverseCulling() const
 	{
-		if (BodySetup)
-		{
-			return BodySetup;
-		}
+		return bReverseCulling;
+	}
 
-		if (GetStaticMesh())
-		{
-			return GetStaticMesh()->GetBodySetup();
-		}
+	bool IsDisallowNanite() const
+	{
+		return bDisallowNanite;
+	}
 
-		return nullptr;
+	bool IsForceDisableNanite() const
+	{
+		return bForceDisableNanite;
 	}
 
-	TOptional<FCollisionResponseContainer> CollisionResponseContainer;
-	
-	UObject* LODParentPrimitive = nullptr;
-	UObject* GetLODParentPrimitive() const { return LODParentPrimitive; }
+	bool IsForceNaniteForMasked() const
+	{
+		return bForceNaniteForMasked;
+	}
 
-	const Nanite::FResources* GetNaniteResources() const { return NaniteResources; }
-	
-	ENGINE_API bool HasValidNaniteData() const;
+	int32 GetForcedLodModel() const
+	{
+		return ForcedLodModel;
+	}
 
-	int32 GetNumMaterials() const { return GetStaticMesh() ? GetStaticMesh()->GetStaticMaterials().Num() : 0; }
+	bool IsDisplayNaniteFallbackMesh() const
+	{
+		return bDisplayNaniteFallbackMesh;
+	}
 
-	UMaterialInterface* GetOverlayMaterial() const { return OverlayMaterial; }
-	float GetOverlayMaterialMaxDrawDistance() const { return OverlayMaterialMaxDrawDistance; }
-	void GetMaterialSlotsOverlayMaterial(TArray<TObjectPtr<UMaterialInterface>>& OutMaterialSlotsOverlayMaterial) const { OutMaterialSlotsOverlayMaterial = MaterialSlotsOverlayMaterial; }
+	UStaticMesh* GetStaticMesh() const
+	{
+		return StaticMesh;
+	}
 
-	ENGINE_API UMaterialInterface* GetMaterial(int32 MaterialIndex, bool bDoingNaniteMaterialAudit = false, bool bIgnoreNaniteOverrideMaterials = false) const;
-	
-	ENGINE_API bool ShouldCreateNaniteProxy(Nanite::FMaterialAudit* OutNaniteMaterials = nullptr) const;
-	ENGINE_API bool UseNaniteOverrideMaterials(bool bDoingMaterialAudit) const;
-	
-	UMaterialInterface* GetNaniteAuditMaterial(int32 MaterialIndex) const
+	UObject* GetLODParentPrimitive() const
 	{
-		return GetMaterial(MaterialIndex, true);
+		return LODParentPrimitive;
 	}
 
-	void SetMaterialRelevance(const FMaterialRelevance& InRelevance) { MaterialRelevance = InRelevance; bUseProvidedMaterialRelevance = true; }
-	UE_DEPRECATED(5.7, "Please use GetMaterialRelevance with EShaderPlatform argument and not ERHIFeatureLevel::Type")
-	ENGINE_API FMaterialRelevance GetMaterialRelevance(ERHIFeatureLevel::Type InFeatureLevel) const;
-	ENGINE_API FMaterialRelevance GetMaterialRelevance(EShaderPlatform ShaderPlatform) const;
-	int32 GetStaticLightMapResolution() const { return StaticLightMapResolution; }
-	
-	FTextureResource* GetMeshPaintTextureResource() const;
-	
-	bool ShouldCreateMaterialCacheProxy() const;
+	const Nanite::FResources* GetNaniteResources() const
+	{
+		return NaniteResources;
+	}
 
-	void SetCollisionResponseToChannels(const FCollisionResponseContainer& InContainer) 
-	{ 
-		if (&InContainer != &FCollisionResponseContainer::GetDefaultResponseContainer())
-		{
-			CollisionResponseContainer = InContainer;
-		}
+	UMaterialInterface* GetOverlayMaterial() const
+	{
+		return OverlayMaterial;
 	}
 
-	const FCollisionResponseContainer& GetCollisionResponseToChannels() const
+	float GetOverlayMaterialMaxDrawDistance() const
 	{
-		if (CollisionResponseContainer.IsSet())
-		{
-			return *CollisionResponseContainer;
-		}
+		return OverlayMaterialMaxDrawDistance;
+	}
 
-		return FCollisionResponseContainer::GetDefaultResponseContainer();
+	int32 GetStaticLightMapResolution() const
+	{
+		return StaticLightMapResolution;
 	}
 
-	UObject* GetObjectForPropertyColoration() const { return Component; }
+	UObject* GetObjectForPropertyColoration() const
+	{
+		return Component;
+	}
 
-	//@todo: share color selection logic according to mobility with USMC?
-	FColor	WireframeColor = FColor(0, 255, 255, 255);
-	FColor GetWireframeColor() const { return WireframeColor; }
+	FColor GetWireframeColor() const
+	{
+		return WireframeColor;
+	}
 
+	// Begin FPrimitiveSceneProxyDesc interface
 	ENGINE_API virtual void GetUsedMaterials(TArray<UMaterialInterface*>& OutMaterials, bool bGetDebugMaterials = false) const override;
+	// End FPrimitiveSceneProxyDesc interface
+
+	ENGINE_API UMaterialInterface* GetMaterial(int32 MaterialIndex, bool bDoingNaniteMaterialAudit = false, bool bIgnoreNaniteOverrideMaterials = false) const;
+	ENGINE_API int32 GetNumMaterials() const;
+
+	bool ShouldCreateMaterialCacheProxy() const;
+
+	ENGINE_API const UStaticMeshComponent* GetUStaticMeshComponent() const;
+	ENGINE_API void GetMaterialSlotsOverlayMaterial(TArray<TObjectPtr<UMaterialInterface>>& OutMaterialSlotsOverlayMaterial) const;
+
+	ENGINE_API void SetMaterialRelevance(const FMaterialRelevance& InRelevance);
+	ENGINE_API FMaterialRelevance GetMaterialRelevance(EShaderPlatform ShaderPlatform) const;
+
+	UE_DEPRECATED(5.7, "Please use GetMaterialRelevance with EShaderPlatform argument and not ERHIFeatureLevel::Type")
+	ENGINE_API FMaterialRelevance GetMaterialRelevance(ERHIFeatureLevel::Type InFeatureLevel) const;
+
+	FTextureResource* GetMeshPaintTextureResource() const;
+
+	ENGINE_API UBodySetup* GetBodySetup() const;
+
+	ENGINE_API bool HasValidNaniteData() const;
+	ENGINE_API bool ShouldCreateNaniteProxy(Nanite::FMaterialAudit* OutNaniteMaterials = nullptr) const;
+	ENGINE_API bool UseNaniteOverrideMaterials(bool bDoingMaterialAudit) const;
+	ENGINE_API UMaterialInterface* GetNaniteAuditMaterial(int32 MaterialIndex) const;
+
+	ENGINE_API void SetCollisionResponseToChannels(const FCollisionResponseContainer& InContainer);
+	ENGINE_API const FCollisionResponseContainer& GetCollisionResponseToChannels() const;
 
-	const UStaticMeshComponent* GetUStaticMeshComponent() const { return Cast<UStaticMeshComponent>(Component); }
 };
+
+#if UE_ENABLE_INCLUDE_ORDER_DEPRECATED_IN_5_7
+#include "StaticMeshSceneProxy.h"
+#include "Engine/StaticMesh.h"
+#include "Materials/Material.h"
+#include "Engine/Engine.h"
+#include "StaticMeshComponentLODInfo.h"
+#include "Rendering/NaniteResources.h"
+#endif
diff --git a/Engine/Source/Runtime/Renderer/Private/StateStream/StaticMeshStateStreamImpl.cpp b/Engine/Source/Runtime/Renderer/Private/StateStream/StaticMeshStateStreamImpl.cpp
--- a/Engine/Source/Runtime/Renderer/Private/StateStream/StaticMeshStateStreamImpl.cpp
+++ b/Engine/Source/Runtime/Renderer/Private/StateStream/StaticMeshStateStreamImpl.cpp
@@ -1,6 +1,7 @@
 // Copyright Epic Games, Inc. All Rights Reserved.
 
 #include "StaticMeshStateStreamImpl.h"
+#include "Engine/StaticMesh.h"
 #include "ScenePrivate.h"
 #include "StateStreamCreator.h"
 #include "StaticMeshSceneProxy.h"
