Few more changes to APlayerController::IsLocalController() to clean up code and logic.

diff --git a/Engine/Source/Runtime/Engine/Private/PlayerController.cpp b/Engine/Source/Runtime/Engine/Private/PlayerController.cpp
--- a/Engine/Source/Runtime/Engine/Private/PlayerController.cpp
+++ b/Engine/Source/Runtime/Engine/Private/PlayerController.cpp
@@ -301,13 +301,21 @@ bool APlayerController::IsLocalController() const
 		checkSlow(!bIsLocalPlayerController);
 		return false;
 	}
-	
+ 
+	// Fallback logic for edge cases like PIE teardown where NetDriver is null.
+	if (GetNetDriver() == nullptr)
+	{
+		// Only treat as local if we’re Authority and were already marked as local before NetDriver was torn down.
+		const bool bIsServer = GetLocalRole() == ROLE_Authority;
+		return bIsServer && bIsLocalPlayerController;
+	}
+ 
 	// Fast path if we have this bool set.
 	if (bIsLocalPlayerController)
 	{
 		return true;
 	}
-
+ 
 	ENetMode NetMode = GetNetMode();
 	if (NetMode == NM_DedicatedServer)
 	{
@@ -315,14 +323,14 @@ bool APlayerController::IsLocalController() const
 		checkSlow(!bIsLocalPlayerController);
 		return false;
 	}
-
+ 
 	if (NetMode == NM_Client || NetMode == NM_Standalone)
 	{
 		// Clients or Standalone only receive their own PC. We are not ROLE_AutonomousProxy until after PostInitializeComponents so we can't check that.
 		bIsLocalPlayerController = true;
 		return true;
 	}
-
+ 
 	return bIsLocalPlayerController;
 }
 
