[MIR] Some bug fixes in new material translator.

- Use AsTextureObject() for texture parameters.
- Return constant zero for custom HLSL derivatives as they are not derivative aware.
- Add FCall to MaterialIIR.natvis.

#test EngineTest with new translator test suite showed no regression
#rnx

diff --git a/Engine/Source/Runtime/Engine/Private/Materials/MaterialIREmitter.cpp b/Engine/Source/Runtime/Engine/Private/Materials/MaterialIREmitter.cpp
--- a/Engine/Source/Runtime/Engine/Private/Materials/MaterialIREmitter.cpp
+++ b/Engine/Source/Runtime/Engine/Private/Materials/MaterialIREmitter.cpp
@@ -3743,9 +3743,11 @@ FValueRef FEmitter::AnalyticalPartialDerivative(FValueRef Value, EDerivativeAxis
 		}
 
 		// These values are uniform (constant), thus their value is always zero.
+		// Custom HLSL calls are not derivative aware, so we also return constant zero for VK_Call.
 		case VK_Constant:
 		case VK_UniformParameter:
 		case VK_TextureRead:
+		case VK_Call:
 			return Cast(ConstantZero(Value->Type.GetPrimitive().ScalarKind), Value->Type);
 
 		default:
diff --git a/Engine/Source/Runtime/Engine/Private/Materials/MaterialIRToHLSLTranslator.cpp b/Engine/Source/Runtime/Engine/Private/Materials/MaterialIRToHLSLTranslator.cpp
--- a/Engine/Source/Runtime/Engine/Private/Materials/MaterialIRToHLSLTranslator.cpp
+++ b/Engine/Source/Runtime/Engine/Private/Materials/MaterialIRToHLSLTranslator.cpp
@@ -1583,23 +1583,7 @@ struct FPrivate : FMaterialIRToHLSLTranslation
 	UTexture* GetTextureFromUniformParameter(const MIR::FUniformParameter* UniformParameter)
 	{
 		const FMaterialParameterValue& TextureParameter = Module->GetParameterMetadata(UniformParameter->ParameterIdInModule).Value;
-		switch (TextureParameter.Type)
-		{
-			case EMaterialParameterType::Texture:
-			{
-				return TextureParameter.Texture;
-			}
-			case EMaterialParameterType::Font:
-			{
-				UTexture2D* FontTexture = TextureParameter.Font.Value->Textures[TextureParameter.Font.Page].Get();
-				checkf(FontTexture, TEXT("Missing texture for font '%s' [page %d]"), *TextureParameter.Font.Value->GetName(), TextureParameter.Font.Page);
-				return FontTexture;
-			}
-			default:
-			{
-				UE_MIR_UNREACHABLE();
-			}
-		}
+		return Cast<UTexture>(TextureParameter.AsTextureObject());
 	}
 
 	void LowerTextureParameterType(const MIR::FUniformParameter* UniformParameter)
diff --git a/Engine/Source/Runtime/Engine/Private/Materials/MaterialIRValueAnalyzer.cpp b/Engine/Source/Runtime/Engine/Private/Materials/MaterialIRValueAnalyzer.cpp
--- a/Engine/Source/Runtime/Engine/Private/Materials/MaterialIRValueAnalyzer.cpp
+++ b/Engine/Source/Runtime/Engine/Private/Materials/MaterialIRValueAnalyzer.cpp
@@ -492,10 +492,7 @@ static void AnalyzeTextureUniformParameter(FMaterialIRValueAnalyzer& Analyzer, M
 {
 	const FMaterialParameterValue& ParameterValue = Analyzer.Module->GetParameterMetadata(Parameter->ParameterIdInModule).Value;
 
-	UTexture* Texture =
-		ParameterValue.Type == EMaterialParameterType::Texture ? ParameterValue.Texture :
-		ParameterValue.Type == EMaterialParameterType::Font ? ParameterValue.Font.Value->Textures[ParameterValue.Font.Page].Get() :
-		nullptr;
+	UTexture* Texture = Cast<UTexture>(ParameterValue.AsTextureObject());
 	check(Texture);
 
 	FMaterialTextureParameterInfo ParamInfo {
