Fixed UInstancedStaticMeshComponent::InitializeInstancingRandomSeed() only generating a fixed random seed during cooks
* Other types of builds (ex. HLOD builds) were getting non-deterministic values in InstancingRandomSeed

diff --git a/Engine/Source/Runtime/Engine/Classes/Components/InstancedStaticMeshComponent.h b/Engine/Source/Runtime/Engine/Classes/Components/InstancedStaticMeshComponent.h
--- a/Engine/Source/Runtime/Engine/Classes/Components/InstancedStaticMeshComponent.h
+++ b/Engine/Source/Runtime/Engine/Classes/Components/InstancedStaticMeshComponent.h
@@ -779,7 +779,7 @@ protected:
 	ENGINE_API virtual bool DeleteSMInstances(TArrayView<const FSMInstanceId> InstanceIds) override;
 	ENGINE_API virtual bool DuplicateSMInstances(TArrayView<const FSMInstanceId> InstanceIds, TArray<FSMInstanceId>& OutNewInstanceIds) override;
 
-	void InitializeInstancingRandomSeed();
+	void InitializeInstancingRandomSeed(const bool bAllowRandomSeedValue = false);
 
 	friend FStaticLightingTextureMapping_InstancedStaticMesh;
 	friend FInstancedLightMap2D;
diff --git a/Engine/Source/Runtime/Engine/Private/InstancedStaticMesh.cpp b/Engine/Source/Runtime/Engine/Private/InstancedStaticMesh.cpp
--- a/Engine/Source/Runtime/Engine/Private/InstancedStaticMesh.cpp
+++ b/Engine/Source/Runtime/Engine/Private/InstancedStaticMesh.cpp
@@ -3409,6 +3409,9 @@ void UInstancedStaticMeshComponent::PreSave(FObjectPreSaveContext ObjectSaveCont
 	{
 		InvalidateCachedBounds();
 	}
+
+	const bool bUseRandomSeedValue = !ObjectSaveContext.IsCooking();
+	InitializeInstancingRandomSeed(bUseRandomSeedValue);
 }
 
 void UInstancedStaticMeshComponent::Serialize(FArchive& Ar)
@@ -5188,7 +5191,6 @@ FInstanceDataFlags UInstancedStaticMeshComponent::MakeInstanceDataFlags(bool bAn
 	return Flags;
 }
 
-
 void UInstancedStaticMeshComponent::PostLoad()
 {
 	Super::PostLoad();
@@ -5834,33 +5836,43 @@ bool UInstancedStaticMeshComponent::DuplicateSMInstances(TArrayView<const FSMIns
 	return true;
 }
 
-void UInstancedStaticMeshComponent::InitializeInstancingRandomSeed()
+void UInstancedStaticMeshComponent::InitializeInstancingRandomSeed(const bool bAllowRandomSeedValue)
 {
+	auto GenerateSeedFromObjectPath = [this]() -> int32
+	{
+		TStringBuilder<256> ResultBuilder;
+		GetPathName(nullptr, ResultBuilder);
+		int32 RandomSeed = (int32)CityHash32((const char*)ResultBuilder.GetData(), ResultBuilder.Len() * sizeof(decltype(ResultBuilder)::ElementType));
+		if (RandomSeed == 0)
+		{
+			// If the resulting hash is still zero (unlikely, but possible), we fall back to a fixed default value.
+			RandomSeed = 42;
+		}
+
+		return RandomSeed;
+	};
+
 	// If we don't have a random seed for this instanced static mesh component yet, then go ahead and
 	// generate one now.  This will be saved with the static mesh component and used for future generation
 	// of random numbers for this component's instances. (Used by the PerInstanceRandom material expression)
-	if (IsRunningCookCommandlet())
+	if (bAllowRandomSeedValue)
 	{
-		// During cook, we must generate a deterministic value. Subsequent save of the asset in editor can save
-		// a new unique value that persists.
-		if (InstancingRandomSeed == 0)
+		// This code path is assumed to be used when saving.
+		// Make sure to generate a truly random seed which is not based on the object path.
+		int32 SeedFromPath = GenerateSeedFromObjectPath();
+		while (InstancingRandomSeed == 0 || InstancingRandomSeed == SeedFromPath)
 		{
-			TStringBuilder<256> ResultBuilder;
-			GetPathName(nullptr, ResultBuilder);
-			InstancingRandomSeed = (int32)CityHash32((const char*)ResultBuilder.GetData(), ResultBuilder.Len()*sizeof(decltype(ResultBuilder)::ElementType));
-			if (InstancingRandomSeed == 0)
-			{
-				InstancingRandomSeed = 42;
-			}
+			InstancingRandomSeed = FMath::Rand();
 		}
 	}
 	else
 	{
-		while (InstancingRandomSeed == 0)
+		// To maintain determinism, the seed is generated from the component's path name.
+		if (InstancingRandomSeed == 0)
 		{
-			InstancingRandomSeed = FMath::Rand();
+			InstancingRandomSeed = GenerateSeedFromObjectPath();
 		}
-	}
+	}	
 }
 
 void FInstancedStaticMeshVertexFactoryShaderParameters::GetElementShaderBindings(
