CIS fix

diff --git a/Engine/Source/Runtime/Engine/Private/InstancedStaticMesh.cpp b/Engine/Source/Runtime/Engine/Private/InstancedStaticMesh.cpp
--- a/Engine/Source/Runtime/Engine/Private/InstancedStaticMesh.cpp
+++ b/Engine/Source/Runtime/Engine/Private/InstancedStaticMesh.cpp
@@ -3410,8 +3410,11 @@ void UInstancedStaticMeshComponent::PreSave(FObjectPreSaveContext ObjectSaveCont
 		InvalidateCachedBounds();
 	}
 
-	const bool bUseRandomSeedValue = !ObjectSaveContext.IsCooking();
-	InitializeInstancingRandomSeed(bUseRandomSeedValue);
+	if (!IsTemplate())
+	{
+		const bool bUseRandomSeedValue = !ObjectSaveContext.IsCooking();
+		InitializeInstancingRandomSeed(bUseRandomSeedValue);
+	}
 }
 
 void UInstancedStaticMeshComponent::Serialize(FArchive& Ar)
