Fixed issue with a vertex declaration for SkinCache VF PSO pre-caching on a non-MVF platforms

diff --git a/Engine/Source/Runtime/Engine/Private/GPUSkinVertexFactory.cpp b/Engine/Source/Runtime/Engine/Private/GPUSkinVertexFactory.cpp
--- a/Engine/Source/Runtime/Engine/Private/GPUSkinVertexFactory.cpp
+++ b/Engine/Source/Runtime/Engine/Private/GPUSkinVertexFactory.cpp
@@ -1454,6 +1454,23 @@ void FGPUSkinPassthroughVertexFactory::InitRHI(FRHICommandListBase& RHICmdList)
 	FLocalVertexFactory::InitRHI(RHICmdList);
 }
 
+void FGPUSkinPassthroughVertexFactory::GetVertexElements(ERHIFeatureLevel::Type FeatureLevel, EVertexInputStreamType InputStreamType, bool bSupportsManualVertexFetch, FLocalVertexFactory::FDataType& Data, FVertexDeclarationElementList& Elements)
+{
+	// This is used only for PSO pre-caching 
+	// See the tangents setup in a FGPUSkinPassthroughVertexFactory::InitRHI
+	if (!bSupportsManualVertexFetch)
+	{
+		Data.TangentBasisComponents[0].Type = VET_Short4N;
+		Data.TangentBasisComponents[0].Stride = 16;
+		Data.TangentBasisComponents[0].Offset = 0;
+		Data.TangentBasisComponents[1].Type = VET_Short4N;
+		Data.TangentBasisComponents[1].Stride = 16;
+		Data.TangentBasisComponents[1].Offset = 8;
+	}
+		
+	return FLocalVertexFactory::GetVertexElements(FeatureLevel, InputStreamType, bSupportsManualVertexFetch, Data, Elements);
+}
+
 void FGPUSkinPassthroughVertexFactory::UpdateUniformBuffer(FRHICommandListBase& RHICmdList, const FGPUBaseSkinVertexFactory* InSourceVertexFactory)
 {
 	if (RHISupportsManualVertexFetch(GetFeatureLevelShaderPlatform(GetFeatureLevel())))
diff --git a/Engine/Source/Runtime/Engine/Private/LocalVertexFactory.cpp b/Engine/Source/Runtime/Engine/Private/LocalVertexFactory.cpp
--- a/Engine/Source/Runtime/Engine/Private/LocalVertexFactory.cpp
+++ b/Engine/Source/Runtime/Engine/Private/LocalVertexFactory.cpp
@@ -553,7 +553,7 @@ void FLocalVertexFactory::GetVertexElements(
 		if (IsGPUSkinPassThroughSupported(GMaxRHIShaderPlatform))
 		{
 			// Fill PreSkinPosition slot for GPUSkinPassThrough vertex factory, or else use a dummy buffer.
-			FVertexStreamComponent NullComponent(&GNullVertexBuffer, 0, 0, VET_Float4);
+			FVertexStreamComponent NullComponent(&GNullVertexBuffer, 0, 12, VET_Float3);
 			Elements.Add(AccessStreamComponent(Data.PreSkinPositionComponent.VertexBuffer ? Data.PreSkinPositionComponent : NullComponent, 14, InOutStreams));
 		}
 
diff --git a/Engine/Source/Runtime/Engine/Public/GPUSkinVertexFactory.h b/Engine/Source/Runtime/Engine/Public/GPUSkinVertexFactory.h
--- a/Engine/Source/Runtime/Engine/Public/GPUSkinVertexFactory.h
+++ b/Engine/Source/Runtime/Engine/Public/GPUSkinVertexFactory.h
@@ -810,6 +810,8 @@ public:
 
 	EPixelFormat GetTangentFormat() const { return TangentFormat; }
 
+	static ENGINE_API void GetVertexElements(ERHIFeatureLevel::Type FeatureLevel, EVertexInputStreamType InputStreamType, bool bSupportsManualVertexFetch, FDataType& Data, FVertexDeclarationElementList& Elements);
+
 public:
 	TUniformBufferRef<FGPUSkinPassThroughFactoryLooseParameters> LooseParametersUniformBuffer;
 
