Reducing WRITE_SCOPE lock in AnimSequence::FinishAsyncTasks as it can compete (deadlock) against to-be-finished-and-explicitly-waited-for async compression task

diff --git a/Engine/Source/Runtime/Engine/Private/Animation/AnimSequence.cpp b/Engine/Source/Runtime/Engine/Private/Animation/AnimSequence.cpp
--- a/Engine/Source/Runtime/Engine/Private/Animation/AnimSequence.cpp
+++ b/Engine/Source/Runtime/Engine/Private/Animation/AnimSequence.cpp
@@ -5458,7 +5458,6 @@ void UAnimSequence::FinishAsyncTasks()
 	const ITargetPlatform* RunningPlatform = GetTargetPlatformManagerRef().GetRunningTargetPlatform();
 	const FIoHash CurrentPlatformHash = CreateDerivedDataKeyHash(RunningPlatform);
 
-	UE_COMPRESSED_DATA_WRITE_SCOPE(this);
 	const bool bHasInflightTasks = CacheTasksByKeyHash.Num() > 0;
 	if (bHasInflightTasks)
 	{
@@ -5467,8 +5466,10 @@ void UAnimSequence::FinishAsyncTasks()
 		
 		for (auto It = CacheTasksByKeyHash.CreateIterator(); It; ++It)
 		{
+			// Wait has potential to hit FAnimationSequenceAsyncCacheTask::BuildData or FAnimationSequenceAsyncCacheTask::EndCache which requires a write-lock
 			It->Value->Wait();
 
+			UE_COMPRESSED_DATA_WRITE_SCOPE(this);
 			FCompressedAnimSequence* TaskData = It->Value->GetTargetCompressedData();
 			const bool bCompressedRunningPlatform = It->Key == CurrentPlatformHash;
 			auto ResetData = [this, TaskData, It]()
