[MIR] Implement TextureSampleParameterVolume and TextureSampleParameter2DArray expressions in new material translator.

These expressions are almost identical to their base class except they need to ensure a texture coordinate is connected as input.
Same applies to TextureSampleParameterCube for which this change uses the same shared function.

#test EngineTest showed no regression in new translator test suit
#rnx

diff --git a/Engine/Source/Runtime/Engine/Private/Materials/MaterialExpressionsIR.cpp b/Engine/Source/Runtime/Engine/Private/Materials/MaterialExpressionsIR.cpp
--- a/Engine/Source/Runtime/Engine/Private/Materials/MaterialExpressionsIR.cpp
+++ b/Engine/Source/Runtime/Engine/Private/Materials/MaterialExpressionsIR.cpp
@@ -219,7 +219,9 @@
 #include "Materials/MaterialExpressionTextureProperty.h"
 #include "Materials/MaterialExpressionTextureSample.h"
 #include "Materials/MaterialExpressionTextureSampleParameter.h"
+#include "Materials/MaterialExpressionTextureSampleParameter2DArray.h"
 #include "Materials/MaterialExpressionTextureSampleParameterCube.h"
+#include "Materials/MaterialExpressionTextureSampleParameterVolume.h"
 #include "Materials/MaterialExpressionThinTranslucentMaterialOutput.h"
 #include "Materials/MaterialExpressionTime.h"
 #include "Materials/MaterialExpressionTransform.h"
@@ -1386,24 +1388,44 @@ static FValueRef BuildTextureObjectParameter(MIR::FEmitter& Em, UMaterialExpress
 	return ParameterValue;
 }
 
-void UMaterialExpressionTextureSampleParameter::Build(MIR::FEmitter& Em)
+static void BuildTextureSampleParameter(MIR::FEmitter& Em, UMaterialExpressionTextureSampleParameter* Expr)
 {
 	FString ErrorMessage;
-	if (!TextureIsValid(Texture, ErrorMessage))
+	if (!Expr->TextureIsValid(Expr->Texture, ErrorMessage))
 	{
 		Em.Error(ErrorMessage);
 		return;
 	}
-	FValueRef ParameterValue = BuildTextureObjectParameter(Em, this);
+	FValueRef ParameterValue = BuildTextureObjectParameter(Em, Expr);
 	UE_MIR_CHECKPOINT(Em);
-	BuildTextureSampleExpression(Em, this, ParameterValue, Texture->GetMaterialType());
+	BuildTextureSampleExpression(Em, Expr, ParameterValue, Expr->Texture->GetMaterialType());
 }
 
-void UMaterialExpressionTextureSampleParameterCube::Build(MIR::FEmitter& Em)
+void UMaterialExpressionTextureSampleParameter::Build(MIR::FEmitter& Em)
+{
+	BuildTextureSampleParameter(Em, this);
+}
+
+static void BuildTextureSampleParameterWithCoordinatesInput(MIR::FEmitter& Em, UMaterialExpressionTextureSampleParameter* Expr)
 {
-	Em.Input(&Coordinates); // Cubemap sampling requires coordinates input specified
+	Em.Input(&Expr->Coordinates); // Cubemap, 2DArray, and Volume sampling requires coordinates input specified
 	UE_MIR_CHECKPOINT(Em);
-	UMaterialExpressionTextureSampleParameter::Build(Em);
+	BuildTextureSampleParameter(Em, Expr);
+}
+
+void UMaterialExpressionTextureSampleParameterCube::Build(MIR::FEmitter& Em)
+{
+	BuildTextureSampleParameterWithCoordinatesInput(Em, this);
+}
+
+void UMaterialExpressionTextureSampleParameter2DArray::Build(MIR::FEmitter& Em)
+{
+	BuildTextureSampleParameterWithCoordinatesInput(Em, this);
+}
+
+void UMaterialExpressionTextureSampleParameterVolume::Build(MIR::FEmitter& Em)
+{
+	BuildTextureSampleParameterWithCoordinatesInput(Em, this);
 }
 
 void UMaterialExpressionTextureObjectParameter::Build(MIR::FEmitter& Em)
diff --git a/Engine/Source/Runtime/Engine/Public/Materials/MaterialExpressionTextureSampleParameter2DArray.h b/Engine/Source/Runtime/Engine/Public/Materials/MaterialExpressionTextureSampleParameter2DArray.h
--- a/Engine/Source/Runtime/Engine/Public/Materials/MaterialExpressionTextureSampleParameter2DArray.h
+++ b/Engine/Source/Runtime/Engine/Public/Materials/MaterialExpressionTextureSampleParameter2DArray.h
@@ -16,6 +16,7 @@ public:
 #if WITH_EDITOR
 	// Begin UMaterialExpression Interface
 	virtual void GetCaption(TArray<FString>& OutCaptions) const override;
+	virtual void Build(MIR::FEmitter& Emitter) override;
 	virtual int32 Compile(class FMaterialCompiler* Compiler, int32 OutputIndex) override;
 	// End UMaterialExpression Interface
 #endif
diff --git a/Engine/Source/Runtime/Engine/Public/Materials/MaterialExpressionTextureSampleParameterVolume.h b/Engine/Source/Runtime/Engine/Public/Materials/MaterialExpressionTextureSampleParameterVolume.h
--- a/Engine/Source/Runtime/Engine/Public/Materials/MaterialExpressionTextureSampleParameterVolume.h
+++ b/Engine/Source/Runtime/Engine/Public/Materials/MaterialExpressionTextureSampleParameterVolume.h
@@ -17,6 +17,7 @@ class UMaterialExpressionTextureSampleParameterVolume : public UMaterialExpressi
 #if WITH_EDITOR
 	//~ Begin UMaterialExpression Interface
 	virtual void GetCaption(TArray<FString>& OutCaptions) const override;
+	virtual void Build(MIR::FEmitter& Emitter) override;
 	virtual int32 Compile(class FMaterialCompiler* Compiler, int32 OutputIndex) override;
 	//~ End UMaterialExpression Interface
 
