World Partition - HLOD: Remove leftover test for bIsSpatiallyLoaded in UHLODLayer::GetParentLayer()
* Property is now irrelevant and not exposed to the UX when dealing with the runtime hash set

diff --git a/Engine/Source/Runtime/Engine/Public/WorldPartition/HLOD/HLODLayer.h b/Engine/Source/Runtime/Engine/Public/WorldPartition/HLOD/HLODLayer.h
--- a/Engine/Source/Runtime/Engine/Public/WorldPartition/HLOD/HLODLayer.h
+++ b/Engine/Source/Runtime/Engine/Public/WorldPartition/HLOD/HLODLayer.h
@@ -55,7 +55,7 @@ public:
 	const UHLODBuilderSettings* GetHLODBuilderSettings() const { return HLODBuilderSettings; }
 	const TSubclassOf<AWorldPartitionHLOD> GetHLODActorClass() const { return HLODActorClass; }
 	const TSubclassOf<UWorldPartitionHLODModifier> GetHLODModifierClass() const { return HLODModifierClass; }
-	UHLODLayer* GetParentLayer() const { return !bIsSpatiallyLoaded ? nullptr : ParentLayer; }
+	UHLODLayer* GetParentLayer() const { return ParentLayer; }
 	void SetParentLayer(UHLODLayer* InParentLayer) { ParentLayer = InParentLayer; }
 	ENGINE_API bool DoesRequireWarmup() const;
 
