Iris
* Made it possible to change the default NetObjectFactories via UEngineReplicationBridge::ActorFactoryName and UEngineReplicationBridge::SubObjectFactoryName
* Fix issues with NetObjectFactory id's if you unregister a single factory.
* Exported NetActorFactory to be able to derive from it
* Move Read/WriteConditionallyQuantizedVector functions into NetBitStreamUtil.h to allow any factory to reuse them

#rnx
#13650 tag

diff --git a/Engine/Plugins/Runtime/ReplicationSystemTestPlugin/Source/Private/Tests/ReplicationSystem/TestNetObjectFactory.cpp b/Engine/Plugins/Runtime/ReplicationSystemTestPlugin/Source/Private/Tests/ReplicationSystem/TestNetObjectFactory.cpp
new file mode 100644
--- /dev/null
+++ b/Engine/Plugins/Runtime/ReplicationSystemTestPlugin/Source/Private/Tests/ReplicationSystem/TestNetObjectFactory.cpp
@@ -0,0 +1,101 @@
+// Copyright Epic Games, Inc. All Rights Reserved.
+
+#include "NetworkAutomationTest.h"
+#include "NetworkAutomationTestMacros.h"
+#include "ReplicationSystemServerClientTestFixture.h"
+#include "ReplicatedTestObjectFactory.h"
+
+#include "HAL/IConsoleManager.h"
+#include "Iris/Core/IrisLog.h"
+#include "Iris/ReplicationSystem/ReplicationSystem.h"
+#include "Iris/ReplicationSystem/ReplicationSystemInternal.h"
+#include "Iris/ReplicationSystem/NetObjectFactoryRegistry.h"
+
+//#include "UObject/CoreNetTypes.h"
+
+namespace UE::Net::Private
+{
+
+class FNetObjectFactoryTestFixture : public FNetworkAutomationTestSuiteFixture
+{
+public:
+
+	virtual void SetUp() override
+	{
+		Super::SetUp();
+	}
+	virtual void TearDown() override
+	{
+		Super::TearDown();
+	}
+
+private:
+	typedef FNetworkAutomationTestSuiteFixture Super;
+};
+
+UE_NET_TEST_FIXTURE(FNetObjectFactoryTestFixture, NetFactoryRegistration)
+{
+	const FName FactoryA = TEXT("NetFactoryA");
+	const FName FactoryB = TEXT("NetFactoryB");
+	const FName FactoryC = TEXT("NetFactoryC");
+
+	FNetObjectFactoryRegistry::RegisterFactory(UReplicatedTestObjectFactory::StaticClass(), FactoryA);
+	FNetObjectFactoryRegistry::RegisterFactory(UReplicatedTestObjectFactory::StaticClass(), FactoryB);
+	FNetObjectFactoryRegistry::RegisterFactory(UReplicatedTestObjectFactory::StaticClass(), FactoryC);
+
+	const FNetObjectFactoryId Id_A = FNetObjectFactoryRegistry::GetFactoryIdFromName(FactoryA);
+	const FNetObjectFactoryId Id_B = FNetObjectFactoryRegistry::GetFactoryIdFromName(FactoryB);
+	const FNetObjectFactoryId Id_C = FNetObjectFactoryRegistry::GetFactoryIdFromName(FactoryC);
+
+	// Valid the initial factories
+	{
+		UE_NET_ASSERT_TRUE(FNetObjectFactoryRegistry::IsValidFactoryId(Id_A));
+		UE_NET_ASSERT_TRUE(FNetObjectFactoryRegistry::IsValidFactoryId(Id_B));
+		UE_NET_ASSERT_TRUE(FNetObjectFactoryRegistry::IsValidFactoryId(Id_C));
+	}
+
+	// Unregister a factory
+	{
+		FNetObjectFactoryRegistry::UnregisterFactory(FactoryB);
+
+		UE_NET_ASSERT_TRUE(FNetObjectFactoryRegistry::GetFactoryIdFromName(FactoryB) == InvalidNetObjectFactoryId);
+
+		UE_NET_ASSERT_TRUE(FNetObjectFactoryRegistry::IsValidFactoryId(Id_A));
+		UE_NET_ASSERT_FALSE(FNetObjectFactoryRegistry::IsValidFactoryId(Id_B));
+		UE_NET_ASSERT_TRUE(FNetObjectFactoryRegistry::IsValidFactoryId(Id_C));
+
+		UE_NET_ASSERT_TRUE(Id_A == FNetObjectFactoryRegistry::GetFactoryIdFromName(FactoryA));
+		UE_NET_ASSERT_TRUE(Id_C == FNetObjectFactoryRegistry::GetFactoryIdFromName(FactoryC));
+	}
+
+	// Register yet another factory
+	const FName FactoryD = TEXT("NetFactoryD");
+	FNetObjectFactoryRegistry::RegisterFactory(UReplicatedTestObjectFactory::StaticClass(), FactoryD);
+	const FNetObjectFactoryId Id_D = FNetObjectFactoryRegistry::GetFactoryIdFromName(FactoryD);
+
+	// Validate the new factory
+	{
+		UE_NET_ASSERT_TRUE(Id_D != InvalidNetObjectFactoryId);
+		UE_NET_ASSERT_TRUE(FNetObjectFactoryRegistry::IsValidFactoryId(Id_D));
+	}
+
+	// Validate the old factories
+	{
+		UE_NET_ASSERT_TRUE(FNetObjectFactoryRegistry::IsValidFactoryId(Id_A));
+		UE_NET_ASSERT_TRUE(FNetObjectFactoryRegistry::IsValidFactoryId(Id_C));
+
+		UE_NET_ASSERT_TRUE(Id_A == FNetObjectFactoryRegistry::GetFactoryIdFromName(FactoryA));
+		UE_NET_ASSERT_TRUE(Id_C == FNetObjectFactoryRegistry::GetFactoryIdFromName(FactoryC));
+	}
+
+	// These aren't API validatations but simply show that you need to clear your cached ID when unregistering a factory.
+	{
+		UE_NET_ASSERT_TRUE(Id_B == Id_D);
+		UE_NET_ASSERT_TRUE(FNetObjectFactoryRegistry::IsValidFactoryId(Id_B));
+		UE_NET_ASSERT_TRUE(FNetObjectFactoryRegistry::GetFactoryIdFromName(FactoryB) == InvalidNetObjectFactoryId);
+	}
+
+	TUniquePtr<FReplicationSystemTestServer> Server(new FReplicationSystemTestServer(GetName()));
+}
+
+} // end namespace UE::Net::Private
\ No newline at end of file
diff --git a/Engine/Source/Runtime/Engine/Private/Net/Iris/ReplicationSystem/EngineReplicationBridge.cpp b/Engine/Source/Runtime/Engine/Private/Net/Iris/ReplicationSystem/EngineReplicationBridge.cpp
--- a/Engine/Source/Runtime/Engine/Private/Net/Iris/ReplicationSystem/EngineReplicationBridge.cpp
+++ b/Engine/Source/Runtime/Engine/Private/Net/Iris/ReplicationSystem/EngineReplicationBridge.cpp
@@ -121,6 +121,9 @@ UEngineReplicationBridge::UEngineReplicationBridge()
 , SubObjectFactoryId(UE::Net::InvalidNetObjectFactoryId)
 {
 	SetInstancePreUpdateFunction(UE::Net::Private::ActorReplicationBridgePreUpdateFunction);
+
+	ActorFactoryName = UNetActorFactory::GetFactoryName();
+	SubObjectFactoryName = UNetSubObjectFactory::GetFactoryName();
 }
 
 void UEngineReplicationBridge::Initialize(UReplicationSystem* InReplicationSystem)
@@ -131,11 +134,11 @@ void UEngineReplicationBridge::Initialize(UReplicationSystem* InReplicationSyste
 
 	ensureMsgf(GDefaultUseSubObjectReplicationList, TEXT("Iris requires replicated actors to use registered subobjectslists. Add \n[SystemSettings]\nnet.SubObjects.DefaultUseSubObjectReplicationList=1\n to your DefaultEngine.ini"));
 
-	ActorFactoryId = FNetObjectFactoryRegistry::GetFactoryIdFromName(UNetActorFactory::GetFactoryName());
-	checkf(ActorFactoryId != InvalidNetObjectFactoryId, TEXT("UNetActorFactory with name %s was not registered"), ToCStr(UNetActorFactory::GetFactoryName().ToString()));
+	ActorFactoryId = FNetObjectFactoryRegistry::GetFactoryIdFromName(ActorFactoryName);
+	checkf(ActorFactoryId != InvalidNetObjectFactoryId, TEXT("UNetActorFactory with name %s was not registered"), ToCStr(ActorFactoryName.ToString()));
 
-	SubObjectFactoryId = FNetObjectFactoryRegistry::GetFactoryIdFromName(UNetSubObjectFactory::GetFactoryName());
-	checkf(SubObjectFactoryId != InvalidNetObjectFactoryId, TEXT("UNetSubObjectFactory with name %s was not registered"), ToCStr(UNetSubObjectFactory::GetFactoryName().ToString()));
+	SubObjectFactoryId = FNetObjectFactoryRegistry::GetFactoryIdFromName(SubObjectFactoryName);
+	checkf(SubObjectFactoryId != InvalidNetObjectFactoryId, TEXT("UNetSubObjectFactory with name %s was not registered"), ToCStr(SubObjectFactoryName.ToString()));
 
 	{
 		auto ShouldSpatialize = [](const UClass* Class)
diff --git a/Engine/Source/Runtime/Engine/Private/Net/Iris/ReplicationSystem/NetActorFactory.cpp b/Engine/Source/Runtime/Engine/Private/Net/Iris/ReplicationSystem/NetActorFactory.cpp
--- a/Engine/Source/Runtime/Engine/Private/Net/Iris/ReplicationSystem/NetActorFactory.cpp
+++ b/Engine/Source/Runtime/Engine/Private/Net/Iris/ReplicationSystem/NetActorFactory.cpp
@@ -20,8 +20,6 @@
 #include "Iris/Serialization/NetBitStreamReader.h"
 #include "Iris/Serialization/NetBitStreamWriter.h"
 #include "Iris/Serialization/ObjectNetSerializer.h"
-#include "Iris/Serialization/PackedVectorNetSerializers.h"
-#include "Iris/Serialization/VectorNetSerializers.h"
 
 #include "Net/DataBunch.h"
 #include "Net/Core/Connection/ConnectionHandle.h"
@@ -44,6 +42,8 @@ namespace UE::Net::Private
 	SerializeNewActorMaxBunchSize,
 	TEXT("Max allowed bits that can be added to the creation header via OnSerializeNewActor."));
 
+	static FName NetActorFactoryName(TEXT("NetActorFactory"));
+
 	enum class EActorNetSpawnInfoFlags : uint32
 	{
 		None = 0U,
@@ -94,74 +94,7 @@ namespace UE::Net::Private
 		return Flags;
 	}
 
-	// Write vector using default value compression using specified serializer
-	static void WriteConditionallyQuantizedVector(FNetBitStreamWriter* Writer, const FVector& Vector, const FVector& DefaultValue, bool bQuantize)
-	{
-		// We use 0.01f for comparing when using quantization, because we will only send a single point of precision anyway.
-		// We could probably get away with 0.1f, but that may introduce edge cases for rounding.
-		static constexpr float Epsilon_Quantized = 0.01f;
-				
-		// We use KINDA_SMALL_NUMBER for comparing when not using quantization, because that's the default for FVector::Equals.
-		const float Epsilon = bQuantize ? Epsilon_Quantized : UE_KINDA_SMALL_NUMBER;
 	
-		if (Writer->WriteBool(!Vector.Equals(DefaultValue, Epsilon)))
-		{
-			Writer->WriteBool(bQuantize);
-
-			const FNetSerializer& Serializer = bQuantize ? UE_NET_GET_SERIALIZER(FVectorNetQuantize10NetSerializer) : UE_NET_GET_SERIALIZER(FVectorNetSerializer);
-
-			FNetSerializationContext Context(Writer);
-
-			alignas(16) uint8 QuantizedState[32] = {};
-			checkSlow(sizeof(QuantizedState) >= Serializer.QuantizedTypeSize);
-
-			FNetQuantizeArgs QuantizeArgs;
-			QuantizeArgs.Version = Serializer.Version;
-			QuantizeArgs.Source = NetSerializerValuePointer(&Vector);
-			QuantizeArgs.Target = NetSerializerValuePointer(&QuantizedState[0]);
-			QuantizeArgs.NetSerializerConfig = NetSerializerConfigParam(Serializer.DefaultConfig);
-			Serializer.Quantize(Context, QuantizeArgs);
-
-			FNetSerializeArgs SerializeArgs;
-			SerializeArgs.Version = Serializer.Version;
-			SerializeArgs.Source = QuantizeArgs.Target;
-			SerializeArgs.NetSerializerConfig = NetSerializerConfigParam(Serializer.DefaultConfig);
-			Serializer.Serialize(Context, SerializeArgs);
-		}
-	}
-
-	// Read default value compressed vector using specified serializer
-	static void ReadConditionallyQuantizedVector(FNetBitStreamReader* Reader, FVector& OutVector, const FVector& DefaultValue)
-	{
-		if (Reader->ReadBool())
-		{
-			const bool bIsQuantized = Reader->ReadBool();
-
-			const FNetSerializer& Serializer = bIsQuantized ? UE_NET_GET_SERIALIZER(FVectorNetQuantize10NetSerializer) : UE_NET_GET_SERIALIZER(FVectorNetSerializer);
-
-			FNetSerializationContext Context(Reader);
-
-			alignas(16) uint8 QuantizedState[32] = {};
-			checkSlow(sizeof(QuantizedState) >= Serializer.QuantizedTypeSize);
-
-			FNetDeserializeArgs DeserializeArgs;
-			DeserializeArgs.Version = Serializer.Version;
-			DeserializeArgs.Target = NetSerializerValuePointer(&QuantizedState[0]);
-			DeserializeArgs.NetSerializerConfig = NetSerializerConfigParam(Serializer.DefaultConfig);
-			Serializer.Deserialize(Context, DeserializeArgs);
-
-			FNetDequantizeArgs DequantizeArgs;
-			DequantizeArgs.Version = Serializer.Version;
-			DequantizeArgs.Source = DeserializeArgs.Target;
-			DequantizeArgs.Target = NetSerializerValuePointer(&OutVector);
-			DequantizeArgs.NetSerializerConfig = NetSerializerConfigParam(Serializer.DefaultConfig);
-			Serializer.Dequantize(Context, DequantizeArgs);
-		}
-		else
-		{
-			OutVector = DefaultValue;
-		}
-	}
 #endif // UE_WITH_IRIS
 } // end namespace UE::Net::Private
 
@@ -171,6 +104,11 @@ namespace UE::Net::Private
 // UNetActorFactory
 //------------------------------------------------------------------------
 
+FName UNetActorFactory::GetFactoryName()
+{
+	return UE::Net::Private::NetActorFactoryName;
+}
+
 void UNetActorFactory::OnInit()
 {
 	Super::OnInit();
@@ -658,6 +596,7 @@ FString FStaticActorNetCreationHeader::ToString() const
 
 bool FDynamicActorNetCreationHeader::Serialize(const FCreationHeaderContext& Context, UE::Net::Private::EActorNetSpawnInfoFlags SpawnFlags, const FActorNetSpawnInfo& DefaultSpawnInfo) const
 {
+	using namespace UE::Net;
 	using namespace UE::Net::Private;
 
 	FNetBitStreamWriter* Writer = Context.Serialization.GetBitStreamWriter();
@@ -697,6 +636,7 @@ bool FDynamicActorNetCreationHeader::Serialize(const FCreationHeaderContext& Con
 
 bool FDynamicActorNetCreationHeader::Deserialize(const FCreationHeaderContext& Context, const FActorNetSpawnInfo& DefaultSpawnInfo)
 {
+	using namespace UE::Net;
 	using namespace UE::Net::Private;
 
 	FNetBitStreamReader* Reader = Context.Serialization.GetBitStreamReader();
diff --git a/Engine/Source/Runtime/Engine/Public/Net/Iris/ReplicationSystem/EngineReplicationBridge.h b/Engine/Source/Runtime/Engine/Public/Net/Iris/ReplicationSystem/EngineReplicationBridge.h
--- a/Engine/Source/Runtime/Engine/Public/Net/Iris/ReplicationSystem/EngineReplicationBridge.h
+++ b/Engine/Source/Runtime/Engine/Public/Net/Iris/ReplicationSystem/EngineReplicationBridge.h
@@ -56,7 +56,7 @@ struct FActorReplicationParams
 
 #endif // UE_WITH_IRIS
 
-UCLASS(Transient, MinimalAPI)
+UCLASS(Transient, MinimalAPI, config=Engine)
 class UEngineReplicationBridge final : public UObjectReplicationBridge
 {
 	GENERATED_BODY()
@@ -163,4 +163,13 @@ private:
 	UPROPERTY(Transient)
 	TObjectPtr<UIrisObjectReferencePackageMap> ObjectReferencePackageMap = nullptr;
 
+	UPROPERTY(Config)
+	/** Name of the NetObjectFactory to use for default replicated actors */
+	FName ActorFactoryName;
+
+	UPROPERTY(Config)
+	/** Name of the NetObjectFactory to use for default replicated subobjects */
+	FName SubObjectFactoryName;
+
+
 };
diff --git a/Engine/Source/Runtime/Engine/Public/Net/Iris/ReplicationSystem/NetActorFactory.h b/Engine/Source/Runtime/Engine/Public/Net/Iris/ReplicationSystem/NetActorFactory.h
--- a/Engine/Source/Runtime/Engine/Public/Net/Iris/ReplicationSystem/NetActorFactory.h
+++ b/Engine/Source/Runtime/Engine/Public/Net/Iris/ReplicationSystem/NetActorFactory.h
@@ -50,7 +50,7 @@ public:
 	bool Serialize(const FCreationHeaderContext& Context) const;
 	bool Deserialize(const FCreationHeaderContext& Context);
 
-	virtual FString ToString() const override;
+	ENGINE_API virtual FString ToString() const override;
 };
 
 
@@ -66,7 +66,7 @@ public:
 		return true;
 	}
 
-	virtual FString ToString() const override;
+	ENGINE_API virtual FString ToString() const override;
 
 	struct FActorNetSpawnInfo
 	{
@@ -102,7 +102,7 @@ public:
 /**
  * Responsible for creating headers allowing remote factories to spawn replicated actors
  */
-UCLASS()
+UCLASS(MinimalAPI)
 class UNetActorFactory : public UNetObjectFactory
 {
 	GENERATED_BODY()
@@ -111,36 +111,36 @@ class UNetActorFactory : public UNetObjectFactory
 
 public:
 
-	static FName GetFactoryName()
-	{ 
-		return TEXT("NetActorFactory"); 
-	}
+	ENGINE_API static FName GetFactoryName();
+
+public:
 
-	virtual void OnInit() override;
+	ENGINE_API virtual void OnInit() override;
 
-	virtual FInstantiateResult InstantiateReplicatedObjectFromHeader(const FInstantiateContext& Context, const UE::Net::FNetObjectCreationHeader* Header) override;
+	ENGINE_API virtual FInstantiateResult InstantiateReplicatedObjectFromHeader(const FInstantiateContext& Context, const UE::Net::FNetObjectCreationHeader* Header) override;
 
-	virtual void PostInstantiation(const FPostInstantiationContext& Context) override;
+	ENGINE_API virtual void PostInstantiation(const FPostInstantiationContext& Context) override;
 
-	virtual void PostInit(const FPostInitContext& Context) override;
+	ENGINE_API virtual void PostInit(const FPostInitContext& Context) override;
 
-	virtual void SubObjectCreatedFromReplication(UE::Net::FNetRefHandle RootObject, UE::Net::FNetRefHandle SubObjectCreated) override;
+	ENGINE_API virtual void SubObjectCreatedFromReplication(UE::Net::FNetRefHandle RootObject, UE::Net::FNetRefHandle SubObjectCreated) override;
 
-	virtual void DestroyReplicatedObject(const FDestroyedContext& Context) override;
+	ENGINE_API virtual void DestroyReplicatedObject(const FDestroyedContext& Context) override;
 
-	virtual TOptional<FWorldInfoData> GetWorldInfo(const FWorldInfoContext& Context) const override;
+	ENGINE_API virtual TOptional<FWorldInfoData> GetWorldInfo(const FWorldInfoContext& Context) const override;
 
-	virtual float GetPollFrequency(UE::Net::FNetRefHandle RootObjectHandle, UObject* RootObjectInstance) override;
+	ENGINE_API virtual float GetPollFrequency(UE::Net::FNetRefHandle RootObjectHandle, UObject* RootObjectInstance) override;
 
 protected:
 
-	virtual TUniquePtr<UE::Net::FNetObjectCreationHeader> CreateAndFillHeader(UE::Net::FNetRefHandle Handle) override;
-	virtual TUniquePtr<UE::Net::FNetObjectCreationHeader> CreateAndDeserializeHeader(const UE::Net::FCreationHeaderContext& Context) override;
+	ENGINE_API virtual TUniquePtr<UE::Net::FNetObjectCreationHeader> CreateAndFillHeader(UE::Net::FNetRefHandle Handle) override;
+	ENGINE_API virtual TUniquePtr<UE::Net::FNetObjectCreationHeader> CreateAndDeserializeHeader(const UE::Net::FCreationHeaderContext& Context) override;
 
-	virtual bool SerializeHeader(const UE::Net::FCreationHeaderContext& Context, const UE::Net::FNetObjectCreationHeader* Header)  override;
+	ENGINE_API virtual bool SerializeHeader(const UE::Net::FCreationHeaderContext& Context, const UE::Net::FNetObjectCreationHeader* Header)  override;
 
 private:
 
+	/** Cache of the spawn info flags created from the net.QuantizeActor cvars */
 	UE::Net::Private::EActorNetSpawnInfoFlags SpawnInfoFlags;
 
 	const UE::Net::FDynamicActorNetCreationHeader::FActorNetSpawnInfo DefaultSpawnInfo;
diff --git a/Engine/Source/Runtime/Engine/Public/Net/Iris/ReplicationSystem/NetSubObjectFactory.h b/Engine/Source/Runtime/Engine/Public/Net/Iris/ReplicationSystem/NetSubObjectFactory.h
--- a/Engine/Source/Runtime/Engine/Public/Net/Iris/ReplicationSystem/NetSubObjectFactory.h
+++ b/Engine/Source/Runtime/Engine/Public/Net/Iris/ReplicationSystem/NetSubObjectFactory.h
@@ -90,7 +90,7 @@ class UNetSubObjectFactory : public UNetObjectFactory
 
 public:
 
-	static FName GetFactoryName();
+	ENGINE_API static FName GetFactoryName();
 	
 	ENGINE_API virtual FInstantiateResult InstantiateReplicatedObjectFromHeader(const FInstantiateContext& Context, const UE::Net::FNetObjectCreationHeader* Header) override;
 
diff --git a/Engine/Source/Runtime/Experimental/Iris/Core/Private/Iris/ReplicationSystem/NetObjectFactoryRegistry.cpp b/Engine/Source/Runtime/Experimental/Iris/Core/Private/Iris/ReplicationSystem/NetObjectFactoryRegistry.cpp
--- a/Engine/Source/Runtime/Experimental/Iris/Core/Private/Iris/ReplicationSystem/NetObjectFactoryRegistry.cpp
+++ b/Engine/Source/Runtime/Experimental/Iris/Core/Private/Iris/ReplicationSystem/NetObjectFactoryRegistry.cpp
@@ -66,14 +66,31 @@ void FNetObjectFactoryRegistry::RegisterFactory(UClass* FactoryClass, FName Fact
 		checkf(FactoryClass->IsChildOf<UNetObjectFactory>(), TEXT("FNetObjectFactoryRegistry::RegisterFactory factory: %s name: %s is not derived from UNetObjectFactory."), *GetNameSafe(FactoryClass), *FactoryName.ToString());
 		return;
 	}
-	
-	const int32 Index = NetFactories.Emplace( 
-		FFactoryData 
+
+	// Find an invalid factory entry
+	int32 Index = NetFactories.IndexOfByPredicate([](const FFactoryData& rhs)
+	{
+		return rhs.Name == NAME_None;
+	});
+
+	if (Index != INDEX_NONE)
+	{
+		NetFactories[Index] = FFactoryData
 		{
 			.Name = FactoryName,
-			.NetFactoryClass = FactoryClass
+			.Id = IntCastChecked<FNetObjectFactoryId>(Index),
+			.NetFactoryClass = FactoryClass,
+		};
+	}
+	else
+	{
+		Index = NetFactories.Emplace(FFactoryData
+		{
+			.Name = FactoryName,
+			.NetFactoryClass = FactoryClass,
 		});
-	NetFactories[Index].Id = IntCastChecked<FNetObjectFactoryId>(Index);
+		NetFactories[Index].Id = IntCastChecked<FNetObjectFactoryId>(Index);
+	}
 
 	UE_LOG(LogIris, Log, TEXT("FNetObjectFactoryRegistry::UnregisterFactory is unregistering factory: %s name: %s id: %u"), *GetNameSafe(NetFactories[Index].NetFactoryClass.Get()), *NetFactories[Index].Name.ToString(), NetFactories[Index].Id);
 #endif
@@ -98,7 +115,9 @@ void FNetObjectFactoryRegistry::UnregisterFactory(FName FactoryName)
 		if (Data.Name == FactoryName)
 		{
 			UE_LOG(LogIris, Log, TEXT("FNetObjectFactoryRegistry::UnregisterFactory is unregistering factory: %s name: %s id: %u"), *GetNameSafe(Data.NetFactoryClass.Get()), *Data.Name.ToString(), Data.Id);
-			NetFactories.RemoveAt(Index, EAllowShrinking::No);
+			
+			// Reset the entry to keep the Id's mapped to their index
+			NetFactories[Index] = FFactoryData{};
 			return;
 		}
 	}
@@ -125,7 +144,12 @@ FNetObjectFactoryId FNetObjectFactoryRegistry::GetFactoryIdFromName(FName Factor
 
 bool FNetObjectFactoryRegistry::IsValidFactoryId(FNetObjectFactoryId Id)
 {
-	return NetFactories.IsValidIndex(Id);
+	if (NetFactories.IsValidIndex(Id))
+	{
+		return !NetFactories[Id].Name.IsNone();
+	}
+
+	return false;
 }
 
 } // end namespace UE::Net
\ No newline at end of file
diff --git a/Engine/Source/Runtime/Experimental/Iris/Core/Private/Iris/ReplicationSystem/ObjectReplicationBridge.cpp b/Engine/Source/Runtime/Experimental/Iris/Core/Private/Iris/ReplicationSystem/ObjectReplicationBridge.cpp
--- a/Engine/Source/Runtime/Experimental/Iris/Core/Private/Iris/ReplicationSystem/ObjectReplicationBridge.cpp
+++ b/Engine/Source/Runtime/Experimental/Iris/Core/Private/Iris/ReplicationSystem/ObjectReplicationBridge.cpp
@@ -228,7 +228,7 @@ void UObjectReplicationBridge::Deinitialize()
 void UObjectReplicationBridge::InitNetObjectFactories()
 {
 	using namespace UE::Net;
-	TConstArrayView<FNetObjectFactoryRegistry::FFactoryData> RegisteredFactories = FNetObjectFactoryRegistry::GetRegisteredFactories();
+	const TConstArrayView<FNetObjectFactoryRegistry::FFactoryData> RegisteredFactories = FNetObjectFactoryRegistry::GetRegisteredFactories();
 
 	check(NetObjectFactories.IsEmpty());
 
diff --git a/Engine/Source/Runtime/Experimental/Iris/Core/Private/Iris/Serialization/NetBitStreamUtil.cpp b/Engine/Source/Runtime/Experimental/Iris/Core/Private/Iris/Serialization/NetBitStreamUtil.cpp
--- a/Engine/Source/Runtime/Experimental/Iris/Core/Private/Iris/Serialization/NetBitStreamUtil.cpp
+++ b/Engine/Source/Runtime/Experimental/Iris/Core/Private/Iris/Serialization/NetBitStreamUtil.cpp
@@ -1,11 +1,16 @@
 // Copyright Epic Games, Inc. All Rights Reserved.
 
 #include "Iris/Serialization/NetBitStreamUtil.h"
+
+#include "Iris/Core/BitTwiddling.h"
+#include "Iris/Core/IrisLog.h"
+
 #include "Iris/Serialization/NetSerializers.h"
 #include "Iris/Serialization/StringNetSerializerUtils.h"
-#include "Iris/Core/BitTwiddling.h"
+
+#include "Math/Vector.h"
+
 #include "Net/Core/NetBitArray.h"
-#include "Iris/Core/IrisLog.h"
 
 namespace UE::Net
 {
@@ -973,4 +978,75 @@ void WriteBytes(FNetBitStreamWriter* Writer, const uint8* Src, uint32 BytesToWri
 	}
 }
 
+void WriteConditionallyQuantizedVector(UE::Net::FNetBitStreamWriter* Writer, const FVector& Vector, const FVector& DefaultValue, bool bQuantize)
+{
+	using namespace UE::Net::Private;
+
+	// We use 0.01f for comparing when using quantization, because we will only send a single point of precision anyway.
+	// We could probably get away with 0.1f, but that may introduce edge cases for rounding.
+	static constexpr float Epsilon_Quantized = 0.01f;
+				
+	// We use KINDA_SMALL_NUMBER for comparing when not using quantization, because that's the default for FVector::Equals.
+	const float Epsilon = bQuantize ? Epsilon_Quantized : UE_KINDA_SMALL_NUMBER;
+	
+	if (Writer->WriteBool(!Vector.Equals(DefaultValue, Epsilon)))
+	{
+		Writer->WriteBool(bQuantize);
+
+		const FNetSerializer& Serializer = bQuantize ? UE_NET_GET_SERIALIZER(FVectorNetQuantize10NetSerializer) : UE_NET_GET_SERIALIZER(FVectorNetSerializer);
+
+		FNetSerializationContext Context(Writer);
+
+		alignas(16) uint8 QuantizedState[32] = {};
+		checkSlow(sizeof(QuantizedState) >= Serializer.QuantizedTypeSize);
+
+		FNetQuantizeArgs QuantizeArgs;
+		QuantizeArgs.Version = Serializer.Version;
+		QuantizeArgs.Source = NetSerializerValuePointer(&Vector);
+		QuantizeArgs.Target = NetSerializerValuePointer(&QuantizedState[0]);
+		QuantizeArgs.NetSerializerConfig = NetSerializerConfigParam(Serializer.DefaultConfig);
+		Serializer.Quantize(Context, QuantizeArgs);
+
+		FNetSerializeArgs SerializeArgs;
+		SerializeArgs.Version = Serializer.Version;
+		SerializeArgs.Source = QuantizeArgs.Target;
+		SerializeArgs.NetSerializerConfig = NetSerializerConfigParam(Serializer.DefaultConfig);
+		Serializer.Serialize(Context, SerializeArgs);
+	}
 }
+
+void ReadConditionallyQuantizedVector(UE::Net::FNetBitStreamReader* Reader, FVector& OutVector, const FVector& DefaultValue)
+{
+	using namespace UE::Net::Private;
+
+	if (Reader->ReadBool())
+	{
+		const bool bIsQuantized = Reader->ReadBool();
+
+		const FNetSerializer& Serializer = bIsQuantized ? UE_NET_GET_SERIALIZER(FVectorNetQuantize10NetSerializer) : UE_NET_GET_SERIALIZER(FVectorNetSerializer);
+
+		FNetSerializationContext Context(Reader);
+
+		alignas(16) uint8 QuantizedState[32] = {};
+		checkSlow(sizeof(QuantizedState) >= Serializer.QuantizedTypeSize);
+
+		FNetDeserializeArgs DeserializeArgs;
+		DeserializeArgs.Version = Serializer.Version;
+		DeserializeArgs.Target = NetSerializerValuePointer(&QuantizedState[0]);
+		DeserializeArgs.NetSerializerConfig = NetSerializerConfigParam(Serializer.DefaultConfig);
+		Serializer.Deserialize(Context, DeserializeArgs);
+
+		FNetDequantizeArgs DequantizeArgs;
+		DequantizeArgs.Version = Serializer.Version;
+		DequantizeArgs.Source = DeserializeArgs.Target;
+		DequantizeArgs.Target = NetSerializerValuePointer(&OutVector);
+		DequantizeArgs.NetSerializerConfig = NetSerializerConfigParam(Serializer.DefaultConfig);
+		Serializer.Dequantize(Context, DequantizeArgs);
+	}
+	else
+	{
+		OutVector = DefaultValue;
+	}
+}
+
+} // end namespace UE::Net
diff --git a/Engine/Source/Runtime/Experimental/Iris/Core/Public/Iris/Serialization/NetBitStreamUtil.h b/Engine/Source/Runtime/Experimental/Iris/Core/Public/Iris/Serialization/NetBitStreamUtil.h
--- a/Engine/Source/Runtime/Experimental/Iris/Core/Public/Iris/Serialization/NetBitStreamUtil.h
+++ b/Engine/Source/Runtime/Experimental/Iris/Core/Public/Iris/Serialization/NetBitStreamUtil.h
@@ -229,4 +229,10 @@ IRISCORE_API void WriteBytes(FNetBitStreamWriter* Writer, const uint8* Src, uint
  */
 IRISCORE_API void ReadBytes(FNetBitStreamReader* Reader, uint8* Destination, uint32 BytesToRead);
 
-}
+/** Send 1 bit if the vector equals the DefaultValue, otherwise write a compressed vector using Quantize10 */
+IRISCORE_API void WriteConditionallyQuantizedVector(FNetBitStreamWriter* Writer, const FVector& Vector, const FVector& DefaultValue, bool bQuantize);
+
+/** Read a vector serialized using WriteConditionallyQuantizedVector. DefaultValue must be identical in both Read and Write calls. */
+IRISCORE_API void ReadConditionallyQuantizedVector(FNetBitStreamReader* Reader, FVector& OutVector, const FVector& DefaultValue);
+
+} // end namespace UE::Net
