Fix assert triggered by RecordTextureStreaming material expression when used with Texture Collections.
Texture Collections don't support the texture streaming stats system.
Now we skip the recording functionality of the expression when using with Texture Collections.
#rnx

diff --git a/Engine/Source/Runtime/Engine/Private/Materials/MaterialExpressions.cpp b/Engine/Source/Runtime/Engine/Private/Materials/MaterialExpressions.cpp
--- a/Engine/Source/Runtime/Engine/Private/Materials/MaterialExpressions.cpp
+++ b/Engine/Source/Runtime/Engine/Private/Materials/MaterialExpressions.cpp
@@ -24431,6 +24431,13 @@ int32 UMaterialExpressionRecordTextureStreamingInfo::Compile(class FMaterialComp
 	}
 	const int32 CoordinatesCodeIndex = Coordinates.Compile(Compiler);
 
+	// Streaming stat collection for texture collections and mesh paint textures are not supported.
+	EMaterialValueType TextureType = Compiler->GetParameterType(TextureCodeIndex);
+	if (TextureType & (MCT_TextureCollection | MCT_TextureMeshPaint))
+	{
+		return TextureCodeIndex;
+	}
+
 	int32 TextureReferenceIndex = INDEX_NONE;
 	EMaterialSamplerType EffectiveSamplerType = SAMPLERTYPE_Color;
 	TOptional<FName> EffectiveParameterName;
