[MIR] Unify Texture2D type in new material translator to represent all texture types.

In MIR we don't have to distinguish between specific texture types in the type system yet.
Once we do, we'll add the EMaterialValueType bitmask as subtype to MIR::FType.

#test EngineTest showed no regression in new translator test suite.
#rnx

diff --git a/Engine/Source/Runtime/Engine/Private/Materials/MaterialExpressionsIR.cpp b/Engine/Source/Runtime/Engine/Private/Materials/MaterialExpressionsIR.cpp
--- a/Engine/Source/Runtime/Engine/Private/Materials/MaterialExpressionsIR.cpp
+++ b/Engine/Source/Runtime/Engine/Private/Materials/MaterialExpressionsIR.cpp
@@ -1379,7 +1379,7 @@ static FValueRef BuildTextureObjectParameter(MIR::FEmitter& Em, UMaterialExpress
 	}
 
 	FValueRef ParameterValue = Em.Parameter(Expr->GetParameterName(), Param, Expr->SamplerType);
-	if (!ParameterValue->Type.IsTexture2D())
+	if (!ParameterValue->Type.IsTexture())
 	{
 		Em.Error(TEXT("Parameter is not a texture"));
 		return nullptr;
@@ -1523,7 +1523,7 @@ void UMaterialExpressionFontSampleParameter::Build(MIR::FEmitter& Em)
 		GetParameterValue(ParameterMetaData);
 
 		FValueRef TextureParameter = Em.Parameter(ParameterName, ParameterMetaData, ExpectedSamplerType);
-		if (!TextureParameter->Type.IsTexture2D())
+		if (!TextureParameter->Type.IsTexture())
 		{
 			Em.Error(TEXT("Parameter is not a texture"));
 			return;
diff --git a/Engine/Source/Runtime/Engine/Private/Materials/MaterialIREmitter.cpp b/Engine/Source/Runtime/Engine/Private/Materials/MaterialIREmitter.cpp
--- a/Engine/Source/Runtime/Engine/Private/Materials/MaterialIREmitter.cpp
+++ b/Engine/Source/Runtime/Engine/Private/Materials/MaterialIREmitter.cpp
@@ -630,7 +630,7 @@ FValueRef FEmitter::CheckIsMatrix(FValueRef Value)
 
 FValueRef FEmitter::CheckIsTexture(FValueRef Value)
 {
-	if (Value.IsValid() && !Value->Type.IsTexture2D())
+	if (Value.IsValid() && !Value->Type.IsTexture())
 	{
 		Errorf(Value, TEXT("Expected a texture value, got a '%s' instead."), *Value->Type.GetSpelling());
 		return Value.ToPoison();
@@ -969,7 +969,7 @@ FValueRef FEmitter::ShadingModel(EMaterialShadingModel Id)
 
 FValueRef FEmitter::TextureObject(UTexture* Texture, EMaterialSamplerType SamplerType)
 {
-	FTextureObject Proto = MakePrototype<FTextureObject>(FType::MakeTexture2D());
+	FTextureObject Proto = MakePrototype<FTextureObject>(FType::MakeTexture());
 	Proto.Texture = Texture;
 	Proto.SamplerType = SamplerType;
 	return EmitPrototype(*this, Proto);
@@ -1030,7 +1030,7 @@ FValueRef FEmitter::Parameter(FName Name, FMaterialParameterMetadata& Metadata,
 		case EMaterialParameterType::Texture:
 		case EMaterialParameterType::Font:
 		{
-			FUniformParameter Prototype = MakePrototype<FUniformParameter>(FType::MakeTexture2D());
+			FUniformParameter Prototype = MakePrototype<FUniformParameter>(FType::MakeTexture());
 			Prototype.ParameterIdInModule = RegisterParameter(Module, Info, Metadata);
 			Prototype.SamplerType = SamplerType;
 			return EmitPrototype(*this, Prototype);
diff --git a/Engine/Source/Runtime/Engine/Private/Materials/MaterialIRToHLSLTranslator.cpp b/Engine/Source/Runtime/Engine/Private/Materials/MaterialIRToHLSLTranslator.cpp
--- a/Engine/Source/Runtime/Engine/Private/Materials/MaterialIRToHLSLTranslator.cpp
+++ b/Engine/Source/Runtime/Engine/Private/Materials/MaterialIRToHLSLTranslator.cpp
@@ -712,7 +712,7 @@ struct FPrivate : FMaterialIRToHLSLTranslation
 
 	void LowerUniformParameter(const MIR::FUniformParameter* UniformParameter)
 	{
-		if (UniformParameter->Type.IsTexture2D())
+		if (UniformParameter->Type.IsTexture())
 		{
 			UTexture* TextureObject = GetTextureFromUniformParameter(UniformParameter);
 			check(TextureObject);
@@ -1679,7 +1679,7 @@ struct FPrivate : FMaterialIRToHLSLTranslation
 			Printer << TEXTVIEW("(");
 		}
 
-		checkf(TextureRead->TextureObject->Type.IsTexture2D(), TEXT("Invalid texture object type"));
+		checkf(TextureRead->TextureObject->Type.IsTexture(), TEXT("Invalid texture object type"));
 
 		if (IsVirtualSamplerType(TextureRead->SamplerType))
 		{
diff --git a/Engine/Source/Runtime/Engine/Private/Materials/MaterialIRTypes.cpp b/Engine/Source/Runtime/Engine/Private/Materials/MaterialIRTypes.cpp
--- a/Engine/Source/Runtime/Engine/Private/Materials/MaterialIRTypes.cpp
+++ b/Engine/Source/Runtime/Engine/Private/Materials/MaterialIRTypes.cpp
@@ -17,7 +17,7 @@ const TCHAR* TypeKindToString(ETypeKind Kind)
 		case ETypeKind::Primitive: return TEXT("Primitive");
 		case ETypeKind::Aggregate: return TEXT("Aggregate");
 		case ETypeKind::ShadingModel: return TEXT("ShadingModel");
-		case ETypeKind::Texture2D: return TEXT("Texture2D");
+		case ETypeKind::Texture: return TEXT("Texture");
 		case ETypeKind::ParameterCollection: return TEXT("ParameterCollection");
 		case ETypeKind::SubstrateData: return TEXT("SubstrateData");
 		case ETypeKind::VTPageTableResult: return TEXT("VTPageTableResult");
@@ -96,7 +96,11 @@ FType FType::FromMaterialValueType(EMaterialValueType Type)
 		case MCT_Float2: return FType::MakeFloatVector(2);
 		case MCT_Float3: return FType::MakeFloatVector(3);
 		case MCT_Float4: return FType::MakeFloatVector(4);
-		case MCT_Texture2D: return FType::MakeTexture2D();
+		case MCT_Texture2D: return FType::MakeTexture();
+		case MCT_TextureCube: return FType::MakeTexture();
+		case MCT_Texture2DArray: return FType::MakeTexture();
+		case MCT_TextureCubeArray: return FType::MakeTexture();
+		case MCT_VolumeTexture: return FType::MakeTexture();
 		case MCT_StaticBool: return FType::MakeBoolScalar();
 		case MCT_UInt1: return FType::MakeIntScalar();
 		case MCT_UInt2: return FType::MakeIntVector(2);
@@ -122,7 +126,7 @@ FType FType::FromMaterialParameterType(EMaterialParameterType Type)
 		case EMaterialParameterType::Scalar: return FType::MakeFloatScalar();
 		case EMaterialParameterType::Vector: return FType::MakeFloatVector(4);
 		case EMaterialParameterType::DoubleVector: return FType::MakeDoubleVector(4);
-		case EMaterialParameterType::Texture: return FType::MakeTexture2D();
+		case EMaterialParameterType::Texture: return FType::MakeTexture();
 		case EMaterialParameterType::TextureCollection: UE_MIR_TODO();
 		case EMaterialParameterType::ParameterCollection: UE_MIR_TODO();
 		case EMaterialParameterType::Font: UE_MIR_TODO();
@@ -183,10 +187,10 @@ FType FType::MakeShadingModel()
 	return Type;
 }
 
-FType FType::MakeTexture2D()
+FType FType::MakeTexture()
 {
 	FType Type{};
-	Type.Kind = ETypeKind::Texture2D;
+	Type.Kind = ETypeKind::Texture;
 	return Type;
 }
 
@@ -201,7 +205,8 @@ FType FType::MakeVTPageTableResult()
 {
 	FType Type{};
 	Type.Kind = ETypeKind::VTPageTableResult;
-	return Type;}
+	return Type;
+}
 
 FString FType::GetSpelling() const
 {
diff --git a/Engine/Source/Runtime/Engine/Private/Materials/MaterialIRValueAnalyzer.cpp b/Engine/Source/Runtime/Engine/Private/Materials/MaterialIRValueAnalyzer.cpp
--- a/Engine/Source/Runtime/Engine/Private/Materials/MaterialIRValueAnalyzer.cpp
+++ b/Engine/Source/Runtime/Engine/Private/Materials/MaterialIRValueAnalyzer.cpp
@@ -614,7 +614,7 @@ static void AnalyzePrimitiveUniformParameter(FMaterialIRValueAnalyzer& Analyzer,
 
 static void AnalyzeUniformParameter(FMaterialIRValueAnalyzer& Analyzer, MIR::FUniformParameter* Parameter)
 {
-	if (Parameter->Type.IsTexture2D())
+	if (Parameter->Type.IsTexture())
 	{
 		AnalyzeTextureUniformParameter(Analyzer, Parameter);
 	}
diff --git a/Engine/Source/Runtime/Engine/Public/Materials/MaterialIRTypes.h b/Engine/Source/Runtime/Engine/Public/Materials/MaterialIRTypes.h
--- a/Engine/Source/Runtime/Engine/Public/Materials/MaterialIRTypes.h
+++ b/Engine/Source/Runtime/Engine/Public/Materials/MaterialIRTypes.h
@@ -22,7 +22,7 @@ enum class ETypeKind : uint8
 	// Opaque Objects
 	
 	ShadingModel,			// Dynamically specified material shading model.
-	Texture2D,				// A 2D texture object.
+	Texture,				// A texture object.
 	ParameterCollection,	// Material parameter collection object.
 	SubstrateData,			// Substrate-specific data structure.
 	VTPageTableResult,		// A virtual texture page table lookup result.
@@ -186,8 +186,8 @@ struct FType
 	// Returns the Material Shading Model object type.
 	static FType MakeShadingModel();
 
-	// Returns the Texture2D object type.
-	static FType MakeTexture2D();
+	// Returns the Texture object type.
+	static FType MakeTexture();
 	
 	// Returns the SubstrateData object type.
 	static FType MakeSubstrateData();
@@ -247,7 +247,7 @@ struct FType
 	bool IsMatrix() const { return Is(ETypeKind::Primitive) && Primitive.IsMatrix(); }
 
 	// Returns whether this type is a Texture2D object.
-	bool IsTexture2D() const { return Kind == ETypeKind::Texture2D; }
+	bool IsTexture() const { return Kind == ETypeKind::Texture; }
 
 	// Returns whether this type is a SubstrateData object.
 	bool IsSubstrateData() const { return Kind == ETypeKind::SubstrateData; }
