Splines: Update SetOrientation & SetQuaternion* to preserve tangent discontinuity

New logic computes effective relative rotation and applies it to both tangent directions independently.

diff --git a/Engine/Source/Runtime/Engine/Private/Components/SplineComponent.cpp b/Engine/Source/Runtime/Engine/Private/Components/SplineComponent.cpp
--- a/Engine/Source/Runtime/Engine/Private/Components/SplineComponent.cpp
+++ b/Engine/Source/Runtime/Engine/Private/Components/SplineComponent.cpp
@@ -1270,15 +1270,24 @@ void USplineComponent::SetQuaternionAtSplinePoint(int32 PointIndex, const FQuat&
 		const FQuat Quat = (CoordinateSpace == ESplineCoordinateSpace::World) ?
 			GetComponentTransform().InverseTransformRotation(InQuaternion) : InQuaternion;
 
+		// Work backwards to compute the rotation that is currently being applied.
+		const FQuat RelativeQuat = InQuaternion * GetQuaternionAtSplinePoint(PointIndex, ESplineCoordinateSpace::Local).Inverse();
+
 		// Align up vector with rotation
-		FVector UpVector = Quat.GetUpVector();
-		SetUpVectorAtSplinePoint(PointIndex, UpVector, ESplineCoordinateSpace::Local, false);
+		SetUpVectorAtSplinePoint(PointIndex, Quat.GetUpVector(), ESplineCoordinateSpace::Local, false);
 
 		// Align tangents with rotation, preserving magnitude.
-		const float ArriveTangentMag = GetArriveTangentAtSplinePoint(PointIndex, ESplineCoordinateSpace::Local).Length();
-		const float LeaveTangentMag = GetLeaveTangentAtSplinePoint(PointIndex, ESplineCoordinateSpace::Local).Length();
-		const FVector NewTangent = Quat.GetForwardVector().GetSafeNormal();
-		SetTangentsAtSplinePoint(PointIndex, ArriveTangentMag * NewTangent, LeaveTangentMag * NewTangent, ESplineCoordinateSpace::Local, false);
+		const FVector OldArriveTangent = GetArriveTangentAtSplinePoint(PointIndex, ESplineCoordinateSpace::Local);
+		const FVector OldArriveTangentDirection = OldArriveTangent.GetSafeNormal();
+		const float ArriveTangentMag = OldArriveTangent.Length();
+		const FVector NewArriveTangent = ArriveTangentMag * RelativeQuat.RotateVector(OldArriveTangentDirection);
+		
+		const FVector OldLeaveTangent = GetLeaveTangentAtSplinePoint(PointIndex, ESplineCoordinateSpace::Local);
+		const FVector OldLeaveTangentDirection = OldLeaveTangent.GetSafeNormal();
+		const float LeaveTangentMag = OldLeaveTangent.Length();
+		const FVector NewLeaveTangent = LeaveTangentMag * RelativeQuat.RotateVector(OldLeaveTangentDirection);
+				
+		SetTangentsAtSplinePoint(PointIndex, NewArriveTangent, NewLeaveTangent, ESplineCoordinateSpace::Local, false);
 
 		if (bUpdateSpline)
 		{
diff --git a/Engine/Source/Runtime/Engine/Private/Curves/Spline.cpp b/Engine/Source/Runtime/Engine/Private/Curves/Spline.cpp
--- a/Engine/Source/Runtime/Engine/Private/Curves/Spline.cpp
+++ b/Engine/Source/Runtime/Engine/Private/Curves/Spline.cpp
@@ -1087,15 +1087,25 @@ void FSpline::SetOrientation(int32 Index, const FQuat& InOrientation)
 		return;
 	}
 	
+	// Work backwards to compute the rotation that is currently being applied.
+	const FQuat RelativeRotation = InOrientation * GetOrientation(Index).Inverse();
+	
 	// Store rotation which transforms the world up vector to the local up vector.
 	SetRotation(Index, FQuat::FindBetween(FVector::UpVector, InOrientation.GetUpVector()));
 
 	// Align tangents with rotation, preserving magnitude.
-	const float InTangentMag = GetInTangent(Index).Length();
-	const float OutTangentMag = GetOutTangent(Index).Length();
-	const FVector NewTangent = InOrientation.GetForwardVector().GetSafeNormal();
-	SetInTangent(Index, InTangentMag * NewTangent);
-	SetOutTangent(Index, OutTangentMag * NewTangent);
+	const FVector OldInTangent = GetInTangent(Index);
+	const FVector OldInTangentDirection = OldInTangent.GetSafeNormal();
+	const float InTangentMag = OldInTangent.Length();
+	const FVector NewInTangentDirection = RelativeRotation.RotateVector(OldInTangentDirection);
+	
+	const FVector OldOutTangent = GetOutTangent(Index);
+	const FVector OldOutTangentDirection = OldOutTangent.GetSafeNormal();
+	const float OutTangentMag = OldOutTangent.Length();
+	const FVector NewOutTangentDirection = RelativeRotation.RotateVector(OldOutTangentDirection);
+	
+	SetInTangent(Index, InTangentMag * NewInTangentDirection);
+	SetOutTangent(Index, OutTangentMag * NewOutTangentDirection);
 }
 
 FVector FSpline::Evaluate(float Param) const
