Removing redundant extra assignment

diff --git a/Engine/Source/Runtime/Engine/Private/Components/StereoLayerComponent.cpp b/Engine/Source/Runtime/Engine/Private/Components/StereoLayerComponent.cpp
--- a/Engine/Source/Runtime/Engine/Private/Components/StereoLayerComponent.cpp
+++ b/Engine/Source/Runtime/Engine/Private/Components/StereoLayerComponent.cpp
@@ -123,7 +123,7 @@ void UStereoLayerComponent::TickComponent(float DeltaTime, enum ELevelTick TickT
 		if (!UE::XRStereoLayerCVars::bUseLegacyFacing)
 		{
 			// The default Open Xr Stereo layer object faces the camera at no rotation where as the unreal object faces away from it at zero rotation
-			LayerDesc.Transform.SetRotation(Transform.GetRotation() *= FQuat(0., 0., 1., 0.));
+			LayerDesc.Transform.SetRotation(Transform.GetRotation() * FQuat(0., 0., 1., 0.));
 		}
 
 		if (Texture)
