Use LastTickDispatchRealtime instead of FApp::GetCurrentTime() for UNetConnection's packet receive time in its ping calculation. This ensures both the receive time and send time are using real time, in case FApp's time is overridden, such as by a UEngineCustomTimeStep.

diff --git a/Engine/Source/Runtime/Engine/Private/NetConnection.cpp b/Engine/Source/Runtime/Engine/Private/NetConnection.cpp
--- a/Engine/Source/Runtime/Engine/Private/NetConnection.cpp
+++ b/Engine/Source/Runtime/Engine/Private/NetConnection.cpp
@@ -2951,8 +2951,9 @@ bool UNetConnection::ReadPacketInfo(FBitReader& Reader, bool bHasPacketInfoPaylo
 		}
 	}
 
-	// Use FApp's time because it is set closer to the beginning of the frame - we don't care about the time so far of the current frame to process the packet
-	const double CurrentTime = (PacketReceiveTime != 0.0 ? PacketReceiveTime : FApp::GetCurrentTime());
+	// Use the driver's LastTickDispatchRealtime because it is set closer to the beginning of the frame - we don't care about the time so far of the current frame to process the packet
+	// This time needs to be a realtime value based on FPlatformTime::Seconds so that it's compatible with the packet send time.
+	const double CurrentTime = (PacketReceiveTime != 0.0 ? PacketReceiveTime : Driver->LastTickDispatchRealtime);
 
 	// Update ping
 
