Fixing a bug in APlayerController::IsLocalController() where it could incorrectly return true for remote clients during teardown.

diff --git a/Engine/Source/Runtime/Engine/Private/PlayerController.cpp b/Engine/Source/Runtime/Engine/Private/PlayerController.cpp
--- a/Engine/Source/Runtime/Engine/Private/PlayerController.cpp
+++ b/Engine/Source/Runtime/Engine/Private/PlayerController.cpp
@@ -301,13 +301,14 @@ bool APlayerController::IsLocalController() const
 		checkSlow(!bIsLocalPlayerController);
 		return false;
 	}
-	
+
 	// Fast path if we have this bool set.
 	if (bIsLocalPlayerController)
 	{
 		return true;
 	}
 
+
 	ENetMode NetMode = GetNetMode();
 	if (NetMode == NM_DedicatedServer)
 	{
@@ -316,6 +317,22 @@ bool APlayerController::IsLocalController() const
 		return false;
 	}
 
+		
+	// If the NetDriver is already gone (e.g., seamless travel back to standalone/PIE teardown),
+	// classify "local" by presence of a ULocalPlayer on this machine. Remote clients' PCs on a listen
+	// server do NOT have a ULocalPlayer, but true clients and the listen host do.
+	if (GetNetDriver() == nullptr)
+	{
+		const bool bHasLocalPlayer = Cast<ULocalPlayer>(Player) != nullptr;
+		if (bHasLocalPlayer)
+		{
+			bIsLocalPlayerController = true;
+			return true;
+		}
+		return false;
+	}
+
+
 	if (NetMode == NM_Client || NetMode == NM_Standalone)
 	{
 		// Clients or Standalone only receive their own PC. We are not ROLE_AutonomousProxy until after PostInitializeComponents so we can't check that.
@@ -326,7 +343,6 @@ bool APlayerController::IsLocalController() const
 	return bIsLocalPlayerController;
 }
 
-
 void APlayerController::FailedToSpawnPawn()
 {
 	Super::FailedToSpawnPawn();
