Add ensure to protect against missing player states during seamless travel

diff --git a/Engine/Source/Runtime/Engine/Private/PlayerState.cpp b/Engine/Source/Runtime/Engine/Private/PlayerState.cpp
--- a/Engine/Source/Runtime/Engine/Private/PlayerState.cpp
+++ b/Engine/Source/Runtime/Engine/Private/PlayerState.cpp
@@ -388,8 +388,12 @@ APlayerState* APlayerState::Duplicate()
 
 void APlayerState::SeamlessTravelTo(APlayerState* NewPlayerState)
 {
-	DispatchCopyProperties(NewPlayerState);
-	NewPlayerState->SetIsOnlyASpectator(IsOnlyASpectator());
+	// By default there should always be a player state by the time it calls SeamlessTravelTo, but this can be overridden in subclasses
+	if (ensure(NewPlayerState))
+	{
+		DispatchCopyProperties(NewPlayerState);
+		NewPlayerState->SetIsOnlyASpectator(IsOnlyASpectator());
+	}
 }
 
 void APlayerState::GetLifetimeReplicatedProps(TArray< FLifetimeProperty > & OutLifetimeProps) const
