Loading: Fix for possible deadlocks (or other corruptions) on loading screens.

In particular, the issue may happen if the loading screen is in the manual mode i.e. not shutting down automatically after the loading has finished, but waiting for the user code to invoke StopMovie. It is caused by the fact that the Game thread invokes RedrawViewports right after the blocking level loading at which time the async loading screen may still be spinning and issuing its own draw requests for slate. This interleaving of the two separate frames breakes the logic in the render/RHI thread waiting for the frame end.

diff --git a/Engine/Source/Runtime/Engine/Private/GameEngine.cpp b/Engine/Source/Runtime/Engine/Private/GameEngine.cpp
--- a/Engine/Source/Runtime/Engine/Private/GameEngine.cpp
+++ b/Engine/Source/Runtime/Engine/Private/GameEngine.cpp
@@ -781,6 +781,11 @@ void UGameEngine::RedrawViewports( bool bShouldPresent /*= true*/ )
 		GameViewport->LayoutPlayers();
 		if ( GameViewport->Viewport != NULL )
 		{
+			// Make sure we don't attempt to render a frame in parallel with the loading screen, which might deadlock because of interleaving two different frames.
+			// It could happen because the Game thread invokes RedrawViewports right after blocking loads when the loading thread is still spinning.
+			FSlateRenderer* SlateRenderer = FSlateApplication::Get().GetRenderer();
+			FScopeLock ScopeLock(SlateRenderer->GetResourceCriticalSection());
+
 			GameViewport->Viewport->Draw(bShouldPresent);
 		}
 	}
