Fix incorrect setting of material dynamic boolean parameters.
Maths for setting the value used an incorrect bitshift.
The bug only affected legacy static bool params which had the DynamicBranch property set on them (exposed from UE5.2 to UE5.4).
The bug comes from an incorrect move/paste of code in 38477538 (affecting UE5.6).
#rnx

diff --git a/Engine/Source/Runtime/Engine/Private/Materials/MaterialUniformExpressions.cpp b/Engine/Source/Runtime/Engine/Private/Materials/MaterialUniformExpressions.cpp
--- a/Engine/Source/Runtime/Engine/Private/Materials/MaterialUniformExpressions.cpp
+++ b/Engine/Source/Runtime/Engine/Private/Materials/MaterialUniformExpressions.cpp
@@ -911,7 +911,7 @@ static void CopyValueToUniformBuffer(const UE::Shader::FValue& Value, float* Pre
 		}
 
 		uint32 BufferOffsetAdjusted = BufferOffset / 32u;
-		uint32 BufferBitOffset = BufferOffsetAdjusted % 32u;
+		uint32 BufferBitOffset = BufferOffset % 32u;
 		check(BufferBitOffset + Value.GetNumComponents ()<= 32u);
 
 		uint32* DestAddress = (uint32*)PreshaderBuffer + BufferOffsetAdjusted;
