[NetPhysics] Optimize NetworkPhysicsSettings memory footprint

diff --git a/Engine/Source/Runtime/Engine/Public/Physics/NetworkPhysicsSettingsComponent.h b/Engine/Source/Runtime/Engine/Public/Physics/NetworkPhysicsSettingsComponent.h
--- a/Engine/Source/Runtime/Engine/Public/Physics/NetworkPhysicsSettingsComponent.h
+++ b/Engine/Source/Runtime/Engine/Public/Physics/NetworkPhysicsSettingsComponent.h
@@ -83,8 +83,10 @@
 {
 	GENERATED_BODY()
 
+	// Override properties
 	UPROPERTY(config, EditDefaultsOnly, Category = "Overrides", Meta = (InlineEditConditionToggle))
-	uint32 bOverrideSimProxyRepMode : 1;
+	bool bOverrideSimProxyRepMode = false;
+
 	// Override the EPhysicsReplicationMode for Actors with ENetRole::ROLE_SimulatedProxy.
 	UPROPERTY(config, EditDefaultsOnly, Category = "Overrides", Meta = (EditCondition = "bOverrideSimProxyRepMode"))
 	EPhysicsReplicationMode SimProxyRepMode = EPhysicsReplicationMode::PredictiveInterpolation;
@@ -104,29 +106,31 @@
 {
 	GENERATED_BODY()
 
+	// Override properties
 	UPROPERTY(config, EditDefaultsOnly, Category = "Overrides", Meta = (InlineEditConditionToggle))
 	uint32 bOverrideMaxLinearHardSnapDistance : 1;
+	UPROPERTY(config, EditDefaultsOnly, Category = "Overrides", Meta = (InlineEditConditionToggle))
+	uint32 bOverrideDefaultLegacyHardsnapInPT : 1;
+	UPROPERTY(config, EditDefaultsOnly, Category = "Overrides", Meta = (InlineEditConditionToggle))
+	uint32 bOverrideCorrectConnectedBodies : 1;
+	UPROPERTY(config, EditDefaultsOnly, Category = "Overrides", Meta = (InlineEditConditionToggle))
+	uint32 bOverrideCorrectConnectedBodiesFriction : 1;
+
 	// Overrides CVar: p.MaxLinearHardSnapDistance -- Hardsnap if distance between current position and extrapolated target position is larger than this value.
 	UPROPERTY(config, EditDefaultsOnly, Category = "Overrides", Meta = (EditCondition = "bOverrideMaxLinearHardSnapDistance"))
 	float MaxLinearHardSnapDistance = 400.0f;
 	float GetMaxLinearHardSnapDistance(float DefaultValue) const { return bOverrideMaxLinearHardSnapDistance ? MaxLinearHardSnapDistance : DefaultValue; }
 
-	UPROPERTY(config, EditDefaultsOnly, Category = "Overrides", Meta = (InlineEditConditionToggle))
-	uint32 bOverrideDefaultLegacyHardsnapInPT : 1;
 	// Overrides CVar: p.DefaultReplication.Legacy.HardsnapInPT -- If default replication is used and it's running the legacy flow through Game Thread, allow hardsnapping to be performed on Physics Thread if async physics is enabled.
 	UPROPERTY(config, EditDefaultsOnly, Category = "Overrides", Meta = (EditCondition = "bOverrideDefaultLegacyHardsnapInPT"))
 	bool bHardsnapInPhysicsThread = PhysicsReplicationCVars::DefaultReplicationCVars::bHardsnapLegacyInPT;
 	bool GetHardsnapDefaultLegacyInPT() const { return bOverrideDefaultLegacyHardsnapInPT ? bHardsnapInPhysicsThread : PhysicsReplicationCVars::DefaultReplicationCVars::bHardsnapLegacyInPT; }
 
-	UPROPERTY(config, EditDefaultsOnly, Category = "Overrides", Meta = (InlineEditConditionToggle))
-	uint32 bOverrideCorrectConnectedBodies : 1;
 	// Overrides CVar: p.DefaultReplication.CorrectConnectedBodies -- When true, transform corrections will also apply to any connected physics object.
 	UPROPERTY(config, EditDefaultsOnly, Category = "Overrides", Meta = (EditCondition = "bOverrideCorrectConnectedBodies"))
 	bool bCorrectConnectedBodies = PhysicsReplicationCVars::DefaultReplicationCVars::bCorrectConnectedBodies;
 	bool GetCorrectConnectedBodies() const { return bOverrideCorrectConnectedBodies ? bCorrectConnectedBodies : PhysicsReplicationCVars::DefaultReplicationCVars::bCorrectConnectedBodies; }
 
-	UPROPERTY(config, EditDefaultsOnly, Category = "Overrides", Meta = (InlineEditConditionToggle))
-	uint32 bOverrideCorrectConnectedBodiesFriction : 1;
 	// Overrides CVar: p.DefaultReplication.CorrectConnectedBodiesFriction -- When true, transform correction on any connected physics object will also recalculate their friction.
 	UPROPERTY(config, EditDefaultsOnly, Category = "Overrides", Meta = (EditCondition = "bOverrideCorrectConnectedBodiesFriction"))
 	bool bCorrectConnectedBodiesFriction = PhysicsReplicationCVars::DefaultReplicationCVars::bCorrectConnectedBodiesFriction;
@@ -138,118 +142,119 @@
 {
 	GENERATED_BODY()
 
+	// Override properties
 	UPROPERTY(config, EditDefaultsOnly, Category = "Overrides", Meta = (InlineEditConditionToggle))
 	uint32 bOverridePosCorrectionTimeBase : 1;
+	UPROPERTY(config, EditDefaultsOnly, Category = "Overrides", Meta = (InlineEditConditionToggle))
+	uint32 bOverridePosCorrectionTimeMin : 1;
+	UPROPERTY(config, EditDefaultsOnly, Category = "Overrides", Meta = (InlineEditConditionToggle))
+	uint32 bOverridePosCorrectionTimeMultiplier : 1;
+	UPROPERTY(config, EditDefaultsOnly, Category = "Overrides", Meta = (InlineEditConditionToggle))
+	uint32 bOverrideRotCorrectionTimeBase : 1;
+	UPROPERTY(config, EditDefaultsOnly, Category = "Overrides", Meta = (InlineEditConditionToggle))
+	uint32 bOverrideRotCorrectionTimeMin : 1;
+	UPROPERTY(config, EditDefaultsOnly, Category = "Overrides", Meta = (InlineEditConditionToggle))
+	uint32 bOverrideRotCorrectionTimeMultiplier : 1;
+	UPROPERTY(config, EditDefaultsOnly, Category = "Overrides", Meta = (InlineEditConditionToggle))
+	uint32 bOverridePosInterpolationTimeMultiplier : 1;
+	UPROPERTY(config, EditDefaultsOnly, Category = "Overrides", Meta = (InlineEditConditionToggle))
+	uint32 bOverrideRotInterpolationTimeMultiplier : 1;
+	UPROPERTY(config, EditDefaultsOnly, Category = "Overrides", Meta = (InlineEditConditionToggle))
+	uint32 bOverrideSoftSnapPosStrength : 1;
+	UPROPERTY(config, EditDefaultsOnly, Category = "Overrides", Meta = (InlineEditConditionToggle))
+	uint32 bOverrideSoftSnapRotStrength : 1;
+	UPROPERTY(config, EditDefaultsOnly, Category = "Overrides", Meta = (InlineEditConditionToggle))
+	uint32 bOverrideSoftSnapToSource : 1;
+	UPROPERTY(config, EditDefaultsOnly, Category = "Overrides", Meta = (InlineEditConditionToggle))
+	uint32 bOverrideDisableSoftSnap : 1;
+	UPROPERTY(config, EditDefaultsOnly, Category = "Overrides", Meta = (InlineEditConditionToggle))
+	uint32 bOverrideSkipVelocityRepOnPosEarlyOut : 1;
+	UPROPERTY(config, EditDefaultsOnly, Category = "Overrides", Meta = (InlineEditConditionToggle))
+	uint32 bOverridePostResimWaitForUpdate : 1;
+	UPROPERTY(config, EditDefaultsOnly, Category = "Overrides", Meta = (InlineEditConditionToggle))
+	uint32 bOverrideCorrectConnectedBodies : 1;
+	UPROPERTY(config, EditDefaultsOnly, Category = "Overrides", Meta = (InlineEditConditionToggle))
+	uint32 bOverrideCorrectConnectedBodiesFriction : 1;
+
 	// Overrides CVar: np2.PredictiveInterpolation.PosCorrectionTimeBase -- Base time to correct positional offset over. RoundTripTime * PosCorrectionTimeMultiplier is added on top of this.
 	UPROPERTY(config, EditDefaultsOnly, Category = "Overrides", Meta = (EditCondition = "bOverridePosCorrectionTimeBase"))
 	float PosCorrectionTimeBase = PhysicsReplicationCVars::PredictiveInterpolationCVars::PosCorrectionTimeBase;
 	float GetPosCorrectionTimeBase() const { return bOverridePosCorrectionTimeBase ? PosCorrectionTimeBase : PhysicsReplicationCVars::PredictiveInterpolationCVars::PosCorrectionTimeBase; }
 
-	UPROPERTY(config, EditDefaultsOnly, Category = "Overrides", Meta = (InlineEditConditionToggle))
-	uint32 bOverridePosCorrectionTimeMin : 1;
 	// Overrides CVar: np2.PredictiveInterpolation.PosCorrectionTimeMin -- Min time to correct positional offset over. DeltaSeconds is added on top of this.
 	UPROPERTY(config, EditDefaultsOnly, Category = "Overrides", Meta = (EditCondition = "bOverridePosCorrectionTimeMin"))
 	float PosCorrectionTimeMin = PhysicsReplicationCVars::PredictiveInterpolationCVars::PosCorrectionTimeMin;
 	float GetPosCorrectionTimeMin() const { return bOverridePosCorrectionTimeMin ? PosCorrectionTimeMin : PhysicsReplicationCVars::PredictiveInterpolationCVars::PosCorrectionTimeMin; }
 
-	UPROPERTY(config, EditDefaultsOnly, Category = "Overrides", Meta = (InlineEditConditionToggle))
-	uint32 bOverridePosCorrectionTimeMultiplier : 1;
 	// Overrides CVar: np2.PredictiveInterpolation.PosCorrectionTimeMultiplier -- Multiplier to adjust how much of RoundTripTime to add to positional offset correction.
 	UPROPERTY(config, EditDefaultsOnly, Category = "Overrides", Meta = (EditCondition = "bOverridePosCorrectionTimeMultiplier"))
 	float PosCorrectionTimeMultiplier = PhysicsReplicationCVars::PredictiveInterpolationCVars::PosCorrectionTimeMultiplier;
 	float GetPosCorrectionTimeMultiplier() const { return bOverridePosCorrectionTimeMultiplier ? PosCorrectionTimeMultiplier : PhysicsReplicationCVars::PredictiveInterpolationCVars::PosCorrectionTimeMultiplier; }
 
-	UPROPERTY(config, EditDefaultsOnly, Category = "Overrides", Meta = (InlineEditConditionToggle))
-	uint32 bOverrideRotCorrectionTimeBase : 1;
 	// Overrides CVar: np2.PredictiveInterpolation.RotCorrectionTimeBase -- Base time to correct rotational offset over. RoundTripTime * RotCorrectionTimeMultiplier is added on top of this.
 	UPROPERTY(config, EditDefaultsOnly, Category = "Overrides", Meta = (EditCondition = "bOverrideRotCorrectionTimeBase"))
 	float RotCorrectionTimeBase = PhysicsReplicationCVars::PredictiveInterpolationCVars::RotCorrectionTimeBase;
 	float GetRotCorrectionTimeBase() const { return bOverrideRotCorrectionTimeBase ? RotCorrectionTimeBase : PhysicsReplicationCVars::PredictiveInterpolationCVars::RotCorrectionTimeBase; }
 
-	UPROPERTY(config, EditDefaultsOnly, Category = "Overrides", Meta = (InlineEditConditionToggle))
-	uint32 bOverrideRotCorrectionTimeMin : 1;
 	// Overrides CVar: np2.PredictiveInterpolation.RotCorrectionTimeMin -- Min time to correct rotational offset over. DeltaSeconds is added on top of this.
 	UPROPERTY(config, EditDefaultsOnly, Category = "Overrides", Meta = (EditCondition = "bOverrideRotCorrectionTimeMin"))
 	float RotCorrectionTimeMin = PhysicsReplicationCVars::PredictiveInterpolationCVars::RotCorrectionTimeMin;
 	float GetRotCorrectionTimeMin() const { return bOverrideRotCorrectionTimeMin ? RotCorrectionTimeMin : PhysicsReplicationCVars::PredictiveInterpolationCVars::RotCorrectionTimeMin; }
 
-	UPROPERTY(config, EditDefaultsOnly, Category = "Overrides", Meta = (InlineEditConditionToggle))
-	uint32 bOverrideRotCorrectionTimeMultiplier : 1;
 	// Overrides CVar: np2.PredictiveInterpolation.RotCorrectionTimeMultiplier -- Multiplier to adjust how much of RoundTripTime to add to rotational offset correction.
 	UPROPERTY(config, EditDefaultsOnly, Category = "Overrides", Meta = (EditCondition = "bOverrideRotCorrectionTimeMultiplier"))
 	float RotCorrectionTimeMultiplier = PhysicsReplicationCVars::PredictiveInterpolationCVars::RotCorrectionTimeMultiplier;
 	float GetRotCorrectionTimeMultiplier() const { return bOverrideRotCorrectionTimeMultiplier ? RotCorrectionTimeMultiplier : PhysicsReplicationCVars::PredictiveInterpolationCVars::RotCorrectionTimeMultiplier; }
 
-	UPROPERTY(config, EditDefaultsOnly, Category = "Overrides", Meta = (InlineEditConditionToggle))
-	uint32 bOverridePosInterpolationTimeMultiplier : 1;
 	// Overrides CVar: np2.PredictiveInterpolation.InterpolationTimeMultiplier -- Multiplier to adjust the interpolation time which is based on the sendrate of state data from the server.
 	UPROPERTY(config, EditDefaultsOnly, Category = "Overrides", Meta = (EditCondition = "bOverridePosInterpolationTimeMultiplier"))
 	float PosInterpolationTimeMultiplier = PhysicsReplicationCVars::PredictiveInterpolationCVars::PosInterpolationTimeMultiplier;
 	float GetPosInterpolationTimeMultiplier() const { return bOverridePosInterpolationTimeMultiplier ? PosInterpolationTimeMultiplier : PhysicsReplicationCVars::PredictiveInterpolationCVars::PosInterpolationTimeMultiplier; }
 
-	UPROPERTY(config, EditDefaultsOnly, Category = "Overrides", Meta = (InlineEditConditionToggle))
-	uint32 bOverrideRotInterpolationTimeMultiplier : 1;
 	// Overrides CVar: np2.PredictiveInterpolation.RotInterpolationTimeMultiplier -- Multiplier to adjust the rotational interpolation time which is based on the sendrate of state data from the server.
 	UPROPERTY(config, EditDefaultsOnly, Category = "Overrides", Meta = (EditCondition = "bOverrideRotInterpolationTimeMultiplier"))
 	float RotInterpolationTimeMultiplier = PhysicsReplicationCVars::PredictiveInterpolationCVars::RotInterpolationTimeMultiplier;
 	float GetRotInterpolationTimeMultiplier() const { return bOverrideRotInterpolationTimeMultiplier ? RotInterpolationTimeMultiplier : PhysicsReplicationCVars::PredictiveInterpolationCVars::RotInterpolationTimeMultiplier; }
 
-	UPROPERTY(config, EditDefaultsOnly, Category = "Overrides", Meta = (InlineEditConditionToggle))
-	uint32 bOverrideSoftSnapPosStrength : 1;
 	// Overrides CVar: np2.PredictiveInterpolation.SoftSnapPosStrength -- Value in percent between 0.0 - 1.0 representing how much to softsnap each tick of the remaining positional distance.
 	UPROPERTY(config, EditDefaultsOnly, Category = "Overrides", Meta = (EditCondition = "bOverrideSoftSnapPosStrength"))
 	float SoftSnapPosStrength = PhysicsReplicationCVars::PredictiveInterpolationCVars::SoftSnapPosStrength;
 	float GetSoftSnapPosStrength() const { return bOverrideSoftSnapPosStrength ? SoftSnapPosStrength : PhysicsReplicationCVars::PredictiveInterpolationCVars::SoftSnapPosStrength; }
 
-	UPROPERTY(config, EditDefaultsOnly, Category = "Overrides", Meta = (InlineEditConditionToggle))
-	uint32 bOverrideSoftSnapRotStrength : 1;
 	// Overrides CVar: np2.PredictiveInterpolation.SoftSnapRotStrength -- Value in percent between 0.0 - 1.0 representing how much to softsnap each tick of the remaining rotational distance.
 	UPROPERTY(config, EditDefaultsOnly, Category = "Overrides", Meta = (EditCondition = "bOverrideSoftSnapRotStrength"))
 	float SoftSnapRotStrength = PhysicsReplicationCVars::PredictiveInterpolationCVars::SoftSnapRotStrength;
 	float GetSoftSnapRotStrength() const { return bOverrideSoftSnapRotStrength ? SoftSnapRotStrength : PhysicsReplicationCVars::PredictiveInterpolationCVars::SoftSnapRotStrength; }
 
-	UPROPERTY(config, EditDefaultsOnly, Category = "Overrides", Meta = (InlineEditConditionToggle))
-	uint32 bOverrideSoftSnapToSource : 1;
 	// Overrides CVar: np2.PredictiveInterpolation.SoftSnapToSource -- If true, softsnap will be performed towards the source state of the current target instead of the predicted state of the current target.
 	UPROPERTY(config, EditDefaultsOnly, Category = "Overrides", Meta = (EditCondition = "bOverrideSoftSnapToSource"))
 	bool bSoftSnapToSource = PhysicsReplicationCVars::PredictiveInterpolationCVars::bSoftSnapToSource;
 	bool GetSoftSnapToSource() const { return bOverrideSoftSnapToSource ? bSoftSnapToSource : PhysicsReplicationCVars::PredictiveInterpolationCVars::bSoftSnapToSource; }
 
-	UPROPERTY(config, EditDefaultsOnly, Category = "Overrides", Meta = (InlineEditConditionToggle))
-	uint32 bOverrideDisableSoftSnap : 1;
 	// Overrides CVar: np2.PredictiveInterpolation.DisableSoftSnap -- When true, predictive interpolation will not use softsnap to correct the replication with when velocity fails. Hardsnap will still eventually kick in if replication can't reach the target.
 	UPROPERTY(config, EditDefaultsOnly, Category = "Overrides", Meta = (EditCondition = "bOverrideDisableSoftSnap"))
 	bool bDisableSoftSnap = PhysicsReplicationCVars::PredictiveInterpolationCVars::bDisableSoftSnap;
 	bool GetDisableSoftSnap() const { return bOverrideDisableSoftSnap ? bDisableSoftSnap : PhysicsReplicationCVars::PredictiveInterpolationCVars::bDisableSoftSnap; }
 
-	UPROPERTY(config, EditDefaultsOnly, Category = "Overrides", Meta = (InlineEditConditionToggle))
-	uint32 bOverrideSkipVelocityRepOnPosEarlyOut : 1;
 	// Overrides CVar: np2.PredictiveInterpolation.SkipVelocityRepOnPosEarlyOut -- If true, don't run linear velocity replication if position can early out but angular can't early out.
 	UPROPERTY(config, EditDefaultsOnly, Category = "Overrides", Meta = (EditCondition = "bOverrideSkipVelocityRepOnPosEarlyOut"))
 	bool bSkipVelocityRepOnPosEarlyOut = PhysicsReplicationCVars::PredictiveInterpolationCVars::bSkipVelocityRepOnPosEarlyOut;
 	bool GetSkipVelocityRepOnPosEarlyOut() const { return bOverrideSkipVelocityRepOnPosEarlyOut ? bSkipVelocityRepOnPosEarlyOut : PhysicsReplicationCVars::PredictiveInterpolationCVars::bSkipVelocityRepOnPosEarlyOut; }
 
-	UPROPERTY(config, EditDefaultsOnly, Category = "Overrides", Meta = (InlineEditConditionToggle))
-	uint32 bOverridePostResimWaitForUpdate : 1;
 	// Overrides CVar: np2.PredictiveInterpolation.PostResimWaitForUpdate -- After a resimulation, wait for replicated states that correspond to post-resim state before processing replication again.
 	UPROPERTY(config, EditDefaultsOnly, Category = "Overrides", Meta = (EditCondition = "bOverridePostResimWaitForUpdate"))
 	bool bPostResimWaitForUpdate = PhysicsReplicationCVars::PredictiveInterpolationCVars::bPostResimWaitForUpdate;
 	bool GetPostResimWaitForUpdate() const { return bOverridePostResimWaitForUpdate ? bPostResimWaitForUpdate : PhysicsReplicationCVars::PredictiveInterpolationCVars::bPostResimWaitForUpdate; }
 
-	UPROPERTY(config, EditDefaultsOnly, Category = "Overrides", Meta = (InlineEditConditionToggle))
-	uint32 bOverrideCorrectConnectedBodies : 1;
 	// Overrides CVar: np2.PredictiveInterpolation.CorrectConnectedBodies -- When true, transform corrections will also apply to any connected physics object.
 	UPROPERTY(config, EditDefaultsOnly, Category = "Overrides", Meta = (EditCondition = "bOverrideCorrectConnectedBodies"))
 	bool bCorrectConnectedBodies = PhysicsReplicationCVars::PredictiveInterpolationCVars::bCorrectConnectedBodies;
 	bool GetCorrectConnectedBodies() const { return bOverrideCorrectConnectedBodies ? bCorrectConnectedBodies : PhysicsReplicationCVars::PredictiveInterpolationCVars::bCorrectConnectedBodies; }
 
-	UPROPERTY(config, EditDefaultsOnly, Category = "Overrides", Meta = (InlineEditConditionToggle))
-	uint32 bOverrideCorrectConnectedBodiesFriction : 1;
 	// Overrides CVar: np2.PredictiveInterpolation.CorrectConnectedBodiesFriction -- When true, transform correction on any connected physics object will also recalculate their friction.
 	UPROPERTY(config, EditDefaultsOnly, Category = "Overrides", Meta = (EditCondition = "bOverrideCorrectConnectedBodiesFriction"))
 	bool bCorrectConnectedBodiesFriction = PhysicsReplicationCVars::PredictiveInterpolationCVars::bCorrectConnectedBodiesFriction;
 	bool GetCorrectConnectedBodiesFriction() const { return bOverrideCorrectConnectedBodiesFriction ? bCorrectConnectedBodiesFriction : PhysicsReplicationCVars::PredictiveInterpolationCVars::bCorrectConnectedBodiesFriction; }
-
 };
 
 USTRUCT()
@@ -302,9 +307,33 @@
 
 	GENERATED_BODY()
 
+	// Override properties
+	UPROPERTY(config, EditDefaultsOnly, Category = "Overrides", Meta = (InlineEditConditionToggle))
+	uint32 bOverrideResimulationErrorPositionThreshold : 1;
+	UPROPERTY(config, EditDefaultsOnly, Category = "Overrides", Meta = (InlineEditConditionToggle))
+	uint32 bOverrideResimulationErrorRotationThreshold : 1;
+	UPROPERTY(config, EditDefaultsOnly, Category = "Overrides", Meta = (InlineEditConditionToggle))
+	uint32 bOverrideResimulationErrorLinearVelocityThreshold : 1;
+	UPROPERTY(config, EditDefaultsOnly, Category = "Overrides", Meta = (InlineEditConditionToggle))
+	uint32 bOverrideResimulationErrorAngularVelocityThreshold : 1;
+	UPROPERTY(config, EditDefaultsOnly, Category = "Overrides", Meta = (InlineEditConditionToggle))
+	uint32 bOverrideRuntimeCorrectionEnabled : 1;
+	UPROPERTY(config, EditDefaultsOnly, Category = "Overrides", Meta = (InlineEditConditionToggle))
+	uint32 bOverrideRuntimeVelocityCorrection : 1;
+	UPROPERTY(config, EditDefaultsOnly, Category = "Overrides", Meta = (InlineEditConditionToggle))
+	uint32 bOverrideRuntimeCorrectConnectedBodies : 1;
+	UPROPERTY(config, EditDefaultsOnly, Category = "Overrides", Meta = (InlineEditConditionToggle))
+	uint32 bOverridePosStabilityMultiplier : 1;
+	UPROPERTY(config, EditDefaultsOnly, Category = "Overrides", Meta = (InlineEditConditionToggle))
+	uint32 bOverrideRotStabilityMultiplier : 1;
+	UPROPERTY(config, EditDefaultsOnly, Category = "Overrides", Meta = (InlineEditConditionToggle))
+	uint32 bOverrideVelStabilityMultiplier : 1;
+	UPROPERTY(config, EditDefaultsOnly, Category = "Overrides", Meta = (InlineEditConditionToggle))
+	uint32 bOverrideAngVelStabilityMultiplier : 1;
+
 	/** Deprecated UE 5.5 - bOverrideResimulationErrorThreshold has been renamed, please use bOverrideResimulationErrorPositionThreshold*/
 	UPROPERTY(config)
-	uint32 bOverrideResimulationErrorThreshold_DEPRECATED;
+	uint32 bOverrideResimulationErrorThreshold_DEPRECATED : 1;
 
 	/** Deprecated UE 5.5 - ResimulationErrorThreshold has been renamed, please use ResimulationErrorPositionThreshold */
 	UPROPERTY(config)
@@ -313,78 +342,56 @@
 	UE_DEPRECATED(5.5, "GetResimulationErrorThreshold has been renamed, please use GetResimulationErrorPositionThreshold.")
 	uint32 GetResimulationErrorThreshold(uint32 DefaultValue) const { return (uint32)FMath::CeilToInt(GetResimulationErrorPositionThreshold(static_cast<float>(DefaultValue))); }
 
-	UPROPERTY(config, EditDefaultsOnly, Category = "Overrides", Meta = (InlineEditConditionToggle))
-	uint32 bOverrideResimulationErrorPositionThreshold : 1;
 	// Overrides Project Settings -> Physics -> Replication -> Physics Prediction -> Resimulation Error Position Threshold -- Distance that the object is allowed to desync from the server before triggering a resimulation, within this threshold runtime correction can be performed if RuntimeCorrectionEnabled is true.
 	UPROPERTY(config, EditDefaultsOnly, Category = "Overrides", Meta = (EditCondition = "bOverrideResimulationErrorPositionThreshold"))
 	float ResimulationErrorPositionThreshold = 10.0f;
 	float GetResimulationErrorPositionThreshold(float DefaultValue) const { return bOverrideResimulationErrorPositionThreshold ? ResimulationErrorPositionThreshold : DefaultValue; }
 
-	UPROPERTY(config, EditDefaultsOnly, Category = "Overrides", Meta = (InlineEditConditionToggle))
-	uint32 bOverrideResimulationErrorRotationThreshold : 1;
 	// Overrides Project Settings -> Physics -> Replication -> Physics Prediction -> Resimulation Error Rotation Threshold -- Rotation difference in degrees that the object is allowed to desync from the server before triggering a resimulation, within this threshold runtime correction can be performed if RuntimeCorrectionEnabled is true.
 	UPROPERTY(config, EditDefaultsOnly, Category = "Overrides", Meta = (EditCondition = "bOverrideResimulationErrorRotationThreshold"))
 	float ResimulationErrorRotationThreshold = 4.0f;
 	float GetResimulationErrorRotationThreshold(float DefaultValue) const { return bOverrideResimulationErrorRotationThreshold ? ResimulationErrorRotationThreshold : DefaultValue; }
 
-	UPROPERTY(config, EditDefaultsOnly, Category = "Overrides", Meta = (InlineEditConditionToggle))
-	uint32 bOverrideResimulationErrorLinearVelocityThreshold : 1;
 	// Overrides Project Settings -> Physics -> Replication -> Physics Prediction -> Resimulation Error Linear Velocity Threshold -- Velocity difference in centimeters / second that the object is allowed to desync from the server before triggering a resimulation, within this threshold runtime correction can be performed if RuntimeCorrectionEnabled is true.
 	UPROPERTY(config, EditDefaultsOnly, Category = "Overrides", Meta = (EditCondition = "bOverrideResimulationErrorLinearVelocityThreshold"))
 	float ResimulationErrorLinearVelocityThreshold = 5.0f;
 	float GetResimulationErrorLinearVelocityThreshold(float DefaultValue) const { return bOverrideResimulationErrorLinearVelocityThreshold ? ResimulationErrorLinearVelocityThreshold : DefaultValue; }
 
-	UPROPERTY(config, EditDefaultsOnly, Category = "Overrides", Meta = (InlineEditConditionToggle))
-	uint32 bOverrideResimulationErrorAngularVelocityThreshold : 1;
 	// Overrides Project Settings -> Physics -> Replication -> Physics Prediction -> Resimulation Error Angular Velocity Threshold -- Degrees / second that the object is allowed to desync from the server before triggering a resimulation, within this threshold runtime correction can be performed if RuntimeCorrectionEnabled is true.
 	UPROPERTY(config, EditDefaultsOnly, Category = "Overrides", Meta = (EditCondition = "bOverrideResimulationErrorAngularVelocityThreshold"))
 	float ResimulationErrorAngularVelocityThreshold = 2.0f;
 	float GetResimulationErrorAngularVelocityThreshold(float DefaultValue) const { return bOverrideResimulationErrorAngularVelocityThreshold ? ResimulationErrorAngularVelocityThreshold : DefaultValue; }
 
-	UPROPERTY(config, EditDefaultsOnly, Category = "Overrides", Meta = (InlineEditConditionToggle))
-	uint32 bOverrideRuntimeCorrectionEnabled : 1;
 	// Overrides CVar: np2.Resim.RuntimeCorrectionEnabled -- Apply positional and rotational runtime corrections while within resim trigger distance.
 	UPROPERTY(config, EditDefaultsOnly, Category = "Overrides", Meta = (EditCondition = "bOverrideRuntimeCorrectionEnabled"))
 	bool bRuntimeCorrectionEnabled = PhysicsReplicationCVars::ResimulationCVars::bRuntimeCorrectionEnabled;
 	bool GetRuntimeCorrectionEnabled() const { return bOverrideRuntimeCorrectionEnabled ? bRuntimeCorrectionEnabled : PhysicsReplicationCVars::ResimulationCVars::bRuntimeCorrectionEnabled; }
 
-	UPROPERTY(config, EditDefaultsOnly, Category = "Overrides", Meta = (InlineEditConditionToggle))
-	uint32 bOverrideRuntimeVelocityCorrection : 1;
 	// Overrides CVar: np2.Resim.RuntimeVelocityCorrection -- Apply linear and angular velocity corrections in runtime while within resim trigger distance. Used if RuntimeCorrectionEnabled is true.
 	UPROPERTY(config, EditDefaultsOnly, Category = "Overrides", Meta = (EditCondition = "bOverrideRuntimeVelocityCorrection"))
 	bool bRuntimeVelocityCorrection = PhysicsReplicationCVars::ResimulationCVars::bRuntimeVelocityCorrection;
 	bool GetRuntimeVelocityCorrectionEnabled() const { return bOverrideRuntimeVelocityCorrection ? bRuntimeVelocityCorrection : PhysicsReplicationCVars::ResimulationCVars::bRuntimeVelocityCorrection; }
 
-	UPROPERTY(config, EditDefaultsOnly, Category = "Overrides", Meta = (InlineEditConditionToggle))
-	uint32 bOverrideRuntimeCorrectConnectedBodies : 1;
 	// Overrides CVar: np2.Resim.RuntimeCorrectConnectedBodies -- If true runtime position and rotation correction will also shift transform of any connected physics objects. Used if RuntimeCorrectionEnabled is true.
 	UPROPERTY(config, EditDefaultsOnly, Category = "Overrides", Meta = (EditCondition = "bOverrideRuntimeCorrectConnectedBodies"))
 	bool bRuntimeCorrectConnectedBodies = PhysicsReplicationCVars::ResimulationCVars::bRuntimeCorrectConnectedBodies;
 	bool GetRuntimeCorrectConnectedBodies() const { return bOverrideRuntimeCorrectConnectedBodies ? bRuntimeCorrectConnectedBodies : PhysicsReplicationCVars::ResimulationCVars::bRuntimeCorrectConnectedBodies; }
 
-	UPROPERTY(config, EditDefaultsOnly, Category = "Overrides", Meta = (InlineEditConditionToggle))
-	uint32 bOverridePosStabilityMultiplier : 1;
 	// Overrides CVar: np2.Resim.PosStabilityMultiplier -- Recommended range between 0.0-1.0. Lower value means more stable positional corrections.
 	UPROPERTY(config, EditDefaultsOnly, Category = "Overrides", Meta = (EditCondition = "bOverridePosStabilityMultiplier"))
 	float PosStabilityMultiplier = PhysicsReplicationCVars::ResimulationCVars::PosStabilityMultiplier;
 	float GetPosStabilityMultiplier() { return bOverridePosStabilityMultiplier ? PosStabilityMultiplier : PhysicsReplicationCVars::ResimulationCVars::PosStabilityMultiplier; }
 
-	UPROPERTY(config, EditDefaultsOnly, Category = "Overrides", Meta = (InlineEditConditionToggle))
-	uint32 bOverrideRotStabilityMultiplier : 1;
 	// Overrides CVar: np2.Resim.RotStabilityMultiplier -- Recommended range between 0.0-1.0. Lower value means more stable rotational corrections.
 	UPROPERTY(config, EditDefaultsOnly, Category = "Overrides", Meta = (EditCondition = "bOverrideRotStabilityMultiplier"))
 	float RotStabilityMultiplier = PhysicsReplicationCVars::ResimulationCVars::RotStabilityMultiplier;
 	float GetRotStabilityMultiplier() const { return bOverrideRotStabilityMultiplier ? RotStabilityMultiplier : PhysicsReplicationCVars::ResimulationCVars::RotStabilityMultiplier; }
 
-	UPROPERTY(config, EditDefaultsOnly, Category = "Overrides", Meta = (InlineEditConditionToggle))
-	uint32 bOverrideVelStabilityMultiplier : 1;
 	// Overrides CVar: np2.Resim.VelStabilityMultiplier -- Recommended range between 0.0-1.0. Lower value means more stable linear velocity corrections.
 	UPROPERTY(config, EditDefaultsOnly, Category = "Overrides", Meta = (EditCondition = "bOverrideVelStabilityMultiplier"))
 	float VelStabilityMultiplier = PhysicsReplicationCVars::ResimulationCVars::VelStabilityMultiplier;
 	float GetVelStabilityMultiplier() const { return bOverrideVelStabilityMultiplier ? VelStabilityMultiplier : PhysicsReplicationCVars::ResimulationCVars::VelStabilityMultiplier; }
 
-	UPROPERTY(config, EditDefaultsOnly, Category = "Overrides", Meta = (InlineEditConditionToggle))
-	uint32 bOverrideAngVelStabilityMultiplier : 1;
 	// Overrides CVar: np2.Resim.AngVelStabilityMultiplier -- Recommended range between 0.0-1.0. Lower value means more stable angular velocity corrections.
 	UPROPERTY(config, EditDefaultsOnly, Category = "Overrides", Meta = (EditCondition = "bOverrideAngVelStabilityMultiplier"))
 	float AngVelStabilityMultiplier = PhysicsReplicationCVars::ResimulationCVars::AngVelStabilityMultiplier;
@@ -411,30 +418,50 @@
 		}
 	}
 
+	// Override properties
 	UPROPERTY(config, EditDefaultsOnly, Category = "Overrides", Meta = (InlineEditConditionToggle))
 	uint32 bOverrideRedundantInputs : 1;
+	UPROPERTY(config, EditDefaultsOnly, Category = "Overrides", Meta = (InlineEditConditionToggle))
+	uint32 bOverrideRedundantRemoteInputs : 1;
+	UPROPERTY(config, EditDefaultsOnly, Category = "Overrides", Meta = (InlineEditConditionToggle))
+	uint32 bOverrideRedundantStates : 1;
+	UPROPERTY(config, EditDefaultsOnly, Category = "Overrides", Meta = (InlineEditConditionToggle))
+	uint32 bOverrideCompareStateToTriggerRewind : 1;
+	UPROPERTY(config, EditDefaultsOnly, Category = "Overrides", Meta = (InlineEditConditionToggle))
+	uint32 bOverridebCompareStateToTriggerRewindIncludeSimProxies : 1;
+	UPROPERTY(config, EditDefaultsOnly, Category = "Overrides", Meta = (InlineEditConditionToggle))
+	uint32 bOverrideCompareInputToTriggerRewind : 1;
+	UPROPERTY(config, EditDefaultsOnly, Category = "Overrides", Meta = (InlineEditConditionToggle))
+	uint32 bOverrideEnableUnreliableFlow : 1;
+	UPROPERTY(config, EditDefaultsOnly, Category = "Overrides", Meta = (InlineEditConditionToggle))
+	uint32 bOverrideEnableReliableFlow : 1;
+	UPROPERTY(config, EditDefaultsOnly, Category = "Overrides", Meta = (InlineEditConditionToggle))
+	uint32 bOverrideApplyDataInsteadOfMergeData : 1;
+	UPROPERTY(config, EditDefaultsOnly, Category = "Overrides", Meta = (InlineEditConditionToggle))
+	uint32 bOverrideAllowInputExtrapolation : 1;
+	UPROPERTY(config, EditDefaultsOnly, Category = "Overrides", Meta = (InlineEditConditionToggle))
+	uint32 bOverrideValidateDataOnGameThread : 1;
+	UPROPERTY(config, EditDefaultsOnly, Category = "Overrides", Meta = (InlineEditConditionToggle))
+	uint32 bOverrideApplySimProxyStateAtRuntime : 1;
+	UPROPERTY(config, EditDefaultsOnly, Category = "Overrides", Meta = (InlineEditConditionToggle))
+	uint32 bOverrideApplySimProxyInputAtRuntime : 1;
+
 	/* Overrides CVar : np2.Resim.RedundantInputs -- How many extra inputs to send with each unreliable network message, to account for packetloss.From owning client to server and server to owning client.
 	 * NOTE: This is disabled while np2.Resim.DynamicInputScaling.Enabled is enabled. */
 	UPROPERTY(config, EditDefaultsOnly, Category = "Overrides", Meta = (EditCondition = "bOverrideRedundantInputs"))
 	uint16 RedundantInputs = (uint16)PhysicsReplicationCVars::ResimulationCVars::RedundantInputs;
 	const uint16 GetRedundantInputs() const { return bOverrideRedundantInputs ? RedundantInputs : (uint16)PhysicsReplicationCVars::ResimulationCVars::RedundantInputs; }
 
-	UPROPERTY(config, EditDefaultsOnly, Category = "Overrides", Meta = (InlineEditConditionToggle))
-	uint32 bOverrideRedundantRemoteInputs : 1;
 	// Overrides CVar: np2.Resim.RedundantRemoteInputs -- How many extra inputs to send with each unreliable network message, to account for packetloss. From Server to remote clients.
 	UPROPERTY(config, EditDefaultsOnly, Category = "Overrides", Meta = (EditCondition = "bOverrideRedundantRemoteInputs"))
 	uint16 RedundantRemoteInputs = (uint16)PhysicsReplicationCVars::ResimulationCVars::RedundantRemoteInputs;
 	const uint16 GetRedundantRemoteInputs() const { return bOverrideRedundantRemoteInputs ? RedundantRemoteInputs : (uint16)PhysicsReplicationCVars::ResimulationCVars::RedundantRemoteInputs; }
 
-	UPROPERTY(config, EditDefaultsOnly, Category = "Overrides", Meta = (InlineEditConditionToggle))
-	uint32 bOverrideRedundantStates : 1;
 	// Overrides CVar: np2.Resim.RedundantStates -- How many extra states to send with each unreliable network message, to account for packetloss. From Server to remote clients.
 	UPROPERTY(config, EditDefaultsOnly, Category = "Overrides", Meta = (EditCondition = "bOverrideRedundantStates"))
 	uint16 RedundantStates = (uint16)PhysicsReplicationCVars::ResimulationCVars::RedundantStates;
 	const uint16 GetRedundantStates() const { return bOverrideRedundantStates ? RedundantStates : (uint16)PhysicsReplicationCVars::ResimulationCVars::RedundantStates; }
 
-	UPROPERTY(config, EditDefaultsOnly, Category = "Overrides", Meta = (InlineEditConditionToggle))
-	uint32 bOverrideCompareStateToTriggerRewind : 1;
 	// Overrides CVar: np2.Resim.CompareStateToTriggerRewind -- When true, cache local FNetworkPhysicsData state in rewind history and compare the predicted state with incoming server state to trigger resimulations if they differ, comparison done through FNetworkPhysicsData::CompareData.
 	UPROPERTY(config, EditDefaultsOnly, Category = "Overrides", Meta = (EditCondition = "bOverrideCompareStateToTriggerRewind"))
 	bool bCompareStateToTriggerRewind = PhysicsReplicationCVars::ResimulationCVars::bCompareStateToTriggerRewind;
@@ -442,15 +469,11 @@
 	UE_DEPRECATED(5.5, "Deprecated, use GetCompareStateToTriggerRewind(const bool bDefaultValue) instead and pass in the value that should apply if not overridden.")
 		const bool GetCompareStateToTriggerRewind() const { return bOverrideCompareStateToTriggerRewind ? bCompareStateToTriggerRewind : PhysicsReplicationCVars::ResimulationCVars::bCompareStateToTriggerRewind; }
 
-	UPROPERTY(config, EditDefaultsOnly, Category = "Overrides", Meta = (InlineEditConditionToggle))
-	uint32 bOverridebCompareStateToTriggerRewindIncludeSimProxies : 1;
 	// Overrides CVar: np2.Resim.CompareStateToTriggerRewind.IncludeSimProxies -- When true, include simulated proxies when np2.Resim.CompareStateToTriggerRewind is enabled.
 	UPROPERTY(config, EditDefaultsOnly, Category = "Overrides", Meta = (EditCondition = "bOverridebCompareStateToTriggerRewindIncludeSimProxies"))
 	bool bCompareStateToTriggerRewindIncludeSimProxies = PhysicsReplicationCVars::ResimulationCVars::bCompareStateToTriggerRewindIncludeSimProxies;
 	const bool GetCompareStateToTriggerRewindIncludeSimProxies(const bool bDefaultValue) const { return bOverridebCompareStateToTriggerRewindIncludeSimProxies ? bCompareStateToTriggerRewindIncludeSimProxies : bDefaultValue; }
 
-	UPROPERTY(config, EditDefaultsOnly, Category = "Overrides", Meta = (InlineEditConditionToggle))
-	uint32 bOverrideCompareInputToTriggerRewind : 1;
 	// Overrides CVar: np2.Resim.CompareInputToTriggerRewind -- When true, compare local predicted FNetworkPhysicsData input with incoming server input to trigger resimulations if they differ, comparison done through FNetworkPhysicsData::CompareData.
 	UPROPERTY(config, EditDefaultsOnly, Category = "Overrides", Meta = (EditCondition = "bOverrideCompareInputToTriggerRewind"))
 	bool bCompareInputToTriggerRewind = PhysicsReplicationCVars::ResimulationCVars::bCompareInputToTriggerRewind;
@@ -458,50 +481,36 @@
 	UE_DEPRECATED(5.5, "Deprecated, use GetCompareInputToTriggerRewind(const bool bDefaultValue) instead and pass in the value that should apply if not overridden.")
 		const bool GetCompareInputToTriggerRewind() const { return bOverrideCompareInputToTriggerRewind ? bCompareInputToTriggerRewind : PhysicsReplicationCVars::ResimulationCVars::bCompareInputToTriggerRewind; }
 
-	UPROPERTY(config, EditDefaultsOnly, Category = "Overrides", Meta = (InlineEditConditionToggle))
-	uint32 bOverrideEnableUnreliableFlow : 1;
 	// Overrides CVar: np2.Resim.EnableUnreliableFlow -- When true, allow data to be sent unreliably. Also sends FNetworkPhysicsData not marked with FNetworkPhysicsData::bimportant unreliably over the network.
 	UPROPERTY(config, EditDefaultsOnly, Category = "Overrides", Meta = (EditCondition = "bOverrideEnableUnreliableFlow"))
 	bool bEnableUnreliableFlow = PhysicsReplicationCVars::ResimulationCVars::bEnableUnreliableFlow;
 	const bool GetEnableUnreliableFlow() const { return bOverrideEnableUnreliableFlow ? bEnableUnreliableFlow : PhysicsReplicationCVars::ResimulationCVars::bEnableUnreliableFlow; }
 
-	UPROPERTY(config, EditDefaultsOnly, Category = "Overrides", Meta = (InlineEditConditionToggle))
-	uint32 bOverrideEnableReliableFlow : 1;
 	// Overrides CVar: np2.Resim.EnableReliableFlow -- EXPERIMENTAL -- When true, allow data to be sent reliably. Also send FNetworkPhysicsData marked with FNetworkPhysicsData::bimportant reliably over the network.
 	UPROPERTY(config, EditDefaultsOnly, Category = "Overrides", Meta = (EditCondition = "bOverrideEnableReliableFlow"))
 	bool bEnableReliableFlow = PhysicsReplicationCVars::ResimulationCVars::bEnableReliableFlow;
 	const bool GetEnableReliableFlow() const { return bOverrideEnableReliableFlow ? bEnableReliableFlow : PhysicsReplicationCVars::ResimulationCVars::bEnableReliableFlow; }
 
-	UPROPERTY(config, EditDefaultsOnly, Category = "Overrides", Meta = (InlineEditConditionToggle))
-	uint32 bOverrideApplyDataInsteadOfMergeData : 1;
 	// Overrides CVar: np2.Resim.ApplyDataInsteadOfMergeData -- When true, call ApplyData for each data instead of MergeData when having to use multiple data entries in one frame.
 	UPROPERTY(config, EditDefaultsOnly, Category = "Overrides", Meta = (EditCondition = "bOverrideApplyDataInsteadOfMergeData"))
 	bool bApplyDataInsteadOfMergeData = PhysicsReplicationCVars::ResimulationCVars::bApplyDataInsteadOfMergeData;
 	const bool GetApplyDataInsteadOfMergeData() const { return bOverrideApplyDataInsteadOfMergeData ? bApplyDataInsteadOfMergeData : PhysicsReplicationCVars::ResimulationCVars::bApplyDataInsteadOfMergeData; }
 
-	UPROPERTY(config, EditDefaultsOnly, Category = "Overrides", Meta = (InlineEditConditionToggle))
-	uint32 bOverrideAllowInputExtrapolation : 1;
 	// Overrides CVar: np2.Resim.AllowInputExtrapolation -- When true and not locally controlled, allow inputs to be extrapolated from last known and if there is a gap allow interpolation between two known inputs.
 	UPROPERTY(config, EditDefaultsOnly, Category = "Overrides", Meta = (EditCondition = "bOverrideAllowInputExtrapolation"))
 	bool bAllowInputExtrapolation = PhysicsReplicationCVars::ResimulationCVars::bAllowInputExtrapolation;
 	const bool GetAllowInputExtrapolation() const { return bOverrideAllowInputExtrapolation ? bAllowInputExtrapolation : PhysicsReplicationCVars::ResimulationCVars::bAllowInputExtrapolation; }
 
-	UPROPERTY(config, EditDefaultsOnly, Category = "Overrides", Meta = (InlineEditConditionToggle))
-	uint32 bOverrideValidateDataOnGameThread : 1;
 	// Overrides CVar: np2.Resim.ValidateDataOnGameThread -- When true, perform server-side input validation through FNetworkPhysicsData::ValidateData on the Game Thread. If false, perform the call on the Physics Thread.
 	UPROPERTY(config, EditDefaultsOnly, Category = "Overrides", Meta = (EditCondition = "bOverrideValidateDataOnGameThread"))
 	bool bValidateDataOnGameThread = PhysicsReplicationCVars::ResimulationCVars::bValidateDataOnGameThread;
 	const bool GetValidateDataOnGameThread() const { return bOverrideValidateDataOnGameThread ? bValidateDataOnGameThread : PhysicsReplicationCVars::ResimulationCVars::bValidateDataOnGameThread; }
 
-	UPROPERTY(config, EditDefaultsOnly, Category = "Overrides", Meta = (InlineEditConditionToggle))
-	uint32 bOverrideApplySimProxyStateAtRuntime : 1;
 	// Overrides CVar: np2.Resim.ApplySimProxyStateAtRuntime -- When true, call ApplyData on received states for simulated proxies at runtime.
 	UPROPERTY(config, EditDefaultsOnly, Category = "Overrides", Meta = (EditCondition = "bOverrideApplySimProxyStateAtRuntime"))
 	bool bApplySimProxyStateAtRuntime = PhysicsReplicationCVars::ResimulationCVars::bApplySimProxyStateAtRuntime;
 	const bool GetApplySimProxyStateAtRuntime() const { return bOverrideApplySimProxyStateAtRuntime ? bApplySimProxyStateAtRuntime : PhysicsReplicationCVars::ResimulationCVars::bApplySimProxyStateAtRuntime; }
 
-	UPROPERTY(config, EditDefaultsOnly, Category = "Overrides", Meta = (InlineEditConditionToggle))
-	uint32 bOverrideApplySimProxyInputAtRuntime : 1;
 	// Overrides CVar: np2.Resim.ApplySimProxyInputAtRuntime -- When true, call ApplyData on received inputs for simulated proxies at runtime.
 	UPROPERTY(config, EditDefaultsOnly, Category = "Overrides", Meta = (EditCondition = "bOverrideApplySimProxyInputAtRuntime"))
 	bool bApplySimProxyInputAtRuntime = PhysicsReplicationCVars::ResimulationCVars::bApplySimProxyInputAtRuntime;
@@ -649,7 +658,7 @@
 	TObjectPtr<UNetworkPhysicsSettingsDataAsset> SettingsDataAsset;
 
 // Deprecated properties
-//#if WITH_EDITORONLY_DATA // TODO UE 5.7 enable the if-statement to hide the deprecated properties
+#if WITH_EDITORONLY_DATA
 	UE_DEPRECATED(5.7, "FNetworkPhysicsSettings GeneralSettings has been deprecated, create a DataAsset of UNetworkPhysicsSettingsDataAsset type and reference that in the UNetworkPhysicsSettingsComponent instead.")
 	UPROPERTY(meta = (DeprecatedProperty, DeprecationMessage = "Deprecated, create a DataAsset of UNetworkPhysicsSettingsDataAsset type and reference that in the UNetworkPhysicsSettingsComponent instead"))
 	FNetworkPhysicsSettings GeneralSettings;
@@ -669,7 +678,7 @@
 	UE_DEPRECATED(5.7, "FNetworkPhysicsSettingsNetworkPhysicsComponent NetworkPhysicsComponentSettings has been deprecated, create a DataAsset of UNetworkPhysicsSettingsDataAsset type and reference that in the UNetworkPhysicsSettingsComponent instead.")
 	UPROPERTY(meta = (DeprecatedProperty, DeprecationMessage = "Deprecated, create a DataAsset of UNetworkPhysicsSettingsDataAsset type and reference that in the UNetworkPhysicsSettingsComponent instead"))
 	FNetworkPhysicsSettingsNetworkPhysicsComponent NetworkPhysicsComponentSettings;
-//#endif
+#endif
 
 private:
 	FNetworkPhysicsSettingsComponentAsync* NetworkPhysicsSettings_Internal;
