Added support to pass '*' to InvestigateRenderAsset console command. This will dump the details on all render assets and result will be sorted based on the asset full name.

diff --git a/Engine/Source/Runtime/Engine/Private/Streaming/StreamingManagerTexture.cpp b/Engine/Source/Runtime/Engine/Private/Streaming/StreamingManagerTexture.cpp
--- a/Engine/Source/Runtime/Engine/Private/Streaming/StreamingManagerTexture.cpp
+++ b/Engine/Source/Runtime/Engine/Private/Streaming/StreamingManagerTexture.cpp
@@ -2860,11 +2860,29 @@ bool FRenderAssetStreamingManager::HandleInvestigateRenderAssetCommand(const TCH
 		StreamingData.ComputeViewInfoExtras(Settings);
 		StreamingData.UpdateBoundSizes_Async(Settings);
 
+		const bool bInvestigateAll = InvestigateAssetName == TEXT("*");
+		TArray<int32> AssetIndices;
+		AssetIndices.Reserve(StreamingRenderAssets.Num());
 		for (int32 AssetIndex = 0; AssetIndex < StreamingRenderAssets.Num(); ++AssetIndex)
 		{
+			AssetIndices.Add(AssetIndex);
+		}
+
+		// When investigating all assets, sort render asset based on their full name to facilitate diffing
+		if (bInvestigateAll)
+		{
+			AssetIndices.Sort([&StreamingRenderAssets](const int32& A, const int32& B)
+			{
+				return StreamingRenderAssets[A].RenderAsset->GetFullName() < StreamingRenderAssets[B].RenderAsset->GetFullName();
+			});
+		}
+
+		for (int32 i = 0; i < AssetIndices.Num(); ++i)
+		{
+			int32 AssetIndex = AssetIndices[i];
 			FStreamingRenderAsset& StreamingRenderAsset = StreamingRenderAssets[AssetIndex];
 			FString AssetName = StreamingRenderAsset.RenderAsset->GetFullName();
-			if (AssetName.Contains(InvestigateAssetName))
+			if (bInvestigateAll || AssetName.Contains(InvestigateAssetName))
 			{
 				UStreamableRenderAsset* RenderAsset = StreamingRenderAsset.RenderAsset;
 				if (!RenderAsset) continue;
