Adding the CSV profiler frame number to PSO precache miss info in the logs

diff --git a/Engine/Source/Runtime/Engine/Private/PSOPrecacheValidation.cpp b/Engine/Source/Runtime/Engine/Private/PSOPrecacheValidation.cpp
--- a/Engine/Source/Runtime/Engine/Private/PSOPrecacheValidation.cpp
+++ b/Engine/Source/Runtime/Engine/Private/PSOPrecacheValidation.cpp
@@ -500,6 +500,16 @@ static const TCHAR* GetPSOMissTypeName(EPSOPrecacheMissType Type)
 	return ResultString;
 }
 
+static void AppendCSVProfilerInfo(PSOMissStringBuilder& StringBuilder)
+{
+	#if CSV_PROFILER_STATS
+	if (FCsvProfiler::Get() && FCsvProfiler::Get()->IsCapturing_Renderthread())
+	{
+		StringBuilder.Appendf(TEXT("\nFrame:\t\t\t\t\t%d"), FCsvProfiler::Get()->GetCaptureFrameNumberRT());
+	}
+	#endif
+}
+
 static void LogGeneralPSOMissInfo(
 	const FGraphicsPipelineStateInitializer& Initializer,
 	const FMaterial* Material,
@@ -512,6 +522,7 @@ static void LogGeneralPSOMissInfo(
 	PSOMissStringBuilder& StringBuilder)
 {
 	StringBuilder.Appendf(TEXT("\n\nPSO PRECACHING MISS:"));
+	AppendCSVProfilerInfo(StringBuilder);
 	StringBuilder.Appendf(TEXT("\n\tType:\t\t\t\t\t%s"), GetPSOMissTypeName(MissType));
 	StringBuilder.Appendf(TEXT("\n\tPSOPrecachingState:\t\t%s"), LexToString(PrecacheResult));
 	StringBuilder.Appendf(TEXT("\n\tMaterial:\t\t\t\t%s"), Material ? *Material->GetAssetName() : TEXT("Unknown"));
@@ -1041,6 +1052,7 @@ void LogPSOMissInfo(
 {
 	PSOMissStringBuilder StringBuilder;
 	StringBuilder.Appendf(TEXT("\n\nPSO PRECACHING MISS:"));
+	AppendCSVProfilerInfo(StringBuilder);
 	StringBuilder.Appendf(TEXT("\n\tType:\t\t\t\t\t%s"), TEXT("Compute"));
 	StringBuilder.Appendf(TEXT("\n\tPSOPrecachingState:\t\t%s"), LexToString(PrecacheResult));
 	StringBuilder.Appendf(TEXT("\n\tMaterial:\t\t\t\t%s"), Material ? *Material->GetAssetName() : TEXT("Unknown"));
