[Audio Insights] Trace audio errors on events tab

diff --git a/Engine/Plugins/AudioInsights/Source/AudioInsights/Private/Views/AudioEventLogDashboardViewFactory.cpp b/Engine/Plugins/AudioInsights/Source/AudioInsights/Private/Views/AudioEventLogDashboardViewFactory.cpp
--- a/Engine/Plugins/AudioInsights/Source/AudioInsights/Private/Views/AudioEventLogDashboardViewFactory.cpp
+++ b/Engine/Plugins/AudioInsights/Source/AudioInsights/Private/Views/AudioEventLogDashboardViewFactory.cpp
@@ -1,6 +1,7 @@
 // Copyright Epic Games, Inc. All Rights Reserved.
 #include "Views/AudioEventLogDashboardViewFactory.h"
 
+#include "AudioDeviceTrace.h"
 #include "AudioEventLogEditorCommands.h"
 #include "AudioInsightsModule.h"
 #include "AudioInsightsStyle.h"
@@ -25,6 +26,7 @@ namespace UE::Audio::Insights
 		const FString CategoryVirtualization = LOCTEXT("EventLogCategory_Virtualization", "Virtualization").ToString();
 		const FString CategoryPlayRequests = LOCTEXT("EventLogCategory_PlayRequests", "Play Requests").ToString();
 		const FString CategoryStopRequests = LOCTEXT("EventLogCategory_StopRequests", "Stop Requests").ToString();
+		const FString CategoryPlayErrors = LOCTEXT("EventLogCategory_PlayErrors", "Play Errors").ToString();
 		const FString CategoryMessages = LOCTEXT("EventLogCategory_Messages", "Messages").ToString();
 		const FString CategoryCustom = LOCTEXT("EventLogCategory_Custom", "Custom").ToString();
 
@@ -87,6 +89,16 @@ namespace UE::Audio::Insights
 				}
 			},
 
+			{
+				CategoryPlayErrors,
+				{
+					ID::PlayFailedNotPlayable,
+					ID::PlayFailedOutOfRange,
+					ID::PlayFailedDebugFiltered,
+					ID::PlayFailedConcurrency,
+				}
+			},
+
 			{
 				CategoryMessages,
 				{
diff --git a/Engine/Source/Runtime/Engine/Private/AudioDevice.cpp b/Engine/Source/Runtime/Engine/Private/AudioDevice.cpp
--- a/Engine/Source/Runtime/Engine/Private/AudioDevice.cpp
+++ b/Engine/Source/Runtime/Engine/Private/AudioDevice.cpp
@@ -4,6 +4,7 @@
 #include "ActiveSoundUpdateInterface.h"
 #include "AudioCompressionSettingsUtils.h"
 #include "AudioDecompress.h"
+#include "AudioDeviceTrace.h"
 #include "AudioEffect.h"
 #include "AudioPluginUtilities.h"
 #include "Audio/AudioDebug.h"
@@ -276,6 +277,16 @@ UE_TRACE_EVENT_BEGIN(Audio, VirtualLoopStopOrRealize)
 	UE_TRACE_EVENT_FIELD(uint64, Timestamp)
 	UE_TRACE_EVENT_FIELD(uint32, PlayOrder)
 UE_TRACE_EVENT_END()
+
+// Audio Insights Event Log built-in event types
+namespace Audio::Trace::EventLog::ID
+{
+	const FString PlayFailedNotPlayable = NSLOCTEXT("AudioInsightsEventLog", "EventLogTraceMessage_PlayFailedNotPlayable", "Play Failed : Not playable").ToString();
+	const FString PlayFailedOutOfRange = NSLOCTEXT("AudioInsightsEventLog", "EventLogTraceMessage_PlayFailedOutOfRange", "Play Failed : Out of range").ToString();
+	const FString PlayFailedDebugFiltered = NSLOCTEXT("AudioInsightsEventLog", "EventLogTraceMessage_PlayFailedDebugFiltered", "Play Failed : Debug filtered").ToString();
+	const FString PlayFailedConcurrency = NSLOCTEXT("AudioInsightsEventLog", "EventLogTraceMessage_PlayFailedConcurrency", "Play Failed : Concurrency").ToString();
+}
+
 #endif // UE_AUDIO_PROFILERTRACE_ENABLED
 
 namespace Audio
@@ -5368,6 +5379,11 @@ void FAudioDevice::AddNewActiveSoundInternal(const FActiveSound& InNewActiveSoun
 	if (Sound == nullptr || !Sound->IsPlayable())
 	{
 		ReportSoundFailedToStart(InNewActiveSound.AudioComponentID, InVirtualLoopToRetrigger);
+
+#if UE_AUDIO_PROFILERTRACE_ENABLED
+		Audio::Trace::EventLog::SendActiveSoundEvent(*this, InNewActiveSound, Audio::Trace::EventLog::ID::PlayFailedNotPlayable);
+#endif // UE_AUDIO_PROFILERTRACE_ENABLED
+
 		return;
 	}
 
@@ -5380,6 +5396,11 @@ void FAudioDevice::AddNewActiveSoundInternal(const FActiveSound& InNewActiveSoun
 			UE_LOG(LogAudio, Log, TEXT("New ActiveSound not created for out of range Sound %s"), *InNewActiveSound.Sound->GetName());
 
 			ReportSoundFailedToStart(InNewActiveSound.AudioComponentID, InVirtualLoopToRetrigger);
+
+#if UE_AUDIO_PROFILERTRACE_ENABLED
+			Audio::Trace::EventLog::SendActiveSoundEvent(*this, InNewActiveSound, Audio::Trace::EventLog::ID::PlayFailedOutOfRange);
+#endif // UE_AUDIO_PROFILERTRACE_ENABLED
+
 			return;
 		}
 	}
@@ -5398,6 +5419,11 @@ void FAudioDevice::AddNewActiveSoundInternal(const FActiveSound& InNewActiveSoun
 			if (!SoundName.Contains(DebugSound))
 			{
 				ReportSoundFailedToStart(InNewActiveSound.AudioComponentID, InVirtualLoopToRetrigger);
+
+#if UE_AUDIO_PROFILERTRACE_ENABLED
+				Audio::Trace::EventLog::SendActiveSoundEvent(*this, InNewActiveSound, Audio::Trace::EventLog::ID::PlayFailedDebugFiltered);
+#endif // UE_AUDIO_PROFILERTRACE_ENABLED
+
 				return;
 			}
 		}
@@ -5446,6 +5472,10 @@ void FAudioDevice::AddNewActiveSoundInternal(const FActiveSound& InNewActiveSoun
 			else
 			{
 				ReportSoundFailedToStart(InNewActiveSound.GetAudioComponentID(), InVirtualLoopToRetrigger);
+
+#if UE_AUDIO_PROFILERTRACE_ENABLED
+				Audio::Trace::EventLog::SendActiveSoundEvent(*this, InNewActiveSound, Audio::Trace::EventLog::ID::PlayFailedConcurrency);
+#endif // UE_AUDIO_PROFILERTRACE_ENABLED
 			}
 		}
 		return;
diff --git a/Engine/Source/Runtime/Engine/Public/AudioDeviceTrace.h b/Engine/Source/Runtime/Engine/Public/AudioDeviceTrace.h
new file mode 100644
--- /dev/null
+++ b/Engine/Source/Runtime/Engine/Public/AudioDeviceTrace.h
@@ -0,0 +1,17 @@
+// Copyright Epic Games, Inc. All Rights Reserved.
+
+#pragma once
+
+#include "Containers/UnrealString.h"
+#include "Internationalization/Internationalization.h"
+
+// Audio Insights Event Log built-in event types
+// actual event list is in AudioDevice.cpp
+namespace Audio::Trace::EventLog::ID
+{
+	// Active Sound : Playing Audio Failures
+	extern ENGINE_API const FString PlayFailedNotPlayable;
+	extern ENGINE_API const FString PlayFailedOutOfRange;
+	extern ENGINE_API const FString PlayFailedDebugFiltered;
+	extern ENGINE_API const FString PlayFailedConcurrency;
+}
