Spline: Add CVar which lets user configure between immediate and lazy update of legacy interp curves in FNewSpline

diff --git a/Engine/Source/Runtime/Engine/Private/Curves/Spline.cpp b/Engine/Source/Runtime/Engine/Private/Curves/Spline.cpp
--- a/Engine/Source/Runtime/Engine/Private/Curves/Spline.cpp
+++ b/Engine/Source/Runtime/Engine/Private/Curves/Spline.cpp
@@ -58,6 +58,13 @@ FAutoConsoleVariableRef CVarFallbackFindNearest(
 	TEXT("True if we should implement FindNearest and FindNearestOnSegment using an intermediate spline representation. Only applies if Spline.Implementation == 2.")
 );
 
+bool GImmediatelyUpdateLegacyCurves = false;
+FAutoConsoleVariableRef CVarImmediatelyUpdateLegacyCurves(
+	TEXT("Spline.ImmediatelyUpdateLegacyCurves"),
+	GImmediatelyUpdateLegacyCurves,
+	TEXT("If true, mutating operations immediately rebuild legacy curves. If false, legacy curves are updated only when requested.")
+);
+
 // WARNING! THIS BREAKS EXISTING TEXT REPRESENTATION WHEN CHANGED, IT WILL ABSOLUTELY TRASH YOUR STUFF IF YOU CHANGE WILLY-NILLY
 bool GEncodeAsHex = true;
 // FAutoConsoleVariableRef CVarEncodeAsHex(
@@ -276,13 +283,12 @@ private:
 	static inline const FName RotationAttrName = FName("Rotation");
 	static inline const FName ScaleAttrName = FName("Scale");
 
-    /**
-     * We keep these up to date with Points.
-     * Exists to implement the legacy interface provided by FSpline.
-     */
+	/** Legacy Curves: */
     FInterpCurveVector PositionCurve;
 	FInterpCurveQuat RotationCurve;
 	FInterpCurveVector ScaleCurve;
+	FChangeStamp LegacyCurvesNextVersion;
+	uint32 LegacyCurvesVersion = 0;
 	
 	int32 ReparamStepsPerSegment = 10;
 
@@ -309,6 +315,8 @@ private:
 	/** Clears and repopulates reparameterization attribute channels. */
 	void UpdateReparamTable();
 
+	void MarkLegacyCurvesDirty();
+
 	/** Clears and repopulates PositionCurve, RotationCurve, and ScaleCurve */
 	void RebuildLegacyCurves();
 
@@ -1551,9 +1559,13 @@ const FInterpCurveVector& FSpline::GetSplinePointsPosition() const
 {
 	switch (CurrentImplementation)
 	{
-	default: return PositionCurve;
-	case 1: return LegacyData->GetSplinePointsPosition();
-	case 2: return NewData->GetSplinePointsPosition();
+	default:
+		return PositionCurve;
+	case 1:
+		return LegacyData->GetSplinePointsPosition();
+	case 2:
+		NewData->RebuildLegacyCurves();
+		return NewData->GetSplinePointsPosition();
 	}
 }
 
@@ -1561,9 +1573,13 @@ const FInterpCurveQuat& FSpline::GetSplinePointsRotation() const
 {
 	switch (CurrentImplementation)
 	{
-	default: return RotationCurve;
-	case 1: return LegacyData->GetSplinePointsRotation();
-	case 2: return NewData->GetSplinePointsRotation();
+	default:
+		return RotationCurve;
+	case 1:
+		return LegacyData->GetSplinePointsRotation();
+	case 2:
+		NewData->RebuildLegacyCurves();
+		return NewData->GetSplinePointsRotation();
 	}
 }
 
@@ -1571,9 +1587,13 @@ const FInterpCurveVector& FSpline::GetSplinePointsScale() const
 {
 	switch (CurrentImplementation)
 	{
-	default: return ScaleCurve;
-	case 1: return LegacyData->GetSplinePointsScale();
-	case 2: return NewData->GetSplinePointsScale();
+	default:
+		return ScaleCurve;
+	case 1:
+		return LegacyData->GetSplinePointsScale();
+	case 2:
+		NewData->RebuildLegacyCurves();
+		return NewData->GetSplinePointsScale();
 	}
 }
 
@@ -2025,7 +2045,7 @@ void FNewSpline::AddPoint(const FSplinePoint& Point)
 		RotChild->AppendPoint(TTangentBezierSpline<FQuat>::FTangentBezierControlPoint(FQuat::Identity));
 		RotChild->SetKnotVector(GetSpline().GetKnotVector());
 		
-		RebuildLegacyCurves();
+		MarkLegacyCurvesDirty();
 		return;
 	}
 	else if (Index == 0)
@@ -2040,7 +2060,7 @@ void FNewSpline::AddPoint(const FSplinePoint& Point)
     	RotChild->PrependPoint(TTangentBezierSpline<FQuat>::FTangentBezierControlPoint(FQuat::Identity));
     	RotChild->SetKnotVector(GetSpline().GetKnotVector());
     	
-    	RebuildLegacyCurves();
+    	MarkLegacyCurvesDirty();
     	return;
     }
 	
@@ -2058,7 +2078,7 @@ void FNewSpline::AddPoint(const FSplinePoint& Point)
 	RotChild->InsertPointAtGlobalParam(RotParam, TTangentBezierSpline<FQuat>::FTangentBezierControlPoint(FQuat::Identity));
 	RotChild->SetKnotVector(GetSpline().GetKnotVector());
 	
-	RebuildLegacyCurves();
+	MarkLegacyCurvesDirty();
 }
 
 void FNewSpline::InsertPoint(const FSplinePoint& Point, int32 Index)
@@ -2087,7 +2107,7 @@ void FNewSpline::InsertPoint(const FSplinePoint& Point, int32 Index)
 		RotChild->AppendPoint(TTangentBezierSpline<FQuat>::FTangentBezierControlPoint(FQuat::Identity));
 		RotChild->SetKnotVector(GetSpline().GetKnotVector());
 		
-		RebuildLegacyCurves();
+		MarkLegacyCurvesDirty();
 		return;
 	}
 	else if (Index == 0)
@@ -2102,7 +2122,7 @@ void FNewSpline::InsertPoint(const FSplinePoint& Point, int32 Index)
 		RotChild->PrependPoint(TTangentBezierSpline<FQuat>::FTangentBezierControlPoint(FQuat::Identity));
 		RotChild->SetKnotVector(GetSpline().GetKnotVector());
 		
-		RebuildLegacyCurves();
+		MarkLegacyCurvesDirty();
 		return;
 	}
 	
@@ -2121,7 +2141,7 @@ void FNewSpline::InsertPoint(const FSplinePoint& Point, int32 Index)
 	RotChild->InsertPointAtGlobalParam(RotParam, TTangentBezierSpline<FQuat>::FTangentBezierControlPoint(FQuat::Identity));
 	RotChild->SetKnotVector(GetSpline().GetKnotVector());
 	
-	RebuildLegacyCurves();
+	MarkLegacyCurvesDirty();
 }
 
 FSplinePoint FNewSpline::GetPoint(int32 Index) const
@@ -2172,7 +2192,7 @@ void FNewSpline::RemovePoint(int32 Index)
 	RotChild->RemovePoint(Index);
 	RotChild->SetKnotVector(GetSpline().GetKnotVector());
 	
-    RebuildLegacyCurves();
+    MarkLegacyCurvesDirty();
 }
 
 void FNewSpline::SetLocation(int32 Index, const FVector& InLocation)
@@ -2185,7 +2205,7 @@ void FNewSpline::SetLocation(int32 Index, const FVector& InLocation)
 	GetSpline().ModifyPoint(Index, {InLocation, GetSpline().GetTangentIn(Index), GetSpline().GetTangentOut(Index), GetSpline().GetTangentMode(Index)});
 
 	MarkReparamTableDirty();
-    RebuildLegacyCurves();
+    MarkLegacyCurvesDirty();
 }
 
 FVector FNewSpline::GetLocation(const int32 Index) const
@@ -2209,7 +2229,7 @@ void FNewSpline::SetInTangent(const int32 Index, const FVector& InTangent)
     GetSpline().SetTangentIn(Index, InTangent);
 
 	MarkReparamTableDirty();
-    RebuildLegacyCurves();
+    MarkLegacyCurvesDirty();
 }
 
 FVector FNewSpline::GetInTangent(const int32 Index) const
@@ -2233,7 +2253,7 @@ void FNewSpline::SetOutTangent(const int32 Index, const FVector& OutTangent)
 	GetSpline().SetTangentOut(Index, OutTangent);
 
 	MarkReparamTableDirty();
-    RebuildLegacyCurves();
+    MarkLegacyCurvesDirty();
 }
 
 FVector FNewSpline::GetOutTangent(const int32 Index) const
@@ -2257,7 +2277,7 @@ void FNewSpline::SetRotation(int32 Index, const FQuat& InRotation)
 
 	SetAttributeValue<FQuat>(RotationAttrName, InRotation, Index);
 	
-    RebuildLegacyCurves();
+    MarkLegacyCurvesDirty();
 }
 
 FQuat FNewSpline::GetRotation(const int32 Index) const
@@ -2281,7 +2301,7 @@ void FNewSpline::SetScale(int32 Index, const FVector& InScale)
 
 	SetAttributeValue<FVector>(ScaleAttrName, InScale, Index);
     
-    RebuildLegacyCurves();
+    MarkLegacyCurvesDirty();
 }
 
 FVector FNewSpline::GetScale(const int32 Index) const
@@ -2304,7 +2324,7 @@ void FNewSpline::SetSplinePointType(int32 Index, EInterpCurveMode Type)
 	GetSpline().SetPointTangentMode(Index, ConvertInterpCurveModeToTangentMode(Type));
 
 	MarkReparamTableDirty();
-    RebuildLegacyCurves();
+    MarkLegacyCurvesDirty();
 }
 
 EInterpCurveMode FNewSpline::GetSplinePointType(int32 Index) const
@@ -2411,7 +2431,7 @@ void FNewSpline::UpdateSpline()
 {
     GetSpline().UpdateTangents();			// updates the tangents on our version of the points
     MarkReparamTableDirty();
-    RebuildLegacyCurves();		// updates legacy curves based on points, never evaluates internal spline.
+    RebuildLegacyCurves();					// updates legacy curves based on points, never evaluates internal spline.
 }
 
 void FNewSpline::UpdateSpline(const FSpline::FUpdateSplineParams& Params)
@@ -2497,8 +2517,33 @@ void FNewSpline::UpdateReparamTable()
 	ReparamTableVersion = CachedNextVersion;
 }
 
+void FNewSpline::MarkLegacyCurvesDirty()
+{
+	if (UE::Spline::GImmediatelyUpdateLegacyCurves)
+	{
+		RebuildLegacyCurves();
+	}
+	else
+	{
+		LegacyCurvesNextVersion.Increment();
+	}
+}
+
 void FNewSpline::RebuildLegacyCurves()
 {
+	// Curves are dirty if:
+	// 1) Current version is not the same as next version (if next version == version, it has not been dirtied)
+	// 2) UE::Spline::GImmediatelyUpdateLegacyCurves == true (legacy curves are updated any time they might be dirty)
+	auto AreLegacyCurvesDirty = [this]()
+	{
+		return LegacyCurvesNextVersion.GetValue() != LegacyCurvesVersion || UE::Spline::GImmediatelyUpdateLegacyCurves;
+	};
+	
+	if (!AreLegacyCurvesDirty())
+	{
+		return;
+	}
+
 	PositionCurve.Points.Reset(GetSpline().GetNumPoints());
 	RotationCurve.Points.Reset(GetSpline().GetNumPoints());
 	ScaleCurve.Points.Reset(GetSpline().GetNumPoints());
@@ -2517,6 +2562,8 @@ void FNewSpline::RebuildLegacyCurves()
     	RotationCurve.AddPoint(AttributeParameter, GetAttributeValue<FQuat>(RotationAttrName, i));
     	ScaleCurve.AddPoint(AttributeParameter, GetAttributeValue<FVector>(ScaleAttrName, i));
     }
+    
+    LegacyCurvesVersion = LegacyCurvesNextVersion.GetValue();
 }
 
 FNewSpline::FTangentBezierControlPoint FNewSpline::ConvertToTangentBezierControlPoint(const FSplinePoint& Point) const
@@ -2540,7 +2587,7 @@ void FNewSpline::UpdatePointAttributes(const FSplinePoint& Point, int32 PointInd
 	SetRotation(PointIndex, Point.Rotation.Quaternion());
 	SetScale(PointIndex, Point.Scale);
         
-	RebuildLegacyCurves();
+	MarkLegacyCurvesDirty();
 }
 
 float FNewSpline::ConvertIndexToInternalParameter(int32 Index, float Fraction) const
