Adds fallback to the old material translator instead of using the new one if Substrate is turned on, since the latter does not support Substrate yet.

#rnx

diff --git a/Engine/Source/Runtime/Engine/Private/Materials/MaterialShared.cpp b/Engine/Source/Runtime/Engine/Private/Materials/MaterialShared.cpp
--- a/Engine/Source/Runtime/Engine/Private/Materials/MaterialShared.cpp
+++ b/Engine/Source/Runtime/Engine/Private/Materials/MaterialShared.cpp
@@ -3478,7 +3478,8 @@ bool FMaterial::Translate(const FMaterialShaderMapId& InShaderMapId,
 	// Not all Insight data will be filled out by both translators.
 	GetMaterialInterface()->MaterialInsight.Reset(new FMaterialInsights);
 
-	if (InShaderMapId.bUsingNewHLSLGenerator)
+	// Use the new translator if the shader map uses it. Note: the new translator does not support Substrate yet, so if Substrate is on, fallback to the old translator.
+	if (InShaderMapId.bUsingNewHLSLGenerator && !Substrate::IsSubstrateEnabled())
 	{
 		return Translate_New(InShaderMapId, InStaticParameters, InTargetPlatform, OutCompilationOutput, OutMaterialEnvironment);
 	}
