Spline: Speculative crash fix which allows Spline.Implementation to exist in non-editor builds

diff --git a/Engine/Source/Runtime/Engine/Private/Curves/Spline.cpp b/Engine/Source/Runtime/Engine/Private/Curves/Spline.cpp
--- a/Engine/Source/Runtime/Engine/Private/Curves/Spline.cpp
+++ b/Engine/Source/Runtime/Engine/Private/Curves/Spline.cpp
@@ -36,13 +36,17 @@ static void SplineImplementationSink(IConsoleVariable*)
 		GPreviousImplementation = GImplementation;
 	}
 }
+#endif
 FAutoConsoleVariableRef CVarSplineImplementation(
 	TEXT("Spline.Implementation"),
 	GImplementation,
-	TEXT("0) Not Implemented - 1) Legacy Implementation - 2) New Implementation"),
-	FConsoleVariableDelegate::CreateStatic(&SplineImplementationSink)
-);
+	TEXT("0) Not Implemented - 1) Legacy Implementation - 2) New Implementation")
+#if WITH_EDITOR
+	,FConsoleVariableDelegate::CreateStatic(&SplineImplementationSink)
+#else
+	,ECVF_ReadOnly		// In non-editor builds, this CVar should exist so that GImplementation is properly set but it should not be changed by the user.
 #endif
+);
 	
 bool GApproximateTangents = false;
 FAutoConsoleVariableRef CVarApproximateTangents(
