Exposed GatherTickRecordData on FAnimNotifyEventReference

diff --git a/Engine/Source/Runtime/Engine/Public/Animation/AnimNotifyQueue.h b/Engine/Source/Runtime/Engine/Public/Animation/AnimNotifyQueue.h
--- a/Engine/Source/Runtime/Engine/Public/Animation/AnimNotifyQueue.h
+++ b/Engine/Source/Runtime/Engine/Public/Animation/AnimNotifyQueue.h
@@ -8,6 +8,7 @@
 #include "Animation/AnimNodeMessages.h"
 #include "AnimNotifyQueue.generated.h"
 
+#define UE_API ENGINE_API
 
 class USkeletalMeshComponent;
 struct FAnimInstanceProxy;
@@ -81,7 +82,7 @@ struct FAnimNotifyEventReference
 	}
 
 	// Pulls relevant data from the tick record
-	void GatherTickRecordData(const FAnimTickRecord& InTickRecord);
+	UE_API void GatherTickRecordData(const FAnimTickRecord& InTickRecord);
 
 	// Allows adding extra context data after GatherTickRecordData has been exectuted
 	template<typename Type, typename... TArgs>
@@ -176,7 +177,7 @@ struct FAnimNotifyQueue
 	void AddAnimNotify(const FAnimNotifyEvent* Notify, const UObject* NotifySource);
 
 	/** Add anim notifies **/
-	ENGINE_API void AddAnimNotifies(bool bSrcIsLeader, const TArray<FAnimNotifyEventReference>& NewNotifies, const float InstanceWeight);
+	UE_API void AddAnimNotifies(bool bSrcIsLeader, const TArray<FAnimNotifyEventReference>& NewNotifies, const float InstanceWeight);
 
 	/** Add anim notifies from montage**/
 	void AddAnimNotifies(bool bSrcIsLeader, const TMap<FName, TArray<FAnimNotifyEventReference>>& NewNotifies, const float InstanceWeight);
@@ -220,3 +221,5 @@ private:
 	UPROPERTY(Transient)
 	TWeakObjectPtr<UWorld> World;
 };
+
+#undef UE_API
\ No newline at end of file
