[Audio Insights] Moving event string definitions to AudioMixerCore from Engine

diff --git a/Engine/Plugins/AudioInsights/Source/AudioInsights/Private/Views/AudioEventLogDashboardViewFactory.cpp b/Engine/Plugins/AudioInsights/Source/AudioInsights/Private/Views/AudioEventLogDashboardViewFactory.cpp
--- a/Engine/Plugins/AudioInsights/Source/AudioInsights/Private/Views/AudioEventLogDashboardViewFactory.cpp
+++ b/Engine/Plugins/AudioInsights/Source/AudioInsights/Private/Views/AudioEventLogDashboardViewFactory.cpp
@@ -1,7 +1,6 @@
 // Copyright Epic Games, Inc. All Rights Reserved.
 #include "Views/AudioEventLogDashboardViewFactory.h"
 
-//#include "AudioDeviceTrace.h"
 #include "AudioEventLogEditorCommands.h"
 #include "AudioInsightsModule.h"
 #include "AudioInsightsStyle.h"
@@ -92,10 +91,10 @@ namespace UE::Audio::Insights
 			{
 				CategoryPlayErrors,
 				{
-					//ID::PlayFailedNotPlayable,
-					//ID::PlayFailedOutOfRange,
-					//ID::PlayFailedDebugFiltered,
-					//ID::PlayFailedConcurrency,
+					ID::PlayFailedNotPlayable,
+					ID::PlayFailedOutOfRange,
+					ID::PlayFailedDebugFiltered,
+					ID::PlayFailedConcurrency,
 				}
 			},
 
diff --git a/Engine/Source/Runtime/AudioMixerCore/Private/AudioMixer.cpp b/Engine/Source/Runtime/AudioMixerCore/Private/AudioMixer.cpp
--- a/Engine/Source/Runtime/AudioMixerCore/Private/AudioMixer.cpp
+++ b/Engine/Source/Runtime/AudioMixerCore/Private/AudioMixer.cpp
@@ -61,6 +61,11 @@ namespace Audio::Trace::EventLog::ID
 	
 	const FString StopAllRequested = NSLOCTEXT("AudioInsightsEventLog", "EventLogTraceMessage_StopAllRequested", "Stop All").ToString();
 	const FString FlushAudioDeviceRequested = NSLOCTEXT("AudioInsightsEventLog", "EventLogTraceMessage_FlushAudioDeviceRequested", "Flush Audio Device").ToString();
+
+	const FString PlayFailedNotPlayable = NSLOCTEXT("AudioInsightsEventLog", "EventLogTraceMessage_PlayFailedNotPlayable", "Play Failed : Not playable").ToString();
+	const FString PlayFailedOutOfRange = NSLOCTEXT("AudioInsightsEventLog", "EventLogTraceMessage_PlayFailedOutOfRange", "Play Failed : Out of range").ToString();
+	const FString PlayFailedDebugFiltered = NSLOCTEXT("AudioInsightsEventLog", "EventLogTraceMessage_PlayFailedDebugFiltered", "Play Failed : Debug filtered").ToString();
+	const FString PlayFailedConcurrency = NSLOCTEXT("AudioInsightsEventLog", "EventLogTraceMessage_PlayFailedConcurrency", "Play Failed : Concurrency").ToString();
 }
 
 // Command to enable logging to display accurate audio render times
diff --git a/Engine/Source/Runtime/AudioMixerCore/Public/AudioMixerTrace.h b/Engine/Source/Runtime/AudioMixerCore/Public/AudioMixerTrace.h
--- a/Engine/Source/Runtime/AudioMixerCore/Public/AudioMixerTrace.h
+++ b/Engine/Source/Runtime/AudioMixerCore/Public/AudioMixerTrace.h
@@ -48,6 +48,12 @@ namespace Audio::Trace::EventLog::ID
 	extern AUDIOMIXERCORE_API const FString PlayRequestSound2D;
 	extern AUDIOMIXERCORE_API const FString PlayRequestSlateSound;
 	
+	// Active Sound : Playing Audio Failures
+	extern AUDIOMIXERCORE_API const FString PlayFailedNotPlayable;
+	extern AUDIOMIXERCORE_API const FString PlayFailedOutOfRange;
+	extern AUDIOMIXERCORE_API const FString PlayFailedDebugFiltered;
+	extern AUDIOMIXERCORE_API const FString PlayFailedConcurrency;
+
 	extern AUDIOMIXERCORE_API const FString StopRequestActiveSound;
 	extern AUDIOMIXERCORE_API const FString StopRequestSoundsUsingResource;
 	extern AUDIOMIXERCORE_API const FString StopRequestConcurrency;
diff --git a/Engine/Source/Runtime/Engine/Private/AudioDevice.cpp b/Engine/Source/Runtime/Engine/Private/AudioDevice.cpp
--- a/Engine/Source/Runtime/Engine/Private/AudioDevice.cpp
+++ b/Engine/Source/Runtime/Engine/Private/AudioDevice.cpp
@@ -4,7 +4,6 @@
 #include "ActiveSoundUpdateInterface.h"
 #include "AudioCompressionSettingsUtils.h"
 #include "AudioDecompress.h"
-#include "AudioDeviceTrace.h"
 #include "AudioEffect.h"
 #include "AudioPluginUtilities.h"
 #include "Audio/AudioDebug.h"
@@ -278,15 +277,6 @@ UE_TRACE_EVENT_BEGIN(Audio, VirtualLoopStopOrRealize)
 	UE_TRACE_EVENT_FIELD(uint32, PlayOrder)
 UE_TRACE_EVENT_END()
 
-// Audio Insights Event Log built-in event types
-namespace Audio::Trace::EventLog::ID
-{
-	const FString PlayFailedNotPlayable = NSLOCTEXT("AudioInsightsEventLog", "EventLogTraceMessage_PlayFailedNotPlayable", "Play Failed : Not playable").ToString();
-	const FString PlayFailedOutOfRange = NSLOCTEXT("AudioInsightsEventLog", "EventLogTraceMessage_PlayFailedOutOfRange", "Play Failed : Out of range").ToString();
-	const FString PlayFailedDebugFiltered = NSLOCTEXT("AudioInsightsEventLog", "EventLogTraceMessage_PlayFailedDebugFiltered", "Play Failed : Debug filtered").ToString();
-	const FString PlayFailedConcurrency = NSLOCTEXT("AudioInsightsEventLog", "EventLogTraceMessage_PlayFailedConcurrency", "Play Failed : Concurrency").ToString();
-}
-
 #endif // UE_AUDIO_PROFILERTRACE_ENABLED
 
 namespace Audio
diff --git a/Engine/Source/Runtime/Engine/Public/AudioDeviceTrace.h b/Engine/Source/Runtime/Engine/Public/AudioDeviceTrace.h
deleted file mode 100644
