Add an Enum to Subtitle Assets to distinguish when the duration property is relevant vs automatically stopping.

diff --git a/Engine/Plugins/Experimental/SubtitlesAndClosedCaptions/Source/SubtitlesAndClosedCaptions/Private/SubtitlesAudioSubsystem.cpp b/Engine/Plugins/Experimental/SubtitlesAndClosedCaptions/Source/SubtitlesAndClosedCaptions/Private/SubtitlesAudioSubsystem.cpp
--- a/Engine/Plugins/Experimental/SubtitlesAndClosedCaptions/Source/SubtitlesAndClosedCaptions/Private/SubtitlesAudioSubsystem.cpp
+++ b/Engine/Plugins/Experimental/SubtitlesAndClosedCaptions/Source/SubtitlesAndClosedCaptions/Private/SubtitlesAudioSubsystem.cpp
@@ -13,7 +13,11 @@ void USubtitlesAudioSubsystem::NotifyActiveSoundCreated(FActiveSound& ActiveSoun
 	{
 		if (const UAssetUserData * Subtitle{ Sound->GetAssetUserDataOfClass(USubtitleAssetUserData::StaticClass()) })
 		{
-			FSubtitlesAndClosedCaptionsDelegates::QueueSubtitle.ExecuteIfBound(FQueueSubtitleParameters{ *Subtitle }, ESubtitleTiming::InternallyTimed);
+			if (const USubtitleAssetUserData* SubtitleAsset = CastChecked<USubtitleAssetUserData>(Subtitle))	// Should never fail since Subtitle would only be valid if there's userdata of that class.
+			{
+				const bool bUseDurationProperty = (SubtitleAsset->SubtitleDurationType == ESubtitleDurationType::UseDurationProperty);
+				FSubtitlesAndClosedCaptionsDelegates::QueueSubtitle.ExecuteIfBound(FQueueSubtitleParameters{ *Subtitle }, bUseDurationProperty ? ESubtitleTiming::InternallyTimed : ESubtitleTiming::ExternallyTimed);
+			}
 		}
 	}
 }
@@ -26,7 +30,13 @@ void USubtitlesAudioSubsystem::NotifyActiveSoundDeleting(const FActiveSound& Act
 	{
 		if (const UAssetUserData * Subtitle{ Sound->GetAssetUserDataOfClass(USubtitleAssetUserData::StaticClass()) })
 		{
-			FSubtitlesAndClosedCaptionsDelegates::StopSubtitle.ExecuteIfBound(*Subtitle);
+			if (const USubtitleAssetUserData* SubtitleAsset = CastChecked<USubtitleAssetUserData>(Subtitle))	// Should never fail since Subtitle would only be valid if there's userdata of that class.
+			{
+				if (SubtitleAsset->SubtitleDurationType == ESubtitleDurationType::UseSoundDuration)
+				{
+					FSubtitlesAndClosedCaptionsDelegates::StopSubtitle.ExecuteIfBound(*Subtitle);
+				}
+			}
 		}
 	}
 }
\ No newline at end of file
diff --git a/Engine/Source/Runtime/Engine/Public/Subtitles/SubtitlesAndClosedCaptionsDelegates.h b/Engine/Source/Runtime/Engine/Public/Subtitles/SubtitlesAndClosedCaptionsDelegates.h
--- a/Engine/Source/Runtime/Engine/Public/Subtitles/SubtitlesAndClosedCaptionsDelegates.h
+++ b/Engine/Source/Runtime/Engine/Public/Subtitles/SubtitlesAndClosedCaptionsDelegates.h
@@ -19,6 +19,15 @@ enum class ESubtitleTiming : uint8
 	ExternallyTimed
 };
 
+// SubtitlesAudioSubsystem will automatically queue and stop subtitles embedded in USoundBase sounds using play/stop callbacks.
+// This enum controls whether the stopping should be done automatically, or if the subtitle should be timed using its existing duration.
+UENUM(BlueprintType)
+enum class ESubtitleDurationType : uint8
+{
+	UseSoundDuration UMETA(ToolTip = "Automatically stops this subtitle when the sound stops playing (only useful when attached as AssetUserData"),
+	UseDurationProperty UMETA(ToolTip = "Time this subtitle using its Duration property")
+};
+
 // ESRB rating categories.
 UENUM(BlueprintType)
 enum class EESRB : uint8
@@ -60,36 +69,40 @@ public:
 	// #SUBTITLES_PRD: carried over from FSubtitleCue, still required
 
 	// The text to appear in the subtitle.
-	UPROPERTY(EditAnywhere, BlueprintReadWrite, Category = Subtitles, meta = (Multiline = true))
+	UPROPERTY(EditAnywhere, BlueprintReadWrite, Category = "Subtitles", meta = (Multiline = true))
 	FText Text;
 
 	// #SUBTITLES_PRD: carried over from FSubtitleCue, still required
 
-	// Time to display in seconds. 
+	// Whether to automatically unqueue the subtitle using an associated Sound, or to use the duration supplied below.
+	UPROPERTY(EditAnywhere, BlueprintReadWrite, Category = "Subtitles|Duration")
+	ESubtitleDurationType SubtitleDurationType = ESubtitleDurationType::UseSoundDuration;
+
+	// Time to display in seconds.
 	// Defaulted to 3 seconds so when adding new subtitles it's not required to enter a placeholder Duration.
 	// Duration can be be set by ingestion pipelines when importing Subtitles in bulk
-	UPROPERTY(EditAnywhere, BlueprintReadWrite, Category = Subtitles, meta = (ClampMin = 0.05f))
+	UPROPERTY(EditAnywhere, BlueprintReadWrite, Category = "Subtitles|Duration", meta = (ClampMin = 0.05f, EditCondition = "SubtitleDurationType == ESubtitleDurationType::UseDurationProperty"))
 	float Duration = SubtitleDefaultDuration;
 
 	// Some subtitles have a delay before they're allowed to be displayed (primarily from the legacy system).
 	// StartOffset measures how long in Seconds, after queuing, before the subtitle is allowed to enter the active subtitles queue.
 	// ESubtitleTiming::ExternallyTimed does not effect this initial delay.
-	UPROPERTY(EditAnywhere, BlueprintReadWrite, Category = Subtitles, meta = (ClampMin = 0.f))
+	UPROPERTY(EditAnywhere, BlueprintReadWrite, Category = "Subtitles", meta = (ClampMin = 0.f))
 	float StartOffset = 0.f;
 
 	// #SUBTITLES_PRD: Priority comes from USoundBase::GetSubtitlePriority, USoundCue::GetSubtitlePriority and USoundWave::GetSubtitlePriority
 	// Consolidate various subtitle properties throughout sound/audio code into this new subtitles plugin.
 
 	// The priority of the subtitle.  Defaults to 1.  Higher values will play instead of lower values.
-	UPROPERTY(EditAnywhere, BlueprintReadWrite, Category = Subtitles)
+	UPROPERTY(EditAnywhere, BlueprintReadWrite, Category = "Subtitles")
 	float Priority = 1.f;
 
 	// ESRB rating category.
-	UPROPERTY(EditAnywhere, BlueprintReadWrite, Category = Subtitles, meta = (DisplayName = "ESRB"))
+	UPROPERTY(EditAnywhere, BlueprintReadWrite, Category = "Subtitles", meta = (DisplayName = "ESRB"))
 	EESRB Esrb = EESRB::Everyone;
 
 	// Subtitle type for type-specific rendering.
-	UPROPERTY(EditAnywhere, BlueprintReadWrite, Category = Subtitles)
+	UPROPERTY(EditAnywhere, BlueprintReadWrite, Category = "Subtitles")
 	ESubtitleType SubtitleType = ESubtitleType::Subtitle;
 };
 
