UE-313548 - ChildActorComponent crash when switching the child actor class.

In rare cases, it's possible to reload the package that a CAC lived in. This can also trigger a reload of the original ChildActorTemplate that the CAC was serialized with. As a result, if the ChildActorClass is ever changed back to its original value, the renaming of the ChildActorTemplate can fail because it detects that there's already an object of the same name (assuming that a GC hasn't occurred yet).

The easiest fix is to just see if an object already exists, and if so, move it to the transient package.

diff --git a/Engine/Source/Runtime/Engine/Private/Components/ChildActorComponent.cpp b/Engine/Source/Runtime/Engine/Private/Components/ChildActorComponent.cpp
--- a/Engine/Source/Runtime/Engine/Private/Components/ChildActorComponent.cpp
+++ b/Engine/Source/Runtime/Engine/Private/Components/ChildActorComponent.cpp
@@ -573,6 +573,13 @@ void UChildActorComponent::SetChildActorClass(TSubclassOf<AActor> Class, AActor*
 				// Now set the actual name and outer to the BPGC.
 				const FString TemplateName = FString::Printf(TEXT("%s_%s_CAT"), *GetName(), *ChildActorClass->GetName());
 
+				// If the CAC's package is ever reloaded, there's a chance that any serialized template with the same class might've also been reloaded.
+				// In order to avoid any collisions, we'll move the reloaded one to the transient package.
+				if (AActor* ExistingChildActorTemplate = FindObject<AActor>(this, *TemplateName, EFindObjectFlags::None))
+				{
+					ExistingChildActorTemplate->Rename(nullptr, GetTransientPackage(), REN_DontCreateRedirectors);
+				}
+
 				ChildActorTemplate->Rename(*TemplateName, this, REN_DoNotDirty | REN_DontCreateRedirectors);
 			}
 		}
