New material translator. Ported over SceneTexture sampling  optimization when the scene texture is referenced in the code of a custom expression.

#rnx

diff --git a/Engine/Source/Runtime/Engine/Private/Materials/HLSLMaterialTranslator.cpp b/Engine/Source/Runtime/Engine/Private/Materials/HLSLMaterialTranslator.cpp
--- a/Engine/Source/Runtime/Engine/Private/Materials/HLSLMaterialTranslator.cpp
+++ b/Engine/Source/Runtime/Engine/Private/Materials/HLSLMaterialTranslator.cpp
@@ -15230,352 +15230,6 @@ int32 FHLSLMaterialTranslator::AccessMaterialAttribute(int32 CodeIndex, const FG
 		*AttributeName);
 }
 
-enum class AsciiFlags
-{
-	TerminatorOrSlash = (1 << 0),	// Null terminator OR slash
-	Whitespace = (1 << 1),			// Includes other special characters below 32 (in addition to tab / newline)
-	Other = (1 << 2),				// Anything else not one of the other types
-	SymbolStart = (1 << 3),			// Letters plus underscore (anything that can start a symbol)
-	Digit = (1 << 4),
-	Dot = (1 << 5),
-	Quote = (1 << 6),
-	Hash = (1 << 7),
-};
-
-static uint8 AsciiFlagTable[256] =
-{
-	1,2,2,2,2,2,2,2, 2,2,2,2,2,2,2,2, 2,2,2,2,2,2,2,2, 2,2,2,2,2,2,2,2,		// Treat all special characters as whitespace
-
-	2,4,64,128,4,4,4,4,			// 34 == Quote  35 == Hash
-	4,4,4,4,4,4,32,1,			// 46 == Dot    47 == Slash
-	16,16,16,16,16,16,16,16,	// Digits 0-7
-	16,16,4,4,4,4,4,4,			// Digits 8-9
-
-	4,8,8,8,8,8,8,8, 8,8,8,8,8,8,8,8, 8,8,8,8,8,8,8,8, 8,8,8,4,4,4,4,8,		// Upper case letters,  95 == Underscore
-	4,8,8,8,8,8,8,8, 8,8,8,8,8,8,8,8, 8,8,8,8,8,8,8,8, 8,8,8,4,4,4,4,4,		// Lower case letters
-
-	4,4,4,4,4,4,4,4, 4,4,4,4,4,4,4,4, 4,4,4,4,4,4,4,4, 4,4,4,4,4,4,4,4,		// Treat all non-ASCII characters as Other
-	4,4,4,4,4,4,4,4, 4,4,4,4,4,4,4,4, 4,4,4,4,4,4,4,4, 4,4,4,4,4,4,4,4,
-	4,4,4,4,4,4,4,4, 4,4,4,4,4,4,4,4, 4,4,4,4,4,4,4,4, 4,4,4,4,4,4,4,4,
-	4,4,4,4,4,4,4,4, 4,4,4,4,4,4,4,4, 4,4,4,4,4,4,4,4, 4,4,4,4,4,4,4,4,
-};
-
-// Utility function to search for the given HLSL Identifier in a block of Code.  Takes into account comments and quoted text
-// that should be ignored when searching for identifiers.  Returns the start of the identifier, and the input pointer "Code"
-// will be advanced to the first non-whitespace, non-comment character after the end of the identifier.  Or "null" will be
-// returned if the identifier isn't found.  OutCompoundTokens contain additional tokens if it's part of a compound identifier
-// separated by periods (for "Identifier.x", OutCompoundTokens would include "x").
-static const TCHAR* FindHlslIdentifierInCode(const TCHAR*& Code, const TCHAR* Identifier, TArray<FStringView>& OutCompoundTokens)
-{
-	FStringView RootIdentifier;
-	OutCompoundTokens.Empty();
-
-	#define FETCH_ASCII_FLAG_TABLE(CHARACTER) AsciiFlagTable[(uint8)FMath::Min(CHARACTER, (TCHAR)UINT8_MAX)]
-
-	const TCHAR* SearchChar = Code;
-	uint8 SearchCharFlag = FETCH_ASCII_FLAG_TABLE(*SearchChar);
-	bool bInCompoundIdentifier = false;
-
-	// Scanning loop
-	while (1)
-	{
-		static constexpr uint8 AsciiFlagsSymbol = (uint8)AsciiFlags::SymbolStart | (uint8)AsciiFlags::Digit;
-		static constexpr uint8 AsciiFlagsStartNumberOrDirective = (uint8)AsciiFlags::Digit | (uint8)AsciiFlags::Dot | (uint8)AsciiFlags::Hash;
-		static constexpr uint8 AsciiFlagsEndNumberOrDirective = (uint8)AsciiFlags::Whitespace | (uint8)AsciiFlags::Other | (uint8)AsciiFlags::Quote | (uint8)AsciiFlags::TerminatorOrSlash;
-
-		if (SearchCharFlag & (uint8)AsciiFlags::Whitespace)
-		{
-			// Always skip over whitespace
-			SearchChar++;
-			SearchCharFlag = FETCH_ASCII_FLAG_TABLE(*SearchChar);
-		}
-		else if (SearchChar[0] == '/' && (SearchChar[1] == '/' || SearchChar[1] == '*'))
-		{
-			// Always skip over comments
-			if (SearchChar[1] == '/')
-			{
-				SearchChar += 2;
-				while (*SearchChar && *SearchChar != '\r' && *SearchChar != '\n')
-				{
-					SearchChar++;
-				}
-
-				// Could be end of string, if not, skip over newline
-				if (*SearchChar)
-				{
-					SearchChar++;
-				}
-			}
-			else
-			{
-				SearchChar += 2;
-				while (*SearchChar && (SearchChar[0] != '*' || SearchChar[1] != '/'))
-				{
-					SearchChar++;
-				}
-
-				// Could be end of string, if not, skip over comment terminator
-				if (*SearchChar)
-				{
-					SearchChar += 2;
-				}
-			}
-			SearchCharFlag = FETCH_ASCII_FLAG_TABLE(*SearchChar);
-		}
-		else if (!RootIdentifier.IsEmpty() && (SearchCharFlag & ((uint8)AsciiFlags::Dot)))
-		{
-			// Dot after an identifier starts a compound identifier
-			bInCompoundIdentifier = true;
-			SearchChar++;
-			SearchCharFlag = FETCH_ASCII_FLAG_TABLE(*SearchChar);
-		}
-		else if (!RootIdentifier.IsEmpty() && !bInCompoundIdentifier)
-		{
-			// Any other character besides whitespace, comment, or dot terminates an identifier.  Check if it matches what we are searching for.
-			if (RootIdentifier == Identifier)
-			{
-				Code = SearchChar;
-				return RootIdentifier.GetData();
-			}
-
-			// Clear the state
-			RootIdentifier = FStringView();
-			OutCompoundTokens.Empty();
-
-			// Don't advance character -- restart loop, no longer tracking an identifier
-		}
-		else if (SearchCharFlag & (uint8)AsciiFlags::SymbolStart)
-		{
-			// Start of an identifier, determine the bounds of the identifier
-			const TCHAR* IdentifierStart = SearchChar;
-
-			SearchChar++;
-			while ((SearchCharFlag = FETCH_ASCII_FLAG_TABLE(*SearchChar)) & AsciiFlagsSymbol)
-			{
-				SearchChar++;
-			}
-
-			FStringView IdentifierView(IdentifierStart, SearchChar - IdentifierStart);
-
-			// Add identifier as a compound identifier or the root identifier
-			if (bInCompoundIdentifier)
-			{
-				OutCompoundTokens.Add(IdentifierView);
-				bInCompoundIdentifier = false;
-			}
-			else
-			{
-				check(RootIdentifier.IsEmpty());
-				RootIdentifier = IdentifierView;
-			}
-		}
-		else if (!RootIdentifier.IsEmpty())
-		{
-			// If we reach this point, we are in an invalid syntax compound identifier, where an identifier token and dot are not followed
-			// by another identifier.  Clear the state and restart parsing.
-			bInCompoundIdentifier = false;
-			RootIdentifier = FStringView();
-			OutCompoundTokens.Empty();
-
-			// Don't advance character -- restart loop, no longer tracking an identifier.
-		}
-		else if ((SearchCharFlag & (uint8)AsciiFlags::Other) || (*SearchChar == '/'))
-		{
-			// Not in an identifier or comment -- skip over various characters with no special handling
-			SearchChar++;
-			SearchCharFlag = FETCH_ASCII_FLAG_TABLE(*SearchChar);
-		}
-		else if (SearchCharFlag & AsciiFlagsStartNumberOrDirective)
-		{
-			// Number or directive, skip to Whitespace, Other, or Quote (numbers may contain letters or #, i.e. "1.#INF" for infinity, or "e" for an exponent)
-			SearchChar++;
-			while (!((SearchCharFlag = FETCH_ASCII_FLAG_TABLE(*SearchChar)) & AsciiFlagsEndNumberOrDirective))
-			{
-				SearchChar++;
-			}
-		}
-		else if (SearchCharFlag & (uint8)AsciiFlags::Quote)
-		{
-			// Quote, skip to next Quote (or maybe end of string if text is malformed), ignoring the quote if it's escaped
-			SearchChar++;
-			while (*SearchChar && (*SearchChar != '\"' || *(SearchChar - 1) == '\\'))
-			{
-				SearchChar++;
-			}
-
-			// Could be end of string or the quote -- skip over the quote if not the null terminator
-			if (*SearchChar)
-			{
-				SearchChar++;
-			}
-			SearchCharFlag = FETCH_ASCII_FLAG_TABLE(*SearchChar);
-		}
-		// Must be null terminator -- we've tested all other possibilities
-		else
-		{
-			// End of string
-			Code = SearchChar;
-			OutCompoundTokens.Empty();
-			return nullptr;
-		}
-	}
-
-	#undef FETCH_ASCII_FLAG_TABLE
-}
-
-// Fixup for SceneTexture and UserSceneTexture inputs to custom HLSL.  Returns a new string, rather than modifying a Code FString
-// in place, to allow the code to be shared with the HLSL tree code path (MaterialExpressionHLSL.cpp), which uses TStringBuilder.
-// An empty string is returned if no fixup was required (the common case), avoiding reallocation.  OutSceneTextureInfo tracks
-// which inputs are SceneTexture or UserSceneTexture (left empty if no scene textures).  A -1 indicates a scene texture input not
-// directly used as a symbol in the custom HLSL, allowing the fetch to be dead stripped, while 1 is an actually used input.
-//
-// The goal is to allow the scene texture ID to be automatically generated in the code, without the user needing to manually
-// insert the correct identifier based on settings in the connected node.  Without this fixup, the user needs to change the
-// custom HLSL any time the input node changes, which is tedious and error prone.  For SceneTexture, if they (for example) change
-// SceneColor to Depth, they need to change PPI_SceneColor to PPI_Depth.  For UserSceneTexture, if the name is changed from Input
-// to NewInput, they need to change PPIUser_Input to PPIUser_NewInput.  Similar manual changes are also required if a completely
-// different SceneTexture node is linked.  Besides substituting the ID, minimalist shorthand ".Fetch" function call style syntax
-// is provided for offset fetches (accepts either float2 or two floats).  The original syntax of referencing "Input" without the
-// ".ID" or ".Fetch" suffixes still works for fetching a non-offset sample.
-//
-//		Input.ID -> Replace text with ID enum of connected input (PPI_* or PPIUser_*)
-//			Example usage:	SceneTextureFetch(Input.ID, float2(0,0))
-//							SceneTextureFetch(/*Input.ID=*/ PPI_SceneColor, float2(0,0))
-//
-//		Input.Fetch(*) -> Replace with SceneTextureFetchFunc(Parameters, ID, *)
-//			Example usage:	Input.Fetch(1,-1)
-//							SceneTextureFetchFunc(Parameters, /*Input=*/ PPI_SceneColor, 1,-1)
-//
-// Note that this text replacement happens before preprocessor macro substitution, and so it's not possible to use the shorthand
-// syntax with the user specified "InputName" as a macro argument, but you can pass "InputName.ID", for example:
-//
-//		// Shorthand syntax will not work, because substitution of "InputName" as the "Input" macro argument happens later in preprocess
-//		#define WEIGHTED_FETCH_SCENE_TEXTURE(Input, SampleIndex) \
-//			(Input.Fetch(Offsets[SampleIndex], 0) * Weights[SampleIndex])
-//
-//		Color += WEIGHTED_FETCH_SCENE_TEXTURE(InputName, SampleIndex);
-//
-//		// But the syntax that accepts a scene texture ID works fine
-//		#define WEIGHTED_FETCH_SCENE_TEXTURE(InputID, SampleIndex) \
-//			(SceneTextureFetch(InputID, float2(Offsets[SampleIndex], 0)) * Weights[SampleIndex])
-//
-//		Color += WEIGHTED_FETCH_SCENE_TEXTURE(InputName.ID, SampleIndex);
-//
-//		// This also works fine, because InputName is inside the macro, and not a macro argument
-//		#define WEIGHTED_FETCH_SCENE_TEXTURE(SampleIndex) \
-//			(InputName.Fetch(Offsets[SampleIndex], 0) * Weights[SampleIndex])
-//
-FString CustomExpressionSceneTextureInputFixup(const UMaterialExpressionCustom* Custom, const TCHAR* Code, TArray<int8>& OutSceneTextureInfo)
-{
-	// Allocated when required fixup is first encountered
-	FString ModifiedCode;
-
-	for (int32 InputIndex = 0; InputIndex < Custom->Inputs.Num(); InputIndex++)
-	{
-		const FCustomInput& Input = Custom->Inputs[InputIndex];
-
-		if (Input.InputName.IsNone())
-		{
-			continue;
-		}
-
-		FExpressionInput ExpressionInput = Input.Input.GetTracedInput();
-		const UMaterialExpression* Expression = ExpressionInput.Expression;
-		if (!Expression || ExpressionInput.OutputIndex != 0 || !(Expression->IsA<UMaterialExpressionSceneTexture>() || Expression->IsA<UMaterialExpressionUserSceneTexture>()))
-		{
-			continue;
-		}
-
-		FString InputNameString = Input.InputName.ToString();
-		FString InputIDString;
-
-		if (const UMaterialExpressionSceneTexture* SceneTextureExpression = Cast<UMaterialExpressionSceneTexture>(Expression))
-		{
-			InputIDString = StaticEnum<ESceneTextureId>()->GetNameStringByValue(static_cast<int>(SceneTextureExpression->SceneTextureId));
-		}
-		else
-		{
-			const UMaterialExpressionUserSceneTexture* UserSceneTextureExpression = Cast<UMaterialExpressionUserSceneTexture>(Expression);
-			InputIDString = TEXT("PPIUser_") + UserSceneTextureExpression->UserSceneTexture.ToString();
-		}
-
-		// We've encountered a scene texture, allocate our info array
-		if (!OutSceneTextureInfo.Num())
-		{
-			OutSceneTextureInfo.SetNumZeroed(Custom->Inputs.Num());
-		}
-
-		// Mark with -1 to indicate that it's a scene texture input, but might not be directly used in the source.
-		OutSceneTextureInfo[InputIndex] = -1;
-
-		int32 CodeOffset = 0;
-		while (1)
-		{
-			TArray<FStringView> CompoundTokens;
-			const TCHAR* CodeStart = ModifiedCode.IsEmpty() ? Code : *ModifiedCode;
-			const TCHAR* CodeString = &CodeStart[CodeOffset];
-			const TCHAR* FoundInputName = FindHlslIdentifierInCode(CodeString, *InputNameString, CompoundTokens);
-
-			if (!FoundInputName)
-			{
-				break;
-			}
-
-			CodeOffset = FoundInputName - CodeStart;
-
-			int32 OriginalCodeOffset = CodeOffset;
-			if (CompoundTokens.Num() == 1)
-			{
-				// Check if this is using the "Fetch" function call shorthand syntax.  If so, replace it with the corresponding fetch function.
-				if (CompoundTokens[0] == TEXT("Fetch") && *CodeString == '(')
-				{
-					FString ReplacementText = FString::Printf(TEXT("SceneTextureFetchFunc(Parameters, /*%s=*/ %s, "), *InputNameString, *InputIDString);
-
-					if (ModifiedCode.IsEmpty())
-					{
-						ModifiedCode.Reserve(FCString::Strlen(Code) + ReplacementText.Len());
-						ModifiedCode = Code;
-					}
-
-					// Add one character to amount removed, as we are overwriting the open parentheses as well
-					ModifiedCode.RemoveAt(CodeOffset, CodeString - FoundInputName + 1, EAllowShrinking::No);
-					ModifiedCode.InsertAt(CodeOffset, ReplacementText);
-
-					CodeOffset += ReplacementText.Len();
-				}
-				else if (CompoundTokens[0] == TEXT("ID"))
-				{
-					// Replace "Input.ID" identifier with ID (including comment for readability if there are compile errors)
-					FString ReplacementText = FString::Printf(TEXT("/*%s.ID=*/ %s"), *InputNameString, *InputIDString);
-
-					if (ModifiedCode.IsEmpty())
-					{
-						ModifiedCode.Reserve(FCString::Strlen(Code) + ReplacementText.Len());
-						ModifiedCode = Code;
-					}
-
-					ModifiedCode.RemoveAt(CodeOffset, CompoundTokens[0].GetData() + CompoundTokens[0].Len() - FoundInputName, EAllowShrinking::No);
-					ModifiedCode.InsertAt(CodeOffset, ReplacementText);
-
-					CodeOffset += ReplacementText.Len();
-				}
-			}
-
-			// If we didn't insert and skip over replacement code, skip over the token that was found (CodeString will have
-			// been set past the end of the token by the function "FindHlslIdentifierInCode").
-			if (CodeOffset == OriginalCodeOffset)
-			{
-				CodeOffset = CodeString - CodeStart;
-
-				// Scene Texture input pin symbol was directly used by the custom HLSL, mark it with a 1, as the input pin expression needs to be compiled in
-				OutSceneTextureInfo[InputIndex] = 1;
-			}
-		}
-	}
-
-	return ModifiedCode;
-}
-
 int32 FHLSLMaterialTranslator::CustomExpression( class UMaterialExpressionCustom* Custom, int32 OutputIndex, TArray<int32>& CompiledInputs )
 {
 	const FMaterialCustomExpressionEntry* CustomEntry = nullptr;
@@ -15782,7 +15436,7 @@ int32 FHLSLMaterialTranslator::CustomExpression( class UMaterialExpressionCustom
 		Code.ReplaceInline(TEXT("\r\n"), TEXT("\n"), ESearchCase::CaseSensitive);
 
 		TArray<int8> SceneTextureInfoIgnored;
-		FString ModifiedCode = CustomExpressionSceneTextureInputFixup(Custom, *Code, SceneTextureInfoIgnored);
+		FString ModifiedCode =  UE::MaterialTranslatorUtils::CustomExpressionSceneTextureInputFixup(Custom, *Code, SceneTextureInfoIgnored);
 		if (!ModifiedCode.IsEmpty())
 		{
 			Code = ModifiedCode;
diff --git a/Engine/Source/Runtime/Engine/Private/Materials/MaterialExpressions.cpp b/Engine/Source/Runtime/Engine/Private/Materials/MaterialExpressions.cpp
--- a/Engine/Source/Runtime/Engine/Private/Materials/MaterialExpressions.cpp
+++ b/Engine/Source/Runtime/Engine/Private/Materials/MaterialExpressions.cpp
@@ -12866,7 +12866,6 @@ UMaterialExpressionCustom::UMaterialExpressionCustom(const FObjectInitializer& O
 }
 
 #if WITH_EDITOR
-extern FString CustomExpressionSceneTextureInputFixup(const UMaterialExpressionCustom* Custom, const TCHAR* Code, TArray<int8>& OutSceneTextureInfo);
 
 int32 UMaterialExpressionCustom::Compile(class FMaterialCompiler* Compiler, int32 OutputIndex)
 {
@@ -12877,7 +12876,7 @@ int32 UMaterialExpressionCustom::Compile(class FMaterialCompiler* Compiler, int3
 	// in the custom HLSL, rather than using the input pin value.  The fixup function has a parser that is aware of HLSL syntax, and able to tokenize out
 	// identifiers, and handle symbol sequences (such as Input.ID or Input.Fetch) that will be substituted into SceneTextureFetch calls.
 	TArray<int8> SceneTextureInfo;
-	CustomExpressionSceneTextureInputFixup(this, *Code, SceneTextureInfo);
+	UE::MaterialTranslatorUtils::CustomExpressionSceneTextureInputFixup(this, *Code, SceneTextureInfo);
 
 	for( int32 i=0;i<Inputs.Num();i++ )
 	{
diff --git a/Engine/Source/Runtime/Engine/Private/Materials/MaterialExpressionsIR.cpp b/Engine/Source/Runtime/Engine/Private/Materials/MaterialExpressionsIR.cpp
--- a/Engine/Source/Runtime/Engine/Private/Materials/MaterialExpressionsIR.cpp
+++ b/Engine/Source/Runtime/Engine/Private/Materials/MaterialExpressionsIR.cpp
@@ -3316,20 +3316,38 @@ void UMaterialExpressionCustom::Build(MIR::FEmitter& Em)
 	MIR::FFunctionHLSLDesc FuncDesc;
 	FuncDesc.Name = Name;
 	FuncDesc.ReturnType = CustomMaterialOutputTypeToMIR(OutputType);
-	FuncDesc.Code = Code;
+	
+	// Fixup the scene texture identifiers in the source string
+	TArray<int8> SceneTextureInfo;
+	FString FixedCode = UE::MaterialTranslatorUtils::CustomExpressionSceneTextureInputFixup(this, *Code, SceneTextureInfo);
+	FuncDesc.Code = FixedCode;
+	if (FuncDesc.Code.IsEmpty())
+	{
+		FuncDesc.Code = Code;
+	}
 
 	// Turn each expression input into an input-only parameter.
-	for (FCustomInput const& Input : Inputs)
+	for (int32 i = 0; i < Inputs.Num(); ++i)
 	{
-		// Ignore input parameters with "None" name.
+		FCustomInput const& Input = Inputs[i];
+
 		if (Input.InputName.IsNone())
 		{
+			// Ignore this input parameters with "None" name.
 			continue;
 		}
-
+		
 		// Read the input argument
 		InputArgs[FuncDesc.NumInputOnlyParams] = Em.Input(&Input.Input);
 
+		// Is this argument an unused scene texture sample?
+		if (SceneTextureInfo.IsValidIndex(i) && SceneTextureInfo[i] == -1)
+		{
+			// I this parameter samples an unused scene texture, skip the parameter, but still
+			// make sure the scene-texture sample is analyzed.
+			InputArgs[FuncDesc.NumInputOnlyParams] = Em.Nop(InputArgs[FuncDesc.NumInputOnlyParams]);
+		}
+		
 		// Declare an input-only parameter
 		if (!FuncDesc.PushInputOnlyParameter(Input.InputName, InputArgs[FuncDesc.NumInputOnlyParams]->Type))
 		{
diff --git a/Engine/Source/Runtime/Engine/Private/Materials/MaterialIR.cpp b/Engine/Source/Runtime/Engine/Private/Materials/MaterialIR.cpp
--- a/Engine/Source/Runtime/Engine/Private/Materials/MaterialIR.cpp
+++ b/Engine/Source/Runtime/Engine/Private/Materials/MaterialIR.cpp
@@ -60,6 +60,7 @@ const TCHAR* LexToString(EValueKind Kind)
 		case VK_PromoteSubstrateParameter: return TEXT("PromoteSubstrateParameter");
 		case VK_StageSwitch: return TEXT("StageSwitch");
 		case VK_HardwarePartialDerivative: return TEXT("HardwarePartialDerivative");
+		case VK_Nop: return TEXT("Nop");
 		case VK_Call: return TEXT("Call");
 		case VK_CallParameterOutput: return TEXT("CallParameterOutput");
 		/* invalid entries */
@@ -124,6 +125,7 @@ uint32 FValue::GetSizeInBytes() const
 		case VK_PromoteSubstrateParameter: return sizeof(FPromoteSubstrateParameter);
 		case VK_StageSwitch: return sizeof(FStageSwitch);
 		case VK_HardwarePartialDerivative: return sizeof(FHardwarePartialDerivative);
+		case VK_Nop: return sizeof(FNop);
 		case VK_Call: return sizeof(FCall);
 		case VK_CallParameterOutput: return sizeof(FCallParameterOutput);
 			/* invalid entries */
@@ -216,11 +218,16 @@ TConstArrayView<FValue*> FValue::GetUses() const
 			return { &This->Arg, FHardwarePartialDerivative::NumStaticUses };
 		}
 
+		case VK_Nop:
+		{
+			auto This = static_cast<const FNop*>(this);
+			return { &This->Arg, FNop::NumStaticUses };
+		}
+
 		case VK_Call:
 		{
 			auto This = static_cast<const FCall*>(this);
 			return { This->Arguments, This->NumArguments };
-
 		}
 
 		case VK_CallParameterOutput:
diff --git a/Engine/Source/Runtime/Engine/Private/Materials/MaterialIREmitter.cpp b/Engine/Source/Runtime/Engine/Private/Materials/MaterialIREmitter.cpp
--- a/Engine/Source/Runtime/Engine/Private/Materials/MaterialIREmitter.cpp
+++ b/Engine/Source/Runtime/Engine/Private/Materials/MaterialIREmitter.cpp
@@ -3121,6 +3121,22 @@ FValueRef FEmitter::StageSwitch(FType Type, TConstArrayView<FValueRef> ValuePerS
 	return EmitPrototype(*this, Prototype);
 }
 
+FValueRef FEmitter::Nop(FValueRef Arg)
+{
+	// Nop can only have primitive arguments
+	Arg = CheckIsPrimitive(Arg);
+
+	if (!Arg.IsValid())
+	{
+		return Arg;
+	}
+
+	MIR::FNop Prototype = MakePrototype<MIR::FNop>(Arg->Type);
+	Prototype.Arg = Arg;
+
+	return EmitPrototype(*this, Prototype);
+}
+
 // Maps a sampler type from standard texture (ST) to virtual texture (VT).
 static EMaterialSamplerType PromoteSamplerTypeFromSTtoVT(EMaterialSamplerType InSamplerType)
 {
@@ -3847,7 +3863,7 @@ const FFunction* FEmitter::FunctionHLSL(const FFunctionHLSLDesc& Desc)
 	// Create the new MIR HLSL function instance and set it up
 	FFunctionHLSL* Function = new (FPrivate::Allocate(*this, sizeof(FFunctionHLSL), alignof(FFunctionHLSL))) FFunctionHLSL { Prototype };
 	Function->Name = Module->InternString(Function->Name);
-	Function->Code = Desc.Code;
+	Function->Code = Module->InternString(Desc.Code);
 	Function->Defines = MakeArrayCopy(*this, Desc.Defines);
 	Function->Includes = MakeArrayCopy(*this, Desc.Includes);
 
diff --git a/Engine/Source/Runtime/Engine/Private/Materials/MaterialIRModuleBuilder.cpp b/Engine/Source/Runtime/Engine/Private/Materials/MaterialIRModuleBuilder.cpp
--- a/Engine/Source/Runtime/Engine/Private/Materials/MaterialIRModuleBuilder.cpp
+++ b/Engine/Source/Runtime/Engine/Private/Materials/MaterialIRModuleBuilder.cpp
@@ -665,6 +665,12 @@ struct FMaterialIRModuleBuilderImpl
 				Instr->Next[StageIndex] = Instr->Block[StageIndex] ->Instructions;
 				Instr->Block[StageIndex]->Instructions = Instr;
 
+				//
+				if (Instr->As<MIR::FNop>())
+				{
+					continue;
+				}
+
 				TConstArrayView<MIR::FValue*> Uses = Instr->GetUsesForStage((MIR::EStage)StageIndex);
 				for (int32 UseIndex = 0; UseIndex < Uses.Num(); ++UseIndex)
 				{
diff --git a/Engine/Source/Runtime/Engine/Private/Materials/MaterialIRToHLSLTranslator.cpp b/Engine/Source/Runtime/Engine/Private/Materials/MaterialIRToHLSLTranslator.cpp
--- a/Engine/Source/Runtime/Engine/Private/Materials/MaterialIRToHLSLTranslator.cpp
+++ b/Engine/Source/Runtime/Engine/Private/Materials/MaterialIRToHLSLTranslator.cpp
@@ -433,6 +433,7 @@ struct FPrivate : FMaterialIRToHLSLTranslation
 			case MIR::VK_PromoteSubstrateParameter: LowerPromoteSubstrateParameter(static_cast<const MIR::FPromoteSubstrateParameter*>(Instr)); break;
 			case MIR::VK_StageSwitch: LowerStageSwitch(static_cast<const MIR::FStageSwitch*>(Instr)); break;
 			case MIR::VK_HardwarePartialDerivative: LowerHardwarePartialDerivative(static_cast<const MIR::FHardwarePartialDerivative*>(Instr)); break;
+			case MIR::VK_Nop: LowerNop(static_cast<const MIR::FNop*>(Instr)); break;
 			case MIR::VK_Call: LowerCall(static_cast<const MIR::FCall*>(Instr)); break;
 			case MIR::VK_CallParameterOutput: LowerCallOutput(static_cast<const MIR::FCallParameterOutput*>(Instr)); break;
 			default: UE_MIR_UNREACHABLE();
@@ -443,10 +444,10 @@ struct FPrivate : FMaterialIRToHLSLTranslation
 
 	void LowerConstant(const MIR::FConstant* Constant)
 	{
-		TOptional<MIR::FPrimitive> checkNoEntry = Constant->Type.AsPrimitive();
-		check(checkNoEntry && checkNoEntry->IsScalar());
+		TOptional<MIR::FPrimitive> Primitive = Constant->Type.AsPrimitive();
+		check(Primitive && Primitive->IsScalar());
 
-		switch (checkNoEntry->ScalarKind)
+		switch (Primitive->ScalarKind)
 		{
 			case MIR::EScalarKind::Bool:
 				Printer << (Constant->Boolean ? TEXTVIEW("true") : TEXTVIEW("false"));
@@ -1845,6 +1846,14 @@ struct FPrivate : FMaterialIRToHLSLTranslation
 		}
 	}
 
+	void LowerNop(const MIR::FNop* Nop)
+	{
+		// NOP instructions are only used to analyze their argument, but have no effect, thus we compile it to a default
+		// value based on its type.
+		checkf(!Nop->Type.IsDouble(), TEXT("NOPs do not support LWC primitive type yet"));
+		Printer << TEXTVIEW("((") << LowerType(Nop->Type) << TEXTVIEW(")0)");
+	}
+
 	void LowerCall(const MIR::FCall* Call)
 	{
 		int32 ParamLocal = NumLocals;
diff --git a/Engine/Source/Runtime/Engine/Private/Materials/MaterialShared.cpp b/Engine/Source/Runtime/Engine/Private/Materials/MaterialShared.cpp
--- a/Engine/Source/Runtime/Engine/Private/Materials/MaterialShared.cpp
+++ b/Engine/Source/Runtime/Engine/Private/Materials/MaterialShared.cpp
@@ -70,6 +70,7 @@
 #include "Materials/MaterialIRToHLSLTranslator.h"
 #include "Materials/MaterialSourceTemplate.h"
 #include "Materials/MaterialInsights.h"
+#include "Materials/MaterialExpressionCustom.h"
 
 #include UE_INLINE_GENERATED_CPP_BY_NAME(MaterialShared)
 
@@ -6138,6 +6139,347 @@ void FinalCompileValidation(
 	}
 }
 
+enum class AsciiFlags
+{
+	TerminatorOrSlash = (1 << 0),	// Null terminator OR slash
+	Whitespace = (1 << 1),			// Includes other special characters below 32 (in addition to tab / newline)
+	Other = (1 << 2),				// Anything else not one of the other types
+	SymbolStart = (1 << 3),			// Letters plus underscore (anything that can start a symbol)
+	Digit = (1 << 4),
+	Dot = (1 << 5),
+	Quote = (1 << 6),
+	Hash = (1 << 7),
+};
+
+static uint8 AsciiFlagTable[256] =
+{
+	1,2,2,2,2,2,2,2, 2,2,2,2,2,2,2,2, 2,2,2,2,2,2,2,2, 2,2,2,2,2,2,2,2,		// Treat all special characters as whitespace
+
+	2,4,64,128,4,4,4,4,			// 34 == Quote  35 == Hash
+	4,4,4,4,4,4,32,1,			// 46 == Dot    47 == Slash
+	16,16,16,16,16,16,16,16,	// Digits 0-7
+	16,16,4,4,4,4,4,4,			// Digits 8-9
+
+	4,8,8,8,8,8,8,8, 8,8,8,8,8,8,8,8, 8,8,8,8,8,8,8,8, 8,8,8,4,4,4,4,8,		// Upper case letters,  95 == Underscore
+	4,8,8,8,8,8,8,8, 8,8,8,8,8,8,8,8, 8,8,8,8,8,8,8,8, 8,8,8,4,4,4,4,4,		// Lower case letters
+
+	4,4,4,4,4,4,4,4, 4,4,4,4,4,4,4,4, 4,4,4,4,4,4,4,4, 4,4,4,4,4,4,4,4,		// Treat all non-ASCII characters as Other
+	4,4,4,4,4,4,4,4, 4,4,4,4,4,4,4,4, 4,4,4,4,4,4,4,4, 4,4,4,4,4,4,4,4,
+	4,4,4,4,4,4,4,4, 4,4,4,4,4,4,4,4, 4,4,4,4,4,4,4,4, 4,4,4,4,4,4,4,4,
+	4,4,4,4,4,4,4,4, 4,4,4,4,4,4,4,4, 4,4,4,4,4,4,4,4, 4,4,4,4,4,4,4,4,
+};
+
+const TCHAR* FindHlslIdentifierInCode(const TCHAR*& Code, const TCHAR* Identifier, TArray<FStringView>& OutCompoundTokens)
+{
+	FStringView RootIdentifier;
+	OutCompoundTokens.Empty();
+
+	#define FETCH_ASCII_FLAG_TABLE(CHARACTER) AsciiFlagTable[(uint8)FMath::Min(CHARACTER, (TCHAR)UINT8_MAX)]
+
+	const TCHAR* SearchChar = Code;
+	uint8 SearchCharFlag = FETCH_ASCII_FLAG_TABLE(*SearchChar);
+	bool bInCompoundIdentifier = false;
+
+	// Scanning loop
+	while (1)
+	{
+		static constexpr uint8 AsciiFlagsSymbol = (uint8)AsciiFlags::SymbolStart | (uint8)AsciiFlags::Digit;
+		static constexpr uint8 AsciiFlagsStartNumberOrDirective = (uint8)AsciiFlags::Digit | (uint8)AsciiFlags::Dot | (uint8)AsciiFlags::Hash;
+		static constexpr uint8 AsciiFlagsEndNumberOrDirective = (uint8)AsciiFlags::Whitespace | (uint8)AsciiFlags::Other | (uint8)AsciiFlags::Quote | (uint8)AsciiFlags::TerminatorOrSlash;
+
+		if (SearchCharFlag & (uint8)AsciiFlags::Whitespace)
+		{
+			// Always skip over whitespace
+			SearchChar++;
+			SearchCharFlag = FETCH_ASCII_FLAG_TABLE(*SearchChar);
+		}
+		else if (SearchChar[0] == '/' && (SearchChar[1] == '/' || SearchChar[1] == '*'))
+		{
+			// Always skip over comments
+			if (SearchChar[1] == '/')
+			{
+				SearchChar += 2;
+				while (*SearchChar && *SearchChar != '\r' && *SearchChar != '\n')
+				{
+					SearchChar++;
+				}
+
+				// Could be end of string, if not, skip over newline
+				if (*SearchChar)
+				{
+					SearchChar++;
+				}
+			}
+			else
+			{
+				SearchChar += 2;
+				while (*SearchChar && (SearchChar[0] != '*' || SearchChar[1] != '/'))
+				{
+					SearchChar++;
+				}
+
+				// Could be end of string, if not, skip over comment terminator
+				if (*SearchChar)
+				{
+					SearchChar += 2;
+				}
+			}
+			SearchCharFlag = FETCH_ASCII_FLAG_TABLE(*SearchChar);
+		}
+		else if (!RootIdentifier.IsEmpty() && (SearchCharFlag & ((uint8)AsciiFlags::Dot)))
+		{
+			// Dot after an identifier starts a compound identifier
+			bInCompoundIdentifier = true;
+			SearchChar++;
+			SearchCharFlag = FETCH_ASCII_FLAG_TABLE(*SearchChar);
+		}
+		else if (!RootIdentifier.IsEmpty() && !bInCompoundIdentifier)
+		{
+			// Any other character besides whitespace, comment, or dot terminates an identifier.  Check if it matches what we are searching for.
+			if (RootIdentifier == Identifier)
+			{
+				Code = SearchChar;
+				return RootIdentifier.GetData();
+			}
+
+			// Clear the state
+			RootIdentifier = FStringView();
+			OutCompoundTokens.Empty();
+
+			// Don't advance character -- restart loop, no longer tracking an identifier
+		}
+		else if (SearchCharFlag & (uint8)AsciiFlags::SymbolStart)
+		{
+			// Start of an identifier, determine the bounds of the identifier
+			const TCHAR* IdentifierStart = SearchChar;
+
+			SearchChar++;
+			while ((SearchCharFlag = FETCH_ASCII_FLAG_TABLE(*SearchChar)) & AsciiFlagsSymbol)
+			{
+				SearchChar++;
+			}
+
+			FStringView IdentifierView(IdentifierStart, SearchChar - IdentifierStart);
+
+			// Add identifier as a compound identifier or the root identifier
+			if (bInCompoundIdentifier)
+			{
+				OutCompoundTokens.Add(IdentifierView);
+				bInCompoundIdentifier = false;
+			}
+			else
+			{
+				check(RootIdentifier.IsEmpty());
+				RootIdentifier = IdentifierView;
+			}
+		}
+		else if (!RootIdentifier.IsEmpty())
+		{
+			// If we reach this point, we are in an invalid syntax compound identifier, where an identifier token and dot are not followed
+			// by another identifier.  Clear the state and restart parsing.
+			bInCompoundIdentifier = false;
+			RootIdentifier = FStringView();
+			OutCompoundTokens.Empty();
+
+			// Don't advance character -- restart loop, no longer tracking an identifier.
+		}
+		else if ((SearchCharFlag & (uint8)AsciiFlags::Other) || (*SearchChar == '/'))
+		{
+			// Not in an identifier or comment -- skip over various characters with no special handling
+			SearchChar++;
+			SearchCharFlag = FETCH_ASCII_FLAG_TABLE(*SearchChar);
+		}
+		else if (SearchCharFlag & AsciiFlagsStartNumberOrDirective)
+		{
+			// Number or directive, skip to Whitespace, Other, or Quote (numbers may contain letters or #, i.e. "1.#INF" for infinity, or "e" for an exponent)
+			SearchChar++;
+			while (!((SearchCharFlag = FETCH_ASCII_FLAG_TABLE(*SearchChar)) & AsciiFlagsEndNumberOrDirective))
+			{
+				SearchChar++;
+			}
+		}
+		else if (SearchCharFlag & (uint8)AsciiFlags::Quote)
+		{
+			// Quote, skip to next Quote (or maybe end of string if text is malformed), ignoring the quote if it's escaped
+			SearchChar++;
+			while (*SearchChar && (*SearchChar != '\"' || *(SearchChar - 1) == '\\'))
+			{
+				SearchChar++;
+			}
+
+			// Could be end of string or the quote -- skip over the quote if not the null terminator
+			if (*SearchChar)
+			{
+				SearchChar++;
+			}
+			SearchCharFlag = FETCH_ASCII_FLAG_TABLE(*SearchChar);
+		}
+		// Must be null terminator -- we've tested all other possibilities
+		else
+		{
+			// End of string
+			Code = SearchChar;
+			OutCompoundTokens.Empty();
+			return nullptr;
+		}
+	}
+
+	#undef FETCH_ASCII_FLAG_TABLE
+}
+
+// Fixup for SceneTexture and UserSceneTexture inputs to custom HLSL.  Returns a new string, rather than modifying a Code FString
+// in place, to allow the code to be shared with the HLSL tree code path (MaterialExpressionHLSL.cpp), which uses TStringBuilder.
+// An empty string is returned if no fixup was required (the common case), avoiding reallocation.  OutSceneTextureInfo tracks
+// which inputs are SceneTexture or UserSceneTexture (left empty if no scene textures).  A -1 indicates a scene texture input not
+// directly used as a symbol in the custom HLSL, allowing the fetch to be dead stripped, while 1 is an actually used input.
+//
+// The goal is to allow the scene texture ID to be automatically generated in the code, without the user needing to manually
+// insert the correct identifier based on settings in the connected node.  Without this fixup, the user needs to change the
+// custom HLSL any time the input node changes, which is tedious and error prone.  For SceneTexture, if they (for example) change
+// SceneColor to Depth, they need to change PPI_SceneColor to PPI_Depth.  For UserSceneTexture, if the name is changed from Input
+// to NewInput, they need to change PPIUser_Input to PPIUser_NewInput.  Similar manual changes are also required if a completely
+// different SceneTexture node is linked.  Besides substituting the ID, minimalist shorthand ".Fetch" function call style syntax
+// is provided for offset fetches (accepts either float2 or two floats).  The original syntax of referencing "Input" without the
+// ".ID" or ".Fetch" suffixes still works for fetching a non-offset sample.
+//
+//		Input.ID -> Replace text with ID enum of connected input (PPI_* or PPIUser_*)
+//			Example usage:	SceneTextureFetch(Input.ID, float2(0,0))
+//							SceneTextureFetch(/*Input.ID=*/ PPI_SceneColor, float2(0,0))
+//
+//		Input.Fetch(*) -> Replace with SceneTextureFetchFunc(Parameters, ID, *)
+//			Example usage:	Input.Fetch(1,-1)
+//							SceneTextureFetchFunc(Parameters, /*Input=*/ PPI_SceneColor, 1,-1)
+//
+// Note that this text replacement happens before preprocessor macro substitution, and so it's not possible to use the shorthand
+// syntax with the user specified "InputName" as a macro argument, but you can pass "InputName.ID", for example:
+//
+//		// Shorthand syntax will not work, because substitution of "InputName" as the "Input" macro argument happens later in preprocess
+//		#define WEIGHTED_FETCH_SCENE_TEXTURE(Input, SampleIndex) \
+//			(Input.Fetch(Offsets[SampleIndex], 0) * Weights[SampleIndex])
+//
+//		Color += WEIGHTED_FETCH_SCENE_TEXTURE(InputName, SampleIndex);
+//
+//		// But the syntax that accepts a scene texture ID works fine
+//		#define WEIGHTED_FETCH_SCENE_TEXTURE(InputID, SampleIndex) \
+//			(SceneTextureFetch(InputID, float2(Offsets[SampleIndex], 0)) * Weights[SampleIndex])
+//
+//		Color += WEIGHTED_FETCH_SCENE_TEXTURE(InputName.ID, SampleIndex);
+//
+//		// This also works fine, because InputName is inside the macro, and not a macro argument
+//		#define WEIGHTED_FETCH_SCENE_TEXTURE(SampleIndex) \
+//			(InputName.Fetch(Offsets[SampleIndex], 0) * Weights[SampleIndex])
+//
+FString CustomExpressionSceneTextureInputFixup(const UMaterialExpressionCustom* Custom, const TCHAR* Code, TArray<int8>& OutSceneTextureInfo)
+{
+	// Allocated when required fixup is first encountered
+	FString ModifiedCode;
+
+	for (int32 InputIndex = 0; InputIndex < Custom->Inputs.Num(); InputIndex++)
+	{
+		const FCustomInput& Input = Custom->Inputs[InputIndex];
+
+		if (Input.InputName.IsNone())
+		{
+			continue;
+		}
+
+		FExpressionInput ExpressionInput = Input.Input.GetTracedInput();
+		const UMaterialExpression* Expression = ExpressionInput.Expression;
+		if (!Expression || ExpressionInput.OutputIndex != 0 || !(Expression->IsA<UMaterialExpressionSceneTexture>() || Expression->IsA<UMaterialExpressionUserSceneTexture>()))
+		{
+			continue;
+		}
+
+		FString InputNameString = Input.InputName.ToString();
+		FString InputIDString;
+
+		if (const UMaterialExpressionSceneTexture* SceneTextureExpression = Cast<UMaterialExpressionSceneTexture>(Expression))
+		{
+			InputIDString = StaticEnum<ESceneTextureId>()->GetNameStringByValue(static_cast<int>(SceneTextureExpression->SceneTextureId));
+		}
+		else
+		{
+			const UMaterialExpressionUserSceneTexture* UserSceneTextureExpression = Cast<UMaterialExpressionUserSceneTexture>(Expression);
+			InputIDString = TEXT("PPIUser_") + UserSceneTextureExpression->UserSceneTexture.ToString();
+		}
+
+		// We've encountered a scene texture, allocate our info array
+		if (!OutSceneTextureInfo.Num())
+		{
+			OutSceneTextureInfo.SetNumZeroed(Custom->Inputs.Num());
+		}
+
+		// Mark with -1 to indicate that it's a scene texture input, but might not be directly used in the source.
+		OutSceneTextureInfo[InputIndex] = -1;
+
+		int32 CodeOffset = 0;
+		while (1)
+		{
+			TArray<FStringView> CompoundTokens;
+			const TCHAR* CodeStart = ModifiedCode.IsEmpty() ? Code : *ModifiedCode;
+			const TCHAR* CodeString = &CodeStart[CodeOffset];
+			const TCHAR* FoundInputName = FindHlslIdentifierInCode(CodeString, *InputNameString, CompoundTokens);
+
+			if (!FoundInputName)
+			{
+				break;
+			}
+
+			CodeOffset = FoundInputName - CodeStart;
+
+			int32 OriginalCodeOffset = CodeOffset;
+			if (CompoundTokens.Num() == 1)
+			{
+				// Check if this is using the "Fetch" function call shorthand syntax.  If so, replace it with the corresponding fetch function.
+				if (CompoundTokens[0] == TEXT("Fetch") && *CodeString == '(')
+				{
+					FString ReplacementText = FString::Printf(TEXT("SceneTextureFetchFunc(Parameters, /*%s=*/ %s, "), *InputNameString, *InputIDString);
+
+					if (ModifiedCode.IsEmpty())
+					{
+						ModifiedCode.Reserve(FCString::Strlen(Code) + ReplacementText.Len());
+						ModifiedCode = Code;
+					}
+
+					// Add one character to amount removed, as we are overwriting the open parentheses as well
+					ModifiedCode.RemoveAt(CodeOffset, CodeString - FoundInputName + 1, EAllowShrinking::No);
+					ModifiedCode.InsertAt(CodeOffset, ReplacementText);
+
+					CodeOffset += ReplacementText.Len();
+				}
+				else if (CompoundTokens[0] == TEXT("ID"))
+				{
+					// Replace "Input.ID" identifier with ID (including comment for readability if there are compile errors)
+					FString ReplacementText = FString::Printf(TEXT("/*%s.ID=*/ %s"), *InputNameString, *InputIDString);
+
+					if (ModifiedCode.IsEmpty())
+					{
+						ModifiedCode.Reserve(FCString::Strlen(Code) + ReplacementText.Len());
+						ModifiedCode = Code;
+					}
+
+					ModifiedCode.RemoveAt(CodeOffset, CompoundTokens[0].GetData() + CompoundTokens[0].Len() - FoundInputName, EAllowShrinking::No);
+					ModifiedCode.InsertAt(CodeOffset, ReplacementText);
+
+					CodeOffset += ReplacementText.Len();
+				}
+			}
+
+			// If we didn't insert and skip over replacement code, skip over the token that was found (CodeString will have
+			// been set past the end of the token by the function "FindHlslIdentifierInCode").
+			if (CodeOffset == OriginalCodeOffset)
+			{
+				CodeOffset = CodeString - CodeStart;
+
+				// Scene Texture input pin symbol was directly used by the custom HLSL, mark it with a 1, as the input pin expression needs to be compiled in
+				OutSceneTextureInfo[InputIndex] = 1;
+			}
+		}
+	}
+
+	return ModifiedCode;
+}
+
 }  // namespace UE::MaterialTranslatorUtils
 
 #endif // WITH_EDITOR
diff --git a/Engine/Source/Runtime/Engine/Private/Materials/MaterialSharedPrivate.h b/Engine/Source/Runtime/Engine/Private/Materials/MaterialSharedPrivate.h
--- a/Engine/Source/Runtime/Engine/Private/Materials/MaterialSharedPrivate.h
+++ b/Engine/Source/Runtime/Engine/Private/Materials/MaterialSharedPrivate.h
@@ -20,6 +20,7 @@
 namespace UE::Cook { class FCookEventContext; }
 struct FAppendToClassSchemaContext;
 struct FMaterialResourceForCooking;
+class UMaterialExpressionCustom;
 
 namespace UE::MaterialInterface::Private
 {
@@ -97,6 +98,16 @@ namespace UE::MaterialTranslatorUtils
 		EShaderPlatform Platform,
 		TArray<FString>& OutErrors);
 
+	// Utility function to search for the given HLSL Identifier in a block of Code.  Takes into account comments and quoted text
+	// that should be ignored when searching for identifiers.  Returns the start of the identifier, and the input pointer "Code"
+	// will be advanced to the first non-whitespace, non-comment character after the end of the identifier.  Or null will be
+	// returned if the identifier isn't found.  OutCompoundTokens contain additional tokens if it's part of a compound identifier
+	// separated by periods (for "Identifier.x", OutCompoundTokens would include "x").
+	const TCHAR* FindHlslIdentifierInCode(const TCHAR*& Code, const TCHAR* Identifier, TArray<FStringView>& OutCompoundTokens);
+
+	// Fixup for SceneTexture and UserSceneTexture inputs to custom HLSL. See implementation for more info.
+	FString CustomExpressionSceneTextureInputFixup(const UMaterialExpressionCustom* Custom, const TCHAR* Code, TArray<int8>& OutSceneTextureInfo);
+
 } // namespace UE::MaterialTranslatorUtils
 
 ///////////////////////////////////////////////////////
diff --git a/Engine/Source/Runtime/Engine/Public/Materials/MaterialIR.h b/Engine/Source/Runtime/Engine/Public/Materials/MaterialIR.h
--- a/Engine/Source/Runtime/Engine/Public/Materials/MaterialIR.h
+++ b/Engine/Source/Runtime/Engine/Public/Materials/MaterialIR.h
@@ -125,6 +125,7 @@ enum EValueKind
 	VK_PromoteSubstrateParameter,
 	VK_StageSwitch,
 	VK_HardwarePartialDerivative,
+	VK_Nop,
 	VK_Call,
 	VK_CallParameterOutput,
 
@@ -977,6 +978,14 @@ struct FHardwarePartialDerivative : TInstruction<VK_HardwarePartialDerivative, 1
 	EDerivativeAxis Axis;
 };
 
+// A "no operation" instruction is a dummy instruction used to perform analysis on a value (and its dependencies)
+// but not push it to the list of instructions in a block.
+// A NOP instructions should effectively bahave as the default value of Arg type.
+struct FNop : TInstruction<VK_Nop, 1>
+{
+	FValue* Arg;
+};
+
 // Kind of user-specified function.
 enum class FFunctionKind
 {
diff --git a/Engine/Source/Runtime/Engine/Public/Materials/MaterialIREmitter.h b/Engine/Source/Runtime/Engine/Public/Materials/MaterialIREmitter.h
--- a/Engine/Source/Runtime/Engine/Public/Materials/MaterialIREmitter.h
+++ b/Engine/Source/Runtime/Engine/Public/Materials/MaterialIREmitter.h
@@ -512,6 +512,11 @@ public:
 	// Note: ValuePerStage must have no more than MIR::NumStages elements.
 	FValueRef StageSwitch(FType Type, TConstArrayView<FValueRef> ValuePerStage);
 
+	// Emits a NOP instruction that behaves as a default value of Arg type, but will still
+	// cause Arg to be analyzed. Use it when you want Arg to be analyzed (to record side effects
+	// on the module), but not evaluated. 
+	FValueRef Nop(FValueRef Arg);
+
 	// Emits a texture gather instruction.
 	// 
 	// @param Texture The texture to sample. Must be of type MIR::FObjectType and represent a valid texture.
