Fixes for ANoPawnPlayerController leaks when migrating player controllers and closing connections.

diff --git a/Engine/Plugins/Runtime/MultiServerReplication/Source/MultiServerReplication/Private/MultiServerProxy.cpp b/Engine/Plugins/Runtime/MultiServerReplication/Source/MultiServerReplication/Private/MultiServerProxy.cpp
--- a/Engine/Plugins/Runtime/MultiServerReplication/Source/MultiServerReplication/Private/MultiServerProxy.cpp
+++ b/Engine/Plugins/Runtime/MultiServerReplication/Source/MultiServerReplication/Private/MultiServerProxy.cpp
@@ -1134,9 +1134,10 @@ void UProxyNetDriver::HandleClosedProxyConnection(UNetConnection* ProxyConnectio
 			continue;
 		}
 
-		// If the connection route to the backend server is fully open, close it immediately, otherwise mark the route
-		// for closing once the connection handshake has been completed.
-		if (Route->State == EProxyConnectionState::Connected || Route->State == EProxyConnectionState::ConnectedPrimary)
+		// If the connection to the backend game server hasn't completed connecting then defer destroying the connection
+		// until the connection handshake is complete.
+		const bool bStillConnecting  = (Route->State == EProxyConnectionState::Connecting || Route->State == EProxyConnectionState::ConnectingPrimary);
+		if (!bStillConnecting)
 		{
 			CloseAndCleanUpConnectionRoute(Route);
 		}
diff --git a/Engine/Source/Runtime/Engine/Private/PlayerController.cpp b/Engine/Source/Runtime/Engine/Private/PlayerController.cpp
--- a/Engine/Source/Runtime/Engine/Private/PlayerController.cpp
+++ b/Engine/Source/Runtime/Engine/Private/PlayerController.cpp
@@ -1667,35 +1667,6 @@ void APlayerController::Destroyed()
 	PlayerInput = NULL;
 	CheatManager = NULL;
 
-#if UE_WITH_REMOTE_OBJECT_HANDLE
-	// Find any instances of ANoPawnPlayerController that were created when this player controller was migrated to a different server.
-	TArray<ANoPawnPlayerController*> NoPawnPlayerControllersDep;
-	for (FConstPlayerControllerIterator It = GetWorld()->GetPlayerControllerIterator(); It; ++It)
-	{
-		const TWeakObjectPtr<APlayerController>& Ptr = *It;
-		if (!Ptr.IsRemote())
-		{
-			ANoPawnPlayerController* Controller = Cast<ANoPawnPlayerController>(It->Get());
-			if (Controller && 
-				Controller != this &&
-				Controller->CachedConnectionPlayerId == CachedConnectionPlayerId &&
-				Controller->GetClientHandshakeId() == GetClientHandshakeId())
-			{
-				NoPawnPlayerControllersDep.Add(Controller);
-			}
-		}
-	}
-
-	// Remove and destroy any identifier instances of ANoPawnPlayerController.
-	for (ANoPawnPlayerController* Controller : NoPawnPlayerControllersDep)
-	{
-		UE_LOG(LogPlayerController, Log, TEXT("Destroying player controller %s which was used when %s was migrated away from this server."), *GetNameSafe(Controller), *GetName());
-
-		GetWorld()->RemoveController(Controller);
-		GetWorld()->DestroyActor(Controller, true);
-	}
-#endif
-
 	Super::Destroyed();
 }
 
@@ -5103,6 +5074,12 @@ void APlayerController::PostMigrate(const FUObjectMigrationContext& MigrationCon
 				// it's going to be detached from a connection.
 				NetConnection->PlayerController->ForceNetUpdate();
 
+				// Destroy the previous player controller used by the connection.
+				UE_LOG(LogPlayerController, Verbose, TEXT("[%s] APLayerController::PostMigrate [%s]: Destroying previous controller %s"), 
+					   *FRemoteServerId::GetLocalServerId().ToString(), *FRemoteObjectId(this).ToString(), *GetNameSafe(NetConnection->PlayerController));
+				GetWorld()->RemoveController(NetConnection->PlayerController);
+				GetWorld()->DestroyActor(NetConnection->PlayerController, true);
+
 				// Detach the previous player controller from the network connection so that any attempts to send an RPC
 				// on that controller to the client will not succeed.
 				NetConnection->PlayerController->NetConnection = nullptr;
@@ -5150,56 +5127,28 @@ void APlayerController::PostMigrate(const FUObjectMigrationContext& MigrationCon
 				PrevNetConnection->PlayerController = nullptr;
 				PrevNetConnection->OwningActor = nullptr;
 
-				// Look for a instance of ANoPawnPlayerController with the matching player id to attach to the connection.
-				for (FConstPlayerControllerIterator It = GetWorld()->GetPlayerControllerIterator(); It; ++It)
-				{
-					const TWeakObjectPtr<APlayerController>& Ptr = *It;
-					if (!Ptr.IsRemote())
-					{
-						APlayerController* OtherPlayerController = It->Get();
-						if (OtherPlayerController->IsA(ANoPawnPlayerController::StaticClass()) &&
-							OtherPlayerController->CachedConnectionPlayerId == CachedConnectionPlayerId)
-						{
-							// The ANoPawnPlayerController actor has a low NetUpdateFrequency, so ForceNetUpdate() must be called because
-							// it's going to be attached to a connection.
-							OtherPlayerController->ForceNetUpdate();
-
-							PrevNetConnection->PlayerController = OtherPlayerController;
-							PrevNetConnection->OwningActor = OtherPlayerController;
-						}
-					}
-				}
-
-				if (PrevNetConnection->PlayerController == nullptr)
-				{
-					FActorSpawnParameters SpawnInfo;
-					SpawnInfo.ObjectFlags |= RF_Transient;
-					SpawnInfo.bDeferConstruction = true;
-
-					PrevNetConnection->PlayerController = GetWorld()->SpawnActor<APlayerController>(ANoPawnPlayerController::StaticClass(), FVector(), FRotator(), SpawnInfo);
-					PrevNetConnection->PlayerController->SetRole(ROLE_Authority);
-					PrevNetConnection->PlayerController->SetReplicates(true);
-					PrevNetConnection->PlayerController->SetAutonomousProxy(true);
-					PrevNetConnection->PlayerController->FinishSpawning(FTransform());
-					PrevNetConnection->PlayerController->CachedConnectionPlayerId = CachedConnectionPlayerId;
-					PrevNetConnection->PlayerController->NetPlayerIndex = NetPlayerIndex;
-					PrevNetConnection->PlayerController->SetClientHandshakeId(PrevNetConnection->GetClientHandshakeId());
-					PrevNetConnection->PlayerController->NetConnection = PrevNetConnection;
-
-					PrevNetConnection->OwningActor = PrevNetConnection->PlayerController;
+				// Create a new instance of ANoPawnPlayerController for the connection to use while the original instance of APlayerController is migrated
+				// on another server. This controller will be destroyed when it's replaced by APlayerController when it returns to the server.
+				FActorSpawnParameters SpawnInfo;
+				SpawnInfo.ObjectFlags |= RF_Transient;
+				SpawnInfo.bDeferConstruction = true;
+				PrevNetConnection->PlayerController = GetWorld()->SpawnActor<APlayerController>(ANoPawnPlayerController::StaticClass(), FVector(), FRotator(), SpawnInfo);
+				PrevNetConnection->PlayerController->SetRole(ROLE_Authority);
+				PrevNetConnection->PlayerController->SetReplicates(true);
+				PrevNetConnection->PlayerController->SetAutonomousProxy(true);
+				PrevNetConnection->PlayerController->FinishSpawning(FTransform());
+				PrevNetConnection->PlayerController->CachedConnectionPlayerId = CachedConnectionPlayerId;
+				PrevNetConnection->PlayerController->NetPlayerIndex = NetPlayerIndex;
+				PrevNetConnection->PlayerController->SetClientHandshakeId(PrevNetConnection->GetClientHandshakeId());
+				PrevNetConnection->PlayerController->NetConnection = PrevNetConnection;
+
+				PrevNetConnection->OwningActor = PrevNetConnection->PlayerController;
 					
-					UE_LOG(LogPlayerController, Verbose, TEXT("[%s] APlayerController::PostMigrate [%s]: Replacing %s with new player controller %s on connection %s"), 
-						   *FRemoteServerId::GetLocalServerId().ToString(), *FRemoteObjectId(this).ToString(), 
-						   *GetName(), *PrevNetConnection->PlayerController->GetName(), *PrevNetConnection->GetName());
+				UE_LOG(LogPlayerController, Verbose, TEXT("[%s] APlayerController::PostMigrate [%s]: Replacing %s with new player controller %s on connection %s"), 
+					*FRemoteServerId::GetLocalServerId().ToString(), *FRemoteObjectId(this).ToString(), 
+					*GetName(), *PrevNetConnection->PlayerController->GetName(), *PrevNetConnection->GetName());
 
-					GetWorld()->AddController(PrevNetConnection->PlayerController);
-				}
-				else
-				{
-					UE_LOG(LogPlayerController, Verbose, TEXT("[%s] APlayerController::PostMigrate [%s]: Replacing %s with existing player controller %s on connection %s"),
-						   *FRemoteServerId::GetLocalServerId().ToString(), *FRemoteObjectId(this).ToString(), 
-						   *GetName(), *PrevNetConnection->PlayerController->GetName(), *PrevNetConnection->GetName());
-				}
+				GetWorld()->AddController(PrevNetConnection->PlayerController);
 
 				ensure(PrevNetConnection->PlayerController);
 				ensure(PrevNetConnection->OwningActor);
