Fix for use-after-free bug when migrating APlayerController and disable actor ticking for ANoPawnPlayerController.

diff --git a/Engine/Source/Runtime/Engine/Private/PlayerController.cpp b/Engine/Source/Runtime/Engine/Private/PlayerController.cpp
--- a/Engine/Source/Runtime/Engine/Private/PlayerController.cpp
+++ b/Engine/Source/Runtime/Engine/Private/PlayerController.cpp
@@ -5073,19 +5073,20 @@ void APlayerController::PostMigrate(const FUObjectMigrationContext& MigrationCon
 				// The ANoPawnPlayerController actor has a low NetUpdateFrequency, so ForceNetUpdate() must be called because
 				// it's going to be detached from a connection.
 				NetConnection->PlayerController->ForceNetUpdate();
+				
+				// Detach the previous player controller from the network connection so that any attempts to send an RPC
+				// on that controller to the client will not succeed.
+				NetConnection->PlayerController->NetConnection = nullptr;
+				
+				const int32 PrevNetPlayerIndex = NetConnection->PlayerController->NetPlayerIndex;
+				const int32 PrevHandshakeTransactionId = NetConnection->PlayerController->GetClientHandshakeId();
 
 				// Destroy the previous player controller used by the connection.
 				UE_LOG(LogPlayerController, Verbose, TEXT("[%s] APLayerController::PostMigrate [%s]: Destroying previous controller %s"), 
 					   *FRemoteServerId::GetLocalServerId().ToString(), *FRemoteObjectId(this).ToString(), *GetNameSafe(NetConnection->PlayerController));
 				GetWorld()->RemoveController(NetConnection->PlayerController);
 				GetWorld()->DestroyActor(NetConnection->PlayerController, true);
-
-				// Detach the previous player controller from the network connection so that any attempts to send an RPC
-				// on that controller to the client will not succeed.
-				NetConnection->PlayerController->NetConnection = nullptr;
 				
-				const int32 PrevNetPlayerIndex = NetConnection->PlayerController->NetPlayerIndex;
-				const int32 PrevHandshakeTransactionId = NetConnection->PlayerController->GetClientHandshakeId();
 				NetConnection->PlayerController = this;
 				NetConnection->OwningActor = this;
 				Player = NetConnection;
@@ -6657,6 +6658,8 @@ ANoPawnPlayerController::ANoPawnPlayerController(const FObjectInitializer& Objec
 	bReplicates = true;
 	SetRole(ROLE_Authority);
 
+	PrimaryActorTick.bCanEverTick = false;
+
 	// This controller's updates don't need to be replicated to the client very frequently because
 	// it's only intended to be used by the server for replication, relevancy and establishing the
 	// initial connection on the client.
