Skeletal Mesh Cloth - Tidy up the code around the component Collision Sources.

#rnx

diff --git a/Engine/Source/Runtime/Engine/Classes/Components/SkeletalMeshComponent.h b/Engine/Source/Runtime/Engine/Classes/Components/SkeletalMeshComponent.h
--- a/Engine/Source/Runtime/Engine/Classes/Components/SkeletalMeshComponent.h
+++ b/Engine/Source/Runtime/Engine/Classes/Components/SkeletalMeshComponent.h
@@ -2478,7 +2478,7 @@ protected:
 	ENGINE_API void CopyChildrenClothCollisionsToParent();
 
 	/** find if this component has collisions for clothing and return the results calculated by bone transforms */
-	ENGINE_API void FindClothCollisions(FClothCollisionData& OutCollisions);
+	ENGINE_API void FindClothCollisions(FClothCollisionData& OutCollisions) const;
 
 #endif
 
@@ -2495,8 +2495,8 @@ protected:
 
 	ENGINE_API virtual bool ShouldBlendPhysicsBones() const;
 
-	/** Extract collisions for cloth from this component (given a component we want to apply the data to) */
-	static ENGINE_API void ExtractCollisionsForCloth(USkeletalMeshComponent* SourceComponent,  UPhysicsAsset* PhysicsAsset, USkeletalMeshComponent* DestClothComponent, FClothCollisionData& OutCollisions, FClothCollisionSource& ClothCollisionSource);
+	UE_DEPRECATED(5.7, "Use ClothCollisionSource::ExtractCollisions() instead.")
+	static ENGINE_API void ExtractCollisionsForCloth(USkeletalMeshComponent* SourceComponent, UPhysicsAsset* PhysicsAsset, USkeletalMeshComponent* DestClothComponent, FClothCollisionData& OutCollisions, FClothCollisionSource& ClothCollisionSource);
 
 	/** Notify called just before syncing physics update, called only if bNotifySyncComponentToRBPhysics flag is set */
 	virtual void OnSyncComponentToRBPhysics() { }
diff --git a/Engine/Source/Runtime/Engine/Private/SkeletalMeshComponentPhysics.cpp b/Engine/Source/Runtime/Engine/Private/SkeletalMeshComponentPhysics.cpp
--- a/Engine/Source/Runtime/Engine/Private/SkeletalMeshComponentPhysics.cpp
+++ b/Engine/Source/Runtime/Engine/Private/SkeletalMeshComponentPhysics.cpp
@@ -138,6 +138,128 @@ FClothCollisionSource::~FClothCollisionSource()
 	}
 }
 
+void FClothCollisionSource::ExtractCollisions(const USkeletalMeshComponent* DestClothComponent, FClothCollisionData& OutCollisions)
+{
+	const USkeletalMeshComponent* const SourceComponentRawPtr = SourceComponent.Get();
+	const UPhysicsAsset* const SourcePhysicsAssetRawPtr = SourcePhysicsAsset.Get();
+
+	// Extract collisions from this mesh 'raw', as this isn't a mesh that has cloth simulation
+	// (but we want it to affect other meshes with cloth simulation)
+	if (SourceComponentRawPtr && SourceComponentRawPtr->GetSkeletalMeshAsset() && SourcePhysicsAssetRawPtr)
+	{
+		FTransform ComponentToComponentTransform;
+		if (SourceComponentRawPtr != DestClothComponent)
+		{
+			FTransform DestClothComponentTransform = DestClothComponent->GetComponentTransform();
+			DestClothComponentTransform.RemoveScaling();  // The collision source doesn't need the scale of the cloth skeletal mesh applied to it (but it does need the source scale from the component transform)
+			ComponentToComponentTransform = SourceComponentRawPtr->GetComponentTransform() * DestClothComponentTransform.Inverse();
+		}
+
+		// Init cache on first copy
+		if (!bCached || CachedSkeletalMesh.Get() != SourceComponentRawPtr->GetSkeletalMeshAsset())
+		{
+			// Clear previous cached data
+			CachedSpheres.Reset();
+			CachedSphereConnections.Reset();
+
+			for (const USkeletalBodySetup* SkeletalBodySetup : SourcePhysicsAssetRawPtr->SkeletalBodySetups)
+			{
+				// Cache bones
+				const int32 MeshBoneIndex = SourceComponentRawPtr->GetSkeletalMeshAsset()->GetRefSkeleton().FindBoneIndex(SkeletalBodySetup->BoneName);
+				if (MeshBoneIndex != INDEX_NONE)
+				{
+					// Cache spheres & capsules form physics asset
+					for (const FKSphereElem& Sphere : SkeletalBodySetup->AggGeom.SphereElems)
+					{
+						FClothCollisionPrim_Sphere NewSphere;
+						NewSphere.LocalPosition = Sphere.Center;
+						NewSphere.Radius = Sphere.Radius;
+						NewSphere.BoneIndex = MeshBoneIndex;
+
+						CachedSpheres.Add(NewSphere);
+					}
+
+					for (const FKSphylElem& Sphyl : SkeletalBodySetup->AggGeom.SphylElems)
+					{
+						FClothCollisionPrim_Sphere Sphere0;
+						FClothCollisionPrim_Sphere Sphere1;
+						FVector OrientedDirection = Sphyl.Rotation.RotateVector(FVector(0.0f, 0.0f, 1.0f));
+						FVector HalfDim = OrientedDirection * (Sphyl.Length / 2.0f);
+						Sphere0.LocalPosition = Sphyl.Center - HalfDim;
+						Sphere1.LocalPosition = Sphyl.Center + HalfDim;
+						Sphere0.Radius = Sphyl.Radius;
+						Sphere1.Radius = Sphyl.Radius;
+						Sphere0.BoneIndex = MeshBoneIndex;
+						Sphere1.BoneIndex = MeshBoneIndex;
+
+						CachedSpheres.Add(Sphere0);
+						CachedSpheres.Add(Sphere1);
+
+						FClothCollisionPrim_SphereConnection Connection;
+						Connection.SphereIndices[0] = CachedSpheres.Num() - 2;
+						Connection.SphereIndices[1] = CachedSpheres.Num() - 1;
+
+						CachedSphereConnections.Add(Connection);
+					}
+
+					for (const FKTaperedCapsuleElem& TaperedCapsuleElems : SkeletalBodySetup->AggGeom.TaperedCapsuleElems)
+					{
+						FClothCollisionPrim_Sphere Sphere0;
+						FClothCollisionPrim_Sphere Sphere1;
+						const FVector OrientedDirection = TaperedCapsuleElems.Rotation.RotateVector(FVector(0.0f, 0.0f, 1.0f));
+						const FVector HalfDim = OrientedDirection * (TaperedCapsuleElems.Length / 2.0f);
+						Sphere0.LocalPosition = TaperedCapsuleElems.Center + HalfDim;
+						Sphere1.LocalPosition = TaperedCapsuleElems.Center - HalfDim;
+						Sphere0.Radius = TaperedCapsuleElems.Radius0;
+						Sphere1.Radius = TaperedCapsuleElems.Radius1;
+						Sphere0.BoneIndex = MeshBoneIndex;
+						Sphere1.BoneIndex = MeshBoneIndex;
+
+						CachedSpheres.Add(Sphere0);
+						CachedSpheres.Add(Sphere1);
+
+						FClothCollisionPrim_SphereConnection Connection;
+						Connection.SphereIndices[0] = CachedSpheres.Num() - 2;
+						Connection.SphereIndices[1] = CachedSpheres.Num() - 1;
+						if (TaperedCapsuleElems.bOneSidedCollision)
+						{
+							Connection.OneSidedPlaneNormal = TaperedCapsuleElems.Rotation.RotateVector(FVector::ForwardVector);
+						}
+
+						CachedSphereConnections.Add(Connection);
+					}
+				}
+			}
+			CachedSkeletalMesh = SourceComponentRawPtr->GetSkeletalMeshAsset();
+			bCached = true;
+		}
+
+		// Presize array allocations
+		OutCollisions.Spheres.Reserve(OutCollisions.Spheres.Num() + CachedSpheres.Num());
+		OutCollisions.SphereConnections.Reserve(OutCollisions.SphereConnections.Num() + CachedSphereConnections.Num());
+
+		// Now transform output data
+		for (const FClothCollisionPrim_Sphere& CachedSphere : CachedSpheres)
+		{
+			FClothCollisionPrim_Sphere& OutSphere = OutCollisions.Spheres.Add_GetRef(CachedSphere);
+
+			const FTransform BoneTransform = SourceComponentRawPtr->GetBoneTransform(OutSphere.BoneIndex, FTransform::Identity) * ComponentToComponentTransform;
+			OutSphere.LocalPosition = BoneTransform.TransformPosition(OutSphere.LocalPosition);
+			OutSphere.Radius *= BoneTransform.GetScale3D().X;  // Cloth collisions only uniformly scale
+			OutSphere.BoneIndex = INDEX_NONE;
+		}
+
+		// Offset connections
+		const int32 ConnectionBaseIndex = OutCollisions.SphereConnections.Num();
+		for (const FClothCollisionPrim_SphereConnection& CachedSphereConnection : CachedSphereConnections)
+		{
+			FClothCollisionPrim_SphereConnection& OutSphereConnection = OutCollisions.SphereConnections.Add_GetRef(CachedSphereConnection);
+			OutSphereConnection.SphereIndices[0] += ConnectionBaseIndex;
+			OutSphereConnection.SphereIndices[1] += ConnectionBaseIndex;
+		}
+	}
+}
+
 void USkeletalMeshComponent::CreateBodySetup()
 {
 	if (BodySetup == NULL)
@@ -2936,7 +3058,7 @@ void USkeletalMeshComponent::GetWindForCloth_GameThread(FVector& WindDirection,
 
 #if WITH_CLOTH_COLLISION_DETECTION
 
-void USkeletalMeshComponent::FindClothCollisions(FClothCollisionData& OutCollisions)
+void USkeletalMeshComponent::FindClothCollisions(FClothCollisionData& OutCollisions) const
 {
 	if(ClothingSimulation)
 	{
@@ -2946,6 +3068,7 @@ void USkeletalMeshComponent::FindClothCollisions(FClothCollisionData& OutCollisi
 	}
 }
 
+// UE_DEPRECATED(5.7, "Use ClothCollisionSource::ExtractCollisions() instead.")
 void USkeletalMeshComponent::ExtractCollisionsForCloth(
 	USkeletalMeshComponent* SourceComponent, 
 	UPhysicsAsset* PhysicsAsset, 
@@ -2953,121 +3076,10 @@ void USkeletalMeshComponent::ExtractCollisionsForCloth(
 	FClothCollisionData& OutCollisions, 
 	FClothCollisionSource& ClothCollisionSource)
 {
-	// Extract collisions from this mesh 'raw', as this isn't a mesh that has cloth simulation
-	// (but we want it to affect other meshes with cloth simulation)
-	if(SourceComponent->GetSkeletalMeshAsset() && PhysicsAsset)
+	if (ensure(ClothCollisionSource.SourceComponent.Get() == SourceComponent &&
+		ClothCollisionSource.SourcePhysicsAsset.Get() == PhysicsAsset))
 	{
-		FTransform ComponentToComponentTransform;
-		if(SourceComponent != DestClothComponent)
-		{
-			FTransform DestClothComponentTransform = DestClothComponent->GetComponentTransform();
-			DestClothComponentTransform.RemoveScaling();  // The collision source doesn't need the scale of the cloth skeletal mesh applied to it (but it does need the source scale from the component transform)
-			ComponentToComponentTransform = SourceComponent->GetComponentTransform() * DestClothComponentTransform.Inverse();
-		}
-
-		// Init cache on first copy
-		if(!ClothCollisionSource.bCached || ClothCollisionSource.CachedSkeletalMesh.Get() != SourceComponent->GetSkeletalMeshAsset())
-		{
-			// Clear previous cached data
-			ClothCollisionSource.CachedSpheres.Reset();
-			ClothCollisionSource.CachedSphereConnections.Reset();
-
-			for(const USkeletalBodySetup* SkeletalBodySetup : PhysicsAsset->SkeletalBodySetups)
-			{
-				// Cache bones
-				int32 MeshBoneIndex = SourceComponent->GetSkeletalMeshAsset()->GetRefSkeleton().FindBoneIndex(SkeletalBodySetup->BoneName);
-				if(MeshBoneIndex != INDEX_NONE)
-				{
-					// Cache spheres & capsules form physics asset
-					for(const FKSphereElem& Sphere : SkeletalBodySetup->AggGeom.SphereElems)
-					{
-						FClothCollisionPrim_Sphere NewSphere;
-						NewSphere.LocalPosition = Sphere.Center;
-						NewSphere.Radius = Sphere.Radius;
-						NewSphere.BoneIndex = MeshBoneIndex;
-
-						ClothCollisionSource.CachedSpheres.Add(NewSphere);
-					}
-
-					for(const FKSphylElem& Sphyl : SkeletalBodySetup->AggGeom.SphylElems)
-					{
-						FClothCollisionPrim_Sphere Sphere0;
-						FClothCollisionPrim_Sphere Sphere1;
-						FVector OrientedDirection = Sphyl.Rotation.RotateVector(FVector(0.0f, 0.0f, 1.0f));
-						FVector HalfDim = OrientedDirection * (Sphyl.Length / 2.0f);
-						Sphere0.LocalPosition = Sphyl.Center - HalfDim;
-						Sphere1.LocalPosition = Sphyl.Center + HalfDim;
-						Sphere0.Radius = Sphyl.Radius;
-						Sphere1.Radius = Sphyl.Radius;
-						Sphere0.BoneIndex = MeshBoneIndex;
-						Sphere1.BoneIndex = MeshBoneIndex;
-
-						ClothCollisionSource.CachedSpheres.Add(Sphere0);
-						ClothCollisionSource.CachedSpheres.Add(Sphere1);
-
-						FClothCollisionPrim_SphereConnection Connection;
-						Connection.SphereIndices[0] = ClothCollisionSource.CachedSpheres.Num() - 2;
-						Connection.SphereIndices[1] = ClothCollisionSource.CachedSpheres.Num() - 1;
-
-						ClothCollisionSource.CachedSphereConnections.Add(Connection);
-					}
-
-					for(const FKTaperedCapsuleElem& TaperedCapsuleElems : SkeletalBodySetup->AggGeom.TaperedCapsuleElems)
-					{
-						FClothCollisionPrim_Sphere Sphere0;
-						FClothCollisionPrim_Sphere Sphere1;
-						FVector OrientedDirection = TaperedCapsuleElems.Rotation.RotateVector(FVector(0.0f, 0.0f, 1.0f));
-						FVector HalfDim = OrientedDirection * (TaperedCapsuleElems.Length / 2.0f);
-						Sphere0.LocalPosition = TaperedCapsuleElems.Center + HalfDim;
-						Sphere1.LocalPosition = TaperedCapsuleElems.Center - HalfDim;
-						Sphere0.Radius = TaperedCapsuleElems.Radius0;
-						Sphere1.Radius = TaperedCapsuleElems.Radius1;
-						Sphere0.BoneIndex = MeshBoneIndex;
-						Sphere1.BoneIndex = MeshBoneIndex;
-
-						ClothCollisionSource.CachedSpheres.Add(Sphere0);
-						ClothCollisionSource.CachedSpheres.Add(Sphere1);
-
-						FClothCollisionPrim_SphereConnection Connection;
-						Connection.SphereIndices[0] = ClothCollisionSource.CachedSpheres.Num() - 2;
-						Connection.SphereIndices[1] = ClothCollisionSource.CachedSpheres.Num() - 1;
-						if (TaperedCapsuleElems.bOneSidedCollision)
-						{
-							Connection.OneSidedPlaneNormal = TaperedCapsuleElems.Rotation.RotateVector(FVector::ForwardVector);
-						}
-
-						ClothCollisionSource.CachedSphereConnections.Add(Connection);
-					}
-				}
-			}
-
-			ClothCollisionSource.CachedSkeletalMesh = SourceComponent->GetSkeletalMeshAsset();
-			ClothCollisionSource.bCached = true;
-		}
-
-		// presize array allocations
-		OutCollisions.Spheres.Reserve(OutCollisions.Spheres.Num() + ClothCollisionSource.CachedSpheres.Num());
-		OutCollisions.SphereConnections.Reserve(OutCollisions.SphereConnections.Num() + ClothCollisionSource.CachedSphereConnections.Num());
-
-		// Now transform output data
-		for(const FClothCollisionPrim_Sphere& CachedSphere : ClothCollisionSource.CachedSpheres)
-		{
-			FClothCollisionPrim_Sphere& OutSphere = OutCollisions.Spheres.Add_GetRef(CachedSphere);
-
-			const FTransform BoneTransform = SourceComponent->GetBoneTransform(OutSphere.BoneIndex, FTransform::Identity) * ComponentToComponentTransform;
-			OutSphere.LocalPosition = BoneTransform.TransformPosition(OutSphere.LocalPosition);
-			OutSphere.Radius *= BoneTransform.GetScale3D().X;  // Cloth collisions only uniformly scale
-			OutSphere.BoneIndex = INDEX_NONE;
-		}
-
-		// Offset connections
-		int32 ConnectionBaseIndex = OutCollisions.SphereConnections.Num();
-		for(const FClothCollisionPrim_SphereConnection& CachedSphereConnection : ClothCollisionSource.CachedSphereConnections)
-		{
-			FClothCollisionPrim_SphereConnection& OutSphereConnection = OutCollisions.SphereConnections.Add_GetRef(CachedSphereConnection);
-			OutSphereConnection.SphereIndices[0] += ConnectionBaseIndex;
-			OutSphereConnection.SphereIndices[1] += ConnectionBaseIndex;
-		}
+		ClothCollisionSource.ExtractCollisions(DestClothComponent, OutCollisions);
 	}
 }
 
@@ -3112,34 +3124,21 @@ void USkeletalMeshComponent::CopyClothCollisionSources()
 
 	for(FClothCollisionSource& ClothCollisionSource : ClothCollisionSources)
 	{
-		USkeletalMeshComponent* SourceComponent = ClothCollisionSource.SourceComponent.Get();
-		UPhysicsAsset* SourcePhysicsAsset = ClothCollisionSource.SourcePhysicsAsset.Get();
-
-		if(SourceComponent && SourcePhysicsAsset)
-		{
-			ExtractCollisionsForCloth(SourceComponent, SourcePhysicsAsset, this, ExternalCollisions, ClothCollisionSource);
-		}
+		ClothCollisionSource.ExtractCollisions(this, ExternalCollisions);
 	}
 
+
 	ClothingSimulation->AddExternalCollisions(ExternalCollisions);
 }
 
-// children's collisions can affect to parent's cloth reversely
 void USkeletalMeshComponent::CopyChildrenClothCollisionsToParent()
 {
-	// 3 steps
-	// 1. release all previous children collisions
-	// 2. find new collisions from children
-	// 3. add new collisions to parent (this component)
-
-	TArray<USkeletalMeshComponent*> ClothCollisionChildren;
-
+	// Find new collisions from children
 	FClothCollisionData NewCollisions;
 
-	for (USceneComponent* AttachedChild : GetAttachChildren())
+	for (const TObjectPtr<USceneComponent>& AttachedChild : GetAttachChildren())
 	{
-		USkeletalMeshComponent* Child = Cast<USkeletalMeshComponent>(AttachedChild);
-		if(Child)
+		if (const USkeletalMeshComponent* const Child = Cast<USkeletalMeshComponent>(AttachedChild))
 		{
 			Child->FindClothCollisions(NewCollisions);
 		}
diff --git a/Engine/Source/Runtime/Engine/Public/ClothCollisionSource.h b/Engine/Source/Runtime/Engine/Public/ClothCollisionSource.h
--- a/Engine/Source/Runtime/Engine/Public/ClothCollisionSource.h
+++ b/Engine/Source/Runtime/Engine/Public/ClothCollisionSource.h
@@ -11,6 +11,9 @@ struct FClothCollisionSource
 	FClothCollisionSource(USkeletalMeshComponent* InSourceComponent, UPhysicsAsset* InSourcePhysicsAsset, const FOnBoneTransformsFinalizedMultiCast::FDelegate& InOnBoneTransformsFinalizedDelegate);
 	ENGINE_API ~FClothCollisionSource();
 
+	/** Extract collisions for cloth from the source component given a target component we want to apply the data to. */
+	void ExtractCollisions(const USkeletalMeshComponent* DestClothComponent, FClothCollisionData& OutCollisions);
+
 	/** Component that collision data will be copied from */
 	TWeakObjectPtr<USkeletalMeshComponent> SourceComponent;
 
