SplineComponent: Fix error in recent change to SetQuaternion* functions

diff --git a/Engine/Source/Runtime/Engine/Private/Components/SplineComponent.cpp b/Engine/Source/Runtime/Engine/Private/Components/SplineComponent.cpp
--- a/Engine/Source/Runtime/Engine/Private/Components/SplineComponent.cpp
+++ b/Engine/Source/Runtime/Engine/Private/Components/SplineComponent.cpp
@@ -1271,7 +1271,7 @@ void USplineComponent::SetQuaternionAtSplinePoint(int32 PointIndex, const FQuat&
 			GetComponentTransform().InverseTransformRotation(InQuaternion) : InQuaternion;
 
 		// Work backwards to compute the rotation that is currently being applied.
-		const FQuat RelativeQuat = InQuaternion * GetQuaternionAtSplinePoint(PointIndex, ESplineCoordinateSpace::Local).Inverse();
+		const FQuat RelativeQuat = Quat * GetQuaternionAtSplinePoint(PointIndex, ESplineCoordinateSpace::Local).Inverse();
 
 		// Align up vector with rotation
 		SetUpVectorAtSplinePoint(PointIndex, Quat.GetUpVector(), ESplineCoordinateSpace::Local, false);
