Fix code that removes ?listen from client URLs to work correctly with # options as well

diff --git a/Engine/Source/Runtime/Engine/Private/GameModeBase.cpp b/Engine/Source/Runtime/Engine/Private/GameModeBase.cpp
--- a/Engine/Source/Runtime/Engine/Private/GameModeBase.cpp
+++ b/Engine/Source/Runtime/Engine/Private/GameModeBase.cpp
@@ -406,8 +406,17 @@ APlayerController* AGameModeBase::ProcessClientTravel(FString& FURL, bool bSeaml
 					// We are going to lop off the listen option when sending to the client, so they don't think they're the server.
 					// That sounds silly, but in GetNetMode, we will eventually arrive at AttemptDeriveFromURL in certain circumstances (such as during travel!)
 					// The client would then think that they're supposed to execute things locally rather than RPC them based on the listen param.
-					ClientUrl = FURL.Replace(TEXT("?listen?"), TEXT("?"));
-					ClientUrl.RemoveFromEnd(TEXT("?listen"));
+					static FString ListenString = TEXT("?listen");
+					int32 ListenIndex = ClientUrl.Find(ListenString, ESearchCase::IgnoreCase);
+					if (ListenIndex >= 0)
+					{
+						// Only remove if the next character isn't part of a longer option string
+						int32 NextIndex = ListenIndex + ListenString.Len();						
+						if (NextIndex >= ClientUrl.Len() || ClientUrl[NextIndex] == TEXT('#') || ClientUrl[NextIndex] == TEXT('?'))
+						{
+							ClientUrl.RemoveAt(ListenIndex, ListenString.Len());
+						}
+					}
 				}
 
 				// Remote player
