Fixes to ref pose transform provider and added a temporary ability to disable anim banks from being used. (Larger refactor coming to greatly improve this).

diff --git a/Engine/Shaders/Private/Skinning/TransformProviders.usf b/Engine/Shaders/Private/Skinning/TransformProviders.usf
--- a/Engine/Shaders/Private/Skinning/TransformProviders.usf
+++ b/Engine/Shaders/Private/Skinning/TransformProviders.usf
@@ -4,12 +4,7 @@
 #include "/Engine/Private/Common.ush"
 #include "/Engine/Private/SceneData.ush"
 #include "/Engine/Private/BoneTransform.ush"
-
-struct FPrimitiveTransforms
-{
-	float3x4 Current;
-	float3x4 Previous;
-};
+#include "SkinningDataDecode.ush"
 
 struct FTransformBlockHeader
 {
@@ -30,18 +25,16 @@ void RefPoseProviderCS(uint3 GroupId : SV_GroupID, uint GroupIndex : SV_GroupInd
 	const FTransformBlockHeader BlockHeader = HeaderBuffer[GlobalBlockIndex];
 	if (LocalTransformIndex < BlockHeader.BlockTransformCount)
 	{
-		const uint TransformOffset = BlockHeader.BlockTransformOffset + (LocalTransformIndex * (uint)sizeof(FPrimitiveTransforms));
-
-		float3x4 Identity;
-		Identity[0] = float4(1.0f, 0.0f, 0.0f, 0.0f);
-		Identity[1] = float4(0.0f, 1.0f, 0.0f, 0.0f);
-		Identity[2] = float4(0.0f, 0.0f, 1.0f, 0.0f);
+		const float4x3 Identity =
+		{
+			float3(1.0f, 0.0f, 0.0f),
+			float3(0.0f, 1.0f, 0.0f),
+			float3(0.0f, 0.0f, 1.0f),
+			float3(0.0f, 0.0f, 0.0f)
+		};
 
 		// Every primitive has 2 transforms (current 3x4 and previous 3x4 packed together)
-		FPrimitiveTransforms Transforms;
-		Transforms.Current  = Identity;
-		Transforms.Previous = Identity;
-
-		TransformBuffer.Store<FPrimitiveTransforms>(TransformOffset, Transforms);
+		StoreCompressedBoneTransform(TransformBuffer, BlockHeader.BlockTransformOffset, LocalTransformIndex * 2, Identity);
+		StoreCompressedBoneTransform(TransformBuffer, BlockHeader.BlockTransformOffset, LocalTransformIndex * 2 + 1, Identity);
 	}
 }
\ No newline at end of file
diff --git a/Engine/Source/Runtime/Engine/Classes/Components/InstancedSkinnedMeshComponent.h b/Engine/Source/Runtime/Engine/Classes/Components/InstancedSkinnedMeshComponent.h
--- a/Engine/Source/Runtime/Engine/Classes/Components/InstancedSkinnedMeshComponent.h
+++ b/Engine/Source/Runtime/Engine/Classes/Components/InstancedSkinnedMeshComponent.h
@@ -97,6 +97,9 @@ protected:
 	UPROPERTY(EditAnywhere, BlueprintReadOnly, Category = Animation, meta = (ShowOnlyInnerProperties))
 	TArray<struct FAnimBankItem> AnimBankItems;
 
+	UPROPERTY(EditAnywhere, BlueprintReadOnly, Category = Animation)
+	bool bAnimBankEnabled = false;
+
 	UPROPERTY()
 	FBox PrimitiveBoundsOverride;
 
diff --git a/Engine/Source/Runtime/Engine/Private/Components/InstancedSkinnedMeshComponent.cpp b/Engine/Source/Runtime/Engine/Private/Components/InstancedSkinnedMeshComponent.cpp
--- a/Engine/Source/Runtime/Engine/Private/Components/InstancedSkinnedMeshComponent.cpp
+++ b/Engine/Source/Runtime/Engine/Private/Components/InstancedSkinnedMeshComponent.cpp
@@ -65,6 +65,7 @@ public:
 	FInstancedSkinnedMeshSceneProxy(const Nanite::FMaterialAudit& MaterialAudit, const FInstancedSkinnedMeshSceneProxyDesc& InMeshDesc, FSkeletalMeshRenderData* InRenderData)
 		: Super(MaterialAudit, InMeshDesc, InRenderData, false /* bAllowScale */)
 		, AnimationMinScreenSize(InMeshDesc.AnimationMinScreenSize)
+		, bAnimBankEnabled(InMeshDesc.bAnimBankEnabled)
 	{
 #if WITH_EDITOR
 		const bool bSupportInstancePicking = HasPerInstanceHitProxies() && SMInstanceElementDataUtil::SMInstanceElementsEnabled();
@@ -81,7 +82,7 @@ public:
 		}
 #endif
 		const bool bForceRefPose = CVarInstancedSkinnedMeshesForceRefPose.GetValueOnAnyThread() != 0;
-		const bool bUseAnimBank = !bForceRefPose && InMeshDesc.AnimBankItems.Num() > 0;
+		const bool bUseAnimBank = !bForceRefPose && bAnimBankEnabled && InMeshDesc.AnimBankItems.Num() > 0;
 
 		InstanceMinDrawDistance = InMeshDesc.InstanceMinDrawDistance;
 		InstanceStartCullDistance = InMeshDesc.InstanceStartCullDistance;
@@ -131,51 +132,57 @@ public:
 	{
 		Super::CreateRenderThreadResources(RHICmdList);
 
-		TArray<FAnimBankDesc> Descs;
-		Descs.Reserve(AnimBankItems.Num());
-
-		for (const FAnimBankItem& Item : AnimBankItems)
+		if (bAnimBankEnabled)
 		{
-			if (Item.BankAsset == nullptr || SkinnedAsset == nullptr)
+			TArray<FAnimBankDesc> Descs;
+			Descs.Reserve(AnimBankItems.Num());
+
+			for (const FAnimBankItem& Item : AnimBankItems)
 			{
-				Descs.Emplace(FAnimBankDesc());
-				continue;
-			}
+				if (Item.BankAsset == nullptr || SkinnedAsset == nullptr)
+				{
+					Descs.Emplace(FAnimBankDesc());
+					continue;
+				}
 
-			const FAnimBankData& BankData = Item.BankAsset->GetData();
+				const FAnimBankData& BankData = Item.BankAsset->GetData();
 
-			if (Item.SequenceIndex >= BankData.Entries.Num())
-			{
-				Descs.Emplace(FAnimBankDesc());
-				continue;
-			}
+				if (Item.SequenceIndex >= BankData.Entries.Num())
+				{
+					Descs.Emplace(FAnimBankDesc());
+					continue;
+				}
 
-			const FAnimBankEntry& BankEntry = BankData.Entries[Item.SequenceIndex];
+				const FAnimBankEntry& BankEntry = BankData.Entries[Item.SequenceIndex];
 
-			FAnimBankDesc& Desc = Descs.Emplace_GetRef();
+				FAnimBankDesc& Desc = Descs.Emplace_GetRef();
 
-			Desc.BankAsset		= MakeWeakObjectPtr(Item.BankAsset);
-			Desc.SequenceIndex	= Item.SequenceIndex;
-			Desc.Asset			= MakeWeakObjectPtr(SkinnedAsset);
-			Desc.Position		= BankEntry.Position;
-			Desc.PlayRate		= BankEntry.PlayRate;
-			Desc.bLooping		= BankEntry.IsLooping()   ? 1 : 0;
-			Desc.bAutoStart		= BankEntry.IsAutoStart() ? 1 : 0;
-		}
+				Desc.BankAsset		= MakeWeakObjectPtr(Item.BankAsset);
+				Desc.SequenceIndex	= Item.SequenceIndex;
+				Desc.Asset			= MakeWeakObjectPtr(SkinnedAsset);
+				Desc.Position		= BankEntry.Position;
+				Desc.PlayRate		= BankEntry.PlayRate;
+				Desc.bLooping		= BankEntry.IsLooping()   ? 1 : 0;
+				Desc.bAutoStart		= BankEntry.IsAutoStart() ? 1 : 0;
+			}
 
-		AnimBankHandles = GetScene().RegisterAnimBank(Descs);
-		AnimBankIds.Reset(AnimBankHandles.Num());
-		for (const FAnimBankRecordHandle& Handle : AnimBankHandles)
-		{
-			AnimBankIds.Emplace((uint64)Handle.Id);
+			AnimBankHandles = GetScene().RegisterAnimBank(Descs);
+			AnimBankIds.Reset(AnimBankHandles.Num());
+			for (const FAnimBankRecordHandle& Handle : AnimBankHandles)
+			{
+				AnimBankIds.Emplace((uint64)Handle.Id);
+			}
 		}
 	}
 
 	virtual void DestroyRenderThreadResources() override
 	{
-		GetScene().UnregisterAnimBank(AnimBankHandles);
-		AnimBankHandles.Reset();
-		AnimBankIds.Reset();
+		if (bAnimBankEnabled)
+		{
+			GetScene().UnregisterAnimBank(AnimBankHandles);
+			AnimBankHandles.Reset();
+			AnimBankIds.Reset();
+		}
 		Super::DestroyRenderThreadResources();
 	}
 
@@ -183,7 +190,7 @@ public:
 	virtual const TConstArrayView<uint64> GetAnimationProviderData(bool& bOutValid) const override
 	{
 		bOutValid = (AnimBankIds.Num() == UniqueAnimationCount);
-		return AnimBankIds;
+		return AnimBankIds; // TODO: Rename generic
 	}
 
 #if RHI_RAYTRACING
@@ -245,6 +252,8 @@ private:
 	uint32 InstanceMinDrawDistance = 0;
 	uint32 InstanceStartCullDistance = 0;
 	uint32 InstanceEndCullDistance = 0;
+
+	uint8 bAnimBankEnabled : 1 = false;
 };
 
 /** 
@@ -992,7 +1001,7 @@ FInstanceDataManagerSourceDataDesc UInstancedSkinnedMeshComponent::GetComponentD
 #endif
 	
 	const bool bForceRefPose = CVarInstancedSkinnedMeshesForceRefPose.GetValueOnAnyThread() != 0;
-	const bool bUseAnimBank = !bForceRefPose && AnimBankItems.Num() > 0;
+	const bool bUseAnimBank = !bForceRefPose && bAnimBankEnabled && AnimBankItems.Num() > 0;
 
 	Flags.bHasPerInstanceHierarchyOffset = false;
 	Flags.bHasPerInstanceLocalBounds = bUseAnimBank && AnimBankItems.Num() > 1;
@@ -1024,26 +1033,30 @@ FInstanceDataManagerSourceDataDesc UInstancedSkinnedMeshComponent::GetComponentD
 		ChangeSet.GetSkinningDataWriter().Gather(
 			[&](int32 InstanceIndex) -> uint32 
 			{ 
-				return InstanceData[InstanceIndex].BankIndex * MaxBoneTransformCount * 2u; 
+				return bUseAnimBank ? (InstanceData[InstanceIndex].BankIndex * MaxBoneTransformCount * 2u) : 0;
 			});
 
 		ChangeSet.GetLocalBoundsWriter().Gather(
 			[&](int32 InstanceIndex) -> FRenderBounds 
-			{				
-				uint32 BankIndex = InstanceData[InstanceIndex].BankIndex;
-				if (BankIndex < uint32(AnimBankItems.Num()))
+			{
+				if (bUseAnimBank)
 				{
-					const FAnimBankItem& BankItem = AnimBankItems[BankIndex];
-					if (BankItem.BankAsset != nullptr)
+					uint32 BankIndex = InstanceData[InstanceIndex].BankIndex;
+					if (BankIndex < uint32(AnimBankItems.Num()))
 					{
-						const FAnimBankData& BankData = BankItem.BankAsset->GetData();
-						if (BankItem.SequenceIndex < BankData.Entries.Num())
+						const FAnimBankItem& BankItem = AnimBankItems[BankIndex];
+						if (BankItem.BankAsset != nullptr)
 						{
-							return BankData.Entries[BankItem.SequenceIndex].SampledBounds;
+							const FAnimBankData& BankData = BankItem.BankAsset->GetData();
+							if (BankItem.SequenceIndex < BankData.Entries.Num())
+							{
+								return BankData.Entries[BankItem.SequenceIndex].SampledBounds;
+							}
 						}
 					}
 				}
-				return MeshBounds; 
+
+				return MeshBounds;
 			});
 
 #if WITH_EDITOR
diff --git a/Engine/Source/Runtime/Engine/Private/InstancedSkinnedMeshSceneProxyDesc.cpp b/Engine/Source/Runtime/Engine/Private/InstancedSkinnedMeshSceneProxyDesc.cpp
--- a/Engine/Source/Runtime/Engine/Private/InstancedSkinnedMeshSceneProxyDesc.cpp
+++ b/Engine/Source/Runtime/Engine/Private/InstancedSkinnedMeshSceneProxyDesc.cpp
@@ -21,6 +21,7 @@ void FInstancedSkinnedMeshSceneProxyDesc::InitializeFromInstancedSkinnedMeshComp
 #endif
 
 	AnimBankItems = InComponent->AnimBankItems;
+	bAnimBankEnabled = InComponent->bAnimBankEnabled;
 	bAllowAlwaysVisible = true;
 
 	InstanceDataSceneProxy = InComponent->GetInstanceDataSceneProxy();
diff --git a/Engine/Source/Runtime/Engine/Public/InstancedSkinnedMeshSceneProxyDesc.h b/Engine/Source/Runtime/Engine/Public/InstancedSkinnedMeshSceneProxyDesc.h
--- a/Engine/Source/Runtime/Engine/Public/InstancedSkinnedMeshSceneProxyDesc.h
+++ b/Engine/Source/Runtime/Engine/Public/InstancedSkinnedMeshSceneProxyDesc.h
@@ -28,4 +28,5 @@ struct FInstancedSkinnedMeshSceneProxyDesc : public FSkinnedMeshSceneProxyDesc
 #endif
 
 	TArrayView<const struct FAnimBankItem> AnimBankItems;
+	uint8 bAnimBankEnabled : 1 = false;
 };
