[NetPhysics] Initialize bitfield properties

diff --git a/Engine/Source/Runtime/Engine/Public/Physics/NetworkPhysicsSettingsComponent.h b/Engine/Source/Runtime/Engine/Public/Physics/NetworkPhysicsSettingsComponent.h
--- a/Engine/Source/Runtime/Engine/Public/Physics/NetworkPhysicsSettingsComponent.h
+++ b/Engine/Source/Runtime/Engine/Public/Physics/NetworkPhysicsSettingsComponent.h
@@ -83,6 +83,8 @@ struct FNetworkPhysicsSettings
 {
 	GENERATED_BODY()
 
+	FNetworkPhysicsSettings() { }
+
 	// Override properties
 	UPROPERTY(config, EditDefaultsOnly, Category = "Overrides", Meta = (InlineEditConditionToggle))
 	bool bOverrideSimProxyRepMode = false;
@@ -106,6 +108,13 @@ struct FNetworkPhysicsSettingsDefaultReplication
 {
 	GENERATED_BODY()
 
+	FNetworkPhysicsSettingsDefaultReplication()
+		: bOverrideMaxLinearHardSnapDistance(0)
+		, bOverrideDefaultLegacyHardsnapInPT(0)
+		, bOverrideCorrectConnectedBodies(0)
+		, bOverrideCorrectConnectedBodiesFriction(0)
+	{}
+
 	// Override properties
 	UPROPERTY(config, EditDefaultsOnly, Category = "Overrides", Meta = (InlineEditConditionToggle))
 	uint32 bOverrideMaxLinearHardSnapDistance : 1;
@@ -142,6 +151,25 @@ struct FNetworkPhysicsSettingsPredictiveInterpolation
 {
 	GENERATED_BODY()
 
+	FNetworkPhysicsSettingsPredictiveInterpolation() 
+		: bOverridePosCorrectionTimeBase(0)
+		, bOverridePosCorrectionTimeMin(0)
+		, bOverridePosCorrectionTimeMultiplier(0)
+		, bOverrideRotCorrectionTimeBase(0)
+		, bOverrideRotCorrectionTimeMin(0)
+		, bOverrideRotCorrectionTimeMultiplier(0)
+		, bOverridePosInterpolationTimeMultiplier(0)
+		, bOverrideRotInterpolationTimeMultiplier(0)
+		, bOverrideSoftSnapPosStrength(0)
+		, bOverrideSoftSnapRotStrength(0)
+		, bOverrideSoftSnapToSource(0)
+		, bOverrideDisableSoftSnap(0)
+		, bOverrideSkipVelocityRepOnPosEarlyOut(0)
+		, bOverridePostResimWaitForUpdate(0)
+		, bOverrideCorrectConnectedBodies(0)
+		, bOverrideCorrectConnectedBodiesFriction(0)
+	{}
+
 	// Override properties
 	UPROPERTY(config, EditDefaultsOnly, Category = "Overrides", Meta = (InlineEditConditionToggle))
 	uint32 bOverridePosCorrectionTimeBase : 1;
@@ -261,11 +289,11 @@ USTRUCT()
 struct FNetworkPhysicsSettingsResimulationErrorCorrection
 {
 	FNetworkPhysicsSettingsResimulationErrorCorrection()
-	: bOverrideResimErrorInterpolationSettings(0)
-	, ResimErrorCorrectionDuration(0.3f)
-	, ResimErrorMaximumDistanceBeforeSnapping(250.0f)
-	, ResimErrorMaximumDesyncTimeBeforeSnapping(0.6f)
-	, ResimErrorDirectionalDecayMultiplier(0.0f)
+		: bOverrideResimErrorInterpolationSettings(0)
+		, ResimErrorCorrectionDuration(0.3f)
+		, ResimErrorMaximumDistanceBeforeSnapping(250.0f)
+		, ResimErrorMaximumDesyncTimeBeforeSnapping(0.6f)
+		, ResimErrorDirectionalDecayMultiplier(0.0f)
 	{};
 
 	GENERATED_BODY()
@@ -300,9 +328,20 @@ USTRUCT()
 struct FNetworkPhysicsSettingsResimulation
 {
 	FNetworkPhysicsSettingsResimulation()
-	: bOverrideResimulationErrorThreshold_DEPRECATED(0)
-	, ResimulationErrorThreshold_DEPRECATED(10)
-	, ResimulationErrorCorrectionSettings()
+		: bOverrideResimulationErrorPositionThreshold(0)
+		, bOverrideResimulationErrorRotationThreshold(0)
+		, bOverrideResimulationErrorLinearVelocityThreshold(0)
+		, bOverrideResimulationErrorAngularVelocityThreshold(0)
+		, bOverrideRuntimeCorrectionEnabled(0)
+		, bOverrideRuntimeVelocityCorrection(0)
+		, bOverrideRuntimeCorrectConnectedBodies(0)
+		, bOverridePosStabilityMultiplier(0)
+		, bOverrideRotStabilityMultiplier(0)
+		, bOverrideVelStabilityMultiplier(0)
+		, bOverrideAngVelStabilityMultiplier(0)
+		, bOverrideResimulationErrorThreshold_DEPRECATED(0)
+		, ResimulationErrorThreshold_DEPRECATED(10)
+		, ResimulationErrorCorrectionSettings()
 	{};
 
 	GENERATED_BODY()
@@ -404,7 +443,21 @@ struct FNetworkPhysicsSettingsResimulation
 USTRUCT()
 struct FNetworkPhysicsSettingsNetworkPhysicsComponent
 {
-	FNetworkPhysicsSettingsNetworkPhysicsComponent() { }
+	FNetworkPhysicsSettingsNetworkPhysicsComponent()
+		: bOverrideRedundantInputs(0)
+		, bOverrideRedundantRemoteInputs(0)
+		, bOverrideRedundantStates(0)
+		, bOverrideCompareStateToTriggerRewind(0)
+		, bOverridebCompareStateToTriggerRewindIncludeSimProxies(0)
+		, bOverrideCompareInputToTriggerRewind(0)
+		, bOverrideEnableUnreliableFlow(0)
+		, bOverrideEnableReliableFlow(0)
+		, bOverrideApplyDataInsteadOfMergeData(0)
+		, bOverrideAllowInputExtrapolation(0)
+		, bOverrideValidateDataOnGameThread(0)
+		, bOverrideApplySimProxyStateAtRuntime(0)
+		, bOverrideApplySimProxyInputAtRuntime(0)
+	{};
 
 	GENERATED_BODY()
 
