[Audio Insights] Play sound at location gives sound owner as actor label

diff --git a/Engine/Source/Runtime/Engine/Private/Audio/AudioTraceUtil.cpp b/Engine/Source/Runtime/Engine/Private/Audio/AudioTraceUtil.cpp
--- a/Engine/Source/Runtime/Engine/Private/Audio/AudioTraceUtil.cpp
+++ b/Engine/Source/Runtime/Engine/Private/Audio/AudioTraceUtil.cpp
@@ -69,7 +69,7 @@ namespace Audio::Trace
 				}
 			}
 
-			return FString();
+			return InActiveSound.GetOwnerName();
 		}
 
 		FName GetOwnerActorIconName(const FActiveSound& InActiveSound)
