Changed bSlideAlongNavMeshEdge in the character movement component to be blueprint read write

diff --git a/Engine/Source/Runtime/Engine/Classes/GameFramework/CharacterMovementComponent.h b/Engine/Source/Runtime/Engine/Classes/GameFramework/CharacterMovementComponent.h
--- a/Engine/Source/Runtime/Engine/Classes/GameFramework/CharacterMovementComponent.h
+++ b/Engine/Source/Runtime/Engine/Classes/GameFramework/CharacterMovementComponent.h
@@ -1099,7 +1099,7 @@ protected:
 	uint8 bProjectNavMeshWalking : 1 = false;
 
 	/** Should we attempt to slide along a navmesh instead of projecting a point on the navmesh and moving to that */
-	UPROPERTY(Category="Character Movement: NavMesh Movement", EditAnywhere, BlueprintReadOnly)
+	UPROPERTY(Category="Character Movement: NavMesh Movement", EditAnywhere, BlueprintReadWrite)
 	uint8 bSlideAlongNavMeshEdge : 1 = false;
 
 	/** Use both WorldStatic and WorldDynamic channels for NavWalking geometry conforming */
