Small camera lens effect optimization to avoid unnecessary calculations.
Shouldn't result in a change in behavior but gated behind the CVar CameraLensEffect.AttachedEmitterTransformOptimizationEnabled in case anything behaves unexpectedly.

#rnx

diff --git a/Engine/Source/Runtime/Engine/Private/Camera/CameraLensEffectInterface.cpp b/Engine/Source/Runtime/Engine/Private/Camera/CameraLensEffectInterface.cpp
--- a/Engine/Source/Runtime/Engine/Private/Camera/CameraLensEffectInterface.cpp
+++ b/Engine/Source/Runtime/Engine/Private/Camera/CameraLensEffectInterface.cpp
@@ -10,6 +10,16 @@
 
 #include UE_INLINE_GENERATED_CPP_BY_NAME(CameraLensEffectInterface)
 
+namespace CameraLensEffectCVars
+{
+	static bool bAttachedEmitterTransformOptimizationEnabled = true;
+	static FAutoConsoleVariableRef CVarAttachedEmitterTransformOptimizationEnabled(
+		TEXT("CameraLensEffect.AttachedEmitterTransformOptimizationEnabled"),
+		bAttachedEmitterTransformOptimizationEnabled,
+		TEXT("If true, GetAttachedEmitterTransform avoids doing unnecessary work when the camera FOV matches the emitter's base FOV, and avoids recalculating constant values."),
+		ECVF_Default);
+}
+
 TArray<UFXSystemComponent*> ICameraLensEffectInterface::GetParticleComponents() const
 {
 	const AActor* const ConstActorThis = Cast<AActor>(this);
@@ -42,8 +52,14 @@ FTransform ICameraLensEffectInterface::GetAttachedEmitterTransform(const AActor*
 		// base dist uses BaseFOV which is set on the indiv camera lens effect class
 		FTransform RelativeTransformAdjustedForFOV = InterfaceEmitter->GetRelativeTransform();
 		FVector AdjustedRelativeLoc = RelativeTransformAdjustedForFOV.GetLocation();
-		AdjustedRelativeLoc.X *= FMath::Tan(InterfaceEmitter->GetBaseFOV() * 0.5f * UE_PI / 180.f) / FMath::Tan(CamFOVDeg * 0.5f * UE_PI / 180.f);
-		RelativeTransformAdjustedForFOV.SetLocation(AdjustedRelativeLoc);
+
+		const float BaseFOV = InterfaceEmitter->GetBaseFOV();
+		if (!CameraLensEffectCVars::bAttachedEmitterTransformOptimizationEnabled || !FMath::IsNearlyEqual(BaseFOV, CamFOVDeg))
+		{
+			constexpr float HalfFovDegreesToRadians = 0.5f * UE_PI / 180.f;
+			AdjustedRelativeLoc.X *= FMath::Tan(BaseFOV * HalfFovDegreesToRadians) / FMath::Tan(CamFOVDeg * HalfFovDegreesToRadians);
+			RelativeTransformAdjustedForFOV.SetLocation(AdjustedRelativeLoc);
+		}
 
 		FTransform const CameraToWorld(CamRot, CamLoc);
 
@@ -52,6 +68,7 @@ FTransform ICameraLensEffectInterface::GetAttachedEmitterTransform(const AActor*
 
 		return EffectToWorld;
 	}
+
 	return FTransform::Identity;
 }
 
