CurveLinearColor remove ClampNonNegative so that negative RGB is allowed to affect the HSV conversion, which restores the output color to previous
Note that negative RGB still does not leave the function due to use of HSV
The previous bug that required ClampNonNegative due to slight negatives turning white is now fixed in Color.cpp so the clamp is not required here.

diff --git a/Engine/Source/Runtime/Engine/Private/Curves/CurveLinearColor.cpp b/Engine/Source/Runtime/Engine/Private/Curves/CurveLinearColor.cpp
--- a/Engine/Source/Runtime/Engine/Private/Curves/CurveLinearColor.cpp
+++ b/Engine/Source/Runtime/Engine/Private/Curves/CurveLinearColor.cpp
@@ -68,20 +68,10 @@ UCurveLinearColor::UCurveLinearColor(const FObjectInitializer& ObjectInitializer
 #endif
 }
 
-static FLinearColor ClampNonNegative(const FLinearColor & InColor)
-{
-	FLinearColor OutColor;
-	OutColor.R = FMath::Max(InColor.R,0.f);
-	OutColor.G = FMath::Max(InColor.G,0.f);
-	OutColor.B = FMath::Max(InColor.B,0.f);
-	OutColor.A = FMath::Max(InColor.A,0.f);
-	return OutColor;
-}
-
 FLinearColor UCurveLinearColor::GetLinearColorValue( float InTime ) const
 {
 	return GetAdjustedColorValue(InTime, false);
-	}
+}
 
 FLinearColor UCurveLinearColor::GetClampedLinearColorValue(float InTime) const
 {
@@ -92,18 +82,16 @@ FLinearColor UCurveLinearColor::GetAdjustedColorValue(float InTime, bool bClampO
 {
 	// Logic copied from .\Engine\Source\Developer\TextureCompressor\Private\TextureCompressorModule.cpp , from AdjustColorsOld
 	//  @todo : see fixes in AdjustColorsNew
-	FLinearColor OriginalColor = GetUnadjustedLinearColorValue(InTime);
+	const FLinearColor OriginalColor = GetUnadjustedLinearColorValue(InTime);
 	
-	// allow > 1.0 at this point but not < 0.0
-	//	we're just about to convert this RGB to HSV, that doesn't support negatives
-	OriginalColor = ClampNonNegative(OriginalColor);
+	// note RGB can be negative or > 1.0 here
 
 	// if bClampOver1Always , always clamp to stay in[ 0,1]
 	// else, Only clamp value if the color is input RGB <= 1 (fuzzy)
 	const bool bShouldClampValue = bClampOver1Always || 
 		( FMath::Max3(OriginalColor.R,OriginalColor.G,OriginalColor.B) < 1.0f+UE_KINDA_SMALL_NUMBER );
 
-	// Convert to HSV
+	// Convert to HSV (note RGB can be out of [0,1] range expected by HSV here)
 	FLinearColor HSVColor = OriginalColor.LinearRGBToHSV();
 	float& PixelHue = HSVColor.R;
 	float& PixelSaturation = HSVColor.G;
@@ -141,10 +129,10 @@ FLinearColor UCurveLinearColor::GetAdjustedColorValue(float InTime, bool bClampO
 
 	// Clamp HSV values
 	// HSVToLinearRGB now does the needed clamps, except for Value
-		if (bShouldClampValue)
-		{
-			PixelValue = FMath::Clamp(PixelValue, 0.0f, 1.0f);
-		}
+	if (bShouldClampValue)
+	{
+		PixelValue = FMath::Clamp(PixelValue, 0.0f, 1.0f);
+	}
 
 	// Convert back to a linear color
 	FLinearColor LinearColor = HSVColor.HSVToLinearRGB();
