VolumetricClouds: Cooker: SkipOnlyEditorOnly: Add a FCookLoadScope around the load of softobjectpath so that the cooker knows it is expected and does not warn.
#rnx

diff --git a/Engine/Source/Runtime/Engine/Private/Components/VolumetricCloudComponent.cpp b/Engine/Source/Runtime/Engine/Private/Components/VolumetricCloudComponent.cpp
--- a/Engine/Source/Runtime/Engine/Private/Components/VolumetricCloudComponent.cpp
+++ b/Engine/Source/Runtime/Engine/Private/Components/VolumetricCloudComponent.cpp
@@ -13,6 +13,7 @@
 #include "SceneInterface.h"
 #include "UObject/UE5MainStreamObjectVersion.h"
 #include "LocalVertexFactory.h"
+#include "UObject/ICookInfo.h"
 
 #if WITH_EDITOR
 #include "MaterialShared.h"
@@ -176,6 +177,7 @@ void UVolumetricCloudComponent::OnRegister()
 
 	// Make sure the assigned TSoftObjectPtr material is loaded by that point so that we can create proxy render side.
 	// And it is assigned to a TObjectPtr to make sure it remains loaded and ignored by the GC.
+	FCookLoadScope CookLoadScope(ECookLoadType::UsedInGame);
 	MaterialLifetimePtr = Material.LoadSynchronous();
 }
 
