Don't call UpdateResource when changing texture LOD settings for render targets

diff --git a/Engine/Source/Runtime/Engine/Private/TextureLODSettings.cpp b/Engine/Source/Runtime/Engine/Private/TextureLODSettings.cpp
--- a/Engine/Source/Runtime/Engine/Private/TextureLODSettings.cpp
+++ b/Engine/Source/Runtime/Engine/Private/TextureLODSettings.cpp
@@ -2,6 +2,7 @@
 
 #include "Engine/TextureLODSettings.h"
 #include "Engine/TextureCube.h"
+#include "Engine/TextureRenderTarget.h"
 #include "HAL/IConsoleManager.h"
 #include "ImageCoreUtils.h"
 #include "Interfaces/ITargetPlatform.h"
@@ -460,11 +461,11 @@ void UTextureLODSettings::SetPreviewTextureGroupRedirect(UTextureLODSettings* In
 	{
 		PreviewTextureGroupRedirect = InPreviewTextureGroupRedirect;
 
-		for (FThreadSafeObjectIterator ObjIterator(UObject::StaticClass(), RF_NoFlags); ObjIterator; ++ObjIterator)
+		for (FThreadSafeObjectIterator ObjIterator(UTexture::StaticClass(), /*bOnlyGCedObjects = */false, /*AdditionalExclusionFlags = */RF_NoFlags); ObjIterator; ++ObjIterator)
 		{
 			//Call UpdateResource on Textures when changing the TextureLodGroup Settings to update the CachePlatformData
 			UTexture* Texture = Cast<UTexture>(*ObjIterator);
-			if (Texture != nullptr)
+			if (!Texture->IsA<UTextureRenderTarget>())
 			{
 				Texture->UpdateResource();
 			}
