Add log event for when virtual texture pools are recreated.
Also added code for recreate of material cache VTs.
#rnx

diff --git a/Engine/Source/Runtime/Engine/Private/VT/VirtualTextureRecreate.cpp b/Engine/Source/Runtime/Engine/Private/VT/VirtualTextureRecreate.cpp
--- a/Engine/Source/Runtime/Engine/Private/VT/VirtualTextureRecreate.cpp
+++ b/Engine/Source/Runtime/Engine/Private/VT/VirtualTextureRecreate.cpp
@@ -2,13 +2,14 @@
 
 #include "VT/VirtualTextureRecreate.h"
 
-#include "Components/RuntimeVirtualTextureComponent.h"
 #include "Components/PrimitiveComponent.h"
+#include "Components/RuntimeVirtualTextureComponent.h"
 #include "Engine/Texture.h"
-#include "TextureResource.h"
 #include "MaterialCache/MaterialCacheVirtualTexture.h"
-#include "VT/RuntimeVirtualTexture.h"
+#include "TextureResource.h"
 #include "UObject/UObjectIterator.h"
+#include "VirtualTexturing.h"
+#include "VT/RuntimeVirtualTexture.h"
 
 namespace VirtualTexture
 {
@@ -17,6 +18,8 @@ namespace VirtualTexture
 	// Note that for pool size scale changes, there will be a transition period (with high memory watermark) where new pools are created before old pools are released.
 	void Recreate()
 	{
+		UE_LOG(LogVirtualTexturing, Display, TEXT("Recreating virtual texture pools."));
+
 		// Reinit streaming virtual textures.
 		for (TObjectIterator<UTexture> It(RF_ClassDefaultObject, true, EInternalObjectFlags::Garbage); It; ++It)
 		{
@@ -31,10 +34,23 @@ namespace VirtualTexture
 		{
 			It->MarkRenderStateDirty();
 		}
+
+		// Reinit material cache virtual textures.
+		for (TObjectIterator<UMaterialCacheVirtualTexture> It(RF_ClassDefaultObject, true, EInternalObjectFlags::Garbage); It; ++It)
+		{
+			It->UpdateResourceWithParams(UTexture::EUpdateResourceFlags::None);
+
+			if (UPrimitiveComponent* Owner = It->OwningComponent.Get())
+			{
+				Owner->MarkRenderStateDirty();
+			}
+		}
 	}
 
-	void Recreate(TConstArrayView < TEnumAsByte<EPixelFormat> > InFormat)
+	void Recreate(TConstArrayView< TEnumAsByte<EPixelFormat> > InFormat)
 	{
+		UE_LOG(LogVirtualTexturing, Display, TEXT("Recreating virtual texture pools for formats."));
+
 		// Reinit streaming virtual textures that match one of our passed in format arrays.
 		for (TObjectIterator<UTexture> It(RF_ClassDefaultObject, true, EInternalObjectFlags::Garbage); It; ++It)
 		{
