PR #12770: Fixed a math precision bug in UInstancedStaticMeshComponent::ApplyComponentInstanceData

diff --git a/Engine/Source/Runtime/Engine/Private/InstancedStaticMesh.cpp b/Engine/Source/Runtime/Engine/Private/InstancedStaticMesh.cpp
--- a/Engine/Source/Runtime/Engine/Private/InstancedStaticMesh.cpp
+++ b/Engine/Source/Runtime/Engine/Private/InstancedStaticMesh.cpp
@@ -2331,7 +2331,7 @@ void UInstancedStaticMeshComponent::ApplyComponentInstanceData(FInstancedStaticM
 
 		for (int32 InstanceIndex = 0; InstanceIndex < PerInstanceSMData.Num(); ++InstanceIndex)
 		{
-			if (PerInstanceSMData[InstanceIndex].Transform != InstancedMeshData->PerInstanceSMData[InstanceIndex].Transform)
+			if (!PerInstanceSMData[InstanceIndex].Transform.Equals(InstancedMeshData->PerInstanceSMData[InstanceIndex].Transform))
 			{
 				bMatch = false;
 				break;
