Added crash logs for UActorDescContainerSubsystem::GetChecked.

diff --git a/Engine/Source/Runtime/Engine/Private/WorldPartition/ActorDescContainerSubsystem.cpp b/Engine/Source/Runtime/Engine/Private/WorldPartition/ActorDescContainerSubsystem.cpp
--- a/Engine/Source/Runtime/Engine/Private/WorldPartition/ActorDescContainerSubsystem.cpp
+++ b/Engine/Source/Runtime/Engine/Private/WorldPartition/ActorDescContainerSubsystem.cpp
@@ -21,7 +21,7 @@ UActorDescContainerSubsystem* UActorDescContainerSubsystem::Get()
 UActorDescContainerSubsystem& UActorDescContainerSubsystem::GetChecked()
 {
 	UActorDescContainerSubsystem* ContainerSubsystem = Get();
-	check(ContainerSubsystem);
+	checkf(ContainerSubsystem, TEXT("Failed to get ActorDescContainerSubsystem: IsEngineExitRequested=%d, GEngine=%s"), IsEngineExitRequested() ? 1 : 0, *GetNameSafe(GEngine));
 	return *Get();
 }
 
