[MIR] Implement TextureProperty via preshader in new material translator.

Texture properties must be evaluated at runtime when their source texture is a uniform parameter.
This evaluation goes beyond the previously integrated numeric parameter evaluation as the texture properties need to refer to a source texture.
This is implemented as a new MIR instruction `PreshaderParameter` which takes a direct UE::Shader::EPreshaderOpcode to encode the desired property.
No extended preshader computations such as nested math operations are supported yet as this serves only to evaluate runtime information that depends on other uniform parameters.

#test Used new translator in TextureInfo_Nearest asset in EngineTest. No regressions in test suite for new translator.
#rnx

diff --git a/Engine/Source/Runtime/Engine/Private/Materials/HLSLMaterialTranslator.cpp b/Engine/Source/Runtime/Engine/Private/Materials/HLSLMaterialTranslator.cpp
--- a/Engine/Source/Runtime/Engine/Private/Materials/HLSLMaterialTranslator.cpp
+++ b/Engine/Source/Runtime/Engine/Private/Materials/HLSLMaterialTranslator.cpp
@@ -4180,18 +4180,13 @@ int32 FHLSLMaterialTranslator::AccessUniformExpression(int32 Index)
 
 			int32 UniformOffset = CurrentBoolUniformOffset * 32u + CurrentNumBoolComponents;
 
-			FUniformExpressionSet& UniformExpressionSet = MaterialCompilationOutput.UniformExpressionSet;
-			FMaterialUniformPreshaderHeader& PreshaderHeader = UniformExpressionSet.UniformPreshaders.AddDefaulted_GetRef();
-			PreshaderHeader.FieldIndex = UniformExpressionSet.UniformPreshaderFields.Num();
-			PreshaderHeader.NumFields = 1;
-			PreshaderHeader.OpcodeOffset = UniformExpressionSet.UniformPreshaderData.Num();
-			CodeChunk.UniformExpression->WriteNumberOpcodes(UniformExpressionSet.UniformPreshaderData);
-			PreshaderHeader.OpcodeSize = UniformExpressionSet.UniformPreshaderData.Num() - PreshaderHeader.OpcodeOffset;
-
-			FMaterialUniformPreshaderField& PreshaderField = MaterialCompilationOutput.UniformExpressionSet.UniformPreshaderFields.AddDefaulted_GetRef();
-			PreshaderField.BufferOffset = UniformOffset;
-			PreshaderField.Type = UE::Shader::MakeValueType(UE::Shader::EValueComponentType::Bool, NumComponents);
-			PreshaderField.ComponentIndex = 0;
+			MaterialCompilationOutput.UniformExpressionSet.WriteUniformPreshaderEntry(
+				UniformOffset, UE::Shader::MakeValueType(UE::Shader::EValueComponentType::Bool, NumComponents),
+				[&CodeChunk](UE::Shader::FPreshaderData& UniformPreshaderData) -> void
+				{
+					CodeChunk.UniformExpression->WriteNumberOpcodes(UniformPreshaderData);
+				}
+			);
 
 			CodeChunk.UniformExpression->UniformOffset = UniformOffset;
 			CurrentNumBoolComponents += NumComponents;
@@ -4244,17 +4239,13 @@ int32 FHLSLMaterialTranslator::AccessUniformExpression(int32 Index)
 
 			const UE::Shader::EValueComponentType ComponentType = bIsLWC ? UE::Shader::EValueComponentType::Double : UE::Shader::EValueComponentType::Float;
 
-			FMaterialUniformPreshaderHeader& Preshader = MaterialCompilationOutput.UniformExpressionSet.UniformPreshaders.AddDefaulted_GetRef();
-			Preshader.OpcodeOffset = MaterialCompilationOutput.UniformExpressionSet.UniformPreshaderData.Num();
-			CodeChunk.UniformExpression->WriteNumberOpcodes(MaterialCompilationOutput.UniformExpressionSet.UniformPreshaderData);
-			Preshader.OpcodeSize = MaterialCompilationOutput.UniformExpressionSet.UniformPreshaderData.Num() - Preshader.OpcodeOffset;
-			Preshader.FieldIndex = MaterialCompilationOutput.UniformExpressionSet.UniformPreshaderFields.Num();
-			Preshader.NumFields = 1;
-
-			FMaterialUniformPreshaderField& PreshaderField = MaterialCompilationOutput.UniformExpressionSet.UniformPreshaderFields.AddDefaulted_GetRef();
-			PreshaderField.BufferOffset = UniformPreshaderOffset;
-			PreshaderField.Type = UE::Shader::MakeValueType(ComponentType, NumComponents);
-			PreshaderField.ComponentIndex = 0;
+			MaterialCompilationOutput.UniformExpressionSet.WriteUniformPreshaderEntry(
+				UniformPreshaderOffset, UE::Shader::MakeValueType(ComponentType, NumComponents),
+				[&CodeChunk](UE::Shader::FPreshaderData& UniformPreshaderData) -> void
+				{
+					CodeChunk.UniformExpression->WriteNumberOpcodes(UniformPreshaderData);
+				}
+			);
 
 			CodeChunk.UniformExpression->UniformOffset = UniformPreshaderOffset;
 			UniformPreshaderOffset += bIsLWC ? NumComponents * 2u : NumComponents;
@@ -7977,7 +7968,7 @@ int32 FHLSLMaterialTranslator::TextureProperty(int32 TextureIndex, EMaterialExpo
 		return Errorf(TEXT("Expected a texture expression"));
 	}
 
-	const EMaterialValueType ValueType = (TextureType == MCT_VolumeTexture || TextureType == MCT_Texture2DArray || TextureType == MCT_SparseVolumeTexture) ? MCT_Float3 : MCT_Float2;
+	const EMaterialValueType ValueType = UE::MaterialTranslatorUtils::GetTexturePropertyValueType(TextureType);
 	FAddUniformExpressionScope Scope(this);
 	return AddUniformExpression(Scope, new FMaterialUniformExpressionTextureProperty(TextureExpression, Property), ValueType, TEXT(""));
 }
diff --git a/Engine/Source/Runtime/Engine/Private/Materials/MaterialExpressionsIR.cpp b/Engine/Source/Runtime/Engine/Private/Materials/MaterialExpressionsIR.cpp
--- a/Engine/Source/Runtime/Engine/Private/Materials/MaterialExpressionsIR.cpp
+++ b/Engine/Source/Runtime/Engine/Private/Materials/MaterialExpressionsIR.cpp
@@ -252,6 +252,7 @@
 #include "MaterialShared.h"
 #include "Misc/MemStackUtility.h"
 #include "RenderUtils.h"
+#include "TextureResource.h"
 #include "ColorManagement/ColorSpace.h"
 
 #include "Materials/MaterialIREmitter.h"
@@ -1366,7 +1367,14 @@ void UMaterialExpressionTextureSample::Build(MIR::FEmitter& Em)
 
 	UE_MIR_CHECKPOINT(Em);
 
-	BuildTextureSampleExpression(Em, this, TextureValue, TextureValue->GetTexture()->GetMaterialType());
+	UTexture* DefaultTexture = Em.GetTextureFromValue(TextureValue);
+	if (!DefaultTexture)
+	{
+		Em.Error(TEXT("Missing texture object from input"));
+		return;
+	}
+
+	BuildTextureSampleExpression(Em, this, TextureValue, DefaultTexture->GetMaterialType());
 }
 
 static FValueRef BuildTextureObjectParameter(MIR::FEmitter& Em, UMaterialExpressionTextureSampleParameter* Expr)
@@ -1446,7 +1454,65 @@ void UMaterialExpressionTextureCoordinate::Build(MIR::FEmitter& Em)
 	// Multiply the UV input by the UV tiling constants
 	OutputValue = Em.Multiply(OutputValue, Em.ConstantFloat2({ UTiling, VTiling }));
 	
-	Em.Output(GetOutput(0), OutputValue);
+	Em.Output(0, OutputValue);
+}
+
+void UMaterialExpressionTextureProperty::Build(MIR::FEmitter& Em)
+{
+	MIR::FValueRef TextureValue = Em.Input(&TextureObject);
+	UE_MIR_CHECKPOINT(Em);
+
+	UTexture* InputTexture = Em.GetTextureFromValue(TextureValue);
+	if (!InputTexture)
+	{
+		Em.Error(TEXT("Missing texture object"));
+		return;
+	}
+
+	const bool bTexelSizeInUVSpace = Property == TMTM_TexelSize;
+
+	if (MIR::FUniformParameter* UniformParameter = TextureValue->As<MIR::FUniformParameter>())
+	{
+		// If the property is connected to a texture paramter, the property needs to be evaluated at runtime via preshaders
+		const UE::Shader::EPreshaderOpcode PreshaderOpcode =
+			bTexelSizeInUVSpace
+				? UE::Shader::EPreshaderOpcode::TexelSize
+				: UE::Shader::EPreshaderOpcode::TextureSize;
+
+		Em.Output(0, Em.PreshaderParameter(PreshaderOpcode, UniformParameter));
+	}
+	else
+	{
+		// If the property is connected to a texture object directly, evaluate the property directly and emit a constant expression
+		const EMaterialValueType TextureType = InputTexture->GetMaterialType();
+		const EMaterialValueType PropertyType = UE::MaterialTranslatorUtils::GetTexturePropertyValueType(TextureType);
+
+		FTextureResource* TextureResource = InputTexture->GetResource();
+		if (!TextureResource)
+		{
+			Em.Error(TEXT("Missing texture resource to encode texture property expression"));
+			return;
+		}
+
+		const FVector3f TextureSize
+		{
+			static_cast<float>(TextureResource->GetSizeX()),
+			static_cast<float>(TextureResource->GetSizeY()),
+			static_cast<float>(TextureResource->GetSizeZ())
+		};
+		if (bTexelSizeInUVSpace)
+		{
+			Em.Output(0, PropertyType == MCT_Float3
+				? Em.ConstantFloat3(FVector3f{ 1.0f / TextureSize.X, 1.0f / TextureSize.Y, 1.0f / TextureSize.Z })
+				: Em.ConstantFloat2(FVector2f{ 1.0f / TextureSize.X, 1.0f / TextureSize.Y }));
+		}
+		else
+		{
+			Em.Output(0, PropertyType == MCT_Float3
+				? Em.ConstantFloat3(TextureSize)
+				: Em.ConstantFloat2(FVector2f{ TextureSize.X, TextureSize.Y }));
+		}
+	}
 }
 
 #if WITH_EDITOR
diff --git a/Engine/Source/Runtime/Engine/Private/Materials/MaterialIR.cpp b/Engine/Source/Runtime/Engine/Private/Materials/MaterialIR.cpp
--- a/Engine/Source/Runtime/Engine/Private/Materials/MaterialIR.cpp
+++ b/Engine/Source/Runtime/Engine/Private/Materials/MaterialIR.cpp
@@ -63,6 +63,7 @@ const TCHAR* LexToString(EValueKind Kind)
 		case VK_Nop: return TEXT("Nop");
 		case VK_Call: return TEXT("Call");
 		case VK_CallParameterOutput: return TEXT("CallParameterOutput");
+		case VK_PreshaderParameter: return TEXT("PreshaderParameter");
 		/* invalid entries */
 		case VK_InstructionBegin:
 		case VK_InstructionEnd:
@@ -128,6 +129,7 @@ uint32 FValue::GetSizeInBytes() const
 		case VK_Nop: return sizeof(FNop);
 		case VK_Call: return sizeof(FCall);
 		case VK_CallParameterOutput: return sizeof(FCallParameterOutput);
+		case VK_PreshaderParameter: return sizeof(FPreshaderParameter);
 			/* invalid entries */
 		case VK_InstructionBegin:
 		case VK_InstructionEnd:
@@ -236,6 +238,12 @@ TConstArrayView<FValue*> FValue::GetUses() const
 			return { &This->Call, FCallParameterOutput::NumStaticUses };
 		}
 
+		case VK_PreshaderParameter:
+		{
+			auto This = static_cast<const FPreshaderParameter*>(this);
+			return { &This->SourceParameter, FPreshaderParameter::NumStaticUses };
+		}
+
 		default: UE_MIR_UNREACHABLE();
 	}
 }
diff --git a/Engine/Source/Runtime/Engine/Private/Materials/MaterialIREmitter.cpp b/Engine/Source/Runtime/Engine/Private/Materials/MaterialIREmitter.cpp
--- a/Engine/Source/Runtime/Engine/Private/Materials/MaterialIREmitter.cpp
+++ b/Engine/Source/Runtime/Engine/Private/Materials/MaterialIREmitter.cpp
@@ -13,6 +13,7 @@
 #include "MaterialExpressionIO.h"
 #include "Engine/Texture.h"
 #include "Engine/Texture2D.h"
+#include "TextureResource.h"
 #include "RenderUtils.h"
 
 #if WITH_EDITOR
@@ -1074,6 +1075,50 @@ FValueRef FEmitter::Parameter(FName Name, FMaterialParameterMetadata& Metadata,
 	}
 }
 
+UTexture* FEmitter::GetTextureFromValue(MIR::FValueRef Texture)
+{
+	if (!Texture)
+	{
+		return nullptr;
+	}
+	// Handle case for uniform parameters here because their metadata is stored in the IR module, which the emitter has access to.
+	if (FUniformParameter* UniformParameter = Texture->As<FUniformParameter>())
+	{
+		const FMaterialParameterMetadata& ParameterMetadata = Module->GetParameterMetadata(UniformParameter->ParameterIdInModule);
+		if (UTexture* ParameterDefaultTexture = ::Cast<UTexture>(ParameterMetadata.Value.AsTextureObject()))
+		{
+			return ParameterDefaultTexture;
+		}
+	}
+	return Texture->GetTexture();
+}
+
+// For now, only a small subset of opcodes are supported: TextureSize, TexelSize.
+FValueRef FEmitter::PreshaderParameter(UE::Shader::EPreshaderOpcode Opcode, FValueRef SourceParameter)
+{
+	checkf(
+		Opcode == UE::Shader::EPreshaderOpcode::TextureSize ||
+		Opcode == UE::Shader::EPreshaderOpcode::TexelSize,
+		TEXT("Preshader opcode (0x%X) not supported for parameters in new material translator"), (int32)Opcode
+	);
+
+	UTexture* SourceParameterTexture = GetTextureFromValue(SourceParameter);
+	if (!SourceParameterTexture)
+	{
+		Error(TEXT("Missing default texture from source parameter"));
+		return Poison();
+	}
+
+	const EMaterialValueType TextureType = SourceParameterTexture->GetMaterialType();
+	const EMaterialValueType PropertyType = UE::MaterialTranslatorUtils::GetTexturePropertyValueType(TextureType);
+
+	FPreshaderParameter Prototype = MakePrototype<FPreshaderParameter>(FType::FromMaterialValueType(PropertyType));
+	Prototype.SourceParameter = SourceParameter;
+	Prototype.Opcode = Opcode;
+	Prototype.TextureIndex = Material->GetReferencedTextures().Find(SourceParameterTexture);
+	return EmitPrototype(*this, Prototype);
+}
+
 FValueRef FEmitter::CustomPrimitiveData(uint32 PrimitiveDataIndex)
 {
 	// UE_MIR_TODO();
diff --git a/Engine/Source/Runtime/Engine/Private/Materials/MaterialIRToHLSLTranslator.cpp b/Engine/Source/Runtime/Engine/Private/Materials/MaterialIRToHLSLTranslator.cpp
--- a/Engine/Source/Runtime/Engine/Private/Materials/MaterialIRToHLSLTranslator.cpp
+++ b/Engine/Source/Runtime/Engine/Private/Materials/MaterialIRToHLSLTranslator.cpp
@@ -436,6 +436,7 @@ struct FPrivate : FMaterialIRToHLSLTranslation
 			case MIR::VK_Nop: LowerNop(static_cast<const MIR::FNop*>(Instr)); break;
 			case MIR::VK_Call: LowerCall(static_cast<const MIR::FCall*>(Instr)); break;
 			case MIR::VK_CallParameterOutput: LowerCallOutput(static_cast<const MIR::FCallParameterOutput*>(Instr)); break;
+			case MIR::VK_PreshaderParameter: LowerPreshaderParameter(static_cast<const MIR::FPreshaderParameter*>(Instr)); break;
 			default: UE_MIR_UNREACHABLE();
 		}
 	
@@ -842,28 +843,28 @@ struct FPrivate : FMaterialIRToHLSLTranslation
 		Printer << NewLine << EndArgs;
 	}
 
-	ENoOp LowerPrimitiveUniformParameter(const MIR::FUniformParameter* UniformParameter)
+	void EmitPreshaderBufferReadDoubleVector(const MIR::FPrimitive& PrimitiveType, uint32 GlobalComponentOffset)
 	{
-		MIR::FPrimitive PrimitiveType = UniformParameter->Type.GetPrimitive();
-		check(PrimitiveType.IsScalar() || PrimitiveType.IsRowVector()); // no matrices yet
+		const int32 NumComponents = PrimitiveType.NumComponents();
 
-		// LWC parameters are handled differently, they have their own dedicated function.
-		if (PrimitiveType.IsDouble())
-		{
-			return LowerPrimitiveUniformParameterLWC(UniformParameter, PrimitiveType.NumComponents());
-		}
+		// Index of the float4 slot.  LWC parameters always have four elements, so there won't be an offset.
+		const uint32 BufferSlotIndex = GlobalComponentOffset / 4;
+
+		Printer << TEXTVIEW("DFToWS(MakeDFVector") << NumComponents;
+		Printer << TEXTVIEW("(Material.PreshaderBuffer[") << BufferSlotIndex << TEXTVIEW("]") << GVector4SwizzleSubset[NumComponents - 1];
+		Printer << TEXTVIEW(", Material.PreshaderBuffer[") << BufferSlotIndex + 1 << TEXTVIEW("]") << GVector4SwizzleSubset[NumComponents - 1] << TEXTVIEW("))");
+	}
+
+	void EmitPreshaderBufferReadFloatVector(const MIR::FPrimitive& PrimitiveType, uint32 GlobalComponentOffset)
+	{
+		const int32 NumComponents = PrimitiveType.NumComponents();
 
-		const FUniformExpressionSet& UniformExpressionSet = Module->GetCompilationOutput().UniformExpressionSet;
-		
-		// Get the global float4 component index (e.g. if this is 13, it refer to PreshaderBuffer[3].y)
-		uint32 GlobalComponentOffset = UniformExpressionSet.GetNumericParameterEvaluationOffset(UniformParameter->Analysis_UniformParameterIndex);
-		
 		// Index of the float4 slot
 		uint32 BufferSlotIndex = GlobalComponentOffset / 4;
 
 		// Starting component of the float4 slot
 		uint32 BufferSlotOffset = GlobalComponentOffset % 4;
-		
+
 		if (PrimitiveType.IsInteger())
 		{
 			Printer << TEXTVIEW("asint(");
@@ -871,12 +872,12 @@ struct FPrivate : FMaterialIRToHLSLTranslation
 
 		Printer << TEXTVIEW("Material.PreshaderBuffer[") << BufferSlotIndex << TEXTVIEW("]");
 
-		if (PrimitiveType.NumComponents() < 4)
+		if (NumComponents < 4)
 		{
 			Printer << TEXTVIEW(".");
 
 			static const TCHAR* Components = TEXT("xyzw");
-			for (int32 i = 0; i < PrimitiveType.NumComponents(); ++i)
+			for (int32 i = 0; i < NumComponents; ++i)
 			{
 				check(BufferSlotOffset + i < sizeof(Components));
 				Printer.Buffer.AppendChar(Components[BufferSlotOffset + i]);
@@ -887,6 +888,40 @@ struct FPrivate : FMaterialIRToHLSLTranslation
 		{
 			Printer << TEXTVIEW(")"); // close the "asint(" bracket
 		}
+	}
+
+	void EmitPreshaderBufferRead(const MIR::FPrimitive& PrimitiveType, uint32 GlobalComponentOffset)
+	{
+		check(PrimitiveType.IsScalar() || PrimitiveType.IsRowVector()); // no matrices yet
+
+		// LWC parameters are handled differently, they have their own dedicated function.
+		if (PrimitiveType.IsDouble())
+		{
+			EmitPreshaderBufferReadDoubleVector(PrimitiveType, GlobalComponentOffset);
+		}
+		else
+		{
+			EmitPreshaderBufferReadFloatVector(PrimitiveType, GlobalComponentOffset);
+		}
+	}
+
+	ENoOp LowerPrimitiveUniformParameter(const MIR::FUniformParameter* UniformParameter)
+	{
+		// Get the global float4 component index (e.g. if this is 13, it refer to PreshaderBuffer[3].y)
+		const FUniformExpressionSet& UniformExpressionSet = Module->GetCompilationOutput().UniformExpressionSet;
+		const uint32 GlobalComponentOffset = UniformExpressionSet.GetNumericParameterEvaluationOffset(UniformParameter->Analysis_UniformParameterIndex);
+
+		EmitPreshaderBufferRead(UniformParameter->Type.GetPrimitive(), GlobalComponentOffset);
+
+		return NoOp;
+	}
+
+	ENoOp LowerPreshaderParameter(const MIR::FPreshaderParameter* PreshaderParameter)
+	{
+		// Get the global float4 component index (e.g. if this is 13, it refer to PreshaderBuffer[3].y)
+		const uint32 GlobalComponentOffset = PreshaderParameter->Analysis_PreshaderOffset;
+
+		EmitPreshaderBufferRead(PreshaderParameter->Type.GetPrimitive(), GlobalComponentOffset);
 
 		return NoOp;
 	}
@@ -900,14 +935,9 @@ struct FPrivate : FMaterialIRToHLSLTranslation
 		const FUniformExpressionSet& UniformExpressionSet = Module->GetCompilationOutput().UniformExpressionSet;
 
 		// Get the global float4 component index (e.g. if this is 13, it refer to PreshaderBuffer[3].y)
-		uint32 GlobalComponentOffset = UniformExpressionSet.GetNumericParameterEvaluationOffset(UniformParameter->Analysis_UniformParameterIndex);
+		const uint32 GlobalComponentOffset = UniformExpressionSet.GetNumericParameterEvaluationOffset(UniformParameter->Analysis_UniformParameterIndex);
 
-		// Index of the float4 slot.  LWC parameters always have four elements, so there won't be an offset.
-		uint32 BufferSlotIndex = GlobalComponentOffset / 4;
-
-		Printer << TEXTVIEW("DFToWS(MakeDFVector") << NumComponents;
-		Printer << TEXTVIEW("(Material.PreshaderBuffer[") << BufferSlotIndex << TEXTVIEW("]") << GVector4SwizzleSubset[NumComponents - 1];
-		Printer << TEXTVIEW(", Material.PreshaderBuffer[") << BufferSlotIndex + 1 << TEXTVIEW("]") << GVector4SwizzleSubset[NumComponents - 1] << TEXTVIEW("))");
+		EmitPreshaderBufferReadDoubleVector(UniformParameter->Type.GetPrimitive(), GlobalComponentOffset);
 	
 		return NoOp;
 	}
diff --git a/Engine/Source/Runtime/Engine/Private/Materials/MaterialIRValueAnalyzer.cpp b/Engine/Source/Runtime/Engine/Private/Materials/MaterialIRValueAnalyzer.cpp
--- a/Engine/Source/Runtime/Engine/Private/Materials/MaterialIRValueAnalyzer.cpp
+++ b/Engine/Source/Runtime/Engine/Private/Materials/MaterialIRValueAnalyzer.cpp
@@ -482,12 +482,13 @@ static void AnalyzeTextureObject(FMaterialIRValueAnalyzer& Analyzer, MIR::FTextu
 {
 	EMaterialTextureParameterType ParamType = MIR::Internal::TextureMaterialValueTypeToParameterType(TextureObject->Texture->GetMaterialType());
 
-	FMaterialTextureParameterInfo ParamInfo {
+	FMaterialTextureParameterInfo ParamInfo
+	{
 		.ParameterInfo = { "", EMaterialParameterAssociation::GlobalParameter, INDEX_NONE },
 		.TextureIndex = Analyzer.Material->GetReferencedTextures().Find(TextureObject->Texture),
 		.SamplerSource = SSM_FromTextureAsset, // TODO - Is this needed?
 	};
-			
+
 	check(ParamInfo.TextureIndex != INDEX_NONE);
 
 	TextureObject->Analysis_UniformParameterIndex = Analyzer.CompilationOutput->UniformExpressionSet.FindOrAddTextureParameter(ParamType, ParamInfo);
@@ -500,7 +501,8 @@ static void AnalyzeTextureUniformParameter(FMaterialIRValueAnalyzer& Analyzer, M
 	UTexture* Texture = Cast<UTexture>(ParameterValue.AsTextureObject());
 	check(Texture);
 
-	FMaterialTextureParameterInfo ParamInfo {
+	FMaterialTextureParameterInfo ParamInfo
+	{
 		.ParameterInfo = Analyzer.Module->GetParameterInfo(Parameter->ParameterIdInModule),
 		.TextureIndex = Analyzer.Material->GetReferencedTextures().Find(Texture),
 		.SamplerSource = SSM_FromTextureAsset, // TODO - Is this needed?
@@ -512,38 +514,9 @@ static void AnalyzeTextureUniformParameter(FMaterialIRValueAnalyzer& Analyzer, M
 	Parameter->Analysis_UniformParameterIndex = Analyzer.CompilationOutput->UniformExpressionSet.FindOrAddTextureParameter(ParamType, ParamInfo);
 }
 
-static void AnalyzePrimitiveUniformParameter(FMaterialIRValueAnalyzer& Analyzer, MIR::FUniformParameter* Parameter)
+// Returns the next available offset into the preshader buffer for a float vector with the specified number of components (1-4).
+static uint32 NextGlobalComponentOffset(FMaterialIRValueAnalyzer& Analyzer, uint32 NumComponents, bool bAreComponentsDouble = false)
 {
-	FUniformExpressionSet& UniformExpressionSet = Analyzer.CompilationOutput->UniformExpressionSet;
-	FMaterialParameterInfo ParameterInfo = Analyzer.Module->GetParameterInfo(Parameter->ParameterIdInModule);
-	FMaterialParameterMetadata ParameterMetadata = Analyzer.Module->GetParameterMetadata(Parameter->ParameterIdInModule);
-
-	UE::Shader::FValue DefaultValue;
-	switch (ParameterMetadata.Value.Type)
-	{
-		case EMaterialParameterType::Scalar: DefaultValue = ParameterMetadata.Value.AsScalar(); break;
-		case EMaterialParameterType::Vector: DefaultValue = ParameterMetadata.Value.AsLinearColor(); break;
-		case EMaterialParameterType::DoubleVector: DefaultValue = ParameterMetadata.Value.AsVector4d(); break;
-		default: UE_MIR_TODO();
-	}
-
-	// TODO: GetParameterOverrideValueForCurrentFunction?
-
-	uint32 DefaultValueOffset;
-	if (!MIR::Internal::Find(Analyzer.UniformDefaultValueOffsets, DefaultValue, DefaultValueOffset))
-	{
-		DefaultValueOffset = UniformExpressionSet.AddDefaultParameterValue(DefaultValue);
-		Analyzer.UniformDefaultValueOffsets.Add(DefaultValue, DefaultValueOffset);
-	}
-
-	Parameter->Analysis_UniformParameterIndex = UniformExpressionSet.FindOrAddNumericParameter(ParameterMetadata.Value.Type, ParameterInfo, DefaultValueOffset);
-
-	// Make sure the parameter type is primitive
-	MIR::FPrimitive PrimitiveType = Parameter->Type.GetPrimitive();
-
-	// Only int, float and LWC parameters supported for now
-	check(PrimitiveType.IsInteger() || PrimitiveType.IsAnyFloat());
-
 	// Following code calculates GlobalComponentOffset.
 	// GlobalComponentOffset is the i-th component in the array of float4s that make the uniform buffer.
 	// For example a GlobalComponentOffset of 13 references PreshaderBuffer[3].y.
@@ -554,7 +527,6 @@ static void AnalyzePrimitiveUniformParameter(FMaterialIRValueAnalyzer& Analyzer,
 	// in the appropriate freelist.
 
 	uint32 GlobalComponentOffset = UINT32_MAX;
-	uint32 NumComponents = PrimitiveType.NumComponents();
 	check(NumComponents <= 4); // Only vectors supported for now
 
 	uint32 UsedNumComponents;
@@ -568,31 +540,78 @@ static void AnalyzePrimitiveUniformParameter(FMaterialIRValueAnalyzer& Analyzer,
 		{
 			continue;
 		}
-		
+
 		// There is space for this number `NumComponents` in a a chunk with `i` free components
 		GlobalComponentOffset = FreeOffsetsForCurrentNumComponents.Last();
 		FreeOffsetsForCurrentNumComponents.Pop();
 
 		break;
 	}
-	
+
 	// If UsedNumComponents is 4, it means we looked for all possible components sizes (1, 2 and 3) and
 	// could not find a chunk that would fit this parameter. Allocate a new chunk.
 	if (UsedNumComponents >= 4)
 	{
-		check(UsedNumComponents == 4); 
-		const uint32 SizeInBytes = PrimitiveType.IsDouble() ? 2 : 1;
-		GlobalComponentOffset = UniformExpressionSet.AllocateFromUniformBuffer(SizeInBytes) * 4;
+		FUniformExpressionSet& UniformExpressionSet = Analyzer.CompilationOutput->UniformExpressionSet;
+		check(UsedNumComponents == 4);
+		const uint32 NumOfFloat4s = bAreComponentsDouble ? 2 : 1;
+		GlobalComponentOffset = UniformExpressionSet.AllocateFromUniformBuffer(NumOfFloat4s) * 4;
 	}
 
 	// Finally, check whether in the used slot there's any slack left, and if so, record it in the appropriate free-list.
-	int NumComponentsLeft = UsedNumComponents - NumComponents;
+	int32 NumComponentsLeft = UsedNumComponents - NumComponents;
 	if (NumComponentsLeft > 0)
 	{
-		int LeftOverOffset = GlobalComponentOffset + NumComponents;
+		int32 LeftOverOffset = GlobalComponentOffset + NumComponents;
 		Analyzer.FreeOffsetsPerNumComponents[NumComponentsLeft - 1].Push(LeftOverOffset);
 	}
 
+	return GlobalComponentOffset;
+}
+
+static uint32 FindOrAddDefaultValueOffset(FMaterialIRValueAnalyzer& Analyzer, const FMaterialParameterValue& ParameterValue)
+{
+	UE::Shader::FValue DefaultValue;
+	switch (ParameterValue.Type)
+	{
+		case EMaterialParameterType::Scalar: DefaultValue = ParameterValue.AsScalar(); break;
+		case EMaterialParameterType::Vector: DefaultValue = ParameterValue.AsLinearColor(); break;
+		case EMaterialParameterType::DoubleVector: DefaultValue = ParameterValue.AsVector4d(); break;
+		default: UE_MIR_TODO();
+	}
+
+	uint32 DefaultValueOffset = 0;
+	if (!MIR::Internal::Find(Analyzer.UniformDefaultValueOffsets, DefaultValue, DefaultValueOffset))
+	{
+		FUniformExpressionSet& UniformExpressionSet = Analyzer.CompilationOutput->UniformExpressionSet;
+		DefaultValueOffset = UniformExpressionSet.AddDefaultParameterValue(DefaultValue);
+		Analyzer.UniformDefaultValueOffsets.Add(DefaultValue, DefaultValueOffset);
+	}
+
+	return DefaultValueOffset;
+}
+
+static void AnalyzePrimitiveUniformParameter(FMaterialIRValueAnalyzer& Analyzer, MIR::FUniformParameter* Parameter)
+{
+	FUniformExpressionSet& UniformExpressionSet = Analyzer.CompilationOutput->UniformExpressionSet;
+	FMaterialParameterInfo ParameterInfo = Analyzer.Module->GetParameterInfo(Parameter->ParameterIdInModule);
+	FMaterialParameterMetadata ParameterMetadata = Analyzer.Module->GetParameterMetadata(Parameter->ParameterIdInModule);
+
+	// TODO: GetParameterOverrideValueForCurrentFunction?
+
+	const uint32 DefaultValueOffset = FindOrAddDefaultValueOffset(Analyzer, ParameterMetadata.Value);
+
+	Parameter->Analysis_UniformParameterIndex = UniformExpressionSet.FindOrAddNumericParameter(ParameterMetadata.Value.Type, ParameterInfo, DefaultValueOffset);
+
+	// Make sure the parameter type is primitive
+	MIR::FPrimitive PrimitiveType = Parameter->Type.GetPrimitive();
+
+	// Only int, float and LWC parameters supported for now
+	check(PrimitiveType.IsInteger() || PrimitiveType.IsAnyFloat());
+
+	const uint32 NumComponents = PrimitiveType.NumComponents();
+	const uint32 GlobalComponentOffset = NextGlobalComponentOffset(Analyzer, NumComponents, PrimitiveType.IsDouble());
+
 	// Add the parameter evaluation to the uniform data
 	UniformExpressionSet.AddNumericParameterEvaluation(Parameter->Analysis_UniformParameterIndex, GlobalComponentOffset);
 
@@ -630,6 +649,45 @@ static void AnalyzeUniformParameter(FMaterialIRValueAnalyzer& Analyzer, MIR::FUn
 	}
 }
 
+static void AnalyzePreshaderParameter(FMaterialIRValueAnalyzer& Analyzer, MIR::FPreshaderParameter* Parameter)
+{
+	FUniformExpressionSet& UniformExpressionSet = Analyzer.CompilationOutput->UniformExpressionSet;
+
+	MIR::FUniformParameter* UniformSourceParameter = Parameter->SourceParameter->As<MIR::FUniformParameter>();
+	const FMaterialParameterInfo& SourceParameterInfo = Analyzer.Module->GetParameterInfo(UniformSourceParameter->ParameterIdInModule);
+
+	// Make sure the parameter type is primitive
+	MIR::FPrimitive PrimitiveType = Parameter->Type.GetPrimitive();
+
+	// Only int, float and LWC parameters supported for now
+	check(PrimitiveType.IsInteger() || PrimitiveType.IsAnyFloat());
+
+	const uint32 NumComponents = PrimitiveType.NumComponents();
+	const uint32 GlobalComponentOffset = NextGlobalComponentOffset(Analyzer, NumComponents);
+	const uint32 UniformPreshaderOffset = GlobalComponentOffset / 4;
+
+	UniformExpressionSet.WriteUniformPreshaderEntry(
+		UniformPreshaderOffset, UE::Shader::MakeValueType(UE::Shader::EValueComponentType::Float, NumComponents),
+		[Parameter, &SourceParameterInfo](UE::Shader::FPreshaderData& UniformPreshaderData) -> void
+		{
+			const FHashedMaterialParameterInfo HashedSourceParameterInfo{ SourceParameterInfo.Name, EMaterialParameterAssociation::GlobalParameter, INDEX_NONE };
+
+			switch (Parameter->Opcode)
+			{
+			case UE::Shader::EPreshaderOpcode::TextureSize:
+			case UE::Shader::EPreshaderOpcode::TexelSize:
+				UniformPreshaderData.WriteOpcode(Parameter->Opcode).Write(HashedSourceParameterInfo).Write((int32)Parameter->TextureIndex);
+				break;
+
+			default:
+				UE_MIR_UNREACHABLE();
+			}
+		}
+	);
+
+	Parameter->Analysis_PreshaderOffset = GlobalComponentOffset;
+}
+
 static void AnalyzeSetMaterialOutput(FMaterialIRValueAnalyzer& Analyzer, MIR::FSetMaterialOutput* SetMaterialOutput)
 {
 	if (SetMaterialOutput->Property == EMaterialProperty::MP_Normal)
@@ -695,6 +753,7 @@ void FMaterialIRValueAnalyzer::Analyze(MIR::FValue* Value)
 		EXPAND_CASE(MaterialParameterCollection);
 		EXPAND_CASE(TextureObject);
 		EXPAND_CASE(UniformParameter);
+		EXPAND_CASE(PreshaderParameter);
 		EXPAND_CASE(SetMaterialOutput);
 		EXPAND_CASE(InlineHLSL);
 		EXPAND_CASE(ScreenTexture);
diff --git a/Engine/Source/Runtime/Engine/Private/Materials/MaterialShared.cpp b/Engine/Source/Runtime/Engine/Private/Materials/MaterialShared.cpp
--- a/Engine/Source/Runtime/Engine/Private/Materials/MaterialShared.cpp
+++ b/Engine/Source/Runtime/Engine/Private/Materials/MaterialShared.cpp
@@ -5965,6 +5965,11 @@ static bool IsTranslucencyAfterMotionBlurEnabled(UMaterial* Material)
 		&& !Material->IsDeferredDecal();
 }
 
+EMaterialValueType GetTexturePropertyValueType(EMaterialValueType TextureType)
+{
+	return (TextureType == MCT_VolumeTexture || TextureType == MCT_Texture2DArray || TextureType == MCT_SparseVolumeTexture) ? MCT_Float3 : MCT_Float2;
+}
+
 void FinalCompileValidation(
 	UMaterial* Material,
 	const FMaterialCompilationOutput& MaterialCompilationOutput,
diff --git a/Engine/Source/Runtime/Engine/Private/Materials/MaterialSharedPrivate.h b/Engine/Source/Runtime/Engine/Private/Materials/MaterialSharedPrivate.h
--- a/Engine/Source/Runtime/Engine/Private/Materials/MaterialSharedPrivate.h
+++ b/Engine/Source/Runtime/Engine/Private/Materials/MaterialSharedPrivate.h
@@ -88,6 +88,9 @@ namespace UE::MaterialTranslatorUtils
 	// Converts a scene texture identifier to its corresponding HLSL string representation.
 	FString SceneTextureIdToHLSLString(ESceneTextureId TexId);
 
+	// Returns the value type of a texture property with the specified texture value type, e.g. MCT_Float3 for Volume textures.
+	EMaterialValueType GetTexturePropertyValueType(EMaterialValueType TextureType);
+
 	// Final compile validation logic, run at the end of translation to HLSL.
 	void FinalCompileValidation(
 		UMaterial* Material,
diff --git a/Engine/Source/Runtime/Engine/Private/Materials/MaterialUniformExpressions.cpp b/Engine/Source/Runtime/Engine/Private/Materials/MaterialUniformExpressions.cpp
--- a/Engine/Source/Runtime/Engine/Private/Materials/MaterialUniformExpressions.cpp
+++ b/Engine/Source/Runtime/Engine/Private/Materials/MaterialUniformExpressions.cpp
@@ -796,6 +796,24 @@ void FUniformExpressionSet::SetVTLayer(int32 StackIndex, int32 VTLayerIndex, int
 	VTStacks[StackIndex].SetLayer(VTLayerIndex, TextureIndex);
 }
 
+void FUniformExpressionSet::WriteUniformPreshaderEntry(uint32 BufferOffset, UE::Shader::EValueType EntryType, const TFunction<void(UE::Shader::FPreshaderData&)>& EntryGenerator)
+{
+	FMaterialUniformPreshaderHeader& PreshaderHeader = UniformPreshaders.AddDefaulted_GetRef();
+	PreshaderHeader.FieldIndex = UniformPreshaderFields.Num();
+	PreshaderHeader.NumFields = 1;
+
+	PreshaderHeader.OpcodeOffset = UniformPreshaderData.Num();
+	{
+		EntryGenerator(UniformPreshaderData);
+	}
+	PreshaderHeader.OpcodeSize = UniformPreshaderData.Num() - PreshaderHeader.OpcodeOffset;
+
+	FMaterialUniformPreshaderField& PreshaderField = UniformPreshaderFields.AddDefaulted_GetRef();
+	PreshaderField.BufferOffset = BufferOffset;
+	PreshaderField.Type = EntryType;
+	PreshaderField.ComponentIndex = 0;
+}
+
 void FUniformExpressionSet::GetGameThreadTextureValue(EMaterialTextureParameterType Type, int32 Index, const UMaterialInterface* MaterialInterface, const FMaterial& Material, UTexture*& OutValue, bool bAllowOverride) const
 {
 	check(IsInGameThread() || IsInParallelGameThread());
diff --git a/Engine/Source/Runtime/Engine/Public/MaterialShared.h b/Engine/Source/Runtime/Engine/Public/MaterialShared.h
--- a/Engine/Source/Runtime/Engine/Public/MaterialShared.h
+++ b/Engine/Source/Runtime/Engine/Public/MaterialShared.h
@@ -739,6 +739,12 @@ public:
 	int32 AddVTLayer(int32 StackIndex, int32 TextureIndex);
 	void SetVTLayer(int32 StackIndex, int32 VTLayerIndex, int32 TextureIndex);
 
+	/// Writes a new uniform preshader entry.
+	/// @param BufferOffset Specifies the location in the preshader buffer where the result of the preshader instruction will be written to. See FMaterialUniformPreshaderField::BufferOffset.
+	/// @param EntryType Specifies the value type of the preshader instruction output. See FMaterialUniformPreshaderField::Type.
+	/// @param EntryGenerator Callback that writes the preshader opcodes.
+	void WriteUniformPreshaderEntry(uint32 BufferOffset, UE::Shader::EValueType EntryType, const TFunction<void(UE::Shader::FPreshaderData&)>& EntryGenerator);
+
 protected:
 	union FVTPackedStackAndLayerIndex
 	{
diff --git a/Engine/Source/Runtime/Engine/Public/Materials/MaterialExpressionTextureProperty.h b/Engine/Source/Runtime/Engine/Public/Materials/MaterialExpressionTextureProperty.h
--- a/Engine/Source/Runtime/Engine/Public/Materials/MaterialExpressionTextureProperty.h
+++ b/Engine/Source/Runtime/Engine/Public/Materials/MaterialExpressionTextureProperty.h
@@ -38,6 +38,7 @@ class UMaterialExpressionTextureProperty : public UMaterialExpression
 	
 	//~ Begin UMaterialExpression Interface
 #if WITH_EDITOR
+	ENGINE_API virtual void Build(MIR::FEmitter& Emitter) override;
 	ENGINE_API virtual int32 Compile(class FMaterialCompiler* Compiler, int32 OutputIndex) override;
 	ENGINE_API virtual void GetTexturesForceMaterialRecompile(TArray<UTexture *> &Textures) const override;
 	ENGINE_API virtual void GetCaption(TArray<FString>& OutCaptions) const override;
diff --git a/Engine/Source/Runtime/Engine/Public/Materials/MaterialIR.h b/Engine/Source/Runtime/Engine/Public/Materials/MaterialIR.h
--- a/Engine/Source/Runtime/Engine/Public/Materials/MaterialIR.h
+++ b/Engine/Source/Runtime/Engine/Public/Materials/MaterialIR.h
@@ -7,6 +7,7 @@
 #include "Materials/MaterialExpressionDBufferTexture.h"
 #include "MaterialSceneTextureId.h"
 #include "Engine/TextureDefines.h"
+#include "Shader/Preshader.h"
 
 #if WITH_EDITOR
 
@@ -128,6 +129,7 @@ enum EValueKind
 	VK_Nop,
 	VK_Call,
 	VK_CallParameterOutput,
+	VK_PreshaderParameter,
 
 	VK_InstructionEnd,
 };
@@ -1098,5 +1100,20 @@ struct FCallParameterOutput : TInstruction<VK_CallParameterOutput, 1>
 	int32 	Index;
 };
 
+struct FPreshaderParameter : TInstruction<VK_PreshaderParameter, 1>
+{
+	// Argument that points to the source parameter this dynamic parameter depends on.
+	FValue* SourceParameter;
+
+	// Specifies the opcode that needs to be encoded for this parameter. For now, only a small subset of opcodes are supported for material preshader parameters.
+	UE::Shader::EPreshaderOpcode Opcode;
+
+	// Index into UMaterialInterface::GetReferencedTextures() pointing to the source parameter's texture.
+	int32 TextureIndex;
+
+	// Offset into the preshader buffer of the material resource.
+	uint32_t Analysis_PreshaderOffset;
+};
+
 } // namespace MIR
 #endif // WITH_EDITOR
diff --git a/Engine/Source/Runtime/Engine/Public/Materials/MaterialIREmitter.h b/Engine/Source/Runtime/Engine/Public/Materials/MaterialIREmitter.h
--- a/Engine/Source/Runtime/Engine/Public/Materials/MaterialIREmitter.h
+++ b/Engine/Source/Runtime/Engine/Public/Materials/MaterialIREmitter.h
@@ -239,6 +239,9 @@ public:
 	// Returns the type of the i-th attribute in the given material aggregate.
 	FType GetMaterialAggregateAttributeType(const UMaterialAggregate* Aggregate, int32 AttributeIndex);
 
+	// Returns the UTexture of the specified texture object or uniform parameter.
+	UTexture* GetTextureFromValue(MIR::FValueRef Texture);
+
 	/*-------------------------------- Input management --------------------------------*/
 
 	// Retrieves the value from the given input, if present. A null reference otherwise.
@@ -403,6 +406,9 @@ public:
 	// Emits a material parameter value with given name, metadata and sampler type (if it is a texture) and returns it.
 	FValueRef Parameter(FName Name, FMaterialParameterMetadata& Metadata, EMaterialSamplerType SamplerType = SAMPLERTYPE_Color);
 
+	// Emits a preshader parameter value with given opcode.
+	FValueRef PreshaderParameter(UE::Shader::EPreshaderOpcode Opcode, FValueRef SourceParameter);
+
 	// Emits a custom primitive data value with given index.
 	FValueRef CustomPrimitiveData(uint32 PrimitiveDataIndex);
 
