SplineComponent: Call OnSplineSet when mutating the spline, not just when overwriting.

diff --git a/Engine/Source/Runtime/Engine/Private/Components/SplineComponent.cpp b/Engine/Source/Runtime/Engine/Private/Components/SplineComponent.cpp
--- a/Engine/Source/Runtime/Engine/Private/Components/SplineComponent.cpp
+++ b/Engine/Source/Runtime/Engine/Private/Components/SplineComponent.cpp
@@ -178,6 +178,8 @@ void USplineComponent::SetDefaultSpline()
     WarninglessSplineCurves().Position.Points.Emplace(EndParam, EndPoint, FVector::ZeroVector, FVector::ZeroVector, CIM_CurveAuto);
     WarninglessSplineCurves().Rotation.Points.Emplace(EndParam, FQuat::Identity, FQuat::Identity, FQuat::Identity, CIM_CurveAuto);
     WarninglessSplineCurves().Scale.Points.Emplace(EndParam, FVector(1.0f), FVector::ZeroVector, FVector::ZeroVector, CIM_CurveAuto);
+
+	OnSplineSet.Broadcast();
 }
 
 void USplineComponent::GetLifetimeReplicatedProps(TArray<FLifetimeProperty>& OutLifetimeProps) const
@@ -843,6 +845,8 @@ void USplineComponent::SetClosedLoop(bool bInClosedLoop, bool bUpdateSpline)
 	{
 		UpdateSpline();
 	}
+
+	OnSplineSet.Broadcast();
 }
 
 void USplineComponent::SetClosedLoopAtPosition(bool bInClosedLoop, float Key, bool bUpdateSpline)
@@ -855,6 +859,8 @@ void USplineComponent::SetClosedLoopAtPosition(bool bInClosedLoop, float Key, bo
 	{
 		UpdateSpline();
 	}
+
+	OnSplineSet.Broadcast();
 }
 
 bool USplineComponent::IsClosedLoop() const
@@ -979,6 +985,8 @@ void USplineComponent::AddPoint(const FSplinePoint& InSplinePoint, bool bUpdateS
 	{
 		UpdateSpline();
 	}
+
+	OnSplineSet.Broadcast();
 }
 
 void USplineComponent::AddPoints(const TArray<FSplinePoint>& InSplinePoints, bool bUpdateSpline)
@@ -998,6 +1006,8 @@ void USplineComponent::AddPoints(const TArray<FSplinePoint>& InSplinePoints, boo
 	{
 		UpdateSpline();
 	}
+
+	OnSplineSet.Broadcast();
 }
 
 
@@ -1041,6 +1051,8 @@ void USplineComponent::AddSplinePoint(const FVector& Position, ESplineCoordinate
 	{
 		UpdateSpline();
 	}
+
+	OnSplineSet.Broadcast();
 }
 
 void USplineComponent::AddSplinePointAtIndex(const FVector& Position, int32 Index, ESplineCoordinateSpace::Type CoordinateSpace, bool bUpdateSpline)
@@ -1092,6 +1104,8 @@ void USplineComponent::AddSplinePointAtIndex(const FVector& Position, int32 Inde
 	{
 		UpdateSpline();
 	}
+
+	OnSplineSet.Broadcast();
 }
 
 void USplineComponent::RemoveSplinePoint(int32 Index, bool bUpdateSpline)
@@ -1132,6 +1146,8 @@ void USplineComponent::RemoveSplinePoint(int32 Index, bool bUpdateSpline)
 	{
 		UpdateSpline();
 	}
+
+	OnSplineSet.Broadcast();
 }
 
 void USplineComponent::SetSplinePoints(const TArray<FVector>& Points, ESplineCoordinateSpace::Type CoordinateSpace, bool bUpdateSpline)
@@ -1184,6 +1200,8 @@ void USplineComponent::SetSplinePoints(const TArray<FVector>& Points, ESplineCoo
 	{
 		UpdateSpline();
 	}
+
+	OnSplineSet.Broadcast();
 }
 
 void USplineComponent::SetLocationAtSplinePoint(int32 PointIndex, const FVector& InLocation, ESplineCoordinateSpace::Type CoordinateSpace, bool bUpdateSpline)
@@ -1202,6 +1220,8 @@ void USplineComponent::SetLocationAtSplinePoint(int32 PointIndex, const FVector&
 		{
 			UpdateSpline();
 		}
+
+		OnSplineSet.Broadcast();
 	}
 }
 
@@ -1232,6 +1252,8 @@ void USplineComponent::SetTangentsAtSplinePoint(int32 PointIndex, const FVector&
 		{
 			UpdateSpline();
 		}
+
+		OnSplineSet.Broadcast();
 	}
 }
 
@@ -1252,6 +1274,8 @@ void USplineComponent::SetUpVectorAtSplinePoint(int32 PointIndex, const FVector&
 		{
 			UpdateSpline();
 		}
+
+		OnSplineSet.Broadcast();
 	}
 }
 
@@ -1293,6 +1317,8 @@ void USplineComponent::SetQuaternionAtSplinePoint(int32 PointIndex, const FQuat&
 		{
 			UpdateSpline();
 		}
+
+		OnSplineSet.Broadcast();
 	}
 }
 
@@ -1309,6 +1335,8 @@ void USplineComponent::SetScaleAtSplinePoint(int32 PointIndex, const FVector& In
 		{
 			UpdateSpline();
 		}
+
+		OnSplineSet.Broadcast();
 	}
 }
 
@@ -1341,6 +1369,8 @@ void USplineComponent::SetSplinePointType(int32 PointIndex, ESplinePointType::Ty
 		{
 			UpdateSpline();
 		}
+
+		OnSplineSet.Broadcast();
 	}
 }
 
@@ -1555,6 +1585,8 @@ void USplineComponent::SetDefaultUpVector(const FVector& UpVector, ESplineCoordi
 	}
 
 	UpdateSpline();
+
+	OnSplineSet.Broadcast();
 }
 
 FVector USplineComponent::GetDefaultUpVector(ESplineCoordinateSpace::Type CoordinateSpace) const
@@ -2709,6 +2741,8 @@ void USplineComponent::ApplyComponentInstanceData(FSplineInstanceData* SplineIns
 	bSplineHasBeenEdited = SplineInstanceData->bSplineHasBeenEdited;
 
 	UpdateSpline();
+
+	OnSplineSet.Broadcast();
 }
 
 void USplineComponent::ApplyComponentInstanceData(struct FSplineComponentInstanceData* SplineInstanceData, const bool bPostUCS)
@@ -2753,6 +2787,8 @@ void USplineComponent::ApplyComponentInstanceData(struct FSplineComponentInstanc
 	bSplineHasBeenEdited = SplineInstanceData->bSplineHasBeenEdited;
 
 	UpdateSpline();
+
+	OnSplineSet.Broadcast();
 }
 
 #if WITH_EDITOR
@@ -2773,6 +2809,8 @@ void USplineComponent::PostEditChangeChainProperty(FPropertyChangedChainEvent& P
 			PropertyName == ClosedLoopName)
 		{
 			UpdateSpline();
+
+			OnSplineSet.Broadcast();
 		}
 	}
 
