Spline: Reparameterize spline when moving points (for FNewSpline only)

diff --git a/Engine/Source/Runtime/Engine/Private/Curves/Spline.cpp b/Engine/Source/Runtime/Engine/Private/Curves/Spline.cpp
--- a/Engine/Source/Runtime/Engine/Private/Curves/Spline.cpp
+++ b/Engine/Source/Runtime/Engine/Private/Curves/Spline.cpp
@@ -2208,6 +2208,9 @@ void FNewSpline::SetLocation(int32 Index, const FVector& InLocation)
 	
 	GetSpline().ModifyPoint(Index, {InLocation, GetSpline().GetTangentIn(Index), GetSpline().GetTangentOut(Index), GetSpline().GetTangentMode(Index)});
 
+	// We do this so that parameterization remains proportional to the square root of segment chord lengths.
+	GetSpline().Reparameterize(EParameterizationPolicy::Centripetal);
+	
 	MarkReparamTableDirty();
     MarkLegacyCurvesDirty();
 }
