PR #13691: Fix: ParticleSystem EPSUM_FixedTime had no effect due to incorrect assignment

diff --git a/Engine/Source/Runtime/Engine/Private/Particles/ParticleSystemComponent.cpp b/Engine/Source/Runtime/Engine/Private/Particles/ParticleSystemComponent.cpp
--- a/Engine/Source/Runtime/Engine/Private/Particles/ParticleSystemComponent.cpp
+++ b/Engine/Source/Runtime/Engine/Private/Particles/ParticleSystemComponent.cpp
@@ -2406,7 +2406,7 @@ void UParticleSystemComponent::TickComponent(float DeltaTime, enum ELevelTick Ti
 	if (Template->SystemUpdateMode == EPSUM_FixedTime)
 	{
 		// Use the fixed delta time!
-		DeltaTime = Template->UpdateTime_Delta;
+		DeltaTimeTick = Template->UpdateTime_Delta;
 	}
 
 	// Clear out the events.
