Add flags to gamplay tracing system in FObjectInfo.
Use flags to differentiate between UAF graph assets, and transient UAF graphs, in order to avoid crashing on double click

diff --git a/Engine/Plugins/Animation/GameplayInsights/Source/GameplayInsights/Private/GameplayAnalyzer.cpp b/Engine/Plugins/Animation/GameplayInsights/Source/GameplayInsights/Private/GameplayAnalyzer.cpp
--- a/Engine/Plugins/Animation/GameplayInsights/Source/GameplayInsights/Private/GameplayAnalyzer.cpp
+++ b/Engine/Plugins/Animation/GameplayInsights/Source/GameplayInsights/Private/GameplayAnalyzer.cpp
@@ -104,7 +104,8 @@ bool FGameplayAnalyzer::OnEvent(uint16 RouteId, EStyle Style, const FOnEventCont
 		uint64 Id = EventData.GetValue<uint64>("Id");
 		uint64 OuterId = EventData.GetValue<uint64>("OuterId");
 		uint64 ClassId = EventData.GetValue<uint64>("ClassId");
-		GameplayProvider.AppendObject(RewindDebugger::FObjectId{Id}, RewindDebugger::FObjectId{OuterId}, ClassId, *ObjectName, *ObjectPathName);
+		EObjectInfoFlags Flags = static_cast<EObjectInfoFlags>(EventData.GetValue<uint8>("Flags", 0));
+		GameplayProvider.AppendObject(RewindDebugger::FObjectId{Id}, RewindDebugger::FObjectId{OuterId}, ClassId, *ObjectName, *ObjectPathName, Flags);
 		break;
 	}
 	case RouteId_ObjectTransform:
diff --git a/Engine/Plugins/Animation/GameplayInsights/Source/GameplayInsights/Private/GameplayProvider.cpp b/Engine/Plugins/Animation/GameplayInsights/Source/GameplayInsights/Private/GameplayProvider.cpp
--- a/Engine/Plugins/Animation/GameplayInsights/Source/GameplayInsights/Private/GameplayProvider.cpp
+++ b/Engine/Plugins/Animation/GameplayInsights/Source/GameplayInsights/Private/GameplayProvider.cpp
@@ -442,6 +442,7 @@ const FObjectInfo& FGameplayProvider::GetObjectInfo(const RewindDebugger::FObjec
 		, /*ClassId*/0
 		, *UnknownText.ToString()
 		, *UnknownText.ToString()
+		, EObjectInfoFlags::None 
 	};
 	return DefaultObjectInfo;
 }
@@ -480,7 +481,7 @@ void FGameplayProvider::AppendClass(uint64 InClassId, uint64 InSuperId, const TC
 	}
 }
 
-void FGameplayProvider::AppendObject(const RewindDebugger::FObjectId& InObjectId, const RewindDebugger::FObjectId& InOuterId, uint64 InClassId, const TCHAR* InObjectName, const TCHAR* InObjectPathName)
+void FGameplayProvider::AppendObject(const RewindDebugger::FObjectId& InObjectId, const RewindDebugger::FObjectId& InOuterId, uint64 InClassId, const TCHAR* InObjectName, const TCHAR* InObjectPathName, EObjectInfoFlags Flags)
 {
 	Session.WriteAccessCheck();
 
@@ -491,7 +492,7 @@ void FGameplayProvider::AppendObject(const RewindDebugger::FObjectId& InObjectId
 		const TCHAR* NewObjectName = Session.StoreString(InObjectName);
 		const TCHAR* NewObjectPathName = Session.StoreString(InObjectPathName);
 
-		ObjectIdToIndexMap.Add(InObjectId, ObjectInfos.Add({ InObjectId, InOuterId, InClassId, NewObjectName, NewObjectPathName }));
+		ObjectIdToIndexMap.Add(InObjectId, ObjectInfos.Add({ InObjectId, InOuterId, InClassId, NewObjectName, NewObjectPathName, Flags }));
 
 		ObjectHierarchy.AddUnique(InOuterId, InObjectId);
 	}
diff --git a/Engine/Plugins/Animation/GameplayInsights/Source/GameplayInsights/Private/GameplayProvider.h b/Engine/Plugins/Animation/GameplayInsights/Source/GameplayInsights/Private/GameplayProvider.h
--- a/Engine/Plugins/Animation/GameplayInsights/Source/GameplayInsights/Private/GameplayProvider.h
+++ b/Engine/Plugins/Animation/GameplayInsights/Source/GameplayInsights/Private/GameplayProvider.h
@@ -65,7 +65,7 @@ public:
 	void AppendClass(uint64 InClassId, uint64 InSuperId, const TCHAR* InClassName, const TCHAR* InClassPathName);
 
 	/** Add an object message */
-	void AppendObject(const RewindDebugger::FObjectId& InObjectId, const RewindDebugger::FObjectId& InOuterId, uint64 InClassId, const TCHAR* InObjectName, const TCHAR* InObjectPathName);
+	void AppendObject(const RewindDebugger::FObjectId& InObjectId, const RewindDebugger::FObjectId& InOuterId, uint64 InClassId, const TCHAR* InObjectName, const TCHAR* InObjectPathName, EObjectInfoFlags Flags);
 
 	/** Add an object transform */
 	void AppendObjectTransform(uint64 InObjectId, double InProfileTime, double InRecordingTime, const FTransform& Transform);
diff --git a/Engine/Plugins/Animation/GameplayInsights/Source/GameplayInsights/Public/IGameplayProvider.h b/Engine/Plugins/Animation/GameplayInsights/Source/GameplayInsights/Public/IGameplayProvider.h
--- a/Engine/Plugins/Animation/GameplayInsights/Source/GameplayInsights/Public/IGameplayProvider.h
+++ b/Engine/Plugins/Animation/GameplayInsights/Source/GameplayInsights/Public/IGameplayProvider.h
@@ -20,6 +20,13 @@ struct FClassInfo
 	const TCHAR* PathName = nullptr;
 };
 
+enum class EObjectInfoFlags : uint8
+{
+	None = 0,
+	TransientObject = 1 << 0,
+	StructInstance = 1 << 1,
+};
+
 struct FObjectInfo
 {
 PRAGMA_DISABLE_DEPRECATION_WARNINGS
@@ -30,12 +37,13 @@ PRAGMA_DISABLE_DEPRECATION_WARNINGS
 PRAGMA_ENABLE_DEPRECATION_WARNINGS
 
 	FObjectInfo(const RewindDebugger::FObjectId& ObjectId, const RewindDebugger::FObjectId& OuterObjectId, const uint64 ClassId, const TCHAR* Name
-		, const TCHAR* PathName)
+				, const TCHAR* PathName, EObjectInfoFlags Flags)
 		: ObjectId(ObjectId)
 		, OuterObjectId(OuterObjectId)
 		, ClassId(ClassId)
 		, Name(Name)
 		, PathName(PathName)
+		, Flags(Flags)
 	{
 	}
 
@@ -78,6 +86,8 @@ public:
 	uint64 ClassId = 0;
 	const TCHAR* Name = nullptr;
 	const TCHAR* PathName = nullptr;
+	EObjectInfoFlags Flags;
+	
 };
 
 struct FObjectPropertiesMessage
diff --git a/Engine/Plugins/Experimental/UAF/UAF/Source/UAFEditor/Private/RewindDebugger/AnimNextModuleTrack.cpp b/Engine/Plugins/Experimental/UAF/UAF/Source/UAFEditor/Private/RewindDebugger/AnimNextModuleTrack.cpp
--- a/Engine/Plugins/Experimental/UAF/UAF/Source/UAFEditor/Private/RewindDebugger/AnimNextModuleTrack.cpp
+++ b/Engine/Plugins/Experimental/UAF/UAF/Source/UAFEditor/Private/RewindDebugger/AnimNextModuleTrack.cpp
@@ -2,6 +2,8 @@
 
 #include "AnimNextModuleTrack.h"
 
+#include "ObjectTrace.h"
+
 #include UE_INLINE_GENERATED_CPP_BY_NAME(AnimNextModuleTrack)
 
 #if ANIMNEXT_TRACE_ENABLED
@@ -225,8 +227,11 @@ bool FAnimNextModuleTrack::HandleDoubleClickInternal()
 			if (const FDataInterfaceData* ModuleData = AnimNextProvider->GetDataInterfaceData(InstanceId))
 			{
 				const FObjectInfo& AssetInfo = GameplayProvider->GetObjectInfo(ModuleData->AssetId);
-				GEditor->GetEditorSubsystem<UAssetEditorSubsystem>()->OpenEditorForAsset(AssetInfo.PathName);
-				return true;
+				if (!EnumHasAnyFlags(AssetInfo.Flags, EObjectInfoFlags::TransientObject))
+				{
+					GEditor->GetEditorSubsystem<UAssetEditorSubsystem>()->OpenEditorForAsset(AssetInfo.PathName);
+					return true;
+				}
 			}
 		}
 	}
diff --git a/Engine/Source/Runtime/Engine/Private/ObjectTrace.cpp b/Engine/Source/Runtime/Engine/Private/ObjectTrace.cpp
--- a/Engine/Source/Runtime/Engine/Private/ObjectTrace.cpp
+++ b/Engine/Source/Runtime/Engine/Private/ObjectTrace.cpp
@@ -29,6 +29,7 @@ UE_TRACE_EVENT_BEGIN(Object, Object)
 	UE_TRACE_EVENT_FIELD(uint64, OuterId)
 	UE_TRACE_EVENT_FIELD(UE::Trace::WideString, Name)
 	UE_TRACE_EVENT_FIELD(UE::Trace::WideString, Path)
+	UE_TRACE_EVENT_FIELD(uint8, Flags)
 UE_TRACE_EVENT_END()
 
 UE_TRACE_EVENT_BEGIN(Object, ObjectLifetimeBegin2)
@@ -97,6 +98,12 @@ UE_TRACE_EVENT_END()
 
 namespace UE::ObjectTrace
 {
+	
+// duplicate ObjectInfoFlags must match FObjectInfo::Flag_
+constexpr uint8 ObjectInfoFlag_Transient = 0x1;
+constexpr uint8 ObjectInfoFlag_StructInstance = 0x2;
+
+	
 bool IsTracingDisabledForObject(const UObject* InObject)
 {
 	if (InObject == nullptr)
@@ -433,7 +440,7 @@ void FObjectTrace::OutputView(const UObject* InPlayer, const FSceneView* InView)
 		<< View.AspectRatio(AspectRatio);
 }
 
-void OutputInstanceInternal(const UObject* InOuterObject, uint64 InInstanceId, uint64 InOuterId, UStruct* InType, const FString& InName, const FString& InPathName)
+void OutputInstanceInternal(const UObject* InOuterObject, uint64 InInstanceId, uint64 InOuterId, UStruct* InType, const FString& InName, const FString& InPathName, uint8 InFlags)
 {
 	// Trace the object's class first
 	TRACE_TYPE(InType);
@@ -443,7 +450,8 @@ void OutputInstanceInternal(const UObject* InOuterObject, uint64 InInstanceId, u
 		<< Object.ClassId(FObjectTrace::GetObjectId(InType))
 		<< Object.OuterId(InOuterId)
 		<< Object.Name(*InName, InName.Len())
-		<< Object.Path(*InPathName, InPathName.Len());
+		<< Object.Path(*InPathName, InPathName.Len())
+		<< Object.Flags(InFlags);
 		
 	UE_TRACE_LOG(Object, ObjectLifetimeBegin2, ObjectChannel)
 		<< ObjectLifetimeBegin2.Cycle(FPlatformTime::Cycles64())
@@ -462,8 +470,10 @@ void FObjectTrace::OutputInstance(const UObject* InOuterObject, uint64 InInstanc
 	{
 		return;
 	}
+
+	uint8 Flags = UE::ObjectTrace::ObjectInfoFlag_StructInstance; 
 	
-	OutputInstanceInternal(InOuterObject, InInstanceId, InOuterId, InType, InName, InPathName);
+	OutputInstanceInternal(InOuterObject, InInstanceId, InOuterId, InType, InName, InPathName, Flags);
 }
 
 
@@ -480,9 +490,15 @@ void FObjectTrace::OutputObject(const UObject* InObject)
 	{
 		return;
 	}
+	
+	uint8 Flags = 0;
+	if (InObject->HasAnyFlags(RF_Transient))
+	{
+		Flags |=  UE::ObjectTrace::ObjectInfoFlag_Transient;
+	}
  	
 	OutputObject(InObject->GetOuter());
-	OutputInstanceInternal(InObject, InstanceId, GetObjectId(InObject->GetOuter()), InObject->GetClass(), InObject->GetFName().ToString(), InObject->GetPathName());
+	OutputInstanceInternal(InObject, InstanceId, GetObjectId(InObject->GetOuter()), InObject->GetClass(), InObject->GetFName().ToString(), InObject->GetPathName(), Flags);
 }
 
 void FObjectTrace::OutputObjectTransform(const UObject* InObject, const FTransform& InTransform, ETeleportType InTeleportType)
