[Modulation] - LPF disabled on a Sound Wave measures Low Frequency at 20Hz

diff --git a/Engine/Source/Runtime/Engine/Private/Audio.cpp b/Engine/Source/Runtime/Engine/Private/Audio.cpp
--- a/Engine/Source/Runtime/Engine/Private/Audio.cpp
+++ b/Engine/Source/Runtime/Engine/Private/Audio.cpp
@@ -241,6 +241,8 @@ void FWaveInstance::CalculateInheritedModulationSettingsAndValue(const EModulati
 {
 	const Audio::FModulationParameter& InParam = GetParameterBasedOnDestination(InDestination);
 	
+	OutValue = InParam.DefaultValue;
+
 	const FActiveSound& InActiveSound = *ActiveSound;
 	const FSoundModulationDestinationSettings& InWaveSettings = WaveData->ModulationSettings.GetSettingsForDestination(InDestination);
 	EModulationRouting InWaveRouting = WaveData->ModulationSettings.GetRoutingForDestination(InDestination);
