Add extra profiling scopes for adding to and removing components from cluster unions.

diff --git a/Engine/Source/Runtime/Engine/Private/PhysicsEngine/ClusterUnionComponent.cpp b/Engine/Source/Runtime/Engine/Private/PhysicsEngine/ClusterUnionComponent.cpp
--- a/Engine/Source/Runtime/Engine/Private/PhysicsEngine/ClusterUnionComponent.cpp
+++ b/Engine/Source/Runtime/Engine/Private/PhysicsEngine/ClusterUnionComponent.cpp
@@ -127,8 +127,11 @@ FPhysScene_Chaos* UClusterUnionComponent::GetChaosScene() const
 	return GWorld->GetPhysicsScene();
 }
 
+DECLARE_CYCLE_STAT(TEXT("UClusterUnionComponent::AddComponentToCluster"), STAT_ClusterUnionComponent_AddComponentToCluster, STATGROUP_Chaos);
 void UClusterUnionComponent::AddComponentToCluster(UPrimitiveComponent* InComponent, const TArray<int32>& BoneIds, bool bRebuildGeometry)
 {
+	SCOPE_CYCLE_COUNTER(STAT_ClusterUnionComponent_AddComponentToCluster);
+
 	if (!InComponent || !PhysicsProxy || !InComponent->GetWorld())
 	{
 		return;
@@ -257,8 +260,11 @@ void UClusterUnionComponent::AddComponentToCluster(UPrimitiveComponent* InCompon
 	}
 }
 
+DECLARE_CYCLE_STAT(TEXT("UClusterUnionComponent::RemoveComponentFromCluster"), STAT_ClusterUnionComponent_RemoveComponentFromCluster, STATGROUP_Chaos);
 void UClusterUnionComponent::RemoveComponentFromCluster(UPrimitiveComponent* InComponent)
 {
+	SCOPE_CYCLE_COUNTER(STAT_ClusterUnionComponent_RemoveComponentFromCluster);
+
 	if (!InComponent || !PhysicsProxy)
 	{
 		return;
diff --git a/Engine/Source/Runtime/Engine/Private/PhysicsEngine/ClusterUnionReplicatedProxyComponent.cpp b/Engine/Source/Runtime/Engine/Private/PhysicsEngine/ClusterUnionReplicatedProxyComponent.cpp
--- a/Engine/Source/Runtime/Engine/Private/PhysicsEngine/ClusterUnionReplicatedProxyComponent.cpp
+++ b/Engine/Source/Runtime/Engine/Private/PhysicsEngine/ClusterUnionReplicatedProxyComponent.cpp
@@ -45,6 +45,8 @@ void UClusterUnionReplicatedProxyComponent::GetLifetimeReplicatedProps(TArray<FL
 
 void UClusterUnionReplicatedProxyComponent::EndPlay(const EEndPlayReason::Type EndPlayReason)
 {
+	QUICK_SCOPE_CYCLE_COUNTER(UClusterUnionReplicatedProxyComponent_EndPlay);
+
 	UActorComponent::EndPlay(EndPlayReason);
 
 	MarkPendingDeletion();
