World Partition - HLOD: Removed leftover debug code

diff --git a/Engine/Source/Runtime/Engine/Private/WorldPartition/HLOD/HLODActor.cpp b/Engine/Source/Runtime/Engine/Private/WorldPartition/HLOD/HLODActor.cpp
--- a/Engine/Source/Runtime/Engine/Private/WorldPartition/HLOD/HLODActor.cpp
+++ b/Engine/Source/Runtime/Engine/Private/WorldPartition/HLOD/HLODActor.cpp
@@ -364,7 +364,7 @@ void AWorldPartitionHLOD::PreRegisterAllComponents()
 	// In PIE, collisions should be disabled
 	if (GetWorld() && !IsRunningCommandlet() && !FApp::IsUnattended())
 	{
-		bool bShouldEnableCollision = false;/*!GetWorld()->IsGameWorld();*/
+		bool bShouldEnableCollision = !GetWorld()->IsGameWorld();
 		if (GetActorEnableCollision() != bShouldEnableCollision)
 		{
 			SetActorEnableCollision(bShouldEnableCollision);
