Crash when trimming frames from animation sequence containing curves
#fix ensure curve key indices exist when deleting individual keys during curve resizing

diff --git a/Engine/Source/Runtime/Engine/Private/Curves/RichCurve.cpp b/Engine/Source/Runtime/Engine/Private/Curves/RichCurve.cpp
--- a/Engine/Source/Runtime/Engine/Private/Curves/RichCurve.cpp
+++ b/Engine/Source/Runtime/Engine/Private/Curves/RichCurve.cpp
@@ -853,6 +853,8 @@ void FRichCurve::ReadjustTimeRange(float NewMinTimeRange, float NewMaxTimeRange,
 	// delete the keys outside of range
 	if (bNeedToDeleteKey)
 	{
+		// Ensure this here as DeleteKey might do it while running through the loop, causing the keyhandle to become invalid
+		EnsureAllIndicesHaveHandles();
 		for (int32 KeyIndex=0; KeyIndex<Keys.Num(); ++KeyIndex)
 		{
 			if (Keys[KeyIndex].Time < NewMinTimeRange || Keys[KeyIndex].Time > NewMaxTimeRange)
