Fix an erroneous reserve statement in UMaterialInstance::CacheResourceShadersForCooking that did not include the existing resource count.

diff --git a/Engine/Source/Runtime/Engine/Private/Materials/MaterialInstance.cpp b/Engine/Source/Runtime/Engine/Private/Materials/MaterialInstance.cpp
--- a/Engine/Source/Runtime/Engine/Private/Materials/MaterialInstance.cpp
+++ b/Engine/Source/Runtime/Engine/Private/Materials/MaterialInstance.cpp
@@ -2756,7 +2756,7 @@ void UMaterialInstance::CacheResourceShadersForCooking(
 			BeginCacheShadersForResources(ShaderPlatform, NewResourcesToCache, PrecompileMode, TargetPlatform);
 		}
 
-		OutCachedMaterialResources.Reserve(NewResourcesToCache.Num());
+		OutCachedMaterialResources.Reserve(NewResourcesToCache.Num() + OutCachedMaterialResources.Num());
 		for (FMaterialResource* NewResource : NewResourcesToCache)
 		{
 			OutCachedMaterialResources.Add({ NewResource , ShaderPlatform });
