Clean up deprecated curve code in skeleton

diff --git a/Engine/Source/Runtime/AnimGraphRuntime/Public/AnimNodes/AnimNode_PoseDriver.h b/Engine/Source/Runtime/AnimGraphRuntime/Public/AnimNodes/AnimNode_PoseDriver.h
--- a/Engine/Source/Runtime/AnimGraphRuntime/Public/AnimNodes/AnimNode_PoseDriver.h
+++ b/Engine/Source/Runtime/AnimGraphRuntime/Public/AnimNodes/AnimNode_PoseDriver.h
@@ -103,9 +103,6 @@ struct FPoseDriverTarget
 	UPROPERTY(EditAnywhere, Category = PoseDriver)
 	FName DrivenName;
 
-	UE_DEPRECATED(5.3, "DrivenUID is no longer used")
-	USkeleton::AnimCurveUID DrivenUID = SmartName::MaxUID;
-
 	/** Pose buffer index */
 	int32 PoseCurveIndex;
 
diff --git a/Engine/Source/Runtime/Engine/Classes/Animation/Skeleton.h b/Engine/Source/Runtime/Engine/Classes/Animation/Skeleton.h
--- a/Engine/Source/Runtime/Engine/Classes/Animation/Skeleton.h
+++ b/Engine/Source/Runtime/Engine/Classes/Animation/Skeleton.h
@@ -381,27 +381,12 @@ public:
 	/** Serializable retarget sources for this skeleton **/
 	TMap< FName, FReferencePose > AnimRetargetSources;
 
-	// DEPRECATED - no longer used
+	UE_DEPRECATED(5.7, "AnimCurveUID is deprecated and no longer used.")
 	typedef SmartName::UID_Type AnimCurveUID;
 
-	UE_DEPRECATED(5.3, "AnimCurveMappingName is no longer used.")
-	static ENGINE_API const FName AnimCurveMappingName;
-
-	UE_DEPRECATED(5.3, "AnimTrackCurveMappingName is no longer used.")
-	static ENGINE_API const FName AnimTrackCurveMappingName;
-
 	ENGINE_API FCurveMetaData* GetCurveMetaData(FName CurveName);
 	ENGINE_API const FCurveMetaData* GetCurveMetaData(FName CurveName) const;
 
-	UE_DEPRECATED(5.3, "Please use GetCurveMetaData with an FName.")
-	ENGINE_API const FCurveMetaData* GetCurveMetaData(const SmartName::UID_Type CurveUID) const;
-
-	UE_DEPRECATED(5.3, "Please use GetCurveMetaData with an FName.")
-	ENGINE_API FCurveMetaData* GetCurveMetaData(const FSmartName& CurveName);
-	
-	UE_DEPRECATED(5.3, "Please use GetCurveMetaData with an FName.")
-	ENGINE_API const FCurveMetaData* GetCurveMetaData(const FSmartName& CurveName) const;
-
 	/**
 	 * Iterate over all curve metadata entries, calling InFunction on each
 	 * @param	InFunction	The function to call
@@ -496,9 +481,6 @@ public:
 	// return version of AnimCurveUidVersion
 	uint16 GetAnimCurveUidVersion() const;
 
-	UE_DEPRECATED(5.3, "This function is no longer used")
-	const TArray<uint16>& GetDefaultCurveUIDList() const;
-
 	const TArray<TSoftObjectPtr<USkeleton>>& GetCompatibleSkeletons() const { return CompatibleSkeletons; }
 
 	UE_DEPRECATED(5.2, "Please use UE::Anim::FSkeletonRemappingRegistry::GetRemapping.")
@@ -522,17 +504,10 @@ public:
 #endif
 
 protected:
-	// DEPRECATED - moved to CurveMetaData
+	UE_DEPRECATED(5.7, "SmartNames in Skeleton have been deprecated, use CurveMetaData instead.")
 	UPROPERTY()
 	FSmartNameContainer SmartNames_DEPRECATED;
 
-	// Cached ptr to the persistent AnimCurveMapping
-	UE_DEPRECATED(5.3, "AnimCurveMapping is no longer used")
-	static FSmartNameMapping* AnimCurveMapping;
-
-	UE_DEPRECATED(5.3, "DefaultCurveUIDList is no longer used")
-	static TArray<uint16> DefaultCurveUIDList;
-
 	//Cached marker sync marker names (stripped for non editor)
 	TArray<FName> ExistingMarkerNames;
 
@@ -597,38 +572,6 @@ public:
 	ENGINE_API void RemoveSlotGroup(const FName& InSlotName);
 	ENGINE_API void RenameSlotName(const FName& OldName, const FName& NewName);
 
-#if WITH_EDITOR
-	UE_DEPRECATED(5.3, "Please use AddCurveMetaData.")
-	bool AddSmartNameAndModify(FName ContainerName, FName NewDisplayName, FSmartName& NewName) { return false; }
-
-	UE_DEPRECATED(5.3, "Please use RenameCurveMetaData.")
-	bool RenameSmartnameAndModify(FName ContainerName, SmartName::UID_Type Uid, FName NewName) { return false; }
-	
-	UE_DEPRECATED(5.3, "Please use RemoveCurveMetaData.")
-	void RemoveSmartnameAndModify(FName ContainerName, SmartName::UID_Type Uid) {}
-
-	UE_DEPRECATED(5.3, "Please use RemoveCurveMetaData.")
-	void RemoveSmartnamesAndModify(FName ContainerName, const TArray<FName>& Names) {}
-#endif// WITH_EDITOR
-
-	UE_DEPRECATED(5.3, "This function is no longer used")
-	SmartName::UID_Type GetUIDByName(const FName& ContainerName, const FName& Name) const { return 0; }
-
-	UE_DEPRECATED(5.3, "This function is no longer used")
-	bool GetSmartNameByUID(const FName& ContainerName, SmartName::UID_Type UID, FSmartName& OutSmartName) const { return false; }
-	
-	UE_DEPRECATED(5.3, "This function is no longer used")
-	bool GetSmartNameByName(const FName& ContainerName, const FName& InName, FSmartName& OutSmartName) const { return false; }
-
-	UE_DEPRECATED(5.3, "This function is no longer used")
-	const FSmartNameMapping* GetSmartNameContainer(const FName& ContainerName) const { return nullptr; }
-
-	UE_DEPRECATED(5.3, "This function is no longer used")
-	void VerifySmartName(const FName&  ContainerName, FSmartName& InOutSmartName) {}
-	
-	UE_DEPRECATED(5.3, "This function is no longer used")
-	void VerifySmartNames(const FName&  ContainerName, TArray<FSmartName>& InOutSmartNames) {}
-
 #if WITH_EDITORONLY_DATA
 private:
 	/** The default skeletal mesh to use when previewing this skeleton */
@@ -808,9 +751,6 @@ public:
 
 	DECLARE_EVENT(USkeleton, FSmartNamesChangedEvent);
 
-	UE_DEPRECATED(5.3, "This member is no longer used. Delegate registration for skeleton metadata can be handled via UAnimCurveMetaData.")
-	FSmartNamesChangedEvent OnSmartNamesChangedEvent;
-
 #if WITH_EDITORONLY_DATA
 	/**
 	 * Global compatibility delegate, used to override skeleton compatibility. If this returns true, no additional
diff --git a/Engine/Source/Runtime/Engine/Private/Animation/Skeleton.cpp b/Engine/Source/Runtime/Engine/Private/Animation/Skeleton.cpp
--- a/Engine/Source/Runtime/Engine/Private/Animation/Skeleton.cpp
+++ b/Engine/Source/Runtime/Engine/Private/Animation/Skeleton.cpp
@@ -37,11 +37,6 @@
 #define LOCTEXT_NAMESPACE "Skeleton"
 #define ROOT_BONE_PARENT	INDEX_NONE
 
-// DEPRECATED - legacy support
-PRAGMA_DISABLE_DEPRECATION_WARNINGS
-TArray<uint16> USkeleton::DefaultCurveUIDList;
-FSmartNameMapping* USkeleton::AnimCurveMapping = nullptr;
-PRAGMA_ENABLE_DEPRECATION_WARNINGS
 
 #if WITH_EDITOR
 const FName USkeleton::AnimNotifyTag = FName(TEXT("AnimNotifyList"));
@@ -63,10 +58,6 @@ const FString USkeleton::CompatibleSkeletonsTagDelimiter = TEXT(";");
 FAreAllSkeletonsCompatible USkeleton::AreAllSkeletonsCompatibleDelegate;
 #endif
 
-// Names of smartname containers for skeleton properties
-const FName USkeleton::AnimCurveMappingName = FName(TEXT("AnimationCurves"));
-const FName USkeleton::AnimTrackCurveMappingName = FName(TEXT("AnimationTrackCurves"));
-
 const FName FAnimSlotGroup::DefaultGroupName = FName(TEXT("DefaultGroup"));
 const FName FAnimSlotGroup::DefaultSlotName = FName(TEXT("DefaultSlot"));
 
@@ -406,6 +397,7 @@ void USkeleton::PostLoad()
 	{
 		// Move curve metadata over to asset user data
 		UAnimCurveMetaData* AnimCurveMetaData = GetOrCreateCurveMetaDataObject();
+		PRAGMA_DISABLE_DEPRECATION_WARNINGS
 		for(const TPair<FName, FSmartNameMapping>& NameMappingPair : SmartNames_DEPRECATED.NameMappings)
 		{
 			for(const TPair<FName, FCurveMetaData>& NameMetaDataPair : NameMappingPair.Value.CurveMetaDataMap)
@@ -413,6 +405,7 @@ void USkeleton::PostLoad()
 				AnimCurveMetaData->CurveMetaData.Add(NameMetaDataPair.Key, NameMetaDataPair.Value);
 			}
 		}
+		PRAGMA_ENABLE_DEPRECATION_WARNINGS
 	}
 	else
 	{
@@ -1937,21 +1930,6 @@ int32 USkeleton::GetNumCurveMetaData() const
 	return 0;
 }
 
-const FCurveMetaData* USkeleton::GetCurveMetaData(const SmartName::UID_Type CurveUID) const
-{
-	return nullptr;
-}
-
-FCurveMetaData* USkeleton::GetCurveMetaData(const FSmartName& CurveName)
-{
-	return GetCurveMetaData(CurveName.DisplayName);
-}
-
-const FCurveMetaData* USkeleton::GetCurveMetaData(const FSmartName& CurveName) const
-{
-	return GetCurveMetaData(CurveName.DisplayName);
-}
-
 void USkeleton::AccumulateCurveMetaData(FName CurveName, bool bMaterialSet, bool bMorphtargetSet)
 {
 	if (bMaterialSet || bMorphtargetSet)
