Break up ServerReplicateActors transactions when a net driver is running transactionally to improve performance. It will use transactions for PrepConnections and BuildConsiderList, and one for replication to each connection. Disable the cvar net.ServerReplicateActorsGranularTransactions to revert to the old behavior of running ServerReplicateActors in one big transaction.

diff --git a/Engine/Source/Runtime/Engine/Classes/Engine/NetDriver.h b/Engine/Source/Runtime/Engine/Classes/Engine/NetDriver.h
--- a/Engine/Source/Runtime/Engine/Classes/Engine/NetDriver.h
+++ b/Engine/Source/Runtime/Engine/Classes/Engine/NetDriver.h
@@ -363,6 +363,7 @@ namespace UE::Net
 {
 	class FNetTokenStore;
 	class FNetIDVariant;
+	struct FServerReplicateActors_ForConnectionParams;
 
 	class FScopedIgnoreStaticActorDestruction
 	{
@@ -2447,6 +2448,9 @@ private:
 
 	void RemoveDestroyedGuidsByLevel(const ULevel* Level, const TArray<FNetworkGUID>& RemovedGUIDs);
 
+	/** Private helper for ServerReplicateActors */
+	void ServerReplicateActors_ForConnection(UE::Net::FServerReplicateActors_ForConnectionParams& Params);
+
 	/** Handle to FNetDelegates::OnSyncLoadDetected delegate */
 	FDelegateHandle ReportSyncLoadDelegateHandle;
 
diff --git a/Engine/Source/Runtime/Engine/Private/NetDriver.cpp b/Engine/Source/Runtime/Engine/Private/NetDriver.cpp
--- a/Engine/Source/Runtime/Engine/Private/NetDriver.cpp
+++ b/Engine/Source/Runtime/Engine/Private/NetDriver.cpp
@@ -343,6 +343,14 @@ namespace UE::Net::Private
 		TEXT("net.iris.UpdateReplicationSystemWithNoConnections"),
 		bUpdateReplicationSystemWithNoConnections,
 		TEXT("If true ReplicationSystem::NetUpdate will be called even if we have no connections to ensure that we do not accumulate pending work."));
+	
+#if UE_WITH_REMOTE_OBJECT_HANDLE
+	bool bServerReplicateActorsGranularTransactions = true;
+	static FAutoConsoleVariableRef CVarServerReplicateActorsGranularTransactions(
+		TEXT("net.ServerReplicateActorsGranularTransactions"),
+		bServerReplicateActorsGranularTransactions,
+		TEXT("If true and bReplicateTransactionally is true, ServerReplicateActors will use a transaction per connection instead of one transaction for the whole function."));
+#endif
 
 	static void DumpActiveNetActors(UWorld* InWorld)
 	{
@@ -581,6 +589,16 @@ namespace UE::Net
 		TEXT("(specifying 'all' will enable this for all NetDriverDefinition's)"),
 		FConsoleVariableDelegate::CreateStatic(&ParseRequiredEncryptionCVars));
 
+	struct FServerReplicateActors_ForConnectionParams
+	{
+		UNetConnection* Connection;
+		TArray<FNetViewer>& ConnectionViewers;
+		float DeltaSeconds;
+		const TArray<FNetworkObjectInfo*>& ConsiderList;
+		int32& InOutUpdated;
+		bool bCPUSaturated;
+	};
+
 	namespace Private
 	{
 		static int32 CleanUpRenamedDynamicActors = 0;
@@ -1138,7 +1156,7 @@ void UNetDriver::TickFlush(float DeltaSeconds)
 		if (ClientConnections.Num() > 0)
 		{
 #if UE_WITH_REMOTE_OBJECT_HANDLE
-			if (bReplicateTransactionally)
+			if (bReplicateTransactionally && !UE::Net::Private::bServerReplicateActorsGranularTransactions)
 			{
 				static FName TransactionalWorkName(TEXT("ServerReplicateActors"));
 				UE::RemoteExecutor::ExecuteTransactional(TransactionalWorkName, [this, DeltaSeconds]()
@@ -1147,12 +1165,10 @@ void UNetDriver::TickFlush(float DeltaSeconds)
 				});
 			}
 			else
+#endif // UE_WITH_REMOTE_OBJECT_HANDLE
 			{
 				ServerReplicateActors(DeltaSeconds);
 			}
-#else
-			ServerReplicateActors(DeltaSeconds);
-#endif // UE_WITH_REMOTE_OBJECT_HANDLE
 		}
 #endif // WITH_SERVER_CODE
 	}
@@ -5331,6 +5347,14 @@ int32 UNetDriver::ServerReplicateActors_PrioritizeActors( UNetConnection* Connec
 
 		for ( FNetworkObjectInfo* ActorInfo : ConsiderList )
 		{
+#if UE_WITH_REMOTE_OBJECT_HANDLE
+			// If we can run transactions entries in the ConsiderList might become invalid
+			if (bReplicateTransactionally && !ActorInfo->WeakActor.IsValid())
+			{
+				continue;
+			}
+#endif
+
 			AActor* Actor = ActorInfo->Actor;
 
 			UActorChannel* Channel = Connection->FindActorChannelRef( ActorInfo->WeakActor );
@@ -5697,6 +5721,72 @@ void UNetDriver::ServerReplicateActors_MarkRelevantActors( UNetConnection* Conne
 
 #endif
 
+void UNetDriver::ServerReplicateActors_ForConnection(UE::Net::FServerReplicateActors_ForConnectionParams& Params)
+{
+#if WITH_SERVER_CODE
+	const int32 LocalNumSaturated = GNumSaturatedConnections;
+
+	// Make a list of viewers this connection should consider (this connection and children of this connection)
+	Params.ConnectionViewers.Reset();
+	new(Params.ConnectionViewers)FNetViewer(Params.Connection, Params.DeltaSeconds);
+	for (int32 ViewerIndex = 0; ViewerIndex < Params.Connection->Children.Num(); ViewerIndex++)
+	{
+		if (Params.Connection->Children[ViewerIndex]->ViewTarget != NULL)
+		{
+			new(Params.ConnectionViewers)FNetViewer(Params.Connection->Children[ViewerIndex], Params.DeltaSeconds);
+		}
+	}
+
+	// send ClientAdjustment if necessary
+	// we do this here so that we send a maximum of one per packet to that client; there is no value in stacking additional corrections
+	if (Params.Connection->PlayerController)
+	{
+		Params.Connection->PlayerController->SendClientAdjustment();
+	}
+
+	for (int32 ChildIdx = 0; ChildIdx < Params.Connection->Children.Num(); ChildIdx++)
+	{
+		if (Params.Connection->Children[ChildIdx]->PlayerController != NULL)
+		{
+			Params.Connection->Children[ChildIdx]->PlayerController->SendClientAdjustment();
+		}
+	}
+
+	FMemMark RelevantActorMark(FMemStack::Get());
+
+	const bool bProcessConsiderListIsBound = OnProcessConsiderListOverride.IsBound();
+
+	if (bProcessConsiderListIsBound)
+	{
+		OnProcessConsiderListOverride.Execute({ Params.DeltaSeconds, Params.Connection, Params.bCPUSaturated }, Params.InOutUpdated, Params.ConsiderList);
+	}
+
+	if (!bProcessConsiderListIsBound)
+	{
+		FActorPriority* PriorityList = NULL;
+		FActorPriority** PriorityActors = NULL;
+
+		// Get a sorted list of actors for this connection
+		const int32 FinalSortedCount = ServerReplicateActors_PrioritizeActors(Params.Connection, Params.ConnectionViewers, Params.ConsiderList, Params.bCPUSaturated, PriorityList, PriorityActors);
+
+		// Process the sorted list of actors for this connection
+		TInterval<int32> ActorsIndexRange(0, FinalSortedCount);
+		const int32 LastProcessedActor = ServerReplicateActors_ProcessPrioritizedActorsRange(Params.Connection, Params.ConnectionViewers, PriorityActors, ActorsIndexRange, Params.InOutUpdated);
+
+		ServerReplicateActors_MarkRelevantActors(Params.Connection, Params.ConnectionViewers, LastProcessedActor, FinalSortedCount, PriorityActors);
+	}
+
+	RelevantActorMark.Pop();
+
+	Params.ConnectionViewers.Reset();
+
+	Params.Connection->LastProcessedFrame = ReplicationFrame;
+
+	const bool bWasSaturated = GNumSaturatedConnections > LocalNumSaturated;
+	Params.Connection->TrackReplicationForAnalytics(bWasSaturated);
+#endif
+}
+
 int64 UNetDriver::SendDestructionInfo(UNetConnection* Connection, FActorDestructionInfo* DestructionInfo)
 {
 	LLM_SCOPE_BYTAG(NetDriver);
@@ -5975,7 +6065,26 @@ int32 UNetDriver::ServerReplicateActors(float DeltaSeconds)
 
 	int32 Updated = 0;
 
-	const int32 NumClientsToTick = ServerReplicateActors_PrepConnections( DeltaSeconds );
+	const int32 NumClientsToTick = [this, DeltaSeconds]()
+		{
+#if UE_WITH_REMOTE_OBJECT_HANDLE
+			if (bReplicateTransactionally && UE::Net::Private::bServerReplicateActorsGranularTransactions)
+			{
+				int32 LocalNumClientsToTick = 0;
+				static FName TransactionalWorkName("ServerReplicateActors_PrepConnections");
+				UE::RemoteExecutor::ExecuteTransactional(TransactionalWorkName, [this, DeltaSeconds, &LocalNumClientsToTick]()
+					{
+						LocalNumClientsToTick = ServerReplicateActors_PrepConnections(DeltaSeconds);
+					});
+				return LocalNumClientsToTick;
+			}
+			else
+#endif // UE_WITH_REMOTE_OBJECT_HANDLE
+			{
+				return ServerReplicateActors_PrepConnections(DeltaSeconds);
+			}
+		}();
+
 
 	if ( NumClientsToTick == 0 )
 	{
@@ -6001,7 +6110,20 @@ int32 UNetDriver::ServerReplicateActors(float DeltaSeconds)
 	ConsiderList.Reserve( GetNetworkObjectList().GetActiveObjects().Num() );
 
 	// Build the consider list (actors that are ready to replicate)
-	ServerReplicateActors_BuildConsiderList( ConsiderList, ServerTickTime );
+#if UE_WITH_REMOTE_OBJECT_HANDLE
+	if (bReplicateTransactionally && UE::Net::Private::bServerReplicateActorsGranularTransactions)
+	{
+		static FName TransactionalWorkName("ServerReplicateActors_BuildConsiderList");
+		UE::RemoteExecutor::ExecuteTransactional(TransactionalWorkName, [this, &ConsiderList, ServerTickTime]()
+			{
+				ServerReplicateActors_BuildConsiderList(ConsiderList, ServerTickTime);
+			});
+	}
+	else
+#endif // UE_WITH_REMOTE_OBJECT_HANDLE
+	{
+		ServerReplicateActors_BuildConsiderList(ConsiderList, ServerTickTime);
+	}
 
 	TSet<UNetConnection*> ConnectionsToClose;
 
@@ -6042,12 +6164,12 @@ int32 UNetDriver::ServerReplicateActors(float DeltaSeconds)
 			// then mark each considered actor as bPendingNetUpdate so that they will be considered again the next frame when the connection is actually ticked
 			for (int32 ConsiderIdx = 0; ConsiderIdx < ConsiderList.Num(); ConsiderIdx++)
 			{
-				AActor *Actor = ConsiderList[ConsiderIdx]->Actor;
+				AActor* Actor = ConsiderList[ConsiderIdx]->Actor;
 				// if the actor hasn't already been flagged by another connection,
 				if (Actor != NULL && !ConsiderList[ConsiderIdx]->bPendingNetUpdate)
 				{
 					// find the channel
-					UActorChannel *Channel = Connection->FindActorChannelRef(ConsiderList[ConsiderIdx]->WeakActor);
+					UActorChannel* Channel = Connection->FindActorChannelRef(ConsiderList[ConsiderIdx]->WeakActor);
 					// and if the channel last update time doesn't match the last net update time for the actor
 					if (Channel != NULL && Channel->LastUpdateTime < ConsiderList[ConsiderIdx]->LastNetUpdateTimestamp)
 					{
@@ -6061,70 +6183,31 @@ int32 UNetDriver::ServerReplicateActors(float DeltaSeconds)
 			Connection->TimeSensitive = false;
 		}
 		else if (Connection->ViewTarget)
-		{		
-
-			const int32 LocalNumSaturated = GNumSaturatedConnections;
-
-			// Make a list of viewers this connection should consider (this connection and children of this connection)
-			TArray<FNetViewer>& ConnectionViewers = WorldSettings->ReplicationViewers;
-
-			ConnectionViewers.Reset();
-			new( ConnectionViewers )FNetViewer( Connection, DeltaSeconds );
-			for ( int32 ViewerIndex = 0; ViewerIndex < Connection->Children.Num(); ViewerIndex++ )
-			{
-				if ( Connection->Children[ViewerIndex]->ViewTarget != NULL )
-				{
-					new( ConnectionViewers )FNetViewer( Connection->Children[ViewerIndex], DeltaSeconds );
-				}
-			}
-
-			// send ClientAdjustment if necessary
-			// we do this here so that we send a maximum of one per packet to that client; there is no value in stacking additional corrections
-			if ( Connection->PlayerController )
-			{
-				Connection->PlayerController->SendClientAdjustment();
-			}
-
-			for ( int32 ChildIdx = 0; ChildIdx < Connection->Children.Num(); ChildIdx++ )
+		{
+			UE::Net::FServerReplicateActors_ForConnectionParams Params =
 			{
-				if ( Connection->Children[ChildIdx]->PlayerController != NULL )
-				{
-					Connection->Children[ChildIdx]->PlayerController->SendClientAdjustment();
-				}
-			}
-
-			FMemMark RelevantActorMark(FMemStack::Get());
-
-			const bool bProcessConsiderListIsBound = OnProcessConsiderListOverride.IsBound();
+				.Connection = Connection,
+				.ConnectionViewers = WorldSettings->ReplicationViewers,
+				.DeltaSeconds = DeltaSeconds,
+				.ConsiderList = ConsiderList,
+				.InOutUpdated = Updated,
+				.bCPUSaturated = bCPUSaturated
+			};
 
-			if (bProcessConsiderListIsBound)
+#if UE_WITH_REMOTE_OBJECT_HANDLE
+			if (bReplicateTransactionally && UE::Net::Private::bServerReplicateActorsGranularTransactions)
 			{
-				OnProcessConsiderListOverride.Execute( { DeltaSeconds, Connection, bCPUSaturated }, Updated, ConsiderList );
+				static FName TransactionalWorkName("ServerReplicateActors_ForConnection");
+				UE::RemoteExecutor::ExecuteTransactional(TransactionalWorkName, [this, &Params]()
+					{
+						ServerReplicateActors_ForConnection(Params);
+					});
 			}
-
-			if (!bProcessConsiderListIsBound)
+			else
+#endif // UE_WITH_REMOTE_OBJECT_HANDLE
 			{
-				FActorPriority* PriorityList = NULL;
-				FActorPriority** PriorityActors = NULL;
-
-				// Get a sorted list of actors for this connection
-				const int32 FinalSortedCount = ServerReplicateActors_PrioritizeActors(Connection, ConnectionViewers, ConsiderList, bCPUSaturated, PriorityList, PriorityActors);
-
-				// Process the sorted list of actors for this connection
-				TInterval<int32> ActorsIndexRange(0, FinalSortedCount);
-				const int32 LastProcessedActor = ServerReplicateActors_ProcessPrioritizedActorsRange(Connection, ConnectionViewers, PriorityActors, ActorsIndexRange, Updated);
-
-				ServerReplicateActors_MarkRelevantActors(Connection, ConnectionViewers, LastProcessedActor, FinalSortedCount, PriorityActors);
+				ServerReplicateActors_ForConnection(Params);
 			}
-
-			RelevantActorMark.Pop();
-
-			ConnectionViewers.Reset();
-
-			Connection->LastProcessedFrame = ReplicationFrame;
-
-			const bool bWasSaturated = GNumSaturatedConnections > LocalNumSaturated;
-			Connection->TrackReplicationForAnalytics(bWasSaturated);
 		}
 
 		if (Connection->GetPendingCloseDueToReplicationFailure())
