[shaders] fix unnecessary retry popup when compile errors are resolved; need to reset value of bAllJobsSucceeded flag when re-issuing failed jobs (remaining jobs not reissued for a map are all guaranteed to be successful)

diff --git a/Engine/Source/Runtime/Engine/Private/ShaderCompiler/ShaderCompiler.cpp b/Engine/Source/Runtime/Engine/Private/ShaderCompiler/ShaderCompiler.cpp
--- a/Engine/Source/Runtime/Engine/Private/ShaderCompiler/ShaderCompiler.cpp
+++ b/Engine/Source/Runtime/Engine/Private/ShaderCompiler/ShaderCompiler.cpp
@@ -2503,6 +2503,9 @@ bool FShaderCompilingManager::HandlePotentialRetry(TMap<int32, FShaderMapFinaliz
 		if (bAnyToReissue)
 		{
 			ReissueMapIds.Add(It.Key());
+			// reset flag indicating failed jobs for a map when reissuing; we will always re-issue all failed jobs if we reissue any,
+			// so the ones remaining in the results array should all be successful
+			It.Value().bAllJobsSucceeded = true;
 		}
 	}
 
