Fixes for using NMT_JoinNoPawn and NMT_JoinNoPawnSplit with Iris.

diff --git a/Engine/Source/Runtime/Engine/Classes/Engine/World.h b/Engine/Source/Runtime/Engine/Classes/Engine/World.h
--- a/Engine/Source/Runtime/Engine/Classes/Engine/World.h
+++ b/Engine/Source/Runtime/Engine/Classes/Engine/World.h
@@ -1160,6 +1160,9 @@ private:
 	/** Returns wether AddToWorld should be skipped on a given level */
 	UE_API bool CanAddLoadedLevelToWorld(ULevel* Level) const;
 
+	/** Spawn an instance of ANoPawnPlayerController and set it as the value of Connection::PlayerController. */
+	void SpawnAndSetNoPawnPlayerController(UNetConnection* Connection, uint8 NetPlayerIndex, int32 LocalPlayerIdentifier);
+
 	/** Stores whether the game world supports for a client to use "making visible" transaction requests to the server. */
 	mutable TOptional<bool> bSupportsMakingVisibleTransactionRequests;
 
diff --git a/Engine/Source/Runtime/Engine/Private/World.cpp b/Engine/Source/Runtime/Engine/Private/World.cpp
--- a/Engine/Source/Runtime/Engine/Private/World.cpp
+++ b/Engine/Source/Runtime/Engine/Private/World.cpp
@@ -4816,6 +4816,27 @@ bool UWorld::CanAddLoadedLevelToWorld(ULevel* Level) const
 	return true;
 }
 
+void UWorld::SpawnAndSetNoPawnPlayerController(UNetConnection* Connection, uint8 NetPlayerIndex, int32 LocalPlayerIdentifier)
+{
+	FActorSpawnParameters SpawnInfo;
+	SpawnInfo.ObjectFlags |= RF_Transient;
+	SpawnInfo.bDeferConstruction = true;
+
+	Connection->PlayerController = GetWorld()->SpawnActor<APlayerController>(ANoPawnPlayerController::StaticClass(), FVector(), FRotator(), SpawnInfo);
+	Connection->PlayerController->NetPlayerIndex = NetPlayerIndex;
+	Connection->PlayerController->NetConnection = Connection;
+	Connection->PlayerController->FinishSpawning(FTransform());
+
+	Connection->PlayerController->SetRole(ROLE_Authority);
+	Connection->PlayerController->SetReplicates(true);
+	Connection->PlayerController->SetAutonomousProxy(true);
+	Connection->PlayerController->SetClientHandshakeId(Connection->GetClientHandshakeId());
+	Connection->PlayerController->SetLocalPlayerConnectionIdentifier(LocalPlayerIdentifier);
+
+	Connection->OwningActor = Connection->PlayerController;
+	Connection->SetClientLoginState(EClientLoginState::ReceivedJoin);
+}
+
 PRAGMA_DISABLE_DEPRECATION_WARNINGS
 void UWorld::SetBegunPlay(bool bHasBegunPlay)
 {
@@ -7532,18 +7553,7 @@ void UWorld::NotifyControlMessage(UNetConnection* Connection, uint8 MessageType,
 
 					UE::Private::World::ParseAndSetClientHandshakeId(InURL, Connection);
 
-					FActorSpawnParameters SpawnInfo;
-					SpawnInfo.ObjectFlags |= RF_Transient;
-					SpawnInfo.bDeferConstruction = true;
-					
-					Connection->PlayerController = GetWorld()->SpawnActor<APlayerController>(ANoPawnPlayerController::StaticClass(), FVector(), FRotator(), SpawnInfo);
-					Connection->PlayerController->SetAutonomousProxy(true);
-					Connection->PlayerController->FinishSpawning(FTransform());
-					Connection->PlayerController->NetPlayerIndex = 0;
-					Connection->PlayerController->NetConnection = Connection;
-					Connection->PlayerController->SetClientHandshakeId(Connection->GetClientHandshakeId());
-
-					Connection->OwningActor = Connection->PlayerController;
+					SpawnAndSetNoPawnPlayerController(Connection, 0, 0);
 
 #if WITH_EDITORONLY_DATA
 					check(CurrentLevel);
@@ -7593,17 +7603,7 @@ void UWorld::NotifyControlMessage(UNetConnection* Connection, uint8 MessageType,
 						
 						const int32 LocalPlayerIdentifier = UE::Private::World::GetLocalPlayerIdentifier(InURL);
 
-						ChildConnection->PlayerController = GetWorld()->SpawnActor<APlayerController>(ANoPawnPlayerController::StaticClass(), FVector(), FRotator(), SpawnInfo);
-						ChildConnection->PlayerController->SetRole(ROLE_Authority);
-						ChildConnection->PlayerController->SetReplicates(true);
-						ChildConnection->PlayerController->SetAutonomousProxy(true);
-						ChildConnection->PlayerController->FinishSpawning(FTransform());
-						ChildConnection->PlayerController->NetPlayerIndex = NetPlayerIndex;
-						ChildConnection->PlayerController->NetConnection = ChildConnection;
-						ChildConnection->PlayerController->SetClientHandshakeId(ChildConnection->GetClientHandshakeId());
-						ChildConnection->PlayerController->SetLocalPlayerConnectionIdentifier(LocalPlayerIdentifier);
-
-						ChildConnection->OwningActor = ChildConnection->PlayerController;
+						SpawnAndSetNoPawnPlayerController(ChildConnection, NetPlayerIndex, LocalPlayerIdentifier);
 
 #if WITH_EDITORONLY_DATA
 						check(CurrentLevel);
