Adding public API to allow adding a new runtime cell transformer from code.

diff --git a/Engine/Source/Runtime/Engine/Private/WorldPartition/WorldPartition.cpp b/Engine/Source/Runtime/Engine/Private/WorldPartition/WorldPartition.cpp
--- a/Engine/Source/Runtime/Engine/Private/WorldPartition/WorldPartition.cpp
+++ b/Engine/Source/Runtime/Engine/Private/WorldPartition/WorldPartition.cpp
@@ -2601,6 +2601,28 @@ bool UWorldPartition::HasRuntimeCellTransformerOfType(const UClass* InCellTransf
 	return false;
 }
 
+void UWorldPartition::AddRuntimeCellTransformerOfType(UClass* InCellTransformerClass, int32 PositionToInsert)
+{
+	int32 Index;
+	if (PositionToInsert >= 0)
+	{
+		Index = PositionToInsert;
+	}
+	else
+	{
+		// negative numbers insert at a location relative to the end: -1 inserts at the end, -2 second form the end, etc
+		Index = RuntimeCellsTransformerStack.Num() + 1 + PositionToInsert;
+	}
+	// Note that for insertion there are N+1 valid indices for a list of length N
+	if (ensure((Index >= 0) && (Index <= RuntimeCellsTransformerStack.Num())))
+	{
+		RuntimeCellsTransformerStack.InsertDefaulted(Index);
+		FRuntimeCellTransformerInstance& TransformerInstance = RuntimeCellsTransformerStack[Index];
+		TransformerInstance.Class = InCellTransformerClass;
+		TransformerInstance.Instance = TransformerInstance.Class ? NewObject<UWorldPartitionRuntimeCellTransformer>(this, TransformerInstance.Class, NAME_None) : nullptr;
+	}
+}
+
 #endif
 
 bool UWorldPartition::SupportsWorldAssetStreaming(const FName& InTargetGrid)
diff --git a/Engine/Source/Runtime/Engine/Public/WorldPartition/WorldPartition.h b/Engine/Source/Runtime/Engine/Public/WorldPartition/WorldPartition.h
--- a/Engine/Source/Runtime/Engine/Public/WorldPartition/WorldPartition.h
+++ b/Engine/Source/Runtime/Engine/Public/WorldPartition/WorldPartition.h
@@ -342,6 +342,7 @@ public:
 
 	ENGINE_API void ApplyRuntimeCellsTransformerStack(ULevel* InLevel);
 	ENGINE_API bool HasRuntimeCellTransformerOfType(const UClass* InCellTransformerClass);
+	ENGINE_API void AddRuntimeCellTransformerOfType(UClass* InCellTransformerClass, int32 PositionToInsert = -1);
 	
 	UHLODLayer* GetDefaultHLODLayer() const { return DefaultHLODLayer; }
 	void SetDefaultHLODLayer(UHLODLayer* InDefaultHLODLayer) { DefaultHLODLayer = InDefaultHLODLayer; }
