UDataTable changes for code size

Extract common type checking into private method;
all impls of FindRow are the same except for the T::StaticStruct() and a typecast at the end.

-140k FNSCT

diff --git a/Engine/Source/Runtime/Engine/Classes/Engine/DataTable.h b/Engine/Source/Runtime/Engine/Classes/Engine/DataTable.h
--- a/Engine/Source/Runtime/Engine/Classes/Engine/DataTable.h
+++ b/Engine/Source/Runtime/Engine/Classes/Engine/DataTable.h
@@ -203,15 +203,7 @@ public:
 	template <class T>
 	void GetAllRows(const TCHAR* ContextString, OUT TArray<T*>& OutRowArray) const
 	{
-		if (RowStruct == nullptr)
-		{
-			UE_LOG(LogDataTable, Error, TEXT("UDataTable::GetAllRows : DataTable '%s' has no RowStruct specified (%s)."), *GetPathName(), ContextString);
-		}
-		else if (!RowStruct->IsChildOf(T::StaticStruct()))
-		{
-			UE_LOG(LogDataTable, Error, TEXT("UDataTable::GetAllRows : Incorrect type specified for DataTable '%s' (%s)."), *GetPathName(), ContextString);
-		}
-		else
+		if (CommonTypeCheck(T::StaticStruct(), TEXT("UDataTable::GetAllRows"), ContextString))
 		{
 			OutRowArray.Reserve(OutRowArray.Num() + GetRowMap().Num());
 			for (TMap<FName, uint8*>::TConstIterator RowMapIter(GetRowMap().CreateConstIterator()); RowMapIter; ++RowMapIter)
@@ -231,38 +223,7 @@ public:
 	template <class T>
 	T* FindRow(FName RowName, const TCHAR* ContextString, bool bWarnIfRowMissing = true) const
 	{
-		if(RowStruct == nullptr)
-		{
-			UE_LOG(LogDataTable, Error, TEXT("UDataTable::FindRow : '%s' specified no row for DataTable '%s'."), ContextString, *GetPathName());
-			return nullptr;
-		}
-
-		if(!RowStruct->IsChildOf(T::StaticStruct()))
-		{
-			UE_CLOG(bWarnIfRowMissing, LogDataTable, Error, TEXT("UDataTable::FindRow : '%s' specified incorrect type for DataTable '%s'."), ContextString, *GetPathName());
-			return nullptr;
-		}
-
-		if(RowName.IsNone())
-		{
-			UE_CLOG(bWarnIfRowMissing, LogDataTable, Warning, TEXT("UDataTable::FindRow : '%s' requested invalid row 'None' from DataTable '%s'."), ContextString, *GetPathName());
-			return nullptr;
-		}
-
-		uint8* const* RowDataPtr = GetRowMap().Find(RowName);
-		if (RowDataPtr == nullptr)
-		{
-			if (bWarnIfRowMissing)
-			{
-				UE_LOG(LogDataTable, Warning, TEXT("UDataTable::FindRow : '%s' requested row '%s' not in DataTable '%s'."), ContextString, *RowName.ToString(), *GetPathName());
-			}
-			return nullptr;
-		}
-
-		uint8* RowData = *RowDataPtr;
-		check(RowData);
-
-		return reinterpret_cast<T*>(RowData);
+		return reinterpret_cast<T*>(FindRowInternal(T::StaticStruct(), RowName, ContextString, bWarnIfRowMissing));
 	}
 
 	template <class T>
@@ -275,15 +236,7 @@ public:
 	template <class T>
 	void ForeachRow(const TCHAR* ContextString, TFunctionRef<void (const FName& Key, const T& Value)> Predicate) const
 	{
-		if (RowStruct == nullptr)
-		{
-			UE_LOG(LogDataTable, Error, TEXT("UDataTable::ForeachRow : DataTable '%s' has no RowStruct specified (%s)."), *GetPathName(), ContextString);
-		}
-		else if (!RowStruct->IsChildOf(T::StaticStruct()))
-		{
-			UE_LOG(LogDataTable, Error, TEXT("UDataTable::ForeachRow : Incorrect type specified for DataTable '%s' (%s)."), *GetPathName(), ContextString);
-		}
-		else
+		if (CommonTypeCheck(T::StaticStruct(), TEXT("UDataTable::ForeachRow"), ContextString))
 		{
 			for (TMap<FName, uint8*>::TConstIterator RowMapIter(GetRowMap().CreateConstIterator()); RowMapIter; ++RowMapIter)
 			{
@@ -430,6 +383,10 @@ protected:
 		static TMap<UDataTable*, int32> ScopeCount;
 		static FTransactionallySafeCriticalSection CriticalSection;
 	};
+
+private:
+	ENGINE_API bool CommonTypeCheck(UScriptStruct* Type, const TCHAR* MethodName, const TCHAR* ContextString) const;
+	ENGINE_API uint8* FindRowInternal(UScriptStruct* Type, FName RowName, const TCHAR* ContextString, bool bWarnIfRowMissing) const;
 };
 
 
diff --git a/Engine/Source/Runtime/Engine/Private/DataTable.cpp b/Engine/Source/Runtime/Engine/Private/DataTable.cpp
--- a/Engine/Source/Runtime/Engine/Private/DataTable.cpp
+++ b/Engine/Source/Runtime/Engine/Private/DataTable.cpp
@@ -1056,6 +1056,51 @@ TArray<FName> UDataTable::GetRowNames() const
 	return Keys;
 }
 
+bool UDataTable::CommonTypeCheck(UScriptStruct* Type, const TCHAR* MethodName, const TCHAR* ContextString) const
+{
+	if (RowStruct == nullptr)
+	{
+		UE_LOG(LogDataTable, Error, TEXT("%s : DataTable '%s' has no RowStruct specified (%s)."), MethodName, *GetPathName(), ContextString);
+		return false;
+	}
+	else if (!RowStruct->IsChildOf(Type))
+	{
+		UE_LOG(LogDataTable, Error, TEXT("%s : Incorrect type specified for DataTable '%s' (%s)."), MethodName, *GetPathName(), ContextString);
+		return false;
+	}
+
+	return true;
+}
+
+uint8* UDataTable::FindRowInternal(UScriptStruct* Type, FName RowName, const TCHAR* ContextString, bool bWarnIfRowMissing) const
+{
+	if (!CommonTypeCheck(Type, TEXT("UDataTabke::FindRow"), ContextString))
+	{
+		return nullptr;
+	}
+
+	if(RowName.IsNone())
+	{
+		UE_CLOG(bWarnIfRowMissing, LogDataTable, Warning, TEXT("UDataTable::FindRow : '%s' requested invalid row 'None' from DataTable '%s'."), ContextString, *GetPathName());
+		return nullptr;
+	}
+
+	uint8* const* RowDataPtr = GetRowMap().Find(RowName);
+	if (RowDataPtr == nullptr)
+	{
+		if (bWarnIfRowMissing)
+		{
+			UE_LOG(LogDataTable, Warning, TEXT("UDataTable::FindRow : '%s' requested row '%s' not in DataTable '%s'."), ContextString, *RowName.ToString(), *GetPathName());
+		}
+		return nullptr;
+	}
+
+	uint8* RowData = *RowDataPtr;
+	check(RowData);
+
+	return RowData;
+}
+
 bool FDataTableRowHandle::operator==(FDataTableRowHandle const& Other) const
 {
 	return DataTable == Other.DataTable && RowName == Other.RowName;
