add rotator-based oriented box 'make' function, so there is a standard way to construct an oriented box in blueprint w/ valid axes

diff --git a/Engine/Source/Runtime/CoreUObject/Public/UObject/NoExportTypes.h b/Engine/Source/Runtime/CoreUObject/Public/UObject/NoExportTypes.h
--- a/Engine/Source/Runtime/CoreUObject/Public/UObject/NoExportTypes.h
+++ b/Engine/Source/Runtime/CoreUObject/Public/UObject/NoExportTypes.h
@@ -1682,7 +1682,7 @@ struct FBoxSphereBounds
  * Structure for arbitrarily oriented boxes (i.e. not necessarily axis-aligned).
  * @note The full C++ class is located here: Engine\Source\Runtime\Core\Public\Math\OrientedBox.h
  */
-USTRUCT(immutable, noexport, BlueprintType, IsAlwaysAccessible, HasDefaults)
+USTRUCT(immutable, noexport, BlueprintType, IsAlwaysAccessible, HasDefaults, meta = (HasNativeMake = "/Script/Engine.KismetMathLibrary.MakeOrientedBox"))
 struct FOrientedBox
 {
 	/** Holds the center of the box. */
diff --git a/Engine/Source/Runtime/Engine/Classes/Kismet/KismetMathLibrary.h b/Engine/Source/Runtime/Engine/Classes/Kismet/KismetMathLibrary.h
--- a/Engine/Source/Runtime/Engine/Classes/Kismet/KismetMathLibrary.h
+++ b/Engine/Source/Runtime/Engine/Classes/Kismet/KismetMathLibrary.h
@@ -3705,6 +3705,17 @@ class UKismetMathLibrary : public UBlueprintFunctionLibrary
 	UFUNCTION(BlueprintPure, Category="Math|Box", meta=(Keywords="construct build"))
 	static ENGINE_API FBox MakeBoxWithOrigin(const FVector& Origin, const FVector& Extent);
 
+	/** 
+	 * Makes an FOrientedBox from Center, Orientation and Extent 
+	 * 
+	 * @param Origin The center of the box
+	 * @param Orientation The rotation of the box
+	 * @param Extent The half-size of the box
+	 * @return A new oriented bounding box
+	 */
+	UFUNCTION(BlueprintPure, Category = "Math|OrientedBox", meta = (Keywords = "construct build", NativeMakeFunc))
+	static ENGINE_API FOrientedBox MakeOrientedBox(FVector Origin, FRotator Orientation, FVector Extent);
+
 	/**
 	 * Returns true if the InnerTest Box is is completely inside of the OuterTest Box
 	 * 
diff --git a/Engine/Source/Runtime/Engine/Private/KismetMathLibrary.cpp b/Engine/Source/Runtime/Engine/Private/KismetMathLibrary.cpp
--- a/Engine/Source/Runtime/Engine/Private/KismetMathLibrary.cpp
+++ b/Engine/Source/Runtime/Engine/Private/KismetMathLibrary.cpp
@@ -1275,6 +1275,17 @@ FBox UKismetMathLibrary::MakeBoxWithOrigin(const FVector& Origin, const FVector&
 	return FBox::BuildAABB(Origin, Extent);
 }
 
+FOrientedBox UKismetMathLibrary::MakeOrientedBox(FVector Origin, FRotator Orientation, FVector Extent)
+{
+	FOrientedBox ToRet;
+	ToRet.Center = Origin;
+	GetAxes(Orientation, ToRet.AxisX, ToRet.AxisY, ToRet.AxisZ);
+	ToRet.ExtentX = Extent.X;
+	ToRet.ExtentY = Extent.Y;
+	ToRet.ExtentZ = Extent.Z;
+	return ToRet;
+}
+
 bool UKismetMathLibrary::Box_IsInside(const FBox& InnerTest, const FBox& OuterTest)
 {
 	return OuterTest.IsInside(InnerTest);
