Adds depth read in VS for mobile supporting SM5 material nodes

diff --git a/Engine/Source/Runtime/Engine/Private/Materials/HLSLMaterialTranslator.cpp b/Engine/Source/Runtime/Engine/Private/Materials/HLSLMaterialTranslator.cpp
--- a/Engine/Source/Runtime/Engine/Private/Materials/HLSLMaterialTranslator.cpp
+++ b/Engine/Source/Runtime/Engine/Private/Materials/HLSLMaterialTranslator.cpp
@@ -8286,7 +8286,7 @@ int32 FHLSLMaterialTranslator::GetScreenAlignedUV(int32 Offset, int32 ViewportUV
 
 int32 FHLSLMaterialTranslator::SceneDepth(int32 Offset, int32 ViewportUV, bool bUseOffset)
 {
-	if (ShaderFrequency == SF_Vertex && FeatureLevel <= ERHIFeatureLevel::ES3_1)
+	if (ShaderFrequency == SF_Vertex && IsMobilePlatform(Platform) && !MobileSupportsSM5MaterialNodes(Platform))
 	{
 		// mobile currently does not support this, we need to read a separate copy of the depth, we must disable framebuffer fetch and force scene texture reads.
 		return Errorf(TEXT("Cannot read scene depth from the vertex shader with the Mobile feature level"));
