Improve comments for DetailMode to indicate that it can result in completely removing components on certain platforms depending on how the game is set up.

diff --git a/Engine/Source/Runtime/Engine/Classes/Components/SceneComponent.h b/Engine/Source/Runtime/Engine/Classes/Components/SceneComponent.h
--- a/Engine/Source/Runtime/Engine/Classes/Components/SceneComponent.h
+++ b/Engine/Source/Runtime/Engine/Classes/Components/SceneComponent.h
@@ -23,7 +23,10 @@ struct FOverlapInfo;
 typedef TArrayView<const FOverlapInfo> TOverlapArrayView;
 namespace ERHIFeatureLevel { enum Type : int; }
 
-/** Detail mode for scene component rendering, corresponds with the integer value of UWorld::GetDetailMode() */
+/** 
+ * Detail mode for scene component rendering, corresponds with the world GetDetailMode which is set based on platform and scalability settings.
+ * If settings like r.CookOutUnusedDetailModeComponents are used, this can result in components being removed entirely when the target detail mode is higher.
+ */
 UENUM()
 enum EDetailMode : int
 {
@@ -294,7 +297,11 @@ public:
 	UPROPERTY(Category = Mobility, EditAnywhere, BlueprintReadOnly, Replicated)
 	TEnumAsByte<EComponentMobility::Type> Mobility;
 
-	/** If detail mode is >= system detail mode, primitive won't be rendered. */
+	/** 
+	 * Set the required detail level of this component.
+	 * It will be deleted or made invisible if the world's detail level is higher, based on platform and scalability settings.
+	 * For example, a detail mode of High may prevent this component from loading on mobile platforms.
+	 */
 	UPROPERTY(EditAnywhere, AdvancedDisplay, BlueprintReadOnly, Category = LOD)
 	TEnumAsByte<enum EDetailMode> DetailMode;
 
