Add support for setting CollisionEnabled per ISM body.

Calling SetCollisionEnabled on ISM bodies currently calls RecreatePhysicsState which reinitializes all bodies and overrides the user's set CollisionEnabled flag. This change adds a new handler to allow components to have custom logic for recreating their individual physics bodies.

diff --git a/Engine/Source/Runtime/Engine/Classes/Components/InstancedStaticMeshComponent.h b/Engine/Source/Runtime/Engine/Classes/Components/InstancedStaticMeshComponent.h
--- a/Engine/Source/Runtime/Engine/Classes/Components/InstancedStaticMeshComponent.h
+++ b/Engine/Source/Runtime/Engine/Classes/Components/InstancedStaticMeshComponent.h
@@ -519,6 +519,7 @@ public:
 	ENGINE_API virtual FPrimitiveSceneProxy* CreateSceneProxy() override;
 	ENGINE_API virtual FMatrix GetRenderMatrix() const override;
 	ENGINE_API virtual FBodyInstance* GetBodyInstance(FName BoneName = NAME_None, bool bGetWelded = true, int32 Index = INDEX_NONE) const override;
+	ENGINE_API virtual void RecreateInstanceBody(int32 InstanceBodyIndex) override;
 protected:
 	ENGINE_API virtual void OnCreatePhysicsState() override;
 	ENGINE_API virtual void OnDestroyPhysicsState() override;
@@ -744,6 +745,8 @@ protected:
 	/** Initializes the body instance for the specified instance of the static mesh. */
 	ENGINE_API void InitInstanceBody(int32 InstanceIdx, FBodyInstance* InBodyInstance);
 
+	ENGINE_API void InitInstanceBody(int32 InstanceIdx, FBodyInstance* InstanceBodyInstance, FBodyInstance* ReferenceBodyInstance, bool bRuntimeBodyInstance, UBodySetup* BodySetup);
+
 	/** The mappings for all the instances of this component. */
 	UPROPERTY(Transient, DuplicateTransient, TextExportTransient)
 	TArray<FInstancedStaticMeshMappingInfo> CachedMappings;
diff --git a/Engine/Source/Runtime/Engine/Classes/Components/PrimitiveComponent.h b/Engine/Source/Runtime/Engine/Classes/Components/PrimitiveComponent.h
--- a/Engine/Source/Runtime/Engine/Classes/Components/PrimitiveComponent.h
+++ b/Engine/Source/Runtime/Engine/Classes/Components/PrimitiveComponent.h
@@ -2337,6 +2337,9 @@ public:
 		return false;
 	}
 
+	/** Called when the CollisionEnabled of the BodyInstance is changed. */
+	ENGINE_API virtual void RecreateInstanceBody(int32 InstanceBodyIndex);
+
 	/** Event called when a component is 'damaged', allowing for component class specific behaviour */
 	ENGINE_API virtual void ReceiveComponentDamage(float DamageAmount, FDamageEvent const& DamageEvent, AController* EventInstigator, AActor* DamageCauser);
 
diff --git a/Engine/Source/Runtime/Engine/Private/Components/PrimitiveComponent.cpp b/Engine/Source/Runtime/Engine/Private/Components/PrimitiveComponent.cpp
--- a/Engine/Source/Runtime/Engine/Private/Components/PrimitiveComponent.cpp
+++ b/Engine/Source/Runtime/Engine/Private/Components/PrimitiveComponent.cpp
@@ -1698,6 +1698,11 @@ void UPrimitiveComponent::UpdateCollisionProfile()
 }
 #endif // WITH_EDITOR
 
+void UPrimitiveComponent::RecreateInstanceBody(int32 InstanceBodyIndex)
+{ 
+	RecreatePhysicsState();
+}
+
 void UPrimitiveComponent::ReceiveComponentDamage(float DamageAmount, FDamageEvent const& DamageEvent, AController* EventInstigator, AActor* DamageCauser)
 {
 	if (bApplyImpulseOnDamage)
diff --git a/Engine/Source/Runtime/Engine/Private/InstancedStaticMesh.cpp b/Engine/Source/Runtime/Engine/Private/InstancedStaticMesh.cpp
--- a/Engine/Source/Runtime/Engine/Private/InstancedStaticMesh.cpp
+++ b/Engine/Source/Runtime/Engine/Private/InstancedStaticMesh.cpp
@@ -2715,25 +2715,33 @@ void UInstancedStaticMeshComponent::BuildComponentInstanceData(EShaderPlatform I
 	};
 }
 
-void UInstancedStaticMeshComponent::InitInstanceBody(int32 InstanceIdx, FBodyInstance* InstanceBodyInstance)
+void UInstancedStaticMeshComponent::InitInstanceBody(int32 InstanceIdx, FBodyInstance* InstanceBodyInstance, FBodyInstance* ReferenceBodyInstance, bool bRuntimeBodyInstance, UBodySetup* BodySetup)
 {
 	if (!GetStaticMesh())
 	{
-		UE_LOG(LogStaticMesh, Warning, TEXT("Unabled to create a body instance for %s in Actor %s. No StaticMesh set."), *GetName(), GetOwner() ? *GetOwner()->GetName() : TEXT("?"));
+		UE_LOG(LogStaticMesh, Warning, TEXT("Unable to create a body instance for %s in Actor %s. No StaticMesh set."), *GetName(), GetOwner() ? *GetOwner()->GetName() : TEXT("?"));
 		return;
 	}
 
 	check(InstanceIdx < PerInstanceSMData.Num());
 	check(InstanceIdx < InstanceBodies.Num());
 	check(InstanceBodyInstance);
-
-	UBodySetup* BodySetup = GetBodySetup();
+	check(ReferenceBodyInstance);
 	check(BodySetup);
 
 	// Get transform of the instance
 	FTransform InstanceTransform = FTransform(PerInstanceSMData[InstanceIdx].Transform) * GetComponentTransform();
+
+	if(bRuntimeBodyInstance)
+	{
+		InstanceBodyInstance->CopyRuntimeBodyInstancePropertiesFrom(ReferenceBodyInstance);
+	}
+	else
+	{
+		InstanceBodyInstance->CopyBodyInstancePropertiesFrom(ReferenceBodyInstance);
+	}
+
 	
-	InstanceBodyInstance->CopyBodyInstancePropertiesFrom(&BodyInstance);
 	InstanceBodyInstance->InstanceBodyIndex = InstanceIdx; // Set body index 
 
 	// make sure we never enable bSimulatePhysics for ISMComps
@@ -2744,6 +2752,11 @@ void UInstancedStaticMeshComponent::InitInstanceBody(int32 InstanceIdx, FBodyIns
 	InstanceBodyInstance->InitBody(BodySetup, InstanceTransform, this, GetWorld()->GetPhysicsScene(), nullptr);
 }
 
+void UInstancedStaticMeshComponent::InitInstanceBody(int32 InstanceIdx, FBodyInstance* InstanceBodyInstance)
+{
+	InitInstanceBody(InstanceIdx, InstanceBodyInstance, &BodyInstance, false, GetBodySetup());
+}
+
 void UInstancedStaticMeshComponent::CreateAllInstanceBodies()
 {
 	TRACE_CPUPROFILER_EVENT_SCOPE(UInstancedStaticMeshComponent::CreateAllInstanceBodies);
@@ -4194,6 +4207,27 @@ void UInstancedStaticMeshComponent::UpdateComponentTransform(EUpdateTransformFla
 	PreviousComponentTransform = GetComponentTransform();
 }
 
+void UInstancedStaticMeshComponent::RecreateInstanceBody(int32 InstanceBodyIndex)
+{
+	if (InstanceBodies.IsValidIndex(InstanceBodyIndex))
+	{
+		FBodyInstance* Body = InstanceBodies[InstanceBodyIndex];
+		if (Body)
+		{
+			const FTransform InstanceTM = FTransform(PerInstanceSMData[InstanceBodyIndex].Transform) * GetComponentTransform();
+
+			FBodyInstance* Instance = new FBodyInstance;
+
+			InitInstanceBody(InstanceBodyIndex, Instance, Body, true, Body->GetBodySetup());
+
+			InstanceBodies[InstanceBodyIndex] = Instance;
+
+			Body->TermBody();
+			delete Body;
+		}
+	}
+}
+
 void UInstancedStaticMeshComponent::UpdateInstanceBodyTransform(int32 InstanceIndex, const FTransform& WorldSpaceInstanceTransform, bool bTeleport)
 {
 	check(bPhysicsStateCreated);
diff --git a/Engine/Source/Runtime/Engine/Private/PhysicsEngine/BodyInstance.cpp b/Engine/Source/Runtime/Engine/Private/PhysicsEngine/BodyInstance.cpp
--- a/Engine/Source/Runtime/Engine/Private/PhysicsEngine/BodyInstance.cpp
+++ b/Engine/Source/Runtime/Engine/Private/PhysicsEngine/BodyInstance.cpp
@@ -761,7 +761,7 @@ void FBodyInstance::SetCollisionEnabled(ECollisionEnabled::Type NewType, bool bU
 		{
 			if(UPrimitiveComponent* PrimComponent = OwnerComponent.Get())
 			{
-				PrimComponent->RecreatePhysicsState();
+				PrimComponent->RecreateInstanceBody(InstanceBodyIndex);
 			}
 
 		}
