[Audio] make sure ActiveSound always has correct information about its Owner, when present.
- Adding the FObjectKey to ActiveSound lets us remove the "set sound owner" notification delegate from AudioDevice as it is no longer needed.

diff --git a/Engine/Source/Runtime/Engine/Private/ActiveSound.cpp b/Engine/Source/Runtime/Engine/Private/ActiveSound.cpp
--- a/Engine/Source/Runtime/Engine/Private/ActiveSound.cpp
+++ b/Engine/Source/Runtime/Engine/Private/ActiveSound.cpp
@@ -512,11 +512,13 @@ void FActiveSound::SetOwner(const AActor* Actor)
 	{
 		OwnerID = Actor->GetUniqueID();
 		OwnerName = Actor->GetFName();
+		OwnerKey = FObjectKey{Actor};
 	}
 	else
 	{
 		OwnerID = 0;
 		OwnerName = NAME_None;
+		OwnerKey = FObjectKey{};
 	}
 }
 
diff --git a/Engine/Source/Runtime/Engine/Private/AudioDevice.cpp b/Engine/Source/Runtime/Engine/Private/AudioDevice.cpp
--- a/Engine/Source/Runtime/Engine/Private/AudioDevice.cpp
+++ b/Engine/Source/Runtime/Engine/Private/AudioDevice.cpp
@@ -133,14 +133,6 @@ FAutoConsoleVariableRef CVarDisableLegacyReverb(
 	TEXT("0: Enabled, 1: Disabled"),
 	ECVF_Default);
 
-static int32 DisableOwnerNotify = 0;
-FAutoConsoleVariableRef CVarDisableOwnerNotify(
-	TEXT("au.DisableOwnerNotify"),
-	DisableOwnerNotify,
-	TEXT("Disables SetOwner notifications.\n")
-	TEXT("0: Enabled, 1: Disabled"),
-	ECVF_Default);
-
 static float SoundDistanceOptimizationLengthCVar = 1.0f;
 FAutoConsoleVariableRef CVarSoundDistanceOptimizationLength(
 	TEXT("au.SoundDistanceOptimizationLength"),
@@ -6823,10 +6815,6 @@ void FAudioDevice::PlaySoundAtLocation(USoundBase* Sound, UWorld* World, float V
 		NewActiveSound.Priority = Sound->Priority;
 
 		NewActiveSound.SetOwner(OwningActor);
-		if (!DisableOwnerNotify)
-		{
-			OnAudioDeviceSetSoundOwner.Broadcast(NewActiveSound, OwningActor);
-		}
 
 #if UE_AUDIO_PROFILERTRACE_ENABLED
 		const FString* EventLogID = &Audio::Trace::EventLog::ID::PlayRequestSoundAtLocation;
diff --git a/Engine/Source/Runtime/Engine/Private/Components/AudioComponent.cpp b/Engine/Source/Runtime/Engine/Private/Components/AudioComponent.cpp
--- a/Engine/Source/Runtime/Engine/Private/Components/AudioComponent.cpp
+++ b/Engine/Source/Runtime/Engine/Private/Components/AudioComponent.cpp
@@ -870,6 +870,8 @@ void UAudioComponent::PlayInternal(const PlayInternalRequestData& InPlayRequestD
 		TArray<FAudioParameter> ActorParams;
 		UActorSoundParameterInterface::Fill(Owner, ActorParams);
 		FAudioParameter::Merge(MoveTemp(ActorParams), SoundParams);
+
+		NewActiveSound.SetOwner(Owner);
 	}
 
 	TArray<FAudioParameter> InstanceParamsCopy = InstanceParameters;
diff --git a/Engine/Source/Runtime/Engine/Public/ActiveSound.h b/Engine/Source/Runtime/Engine/Public/ActiveSound.h
--- a/Engine/Source/Runtime/Engine/Public/ActiveSound.h
+++ b/Engine/Source/Runtime/Engine/Public/ActiveSound.h
@@ -304,6 +304,8 @@ private:
 	FName AudioComponentName;
 	FName OwnerName;
 
+	FObjectKey OwnerKey;
+
 	uint32 PlayOrder;
 
 public:
@@ -776,6 +778,9 @@ public:
 	/** Returns the unique ID of the active sound's owner if it exists. Returns 0 if the sound doesn't have an owner. */
 	FSoundOwnerObjectID GetOwnerID() const { return OwnerID; }
 
+	/** Returns the unique Key of the active sound's owner if it exists. Returns default FObjectKey() if the sound doesn't have an owner. */
+	FObjectKey GetOwnerKey() const { return OwnerKey; }
+
 	/** Gets the sound concurrency handles applicable to this sound instance*/
 	ENGINE_API void GetConcurrencyHandles(TArray<FConcurrencyHandle>& OutConcurrencyHandles) const;
 
diff --git a/Engine/Source/Runtime/Engine/Public/AudioDevice.h b/Engine/Source/Runtime/Engine/Public/AudioDevice.h
--- a/Engine/Source/Runtime/Engine/Public/AudioDevice.h
+++ b/Engine/Source/Runtime/Engine/Public/AudioDevice.h
@@ -417,8 +417,6 @@ struct FAudioDeviceRenderInfo
 DECLARE_MULTICAST_DELEGATE_OneParam(FOnAudioDevicePreRender, const FAudioDeviceRenderInfo&);
 DECLARE_MULTICAST_DELEGATE_OneParam(FOnAudioDevicePostRender, const FAudioDeviceRenderInfo&);
 
-DECLARE_MULTICAST_DELEGATE_TwoParams(FOnAudioDeviceSetSoundOwner, const FActiveSound&, const UObject*);
-
 class FAudioDevice : public FExec
 {
 public:
@@ -906,14 +904,11 @@ protected:
 		NotifySubsystemsActiveSoundDeleting(ActiveSound);
 	}
 
-	UE_DEPRECATED(5.7, "NotifySubsystemsSetOwner is deprecated. Use OnAudioDeviceSetSoundOwner instead.")
+	UE_DEPRECATED(5.7, "NotifySubsystemsSetOwner is deprecated. Use ActiveSound.GetOwnerKey instead.")
 	ENGINE_API void NotifySubsystemsSetOwner(FActiveSound& ActiveSound, const UObject* Owner) const;
 
 public:
 
-	// Broadcast on the Game Thread when an owning object has been associated with a sound.
-	FOnAudioDeviceSetSoundOwner OnAudioDeviceSetSoundOwner;
-
 	/**
 	 * Gets the default reverb and interior settings for the provided world.  Returns true if settings with WorldID were located
 	 */
