Handle edge case where faiiling to save a copy of a level results in a crash when trying to remove loaded actors.

diff --git a/Engine/Source/Runtime/Engine/Private/WorldPartition/WorldPartitionHandle.cpp b/Engine/Source/Runtime/Engine/Private/WorldPartition/WorldPartitionHandle.cpp
--- a/Engine/Source/Runtime/Engine/Private/WorldPartition/WorldPartitionHandle.cpp
+++ b/Engine/Source/Runtime/Engine/Private/WorldPartition/WorldPartitionHandle.cpp
@@ -123,7 +123,7 @@ void FWorldPartitionLoadingContext::FImmediate::UnregisterActor(FWorldPartitionA
 	TGuardValueAccessors<bool> IsEditorLoadingPackageGuard(UE::GetIsEditorLoadingPackage, UE::SetIsEditorLoadingPackage, true);
 
 	// When cleaning up worlds, actors are already marked as garbage at this point, so no need to remove them from the world
-	if (AActor* Actor = InActorDescInstance->GetActor(); IsValid(Actor))
+	if (AActor* Actor = InActorDescInstance->GetActor(); IsValid(Actor) && Actor->GetLevel())
 	{
 		UActorDescContainerInstance* Container = InActorDescInstance->GetContainerInstance();
 		check(Container);
