PR #13098: UActorChannel::ReplicateActor does not account written bits in one of the paths

diff --git a/Engine/Source/Runtime/Engine/Private/DataChannel.cpp b/Engine/Source/Runtime/Engine/Private/DataChannel.cpp
--- a/Engine/Source/Runtime/Engine/Private/DataChannel.cpp
+++ b/Engine/Source/Runtime/Engine/Private/DataChannel.cpp
@@ -3870,15 +3870,18 @@ int64 UActorChannel::ReplicateActor()
 
 		if (Connection->ResendAllDataState != EResendAllDataState::None)
 		{
+			int64 NumBitsWrote = 0;
 			if (bWroteSomethingImportant)
 			{
 				SendBunch(&Bunch, 1);
+				NumBitsWrote = Bunch.GetNumBits();
 			}
 
 			MemMark.Pop();
 			NETWORK_PROFILER(GNetworkProfiler.TrackReplicateActor(Actor, RepFlags, FPlatformTime::Cycles() - ActorReplicateStartTime, Connection));
-
-			return bWroteSomethingImportant;
+			Connection->GetDriver()->GetMetrics()->IncrementInt(UE::Net::Metric::NumReplicatedActorBytes, (NumBitsWrote + 7) >> 3);
+			
+			return NumBitsWrote;
 		}
 
 		{
