Fix attached spatially loaded actors not being GCed when AttachParent exists

diff --git a/Engine/Source/Runtime/Engine/Private/Level.cpp b/Engine/Source/Runtime/Engine/Private/Level.cpp
--- a/Engine/Source/Runtime/Engine/Private/Level.cpp
+++ b/Engine/Source/Runtime/Engine/Private/Level.cpp
@@ -1199,6 +1199,7 @@ void ULevel::RemoveLoadedActors(const TArray<AActor*>& ActorList, const FTransfo
 		// UnregisterAllComponents can destroy child actors
 		if (IsValid(Actor))
 		{
+			Actor->DetachFromActor(FDetachmentTransformRules(EDetachmentRule::KeepWorld, /*bCallModify=*/false));
 			Actor->UnregisterAllComponents();
 			Actor->RegisterAllActorTickFunctions(false, true);
 			Actor->OnLoadedActorRemovedFromLevel();
