[MIR] Fix TextureProperty by always using preshader path and only register uniform parameters for textures that are used for sampling.

Always use the preshader path to determine texture properties as even a fixed texture object can change its render resource depending on LODs.
Also avoid registering uniform parameters for textures that are never used for sampling. If a virtual texture is only used for texture properties, the runtime crashes when it expects a VT stack to be assigned to such a texture.
Add NatVis entry for MIR::FPreshaderParamete.

#rnx

diff --git a/Engine/Source/Runtime/Engine/Private/Materials/MaterialExpressionsIR.cpp b/Engine/Source/Runtime/Engine/Private/Materials/MaterialExpressionsIR.cpp
--- a/Engine/Source/Runtime/Engine/Private/Materials/MaterialExpressionsIR.cpp
+++ b/Engine/Source/Runtime/Engine/Private/Materials/MaterialExpressionsIR.cpp
@@ -1462,57 +1462,14 @@ void UMaterialExpressionTextureProperty::Build(MIR::FEmitter& Em)
 	MIR::FValueRef TextureValue = Em.Input(&TextureObject);
 	UE_MIR_CHECKPOINT(Em);
 
-	UTexture* InputTexture = Em.GetTextureFromValue(TextureValue);
-	if (!InputTexture)
-	{
-		Em.Error(TEXT("Missing texture object"));
-		return;
-	}
-
 	const bool bTexelSizeInUVSpace = Property == TMTM_TexelSize;
 
-	if (MIR::FUniformParameter* UniformParameter = TextureValue->As<MIR::FUniformParameter>())
-	{
-		// If the property is connected to a texture paramter, the property needs to be evaluated at runtime via preshaders
-		const UE::Shader::EPreshaderOpcode PreshaderOpcode =
-			bTexelSizeInUVSpace
-				? UE::Shader::EPreshaderOpcode::TexelSize
-				: UE::Shader::EPreshaderOpcode::TextureSize;
+	const UE::Shader::EPreshaderOpcode PreshaderOpcode =
+		bTexelSizeInUVSpace
+			? UE::Shader::EPreshaderOpcode::TexelSize
+			: UE::Shader::EPreshaderOpcode::TextureSize;
 
-		Em.Output(0, Em.PreshaderParameter(PreshaderOpcode, UniformParameter));
-	}
-	else
-	{
-		// If the property is connected to a texture object directly, evaluate the property directly and emit a constant expression
-		const EMaterialValueType TextureType = InputTexture->GetMaterialType();
-		const EMaterialValueType PropertyType = UE::MaterialTranslatorUtils::GetTexturePropertyValueType(TextureType);
-
-		FTextureResource* TextureResource = InputTexture->GetResource();
-		if (!TextureResource)
-		{
-			Em.Error(TEXT("Missing texture resource to encode texture property expression"));
-			return;
-		}
-
-		const FVector3f TextureSize
-		{
-			static_cast<float>(TextureResource->GetSizeX()),
-			static_cast<float>(TextureResource->GetSizeY()),
-			static_cast<float>(TextureResource->GetSizeZ())
-		};
-		if (bTexelSizeInUVSpace)
-		{
-			Em.Output(0, PropertyType == MCT_Float3
-				? Em.ConstantFloat3(FVector3f{ 1.0f / TextureSize.X, 1.0f / TextureSize.Y, 1.0f / TextureSize.Z })
-				: Em.ConstantFloat2(FVector2f{ 1.0f / TextureSize.X, 1.0f / TextureSize.Y }));
-		}
-		else
-		{
-			Em.Output(0, PropertyType == MCT_Float3
-				? Em.ConstantFloat3(TextureSize)
-				: Em.ConstantFloat2(FVector2f{ TextureSize.X, TextureSize.Y }));
-		}
-	}
+	Em.Output(0, Em.PreshaderParameter(PreshaderOpcode, TextureValue));
 }
 
 #if WITH_EDITOR
diff --git a/Engine/Source/Runtime/Engine/Private/Materials/MaterialIRDebug.cpp b/Engine/Source/Runtime/Engine/Private/Materials/MaterialIRDebug.cpp
--- a/Engine/Source/Runtime/Engine/Private/Materials/MaterialIRDebug.cpp
+++ b/Engine/Source/Runtime/Engine/Private/Materials/MaterialIRDebug.cpp
@@ -330,6 +330,11 @@ struct FDebugDumpIRState
 			Out.Appendf(TEXT(" SamplerSourceMode=\"%s\""), *StaticEnum<ESamplerSourceMode>()->GetDisplayNameTextByValue(TextureRead->SamplerSourceMode).ToString());
 			Out.Appendf(TEXT(" SamplerType=\"%s\""), *StaticEnum<EMaterialSamplerType>()->GetDisplayNameTextByValue(TextureRead->SamplerType).ToString());
 		}
+		else if (const FPreshaderParameter* PreshaderParameter = Instr->As<FPreshaderParameter>())
+		{
+			Out.Appendf(TEXT(" TextureIndex=%d"), PreshaderParameter->TextureIndex);
+			Out.Appendf(TEXT(" PreshaderOffset=%u"), PreshaderParameter->Analysis_PreshaderOffset);
+		}
 		else if (const FInlineHLSL* InlineHLSL = Instr->As<FInlineHLSL>())
 		{
 			if (InlineHLSL->HasFlags(EValueFlags::HasDynamicHLSLCode))
diff --git a/Engine/Source/Runtime/Engine/Private/Materials/MaterialIREmitter.cpp b/Engine/Source/Runtime/Engine/Private/Materials/MaterialIREmitter.cpp
--- a/Engine/Source/Runtime/Engine/Private/Materials/MaterialIREmitter.cpp
+++ b/Engine/Source/Runtime/Engine/Private/Materials/MaterialIREmitter.cpp
@@ -974,6 +974,7 @@ FValueRef FEmitter::TextureObject(UTexture* Texture, EMaterialSamplerType Sample
 	FTextureObject Proto = MakePrototype<FTextureObject>(FType::MakeTexture());
 	Proto.Texture = Texture;
 	Proto.SamplerType = SamplerType;
+	Proto.Analysis_UniformParameterIndex = INDEX_NONE;
 	return EmitPrototype(*this, Proto);
 }
 
@@ -995,6 +996,15 @@ FValueRef FEmitter::Parameter(FName Name, FMaterialParameterMetadata& Metadata,
 		return Id;
 	};
 
+	auto MakeUniformParameter = [this, &RegisterParameter](FType Type, FMaterialParameterInfo InInfo, const FMaterialParameterMetadata& InMetadata, EMaterialSamplerType InSamplerType = SAMPLERTYPE_Color) -> FValueRef
+	{
+		FUniformParameter Prototype = MakePrototype<FUniformParameter>(MoveTemp(Type));
+		Prototype.ParameterIdInModule = RegisterParameter(Module, MoveTemp(InInfo), InMetadata);
+		Prototype.SamplerType = InSamplerType;
+		Prototype.Analysis_UniformParameterIndex = INDEX_NONE;
+		return EmitPrototype(*this, Prototype);
+	};
+
 	FMaterialParameterInfo Info{ Name };
 
 	switch (Metadata.Value.Type)
@@ -1005,10 +1015,7 @@ FValueRef FEmitter::Parameter(FName Name, FMaterialParameterMetadata& Metadata,
 			{
 				return CustomPrimitiveData(Metadata.PrimitiveDataIndex);
 			}
-
-			FUniformParameter Prototype = MakePrototype<FUniformParameter>(FType::MakeFloatScalar());
-			Prototype.ParameterIdInModule = RegisterParameter(Module, Info, Metadata);
-			return EmitPrototype(*this, Prototype);
+			return MakeUniformParameter(FType::MakeFloatScalar(), Info, Metadata);
 		}
 
 		case EMaterialParameterType::Vector:
@@ -1023,19 +1030,14 @@ FValueRef FEmitter::Parameter(FName Name, FMaterialParameterMetadata& Metadata,
 				return Vector4(X, Y, Z, W);
 			}
 
-			EScalarKind ScalarKind = Metadata.Value.Type == EMaterialParameterType::Vector ? EScalarKind::Float : EScalarKind::Double;
-			FUniformParameter Prototype = MakePrototype<FUniformParameter>(FType::MakeVector(ScalarKind, 4));
-			Prototype.ParameterIdInModule = RegisterParameter(Module, Info, Metadata);
-			return EmitPrototype(*this, Prototype);
+			const EScalarKind ScalarKind = Metadata.Value.Type == EMaterialParameterType::Vector ? EScalarKind::Float : EScalarKind::Double;
+			return MakeUniformParameter(FType::MakeVector(ScalarKind, 4), Info, Metadata);
 		}
 
 		case EMaterialParameterType::Texture:
 		case EMaterialParameterType::Font:
 		{
-			FUniformParameter Prototype = MakePrototype<FUniformParameter>(FType::MakeTexture());
-			Prototype.ParameterIdInModule = RegisterParameter(Module, Info, Metadata);
-			Prototype.SamplerType = SamplerType;
-			return EmitPrototype(*this, Prototype);
+			return MakeUniformParameter(FType::MakeTexture(), Info, Metadata, SamplerType);
 		}
 
 		case EMaterialParameterType::StaticSwitch:
@@ -3824,6 +3826,7 @@ FValueRef FEmitter::AnalyticalPartialDerivative(FValueRef Value, EDerivativeAxis
 		// Custom HLSL calls are not derivative aware, so we also return constant zero for VK_Call.
 		case VK_Constant:
 		case VK_UniformParameter:
+		case VK_PreshaderParameter:
 		case VK_TextureRead:
 		case VK_Call:
 			return Cast(ConstantZero(Value->Type.GetPrimitive().ScalarKind), Value->Type);
diff --git a/Engine/Source/Runtime/Engine/Private/Materials/MaterialIRToHLSLTranslator.cpp b/Engine/Source/Runtime/Engine/Private/Materials/MaterialIRToHLSLTranslator.cpp
--- a/Engine/Source/Runtime/Engine/Private/Materials/MaterialIRToHLSLTranslator.cpp
+++ b/Engine/Source/Runtime/Engine/Private/Materials/MaterialIRToHLSLTranslator.cpp
@@ -876,10 +876,10 @@ struct FPrivate : FMaterialIRToHLSLTranslation
 		{
 			Printer << TEXTVIEW(".");
 
-			static const TCHAR* Components = TEXT("xyzw");
+			const TCHAR* Components = TEXT("xyzw");
 			for (int32 i = 0; i < NumComponents; ++i)
 			{
-				check(BufferSlotOffset + i < sizeof(Components));
+				check(BufferSlotOffset + i < 4);
 				Printer.Buffer.AppendChar(Components[BufferSlotOffset + i]);
 			}
 		}
@@ -1859,6 +1859,7 @@ struct FPrivate : FMaterialIRToHLSLTranslation
 
 	ENoOp LowerTextureReference(EMaterialValueType TextureType, int32 TextureParameterIndex)
 	{
+		checkf(TextureParameterIndex != INDEX_NONE, TEXT("Texture uniform parameter not assigned! A texture is used for sampling but hasn't been properly registered during IR analysis."));
 		Printer << TEXTVIEW("Material.") << LowerTextureMaterialType(TextureType, true) << TEXTVIEW("_") << TextureParameterIndex;
 		return NoOp;
 	}
diff --git a/Engine/Source/Runtime/Engine/Private/Materials/MaterialIRValueAnalyzer.cpp b/Engine/Source/Runtime/Engine/Private/Materials/MaterialIRValueAnalyzer.cpp
--- a/Engine/Source/Runtime/Engine/Private/Materials/MaterialIRValueAnalyzer.cpp
+++ b/Engine/Source/Runtime/Engine/Private/Materials/MaterialIRValueAnalyzer.cpp
@@ -154,10 +154,65 @@ static int32 AcquireVTStackIndex(FMaterialIRValueAnalyzer& Analyzer, MIR::FVTPag
 	return Analyzer.CompilationOutput->UniformExpressionSet.AddVTStack(INDEX_NONE);
 }
 
+static int32 RegisterTextureParameter(FMaterialIRValueAnalyzer& Analyzer, UTexture* Texture, const FHashedMaterialParameterInfo& ParameterInfo, uint8 VTLayerIndex = 0xFF)
+{
+	check(Texture);
+
+	FMaterialTextureParameterInfo TextureParameterInfo
+	{
+		.ParameterInfo = ParameterInfo,
+		.TextureIndex = Analyzer.Material->GetReferencedTextures().Find(Texture),
+		.SamplerSource = SSM_FromTextureAsset,
+		.VirtualTextureLayerIndex = VTLayerIndex,
+	};
+
+	check(TextureParameterInfo.TextureIndex != INDEX_NONE);
+
+	EMaterialTextureParameterType ParamType = MIR::Internal::TextureMaterialValueTypeToParameterType(Texture->GetMaterialType());
+
+	return Analyzer.CompilationOutput->UniformExpressionSet.FindOrAddTextureParameter(ParamType, TextureParameterInfo);
+}
+
+static void AccessTextureObject(FMaterialIRValueAnalyzer& Analyzer, MIR::FTextureObject* TextureObject)
+{
+	if (TextureObject->Analysis_UniformParameterIndex == INDEX_NONE)
+	{
+		const FHashedMaterialParameterInfo ParameterInfo{ TextureObject->Texture->GetFName(), EMaterialParameterAssociation::GlobalParameter, INDEX_NONE };
+		TextureObject->Analysis_UniformParameterIndex = RegisterTextureParameter(Analyzer, TextureObject->Texture, ParameterInfo, 0);
+	}
+}
+
+static void AccessTextureUniformParameter(FMaterialIRValueAnalyzer& Analyzer, MIR::FUniformParameter* UniformParameter)
+{
+	if (UniformParameter->Analysis_UniformParameterIndex == INDEX_NONE)
+	{
+		const FMaterialParameterValue& ParameterValue = Analyzer.Module->GetParameterMetadata(UniformParameter->ParameterIdInModule).Value;
+		UTexture* Texture = Cast<UTexture>(ParameterValue.AsTextureObject());
+		UniformParameter->Analysis_UniformParameterIndex = RegisterTextureParameter(Analyzer, Texture, Analyzer.Module->GetParameterInfo(UniformParameter->ParameterIdInModule));
+	}
+}
+
+static void AccessTexture(FMaterialIRValueAnalyzer& Analyzer, MIR::FValue* TextureValue)
+{
+	if (MIR::FTextureObject* TextureObject = TextureValue->As<MIR::FTextureObject>())
+	{
+		AccessTextureObject(Analyzer, TextureObject);
+	}
+	else if (MIR::FUniformParameter* UniformParameter = TextureValue->As<MIR::FUniformParameter>())
+	{
+		AccessTextureUniformParameter(Analyzer, UniformParameter);
+	}
+	else
+	{
+		UE_MIR_UNREACHABLE();
+	}
+}
+
 static void AnalyzeInStageVTPageTableRead(FMaterialIRValueAnalyzer& Analyzer, MIR::FVTPageTableRead* VTPageTableRead, MIR::EStage Stage)
 {
 	check(VTPageTableRead->TextureObject);
 
+	AccessTexture(Analyzer, VTPageTableRead->TextureObject);
 	const int32 TextureUniformIndex = VTPageTableRead->TextureObject->GetUniformParameterIndex();
 	check(TextureUniformIndex >= 0);
 
@@ -478,40 +533,11 @@ static void AnalyzeShadingModel(FMaterialIRValueAnalyzer& Analyzer, MIR::FShadin
 	}
 }
 
-static void AnalyzeTextureObject(FMaterialIRValueAnalyzer& Analyzer, MIR::FTextureObject* TextureObject)
+static void AnalyzeTextureRead(FMaterialIRValueAnalyzer& Analyzer, MIR::FTextureRead* TextureRead)
 {
-	EMaterialTextureParameterType ParamType = MIR::Internal::TextureMaterialValueTypeToParameterType(TextureObject->Texture->GetMaterialType());
-
-	FMaterialTextureParameterInfo ParamInfo
-	{
-		.ParameterInfo = { "", EMaterialParameterAssociation::GlobalParameter, INDEX_NONE },
-		.TextureIndex = Analyzer.Material->GetReferencedTextures().Find(TextureObject->Texture),
-		.SamplerSource = SSM_FromTextureAsset, // TODO - Is this needed?
-	};
-
-	check(ParamInfo.TextureIndex != INDEX_NONE);
-
-	TextureObject->Analysis_UniformParameterIndex = Analyzer.CompilationOutput->UniformExpressionSet.FindOrAddTextureParameter(ParamType, ParamInfo);
-}
-
-static void AnalyzeTextureUniformParameter(FMaterialIRValueAnalyzer& Analyzer, MIR::FUniformParameter* Parameter)
-{
-	const FMaterialParameterValue& ParameterValue = Analyzer.Module->GetParameterMetadata(Parameter->ParameterIdInModule).Value;
-
-	UTexture* Texture = Cast<UTexture>(ParameterValue.AsTextureObject());
-	check(Texture);
-
-	FMaterialTextureParameterInfo ParamInfo
-	{
-		.ParameterInfo = Analyzer.Module->GetParameterInfo(Parameter->ParameterIdInModule),
-		.TextureIndex = Analyzer.Material->GetReferencedTextures().Find(Texture),
-		.SamplerSource = SSM_FromTextureAsset, // TODO - Is this needed?
-		.VirtualTextureLayerIndex = 0xff,
-	};
-
-	EMaterialTextureParameterType ParamType = MIR::Internal::TextureMaterialValueTypeToParameterType(Texture->GetMaterialType());
-	
-	Parameter->Analysis_UniformParameterIndex = Analyzer.CompilationOutput->UniformExpressionSet.FindOrAddTextureParameter(ParamType, ParamInfo);
+	// Ensure a uniform parameter is allocated when a texture read instruction accesses the texture.
+	// Otherwise, no uniform parameter must be allocated for the texture object in case it's used for other nodes only such as texture properties.
+	AccessTexture(Analyzer, TextureRead->TextureObject);
 }
 
 // Returns the next available offset into the preshader buffer for a float vector with the specified number of components (1-4).
@@ -638,11 +664,7 @@ static void AnalyzePrimitiveUniformParameter(FMaterialIRValueAnalyzer& Analyzer,
 
 static void AnalyzeUniformParameter(FMaterialIRValueAnalyzer& Analyzer, MIR::FUniformParameter* Parameter)
 {
-	if (Parameter->Type.IsTexture())
-	{
-		AnalyzeTextureUniformParameter(Analyzer, Parameter);
-	}
-	else
+	if (!Parameter->Type.IsTexture())
 	{
 		check(Parameter->Type.AsPrimitive());
 		AnalyzePrimitiveUniformParameter(Analyzer, Parameter);
@@ -653,8 +675,20 @@ static void AnalyzePreshaderParameter(FMaterialIRValueAnalyzer& Analyzer, MIR::F
 {
 	FUniformExpressionSet& UniformExpressionSet = Analyzer.CompilationOutput->UniformExpressionSet;
 
-	MIR::FUniformParameter* UniformSourceParameter = Parameter->SourceParameter->As<MIR::FUniformParameter>();
-	const FMaterialParameterInfo& SourceParameterInfo = Analyzer.Module->GetParameterInfo(UniformSourceParameter->ParameterIdInModule);
+	// Get parameter name from source parameter
+	FName SourceParameterName;
+	if (MIR::FUniformParameter* UniformSourceParameter = Parameter->SourceParameter->As<MIR::FUniformParameter>())
+	{
+		SourceParameterName = Analyzer.Module->GetParameterInfo(UniformSourceParameter->ParameterIdInModule).Name;
+	}
+	else if (MIR::FTextureObject* TextureObject = Parameter->SourceParameter->As<MIR::FTextureObject>())
+	{
+		SourceParameterName = TextureObject->Texture->GetFName();
+	}
+	else
+	{
+		SourceParameterName = NAME_None;
+	}
 
 	// Make sure the parameter type is primitive
 	MIR::FPrimitive PrimitiveType = Parameter->Type.GetPrimitive();
@@ -664,13 +698,12 @@ static void AnalyzePreshaderParameter(FMaterialIRValueAnalyzer& Analyzer, MIR::F
 
 	const uint32 NumComponents = PrimitiveType.NumComponents();
 	const uint32 GlobalComponentOffset = NextGlobalComponentOffset(Analyzer, NumComponents);
-	const uint32 UniformPreshaderOffset = GlobalComponentOffset / 4;
 
 	UniformExpressionSet.WriteUniformPreshaderEntry(
-		UniformPreshaderOffset, UE::Shader::MakeValueType(UE::Shader::EValueComponentType::Float, NumComponents),
-		[Parameter, &SourceParameterInfo](UE::Shader::FPreshaderData& UniformPreshaderData) -> void
+		GlobalComponentOffset, UE::Shader::MakeValueType(UE::Shader::EValueComponentType::Float, NumComponents),
+		[Parameter, SourceParameterName](UE::Shader::FPreshaderData& UniformPreshaderData) -> void
 		{
-			const FHashedMaterialParameterInfo HashedSourceParameterInfo{ SourceParameterInfo.Name, EMaterialParameterAssociation::GlobalParameter, INDEX_NONE };
+			const FHashedMaterialParameterInfo HashedSourceParameterInfo{ SourceParameterName, EMaterialParameterAssociation::GlobalParameter, INDEX_NONE };
 
 			switch (Parameter->Opcode)
 			{
@@ -751,8 +784,8 @@ void FMaterialIRValueAnalyzer::Analyze(MIR::FValue* Value)
 	{
 		EXPAND_CASE(ExternalInput);
 		EXPAND_CASE(MaterialParameterCollection);
-		EXPAND_CASE(TextureObject);
 		EXPAND_CASE(UniformParameter);
+		EXPAND_CASE(TextureRead);
 		EXPAND_CASE(PreshaderParameter);
 		EXPAND_CASE(SetMaterialOutput);
 		EXPAND_CASE(InlineHLSL);
diff --git a/Engine/Source/Runtime/Engine/Public/Materials/MaterialIR.h b/Engine/Source/Runtime/Engine/Public/Materials/MaterialIR.h
--- a/Engine/Source/Runtime/Engine/Public/Materials/MaterialIR.h
+++ b/Engine/Source/Runtime/Engine/Public/Materials/MaterialIR.h
@@ -494,7 +494,7 @@ struct FTextureObject : TValue<VK_TextureObject>
 
 	// Index of this parameter in the uniform expression set.
 	// Note: This field is automatically set by the builder.
-	uint32 Analysis_UniformParameterIndex;
+	int32 Analysis_UniformParameterIndex;
 };
 
 // It represents a material uniform parameter.
@@ -508,7 +508,7 @@ struct FUniformParameter: TValue<VK_UniformParameter>
 
 	// Index of this parameter in the uniform expression set
 	// Note: This field is automatically set by the builder.
-	uint32 Analysis_UniformParameterIndex;
+	int32 Analysis_UniformParameterIndex;
 };
 
 enum class EScreenTexture : uint8
@@ -1112,7 +1112,7 @@ struct FPreshaderParameter : TInstruction<VK_PreshaderParameter, 1>
 	int32 TextureIndex;
 
 	// Offset into the preshader buffer of the material resource.
-	uint32_t Analysis_PreshaderOffset;
+	uint32 Analysis_PreshaderOffset;
 };
 
 } // namespace MIR
