Categorized bloom settings under the method to which they apply. Added indicator to settings panel to display which bloom method is currently being previewed. Added separate intensity modifier values for gaussian and fft bloom.
This CL adds FRenderViewportFeedback, which communicates state from the renderer back to the viewport (and settings panel). For example, which bloom method was finally actually used during rendering.

diff --git a/Engine/Plugins/Cameras/GameplayCameras/Source/GameplayCameras/Private/Compat/PostProcessUtils.cpp b/Engine/Plugins/Cameras/GameplayCameras/Source/GameplayCameras/Private/Compat/PostProcessUtils.cpp
--- a/Engine/Plugins/Cameras/GameplayCameras/Source/GameplayCameras/Private/Compat/PostProcessUtils.cpp
+++ b/Engine/Plugins/Cameras/GameplayCameras/Source/GameplayCameras/Private/Compat/PostProcessUtils.cpp
@@ -65,6 +65,7 @@ bool FPostProcessUtils::OverridePostProcessSettings(FPostProcessSettings& ThisFr
 		UE_SET_PP(SceneFringeIntensity);
 		UE_SET_PP(ChromaticAberrationStartOffset);
 		UE_SET_PP(BloomIntensity);
+		UE_SET_PP(BloomGaussianIntensity);
 		UE_SET_PP(BloomThreshold);
 		UE_SET_PP(Bloom1Tint);
 		UE_SET_PP(BloomSizeScale);
@@ -81,6 +82,7 @@ bool FPostProcessUtils::OverridePostProcessSettings(FPostProcessSettings& ThisFr
 		UE_SET_PP(Bloom6Size);
 		UE_SET_PP(BloomDirtMaskIntensity);
 		UE_SET_PP(BloomDirtMaskTint);
+		UE_SET_PP(BloomConvolutionIntensity);
 		UE_SET_PP(BloomConvolutionScatterDispersion);
 		UE_SET_PP(BloomConvolutionSize);
 		UE_SET_PP(BloomConvolutionCenterUV);
@@ -347,6 +349,7 @@ bool FPostProcessUtils::BlendPostProcessSettings(FPostProcessSettings& ThisFrom,
 		UE_LERP_PP(SceneFringeIntensity);
 		UE_LERP_PP(ChromaticAberrationStartOffset);
 		UE_LERP_PP(BloomIntensity);
+		UE_LERP_PP(BloomGaussianIntensity);
 		UE_LERP_PP(BloomThreshold);
 		UE_LERP_PP(Bloom1Tint);
 		UE_LERP_PP(BloomSizeScale);
@@ -363,6 +366,7 @@ bool FPostProcessUtils::BlendPostProcessSettings(FPostProcessSettings& ThisFrom,
 		UE_LERP_PP(Bloom6Size);
 		UE_LERP_PP(BloomDirtMaskIntensity);
 		UE_LERP_PP(BloomDirtMaskTint);
+		UE_LERP_PP(BloomConvolutionIntensity);
 		UE_LERP_PP(BloomConvolutionScatterDispersion);
 		UE_LERP_PP(BloomConvolutionSize);
 		UE_LERP_PP(BloomConvolutionCenterUV);
diff --git a/Engine/Plugins/Cameras/GameplayCameras/Source/GameplayCameras/Private/Debug/PostProcessSettingsDebugBlock.cpp b/Engine/Plugins/Cameras/GameplayCameras/Source/GameplayCameras/Private/Debug/PostProcessSettingsDebugBlock.cpp
--- a/Engine/Plugins/Cameras/GameplayCameras/Source/GameplayCameras/Private/Debug/PostProcessSettingsDebugBlock.cpp
+++ b/Engine/Plugins/Cameras/GameplayCameras/Source/GameplayCameras/Private/Debug/PostProcessSettingsDebugBlock.cpp
@@ -92,6 +92,7 @@ void FPostProcessSettingsDebugBlock::OnDebugDraw(const FCameraDebugBlockDrawPara
 		UE_DRAW_PP(SceneFringeIntensity);
 		UE_DRAW_PP(ChromaticAberrationStartOffset);
 		UE_DRAW_PP(BloomIntensity);
+		UE_DRAW_PP(BloomGaussianIntensity);
 		UE_DRAW_PP(BloomThreshold);
 		UE_DRAW_PP(Bloom1Tint);
 		UE_DRAW_PP(BloomSizeScale);
@@ -108,6 +109,7 @@ void FPostProcessSettingsDebugBlock::OnDebugDraw(const FCameraDebugBlockDrawPara
 		UE_DRAW_PP(Bloom6Size);
 		UE_DRAW_PP(BloomDirtMaskIntensity);
 		UE_DRAW_PP(BloomDirtMaskTint);
+		UE_DRAW_PP(BloomConvolutionIntensity);
 		UE_DRAW_PP(BloomConvolutionScatterDispersion);
 		UE_DRAW_PP(BloomConvolutionSize);
 		UE_DRAW_PP(BloomConvolutionCenterUV);
diff --git a/Engine/Source/Editor/DetailCustomizations/Private/PostProcessSettingsCustomization.cpp b/Engine/Source/Editor/DetailCustomizations/Private/PostProcessSettingsCustomization.cpp
--- a/Engine/Source/Editor/DetailCustomizations/Private/PostProcessSettingsCustomization.cpp
+++ b/Engine/Source/Editor/DetailCustomizations/Private/PostProcessSettingsCustomization.cpp
@@ -11,6 +11,9 @@
 #include "DetailWidgetRow.h"
 #include "Editor.h"
 #include "Editor/EditorEngine.h"
+#include "LevelEditor.h"
+#include "SLevelViewport.h"
+#include "LevelEditorViewport.h"
 #include "Engine/BlendableInterface.h"
 #include "Factories/Factory.h"
 #include "Framework/Commands/UIAction.h"
@@ -57,6 +60,10 @@
 #include "Widgets/SBoxPanel.h"
 #include "Widgets/Notifications/SNotificationList.h"
 #include "Widgets/Text/STextBlock.h"
+#include "Widgets/SToolTip.h"
+#include "Widgets/Images/SImage.h"
+#include "SWarningOrErrorBox.h"
+#include "RenderViewportFeedback.h"
 
 class IDetailPropertyRow;
 class SWidget;
@@ -117,12 +124,27 @@ struct FCategoryOrGroup
 	}
 };
 
+struct FAlert
+{
+	enum class EType
+	{
+		Warning,
+		Error
+	};
+
+	EType Type;
+	FText Text;
+	TAttribute<EVisibility> Visibility;
+};
 
 struct FPostProcessGroup
 {
 	FString RawGroupName;
 	FString DisplayName;
+	FString ParentPath;
 	FCategoryOrGroup RootCategory;
+	FCategoryOrGroup Detail;
+	TOptional<FAlert> Alert;
 	TArray<TSharedPtr<IPropertyHandle>> SimplePropertyHandles;
 	TArray<TSharedPtr<IPropertyHandle>> AdvancedPropertyHandles;
 
@@ -133,6 +155,7 @@ struct FPostProcessGroup
 
 	FPostProcessGroup()
 		: RootCategory()
+		, Detail()
 	{}
 };
 
@@ -252,16 +275,270 @@ private:
 	}
 };
 
-void FPostProcessSettingsCustomization::CustomizeChildren( TSharedRef<IPropertyHandle> StructPropertyHandle, class IDetailChildrenBuilder& StructBuilder, IPropertyTypeCustomizationUtils& StructCustomizationUtils )
+// See FDetailGroup::MakeNameWidget()
+static TSharedRef<SWidget> MakeNameWidget(FText LocalizedDisplayName, TOptional<FText> LocalizedToolTip)
 {
-	uint32 NumChildren = 0;
-	FPropertyAccess::Result Result = StructPropertyHandle->GetNumChildren(NumChildren);
+	TSharedRef<STextBlock> TextBlock = 
+		SNew(STextBlock)
+		.Font(IDetailLayoutBuilder::GetDetailFont())
+		.Text(LocalizedDisplayName);
+
+	if (LocalizedToolTip.IsSet() && !LocalizedToolTip->IsEmpty())
+	{
+		TextBlock->SetToolTipText(LocalizedToolTip.GetValue());
+	}
+
+	return
+		SNew(SButton)
+		.ButtonStyle(FAppStyle::Get(), "NoBorder")
+		.ContentPadding(FMargin(0, 2, 0, 2))
+		//.OnClicked(this, &FDetailGroup::OnNameClicked)
+		.ForegroundColor(FSlateColor::UseForeground())
+		.Content()
+		[
+			TextBlock
+		];
+};
+
+static bool GetRenderFeedbackForActiveViewport(FRenderViewportFeedback& OutValue)
+{
+	static FLevelEditorModule* LevelEditorModule = nullptr;
+	if (!LevelEditorModule && FModuleManager::Get().IsModuleLoaded("LevelEditor"))
+	{
+		LevelEditorModule = &FModuleManager::LoadModuleChecked<FLevelEditorModule>("LevelEditor");
+	}
+
+	if (!LevelEditorModule)
+	{
+		return false;
+	}
+
+	const TSharedPtr<SLevelViewport> ActiveLevelViewport = LevelEditorModule->GetFirstActiveLevelViewport();
+	if (!ActiveLevelViewport)
+	{
+		return false;
+	}
+
+	FLevelEditorViewportClient& ViewportClient = ActiveLevelViewport->GetLevelViewportClient();
+	if (ViewportClient.RenderViewportFeedback)
+	{
+		OutValue = ViewportClient.RenderViewportFeedback->Data();
+		return true;
+	}
+	else
+	{
+		// Request data collection for next frame
+		ViewportClient.RenderViewportFeedback = MakeShared<UE::RenderViewportFeedback::FReceiver>();
+		return false;
+	}
+}
+
+static void CustomizeGroup(FPostProcessGroup& Owner, TSharedRef<IPropertyHandle>& StructPropertyHandle);
+
+static void MakeDetail(FPostProcessGroup& Elem, TMap<FString, FPostProcessGroup>& NameToGroupMap, TSharedRef<IPropertyHandle>& StructPropertyHandle)
+{
+	if (Elem.Detail.IsValid())
+	{
+		return;
+	}
+
+	FCategoryOrGroup* ParentDetail;
+	int _;
+	bool bIsRoot = !Elem.ParentPath.FindChar('|', _);
+	if (bIsRoot)
+	{
+		ParentDetail = &Elem.RootCategory;
+	}
+	else
+	{
+		ParentDetail = &NameToGroupMap.Find(Elem.ParentPath)->Detail;
+	}
+	Elem.Detail = ParentDetail->AddGroup(*Elem.RawGroupName, FText::FromString(Elem.DisplayName));
+	CustomizeGroup(Elem, StructPropertyHandle);
+}
+
+static FPostProcessGroup* FindOrCreateGroup(TMap<FString, FPostProcessGroup>& NameToGroupMap, FCategoryOrGroup* Root, FStringView GroupPath, TSharedRef<IPropertyHandle>& StructPropertyHandle)
+{
+	FPostProcessGroup* PPGroup = NameToGroupMap.Find(FString(GroupPath));
+
+	if (!PPGroup)
+	{
+		int32 LastSeperator;
+		bool bIsRoot = !GroupPath.FindLastChar('|', LastSeperator);
+		if (bIsRoot)
+		{
+			return nullptr;
+		}
+		else
+		{
+			PPGroup = &NameToGroupMap.Add(FString(GroupPath));
+		}
+	}
+
+	// Is this a new group?
+	if (!PPGroup->IsValid())
+	{
+		PPGroup->RootCategory = *Root;
+		PPGroup->RawGroupName = GroupPath;
+
+		int32 LastSeparatorIndex;
+		if (GroupPath.FindLastChar('|', LastSeparatorIndex))
+		{
+			PPGroup->ParentPath = GroupPath.Left(LastSeparatorIndex);
+			FPostProcessGroup* Parent = FindOrCreateGroup(NameToGroupMap, Root, PPGroup->ParentPath, StructPropertyHandle);
+			
+			// Parent is now certainly non-empty, so build its IDetailGroup that the child will go into
+			if (Parent)
+			{
+				MakeDetail(*Parent, NameToGroupMap, StructPropertyHandle);
+			}
+
+			PPGroup->DisplayName = GroupPath.RightChop(LastSeparatorIndex+1);
+		}
+		else
+		{
+			// Path is root
+			checkNoEntry();
+		}
+	}
+
+	return PPGroup;
+}
+
+static void CustomizeGroup(FPostProcessGroup& Group, TSharedRef<IPropertyHandle>& StructPropertyHandle)
+{
+	if (Group.RawGroupName == TEXT("Lens|Bloom|Gaussian"))
+	{
+		Group.Alert = FAlert {
+			.Type = FAlert::EType::Warning,
+			.Text = LOCTEXT("PostProcessBloomGaussianWarning", "Primary viewport currently shows convolution bloom"),
+			.Visibility = MakeAttributeLambda([=, CurVisibility = EVisibility::Collapsed]() mutable
+			{
+				// Show if Gaussian bloom is selected, but not currently used in the viewport.
+				// Visibility is updated only if we can access all the required data.
+				FRenderViewportFeedback RenderViewportFeedback;
+				if (GetRenderFeedbackForActiveViewport(RenderViewportFeedback))
+				{
+					FPostProcessSettings* PostProcessSettings;
+					if (StructPropertyHandle->GetValueData((void*&)PostProcessSettings) == FPropertyAccess::Success)
+					{
+						bool bBloomFFTSelected = PostProcessSettings->bOverride_BloomMethod && PostProcessSettings->BloomMethod == EBloomMethod::BM_FFT;
+						CurVisibility = 
+							bBloomFFTSelected && RenderViewportFeedback.BloomMethod == EBloomMethod::BM_SOG 
+							? EVisibility::Visible : EVisibility::Collapsed;
+					}
+				}
+
+				return CurVisibility;
+			})
+		};
+	}
+	else if (Group.RawGroupName == TEXT("Lens|Bloom|Convolution"))
+	{
+		Group.Alert = FAlert {
+			.Type = FAlert::EType::Warning,
+			.Text = LOCTEXT("PostProcessBloomConvolutionWarning", "Primary viewport currently shows Gaussian bloom"),
+			.Visibility = MakeAttributeLambda([=, CurVisibility = EVisibility::Collapsed]() mutable
+			{
+				// Show if convolution bloom is selected, but not currently used in the viewport.
+				// Visibility is updated only if we can access all the required data.
+				FRenderViewportFeedback RenderViewportFeedback;
+				if (GetRenderFeedbackForActiveViewport(RenderViewportFeedback))
+				{
+					FPostProcessSettings* PostProcessSettings;
+					if (StructPropertyHandle->GetValueData((void*&)PostProcessSettings) == FPropertyAccess::Success)
+					{
+						bool bBloomSOGSelected = PostProcessSettings->bOverride_BloomMethod && PostProcessSettings->BloomMethod == EBloomMethod::BM_SOG;
+						CurVisibility = 
+							bBloomSOGSelected && RenderViewportFeedback.BloomMethod == EBloomMethod::BM_FFT 
+							? EVisibility::Visible : EVisibility::Collapsed;
+					}
+				}
+
+				return CurVisibility;
+			})
+		};
+	}
+}
+
+// Return false to hide the property
+static bool CustomizeProperty(TSharedPtr<IPropertyHandle>& PropertyHandle, FProperty* Property, const FName& CategoryFName, TSharedRef<IPropertyHandle>& StructHandle, IPropertyTypeCustomizationUtils& StructCustomizationUtils)
+{
+	static const FName ExposureCategory("Lens|Exposure");
+	static const FName TranslucencyCategory("Rendering Features|Translucency");
+
+	if (CategoryFName == ExposureCategory)
+	{
+		static const auto VarDefaultAutoExposureExtendDefaultLuminanceRange = IConsoleManager::Get().FindTConsoleVariableDataInt(TEXT("r.DefaultFeature.AutoExposure.ExtendDefaultLuminanceRange"));
+		const bool bExtendedLuminanceRange = VarDefaultAutoExposureExtendDefaultLuminanceRange->GetValueOnGameThread() == 1;
+
+		if (bExtendedLuminanceRange)
+		{
+			if (Property->GetName() == GET_MEMBER_NAME_CHECKED(FPostProcessSettings, AutoExposureMinBrightness))
+			{
+				Property->SetMetaData(TEXT("DisplayName"), TEXT("Min EV100"));
+			}
+			else if (Property->GetName() == GET_MEMBER_NAME_CHECKED(FPostProcessSettings, AutoExposureMaxBrightness))
+			{
+				Property->SetMetaData(TEXT("DisplayName"), TEXT("Max EV100"));
+			}
+			else if (Property->GetName() == GET_MEMBER_NAME_CHECKED(FPostProcessSettings, HistogramLogMin))
+			{
+				Property->SetMetaData(TEXT("DisplayName"), TEXT("Histogram Min EV100"));
+			}
+			else if (Property->GetName() == GET_MEMBER_NAME_CHECKED(FPostProcessSettings, HistogramLogMax))
+			{
+				Property->SetMetaData(TEXT("DisplayName"), TEXT("Histogram Max EV100"));
+			}
+		}
+	}
+	// Special handling for the ray tracing translucency section, which needs to show or hide specific properties based on the translucency type
+	else if (CategoryFName == TranslucencyCategory)
+	{
+		TSharedPtr<IPropertyHandle> TranslucencyTypeHandle = StructHandle->GetChildHandle(GET_MEMBER_NAME_CHECKED(FPostProcessSettings, TranslucencyType));
+		if (PropertyHandle->GetProperty()->GetFName() != TranslucencyTypeHandle->GetProperty()->GetName())
+		{
+			// Based on the translucency type property value, hide any properties that don't match the current translucency type
+			uint8 TranslucencyType = 0;
+			TranslucencyTypeHandle->GetValue(TranslucencyType);
 
-	FProperty* Prop = StructPropertyHandle->GetProperty();
-	FStructProperty* StructProp = CastField<FStructProperty>(Prop);
+			const FString& TranslucencyTypeMetaData = PropertyHandle->GetMetaData(TEXT("TranslucencyType"));
+			const FString& TranslucencyTypeName = StaticEnum<ETranslucencyType>()->GetNameStringByValue(TranslucencyType);
+			if (!TranslucencyTypeMetaData.Equals(TranslucencyTypeName, ESearchCase::IgnoreCase))
+			{
+				PropertyHandle->MarkHiddenByCustomization();
+				return false;
+			}
+		}
+		else
+		{
+			// Add a property changed event on the translucency type property to refresh the details panel so that the necessary properties
+			// are made visible or hidden, as well as notifying the user if the project settings don't have hardware ray tracing enabled,
+			// which is required for ray traced translucency
+			PropertyHandle->SetOnPropertyValueChanged(FSimpleDelegate::CreateLambda([StructHandle, &StructCustomizationUtils]()
+			{
+				TSharedPtr<IPropertyHandle> TranslucencyTypeHandle = StructHandle->GetChildHandle(GET_MEMBER_NAME_CHECKED(FPostProcessSettings, TranslucencyType));
+								
+				uint8 TranslucencyType = 0;
+				TranslucencyTypeHandle->GetValue(TranslucencyType);
 
-	// a category with this name should be one level higher, should be "PostProcessSettings"
-	FName ClassName = StructProp->Struct->GetFName();
+				if (static_cast<ETranslucencyType>(TranslucencyType) == ETranslucencyType::RayTraced)
+				{
+					FPostProcessRayTracingDisabledWarning::DisplayRayTracingDisabledWarning();
+				}
+								
+				StructCustomizationUtils.GetPropertyUtilities()->ForceRefresh();
+			}));
+		}
+	}
+
+	return true;
+}
+
+void FPostProcessSettingsCustomization::CustomizeChildren(TSharedRef<IPropertyHandle> StructPropertyHandle, class IDetailChildrenBuilder& StructBuilder, IPropertyTypeCustomizationUtils& StructCustomizationUtils)
+{
+	uint32 NumChildren = 0;
+	FPropertyAccess::Result Result = StructPropertyHandle->GetNumChildren(NumChildren);
 
 	// Create new categories in the parent layout rather than adding all post process settings to one category
 	IDetailLayoutBuilder& LayoutBuilder = StructBuilder.GetParentCategory().GetParentLayout();
@@ -269,60 +546,34 @@ void FPostProcessSettingsCustomization::CustomizeChildren( TSharedRef<IPropertyH
 	TMap<FString, FCategoryOrGroup> NameToCategoryBuilderMap;
 	TMap<FString, FPostProcessGroup> NameToGroupMap;
 
-	static const auto VarDefaultAutoExposureExtendDefaultLuminanceRange = IConsoleManager::Get().FindTConsoleVariableDataInt(TEXT("r.DefaultFeature.AutoExposure.ExtendDefaultLuminanceRange"));
-	const bool bExtendedLuminanceRange = VarDefaultAutoExposureExtendDefaultLuminanceRange->GetValueOnGameThread() == 1;
-	static const FName ExposureCategory("Lens|Exposure");
-	static const FName TranslucencyCategory("Rendering Features|Translucency");
+	TArray<FString> CategoryAndGroups;
 
 	bool bShowPostProcessCategories = StructPropertyHandle->HasMetaData(ShowPostProcessCategoriesName);
 
-	if(Result == FPropertyAccess::Success && NumChildren > 0)
+	if (Result == FPropertyAccess::Success && NumChildren > 0)
 	{
-		for( uint32 ChildIndex = 0; ChildIndex < NumChildren; ++ChildIndex )
+		for (uint32 ChildIndex = 0; ChildIndex < NumChildren; ++ChildIndex)
 		{
-			TSharedPtr<IPropertyHandle> ChildHandle = StructPropertyHandle->GetChildHandle( ChildIndex );
+			TSharedPtr<IPropertyHandle> ChildHandle = StructPropertyHandle->GetChildHandle(ChildIndex);
 
-			if( ChildHandle.IsValid() && ChildHandle->GetProperty() )
+			if (ChildHandle.IsValid() && ChildHandle->GetProperty())
 			{
 				FProperty* Property = ChildHandle->GetProperty();
 
 				FName CategoryFName = FObjectEditorUtils::GetCategoryFName(Property);
-					
-				if (CategoryFName == ExposureCategory && bExtendedLuminanceRange)
-				{
-					if (Property->GetName() == TEXT("AutoExposureMinBrightness"))
-					{
-						Property->SetMetaData(TEXT("DisplayName"), TEXT("Min EV100"));
-					}
-					else if (Property->GetName() == TEXT("AutoExposureMaxBrightness"))
-					{
-						Property->SetMetaData(TEXT("DisplayName"), TEXT("Max EV100"));
-					}
-					else if (Property->GetName() == TEXT("HistogramLogMin"))
-					{
-						Property->SetMetaData(TEXT("DisplayName"), TEXT("Histogram Min EV100"));
-					}
-					else if (Property->GetName() == TEXT("HistogramLogMax"))
-					{
-						Property->SetMetaData(TEXT("DisplayName"), TEXT("Histogram Max EV100"));
-					}
-				}
-				
-				
 				FString RawCategoryName = CategoryFName.ToString();
 
-				TArray<FString> CategoryAndGroups;
+				CategoryAndGroups.Reset();
 				RawCategoryName.ParseIntoArray(CategoryAndGroups, TEXT("|"), 1);
 
 				FString RootCategoryName = CategoryAndGroups.Num() > 0 ? CategoryAndGroups[0] : RawCategoryName;
 
-
 				FCategoryOrGroup* Category = NameToCategoryBuilderMap.Find(RootCategoryName);
-				if(!Category)
+				if (!Category)
 				{
-					if(bShowPostProcessCategories)
+					if (bShowPostProcessCategories)
 					{
-					IDetailCategoryBuilder& NewCategory = LayoutBuilder.EditCategory(*RootCategoryName, FText::GetEmpty(), ECategoryPriority::TypeSpecific);
+						IDetailCategoryBuilder& NewCategory = LayoutBuilder.EditCategory(*RootCategoryName, FText::GetEmpty(), ECategoryPriority::TypeSpecific);
 						Category = &NameToCategoryBuilderMap.Emplace(RootCategoryName, NewCategory);
 					}
 					else
@@ -332,85 +583,71 @@ void FPostProcessSettingsCustomization::CustomizeChildren( TSharedRef<IPropertyH
 					}
 				}
 
-				if(CategoryAndGroups.Num() > 1)
+				bool bVisible = CustomizeProperty(ChildHandle, Property, CategoryFName, StructPropertyHandle, StructCustomizationUtils);
+				if (!bVisible)
 				{
-					// Only handling one group for now
-					// There are sub groups so add them now
-					FPostProcessGroup& PPGroup = NameToGroupMap.FindOrAdd(RawCategoryName);
-					
-					// Is this a new group? It wont be valid if it is
-					if(!PPGroup.IsValid())
-					{
-						PPGroup.RootCategory = *Category;
-						PPGroup.RawGroupName = RawCategoryName;
-						PPGroup.DisplayName = CategoryAndGroups[1].TrimStartAndEnd();
-					}
-	
-					// Special handling for the ray tracing translucency section, which needs to show or hide specific properties based on the translucy type
-					if (RawCategoryName == TranslucencyCategory)
-					{
-						TSharedPtr<IPropertyHandle> TranslucencyTypeHandle = StructPropertyHandle->GetChildHandle(GET_MEMBER_NAME_CHECKED(FPostProcessSettings, TranslucencyType));
-						if (ChildHandle->GetProperty()->GetFName() != TranslucencyTypeHandle->GetProperty()->GetName())
-						{
-							// Based on the translucency type property value, hide any properties that don't match the current translucency type
-							uint8 TranslucencyType = 0;
-							TranslucencyTypeHandle->GetValue(TranslucencyType);
-
-							const FString& TranslucencyTypeMetaData = ChildHandle->GetMetaData(TEXT("TranslucencyType"));
-							const FString& TranslucencyTypeName = StaticEnum<ETranslucencyType>()->GetNameStringByValue(TranslucencyType);
-							if (!TranslucencyTypeMetaData.Equals(TranslucencyTypeName, ESearchCase::IgnoreCase))
-							{
-								ChildHandle->MarkHiddenByCustomization();
-								continue;
-							}
-						}
-						else
-						{
-							// Add a property changed event on the translucency type property to refresh the details panel so that the necessary properties
-							// are made visible or hidden, as well as notifying the user if the project settings don't have hardware ray tracing enabled,
-							// which is required for ray traced translucency
-							ChildHandle->SetOnPropertyValueChanged(FSimpleDelegate::CreateLambda([StructPropertyHandle, &StructCustomizationUtils]()
-							{
-								TSharedPtr<IPropertyHandle> TranslucencyTypeHandle = StructPropertyHandle->GetChildHandle(GET_MEMBER_NAME_CHECKED(FPostProcessSettings, TranslucencyType));
-								
-								uint8 TranslucencyType = 0;
-								TranslucencyTypeHandle->GetValue(TranslucencyType);
+					continue;
+				}
 
-								if (static_cast<ETranslucencyType>(TranslucencyType) == ETranslucencyType::RayTraced)
-								{
-									FPostProcessRayTracingDisabledWarning::DisplayRayTracingDisabledWarning();
-								}
-								
-								StructCustomizationUtils.GetPropertyUtilities()->ForceRefresh();
-							}));
-						}
-					}
+				if (CategoryAndGroups.Num() > 1)
+				{
+					FPostProcessGroup* PPGroup = FindOrCreateGroup(NameToGroupMap, Category, RawCategoryName, StructPropertyHandle);
+					check(PPGroup);
 					
 					bool bIsSimple = !ChildHandle->GetProperty()->HasAnyPropertyFlags(CPF_AdvancedDisplay);
-					if(bIsSimple)
+					if (bIsSimple)
 					{
-						PPGroup.SimplePropertyHandles.Add(ChildHandle);
+						PPGroup->SimplePropertyHandles.Add(ChildHandle);
 					}
 					else
 					{
-						PPGroup.AdvancedPropertyHandles.Add(ChildHandle);
+						PPGroup->AdvancedPropertyHandles.Add(ChildHandle);
 					}
 				}
 				else
 				{
 					Category->AddProperty(ChildHandle.ToSharedRef());
 				}
-
 			}
 		}
 
-		for(auto& NameAndGroup : NameToGroupMap)
+		for (auto& NameAndGroup : NameToGroupMap)
 		{
 			FPostProcessGroup& PPGroup = NameAndGroup.Value;
 
-			if(PPGroup.SimplePropertyHandles.Num() > 0 || PPGroup.AdvancedPropertyHandles.Num() > 0 )
+			if (PPGroup.SimplePropertyHandles.Num() > 0 || PPGroup.AdvancedPropertyHandles.Num() > 0)
 			{
-				IDetailGroup& SimpleGroup = PPGroup.RootCategory.AddGroup(*PPGroup.RawGroupName, FText::FromString(PPGroup.DisplayName));
+				MakeDetail(PPGroup, NameToGroupMap, StructPropertyHandle);
+				IDetailGroup& SimpleGroup = *PPGroup.Detail.Group;
+
+				if (PPGroup.Alert.IsSet())
+				{
+					FDetailWidgetRow Header = SimpleGroup.HeaderRow()
+					.NameContent()
+					[
+						MakeNameWidget(FText::FromString(PPGroup.DisplayName), {})
+					]
+					.ExtensionContent()
+					[
+						SNew(SBorder)
+						.Visibility(PPGroup.Alert->Visibility)
+						.ToolTip(
+							SNew(SToolTip)
+							.Text(PPGroup.Alert->Text)
+						)
+						.Padding(3, 3, 3, 3)
+						.BorderBackgroundColor(FColor::Transparent)
+						[
+							SNew(SImage)
+							.DesiredSizeOverride(FVector2D(16,16))
+							.Image_Lambda([Type=PPGroup.Alert->Type]()
+							{ 
+								return (Type == FAlert::EType::Warning) ? FAppStyle::Get().GetBrush("Icons.WarningWithColor") : FAppStyle::Get().GetBrush("Icons.ErrorWithColor"); 
+							})
+							.ColorAndOpacity(FSlateColor::UseForeground())
+						]
+					];
+				}
 
 				static const FString ColorGradingName = TEXT("Color Grading");
 
@@ -420,16 +657,16 @@ void FPostProcessSettingsCustomization::CustomizeChildren( TSharedRef<IPropertyH
 					SimpleGroup.EnableReset(true);
 				}
 
-				for(auto& SimpleProperty : PPGroup.SimplePropertyHandles)
+				for (auto& SimpleProperty : PPGroup.SimplePropertyHandles)
 				{
 					SimpleGroup.AddPropertyRow(SimpleProperty.ToSharedRef());
 				}
 
-				if(PPGroup.AdvancedPropertyHandles.Num() > 0)
+				if (PPGroup.AdvancedPropertyHandles.Num() > 0)
 				{
 					IDetailGroup& AdvancedGroup = SimpleGroup.AddGroup(*(PPGroup.RawGroupName+TEXT("Advanced")), LOCTEXT("PostProcessAdvancedGroup", "Advanced"));
 					
-					for(auto& AdvancedProperty : PPGroup.AdvancedPropertyHandles)
+					for (auto& AdvancedProperty : PPGroup.AdvancedPropertyHandles)
 					{
 						AdvancedGroup.AddPropertyRow(AdvancedProperty.ToSharedRef());
 					}
@@ -439,10 +676,10 @@ void FPostProcessSettingsCustomization::CustomizeChildren( TSharedRef<IPropertyH
 	}
 }
 
-void FPostProcessSettingsCustomization::CustomizeHeader( TSharedRef<IPropertyHandle> StructPropertyHandle, class FDetailWidgetRow& HeaderRow, IPropertyTypeCustomizationUtils& StructCustomizationUtils )
+void FPostProcessSettingsCustomization::CustomizeHeader(TSharedRef<IPropertyHandle> StructPropertyHandle, class FDetailWidgetRow& HeaderRow, IPropertyTypeCustomizationUtils& StructCustomizationUtils)
 {
 	bool bShowHeader = !StructPropertyHandle->HasMetaData(ShowPostProcessCategoriesName) && !StructPropertyHandle->HasMetaData(ShowOnlyInnerPropertiesName);
-	if(bShowHeader)
+	if (bShowHeader)
 	{
 		HeaderRow.NameContent()
 		[
@@ -466,7 +703,7 @@ void FWeightedBlendableCustomization::AddDirectAsset(TSharedRef<IPropertyHandle>
 
 		TArray<FString> Values;
 
-		for(TArray<UObject*>::TConstIterator It = Objects.CreateConstIterator(); It; It++)
+		for (TArray<UObject*>::TConstIterator It = Objects.CreateConstIterator(); It; It++)
 		{
 			UObject* Obj = *It;
 
@@ -523,9 +760,9 @@ TSharedRef<SWidget> FWeightedBlendableCustomization::GenerateContentWidget(TShar
 
 	FMenuBuilder MenuBuilder(true, NULL);
 	{
-		for(TObjectIterator<UClass> It; It; ++It)
+		for (TObjectIterator<UClass> It; It; ++It)
 		{
-			if( It->IsChildOf(UFactory::StaticClass()))
+			if (It->IsChildOf(UFactory::StaticClass()))
 			{
 				UFactory* Factory = It->GetDefaultObject<UFactory>();
 
@@ -533,13 +770,13 @@ TSharedRef<SWidget> FWeightedBlendableCustomization::GenerateContentWidget(TShar
 
 				UClass* SupportedClass = Factory->GetSupportedClass();
 
-				if(SupportedClass)
+				if (SupportedClass)
 				{
-					if(SupportedClass->ImplementsInterface(UBlendableInterface::StaticClass()))
+					if (SupportedClass->ImplementsInterface(UBlendableInterface::StaticClass()))
 					{
 						// At the moment we know about 3 Blendables: Material, UMaterialInstanceConstant, LightPropagationVolumeBlendable
 						// The materials are not that useful to have here (hard to reference) so we suppress them here
-						if(!(
+						if (!(
 							SupportedClass == UMaterial::StaticClass() ||
 							SupportedClass == UMaterialInstanceConstant::StaticClass()
 							))
@@ -558,7 +795,7 @@ TSharedRef<SWidget> FWeightedBlendableCustomization::GenerateContentWidget(TShar
 			}
 		}
 
-		if(bSeparatorIsNeeded)
+		if (bSeparatorIsNeeded)
 		{
 			MenuBuilder.AddMenuSeparator();
 		}
@@ -613,7 +850,7 @@ int32 FWeightedBlendableCustomization::ComputeSwitcherIndex(TSharedRef<IProperty
 	Weight->GetValue(WeightValue);
 	Value->GetValue(RefObject);
 
-	if(RefObject)
+	if (RefObject)
 	{
 		UPackage* PropPackage = RefObject->GetOutermost();
 
@@ -625,12 +862,12 @@ int32 FWeightedBlendableCustomization::ComputeSwitcherIndex(TSharedRef<IProperty
 	}
 }
 
-void FWeightedBlendableCustomization::CustomizeChildren( TSharedRef<class IPropertyHandle> StructPropertyHandle, class IDetailChildrenBuilder& StructBuilder, IPropertyTypeCustomizationUtils& StructCustomizationUtils )
+void FWeightedBlendableCustomization::CustomizeChildren(TSharedRef<class IPropertyHandle> StructPropertyHandle, class IDetailChildrenBuilder& StructBuilder, IPropertyTypeCustomizationUtils& StructCustomizationUtils)
 {
 	// we don't have children but this is a pure virtual so we need to override
 }
 
-void FWeightedBlendableCustomization::CustomizeHeader( TSharedRef<IPropertyHandle> StructPropertyHandle, class FDetailWidgetRow& HeaderRow, IPropertyTypeCustomizationUtils& StructCustomizationUtils )
+void FWeightedBlendableCustomization::CustomizeHeader(TSharedRef<IPropertyHandle> StructPropertyHandle, class FDetailWidgetRow& HeaderRow, IPropertyTypeCustomizationUtils& StructCustomizationUtils)
 {
 	TSharedPtr<IPropertyHandle> SharedWeightProp;
 	{
diff --git a/Engine/Source/Editor/UnrealEd/Private/EditorViewportClient.cpp b/Engine/Source/Editor/UnrealEd/Private/EditorViewportClient.cpp
--- a/Engine/Source/Editor/UnrealEd/Private/EditorViewportClient.cpp
+++ b/Engine/Source/Editor/UnrealEd/Private/EditorViewportClient.cpp
@@ -89,6 +89,7 @@
 #include "ViewportInteractions/ViewportInteractionsBehaviorSource.h"
 #include "ViewportInteractions/ViewportInteraction.h"
 #include "WorldPartition/WorldPartition.h"
+#include "RenderViewportFeedback.h"
 
 #define LOCTEXT_NAMESPACE "EditorViewportClient"
 
@@ -1581,6 +1582,11 @@ FSceneView* FEditorViewportClient::CalcSceneView(FSceneViewFamily* ViewFamily, c
 
 	View->EndFinalPostprocessSettings(ViewInitOptions);
 
+	if (RenderViewportFeedback && (StereoViewIndex == INDEX_NONE || StereoViewIndex == 0))
+	{
+		View->ViewportFeedback = RenderViewportFeedback->MakeCollector();
+	}
+
 	for (int ViewExt = 0; ViewExt < ViewFamily->ViewExtensions.Num(); ViewExt++)
 	{
 		ViewFamily->ViewExtensions[ViewExt]->SetupView(*ViewFamily, *View);
diff --git a/Engine/Source/Editor/UnrealEd/Public/EditorViewportClient.h b/Engine/Source/Editor/UnrealEd/Public/EditorViewportClient.h
--- a/Engine/Source/Editor/UnrealEd/Public/EditorViewportClient.h
+++ b/Engine/Source/Editor/UnrealEd/Public/EditorViewportClient.h
@@ -42,6 +42,10 @@ class UTypedElementViewportInteraction;
 enum class EViewStatusForScreenPercentage;
 struct FGizmoState;
 struct FViewportClientNavigationHelper;
+namespace UE::RenderViewportFeedback
+{
+	class FReceiver;
+}
 
 /** Delegate called by FEditorViewportClient to check its visibility */
 DECLARE_DELEGATE_RetVal( bool, FViewportStateGetter );
@@ -1964,6 +1968,8 @@ public:
 	/** Whether view modifiers should be called and applied. */
 	bool bShouldApplyViewModifiers;
 
+	TSharedPtr<UE::RenderViewportFeedback::FReceiver, ESPMode::ThreadSafe> RenderViewportFeedback = nullptr;
+
 protected:
 	/** Editor mode tools provided to this instance. */
 	TSharedPtr<FEditorModeTools> ModeTools;
diff --git a/Engine/Source/Runtime/Core/Public/UObject/FortniteMainBranchObjectVersions.inl b/Engine/Source/Runtime/Core/Public/UObject/FortniteMainBranchObjectVersions.inl
--- a/Engine/Source/Runtime/Core/Public/UObject/FortniteMainBranchObjectVersions.inl
+++ b/Engine/Source/Runtime/Core/Public/UObject/FortniteMainBranchObjectVersions.inl
@@ -673,5 +673,8 @@ FORTNITE_MAIN_VERSION(HLODLayerEditorOnlyObject, 219)
 // Deduplicated particle debug names serialization in the Chaos Visual Debugger
 FORTNITE_MAIN_VERSION(DeduplicatedDebugNameSerializationInCVD, 220)
 
+// Add BloomGaussianIntensity and BloomConvolutionIntensity
+FORTNITE_MAIN_VERSION(SpecializeBloomIntensity, 221)
+
 // -----<new versions can be added above this line>-------------------------------------------------
 #undef FORTNITE_MAIN_VERSION
\ No newline at end of file
diff --git a/Engine/Source/Runtime/Engine/Classes/Engine/Scene.h b/Engine/Source/Runtime/Engine/Classes/Engine/Scene.h
--- a/Engine/Source/Runtime/Engine/Classes/Engine/Scene.h
+++ b/Engine/Source/Runtime/Engine/Classes/Engine/Scene.h
@@ -229,21 +229,21 @@ struct FGaussianSumBloomSettings
 	GENERATED_USTRUCT_BODY()
 
 
-	/** Multiplier for all bloom contributions >=0: off, 1(default), >1 brighter */
-	UPROPERTY(Interp, BlueprintReadWrite, Category="Lens|Bloom", meta=(ClampMin = "0.0", UIMax = "8.0", DisplayName = "Intensity"))
+	/** Multiplier for Gaussian bloom contributions >=0: off, 1(default), >1 brighter */
+	UPROPERTY(Interp, BlueprintReadWrite, Category="Lens|Bloom|Gaussian", meta=(ClampMin = "0.0", UIMax = "8.0", DisplayName = "Intensity"))
 	float Intensity;
 
 	/**
 	 * minimum brightness the bloom starts having effect
 	 * -1:all pixels affect bloom equally (physically correct, faster as a threshold pass is omitted), 0:all pixels affect bloom brights more, 1(default), >1 brighter
 	 */
-	UPROPERTY(interp, BlueprintReadWrite, Category="Lens|Bloom", meta=(ClampMin = "-1.0", UIMax = "8.0", DisplayName = "Threshold"))
+	UPROPERTY(interp, BlueprintReadWrite, Category="Lens|Bloom|Gaussian", meta=(ClampMin = "-1.0", UIMax = "8.0", DisplayName = "Threshold"))
 	float Threshold;
 
 	/**
 	 * Scale for all bloom sizes
 	 */
-	UPROPERTY(Interp, BlueprintReadWrite, Category="Lens|Bloom", AdvancedDisplay, meta=(ClampMin = "0.0", UIMax = "64.0", DisplayName = "Size scale"))
+	UPROPERTY(Interp, BlueprintReadWrite, Category="Lens|Bloom|Gaussian", AdvancedDisplay, meta=(ClampMin = "0.0", UIMax = "64.0", DisplayName = "Size scale"))
 	float SizeScale;
 
 	/**
@@ -251,7 +251,7 @@ struct FGaussianSumBloomSettings
 	 * (is done in 1/2 resolution, larger values cost more performance, good for high frequency details)
 	 * >=0: can be clamped because of shader limitations
 	 */
-	UPROPERTY(Interp, BlueprintReadWrite, Category="Lens|Bloom", AdvancedDisplay, meta=(ClampMin = "0.0", UIMax = "4.0", DisplayName = "#1 Size"))
+	UPROPERTY(Interp, BlueprintReadWrite, Category="Lens|Bloom|Gaussian", AdvancedDisplay, meta=(ClampMin = "0.0", UIMax = "4.0", DisplayName = "#1 Size"))
 	float Filter1Size;
 
 	/**
@@ -259,7 +259,7 @@ struct FGaussianSumBloomSettings
 	 * (is done in 1/4 resolution, larger values cost more performance)
 	 * >=0: can be clamped because of shader limitations
 	 */
-	UPROPERTY(Interp, BlueprintReadWrite, Category="Lens|Bloom", AdvancedDisplay, meta=(ClampMin = "0.0", UIMax = "8.0", DisplayName = "#2 Size"))
+	UPROPERTY(Interp, BlueprintReadWrite, Category="Lens|Bloom|Gaussian", AdvancedDisplay, meta=(ClampMin = "0.0", UIMax = "8.0", DisplayName = "#2 Size"))
 	float Filter2Size;
 
 	/**
@@ -267,7 +267,7 @@ struct FGaussianSumBloomSettings
 	 * (is done in 1/8 resolution, larger values cost more performance)
 	 * >=0: can be clamped because of shader limitations
 	 */
-	UPROPERTY(Interp, BlueprintReadWrite, Category="Lens|Bloom", AdvancedDisplay, meta=(ClampMin = "0.0", UIMax = "16.0", DisplayName = "#3 Size"))
+	UPROPERTY(Interp, BlueprintReadWrite, Category="Lens|Bloom|Gaussian", AdvancedDisplay, meta=(ClampMin = "0.0", UIMax = "16.0", DisplayName = "#3 Size"))
 	float Filter3Size;
 
 	/**
@@ -275,7 +275,7 @@ struct FGaussianSumBloomSettings
 	 * (is done in 1/16 resolution, larger values cost more performance, best for wide contributions)
 	 * >=0: can be clamped because of shader limitations
 	 */
-	UPROPERTY(Interp, BlueprintReadWrite, Category="Lens|Bloom", AdvancedDisplay, meta=(ClampMin = "0.0", UIMax = "32.0", DisplayName = "#4 Size"))
+	UPROPERTY(Interp, BlueprintReadWrite, Category="Lens|Bloom|Gaussian", AdvancedDisplay, meta=(ClampMin = "0.0", UIMax = "32.0", DisplayName = "#4 Size"))
 	float Filter4Size;
 
 	/**
@@ -283,7 +283,7 @@ struct FGaussianSumBloomSettings
 	 * (is done in 1/32 resolution, larger values cost more performance, best for wide contributions)
 	 * >=0: can be clamped because of shader limitations
 	 */
-	UPROPERTY(Interp, BlueprintReadWrite, Category="Lens|Bloom", AdvancedDisplay, meta=(ClampMin = "0.0", UIMax = "64.0", DisplayName = "#5 Size"))
+	UPROPERTY(Interp, BlueprintReadWrite, Category="Lens|Bloom|Gaussian", AdvancedDisplay, meta=(ClampMin = "0.0", UIMax = "64.0", DisplayName = "#5 Size"))
 	float Filter5Size;
 
 	/**
@@ -291,37 +291,37 @@ struct FGaussianSumBloomSettings
 	 * (is done in 1/64 resolution, larger values cost more performance, best for wide contributions)
 	 * >=0: can be clamped because of shader limitations
 	 */
-	UPROPERTY(Interp, BlueprintReadWrite, Category="Lens|Bloom", AdvancedDisplay, meta=(ClampMin = "0.0", UIMax = "128.0", DisplayName = "#6 Size"))
+	UPROPERTY(Interp, BlueprintReadWrite, Category="Lens|Bloom|Gaussian", AdvancedDisplay, meta=(ClampMin = "0.0", UIMax = "128.0", DisplayName = "#6 Size"))
 	float Filter6Size;
 
 	/** Bloom1 tint color */
-	UPROPERTY(Interp, BlueprintReadWrite, Category="Lens|Bloom", AdvancedDisplay, meta=(DisplayName = "#1 Tint", HideAlphaChannel))
+	UPROPERTY(Interp, BlueprintReadWrite, Category="Lens|Bloom|Gaussian", AdvancedDisplay, meta=(DisplayName = "#1 Tint", HideAlphaChannel))
 	FLinearColor Filter1Tint;
 
 	/** Bloom2 tint color */
-	UPROPERTY(Interp, BlueprintReadWrite, Category="Lens|Bloom", AdvancedDisplay, meta=(DisplayName = "#2 Tint", HideAlphaChannel))
+	UPROPERTY(Interp, BlueprintReadWrite, Category="Lens|Bloom|Gaussian", AdvancedDisplay, meta=(DisplayName = "#2 Tint", HideAlphaChannel))
 	FLinearColor Filter2Tint;
 
 	/** Bloom3 tint color */
-	UPROPERTY(Interp, BlueprintReadWrite, Category="Lens|Bloom", AdvancedDisplay, meta=(DisplayName = "#3 Tint", HideAlphaChannel))
+	UPROPERTY(Interp, BlueprintReadWrite, Category="Lens|Bloom|Gaussian", AdvancedDisplay, meta=(DisplayName = "#3 Tint", HideAlphaChannel))
 	FLinearColor Filter3Tint;
 
 	/** Bloom4 tint color */
-	UPROPERTY(Interp, BlueprintReadWrite, Category="Lens|Bloom", AdvancedDisplay, meta=(DisplayName = "#4 Tint", HideAlphaChannel))
+	UPROPERTY(Interp, BlueprintReadWrite, Category="Lens|Bloom|Gaussian", AdvancedDisplay, meta=(DisplayName = "#4 Tint", HideAlphaChannel))
 	FLinearColor Filter4Tint;
 
 	/** Bloom5 tint color */
-	UPROPERTY(Interp, BlueprintReadWrite, Category="Lens|Bloom", AdvancedDisplay, meta=(DisplayName = "#5 Tint", HideAlphaChannel))
+	UPROPERTY(Interp, BlueprintReadWrite, Category="Lens|Bloom|Gaussian", AdvancedDisplay, meta=(DisplayName = "#5 Tint", HideAlphaChannel))
 	FLinearColor Filter5Tint;
 
 	/** Bloom6 tint color */
-	UPROPERTY(Interp, BlueprintReadWrite, Category="Lens|Bloom", AdvancedDisplay, meta=(DisplayName = "#6 Tint", HideAlphaChannel))
+	UPROPERTY(Interp, BlueprintReadWrite, Category="Lens|Bloom|Gaussian", AdvancedDisplay, meta=(DisplayName = "#6 Tint", HideAlphaChannel))
 	FLinearColor Filter6Tint;
 
 
 	FGaussianSumBloomSettings()
 	{
-		Intensity = 0.675f;
+		Intensity = 1.0f;
 		Threshold = -1.0f;
 		// default is 4 to maintain old settings after fixing something that caused a factor of 4
 		SizeScale = 4.0;
@@ -349,41 +349,46 @@ struct FConvolutionBloomSettings
 	GENERATED_USTRUCT_BODY()
 
 
+	/** Multiplier for convolution bloom contributions >=0: off, 1(default), >1 brighter */
+	UPROPERTY(Interp, BlueprintReadWrite, Category="Lens|Bloom|Convolution", meta=(ClampMin = "0.0", UIMax = "8.0", DisplayName = "Intensity"))
+	float Intensity;
+
 	/** Texture to replace default convolution bloom kernel */
-	UPROPERTY(EditAnywhere, BlueprintReadWrite, Category = "Lens|Bloom", meta = (DisplayName = "Convolution Kernel"))
+	UPROPERTY(EditAnywhere, BlueprintReadWrite, Category = "Lens|Bloom|Convolution", meta = (DisplayName = "Convolution Kernel"))
 	TObjectPtr<class UTexture2D> Texture;
 
 	/** Intensity multiplier on the scatter dispersion energy of the kernel. 1.0 means exactly use the same energy as the kernel scatter dispersion. */
-	UPROPERTY(Interp, BlueprintReadWrite, Category = "Lens|Bloom", AdvancedDisplay, meta = (ClampMin = "0.0", UIMax = "1.0", DisplayName = "Scatter Dispersion"))
+	UPROPERTY(Interp, BlueprintReadWrite, Category = "Lens|Bloom|Convolution", AdvancedDisplay, meta = (ClampMin = "0.0", UIMax = "1.0", DisplayName = "Scatter Dispersion"))
 	float ScatterDispersion;
 
 	/** Relative size of the convolution kernel image compared to the minor axis of the viewport  */
-	UPROPERTY(Interp, BlueprintReadWrite, Category = "Lens|Bloom", AdvancedDisplay, meta = (ClampMin = "0.0", UIMax = "1.0", DisplayName = "Convolution Scale"))
+	UPROPERTY(Interp, BlueprintReadWrite, Category = "Lens|Bloom|Convolution", AdvancedDisplay, meta = (ClampMin = "0.0", UIMax = "1.0", DisplayName = "Convolution Scale"))
 	float Size;
 
 	/** The UV location of the center of the kernel.  Should be very close to (.5,.5) */
-	UPROPERTY(Interp, BlueprintReadWrite, Category = "Lens|Bloom", AdvancedDisplay, meta = (DisplayName = "Convolution Center"))
+	UPROPERTY(Interp, BlueprintReadWrite, Category = "Lens|Bloom|Convolution", AdvancedDisplay, meta = (DisplayName = "Convolution Center"))
 	FVector2D CenterUV;
 
 	/** Boost intensity of select pixels  prior to computing bloom convolution (Min, Max, Multiplier).  Max < Min disables convolution boost */
-	UPROPERTY(Interp, BlueprintReadWrite, Category = "Lens|Bloom", AdvancedDisplay, meta = (DisplayName = "Convolution Boost Min"))
+	UPROPERTY(Interp, BlueprintReadWrite, Category = "Lens|Bloom|Convolution", AdvancedDisplay, meta = (DisplayName = "Convolution Boost Min"))
 	float PreFilterMin;
 
 	/** Boost intensity of select pixels  prior to computing bloom convolution (Min, Max, Multiplier).  Max < Min disables convolution boost */
-	UPROPERTY(Interp, BlueprintReadWrite, Category = "Lens|Bloom", AdvancedDisplay, meta = (DisplayName = "Convolution Boost Max"))
+	UPROPERTY(Interp, BlueprintReadWrite, Category = "Lens|Bloom|Convolution", AdvancedDisplay, meta = (DisplayName = "Convolution Boost Max"))
 	float PreFilterMax;
 
 	/** Boost intensity of select pixels  prior to computing bloom convolution (Min, Max, Multiplier).  Max < Min disables convolution boost */
-	UPROPERTY(Interp, BlueprintReadWrite, Category = "Lens|Bloom", AdvancedDisplay, meta = (DisplayName = "Convolution Boost Mult"))
+	UPROPERTY(Interp, BlueprintReadWrite, Category = "Lens|Bloom|Convolution", AdvancedDisplay, meta = (DisplayName = "Convolution Boost Mult"))
 	float PreFilterMult;
 
 	/** Implicit buffer region as a fraction of the screen size to insure the bloom does not wrap across the screen.  Larger sizes have perf impact.*/
-	UPROPERTY(Interp, BlueprintReadWrite, Category = "Lens|Bloom", AdvancedDisplay, meta = (ClampMin = "0.0", UIMax = "1.0", DisplayName = "Convolution Buffer"))
+	UPROPERTY(Interp, BlueprintReadWrite, Category = "Lens|Bloom|Convolution", AdvancedDisplay, meta = (ClampMin = "0.0", UIMax = "1.0", DisplayName = "Convolution Buffer"))
 	float BufferScale;
 
 
 	FConvolutionBloomSettings()
 	{
+		Intensity = 1.0f;
 		Texture = nullptr;
 		ScatterDispersion = 1.0f;
 		Size = 1.f;
@@ -402,6 +407,10 @@ USTRUCT(BlueprintType)
 struct FLensBloomSettings
 {
 	GENERATED_USTRUCT_BODY()
+	
+	/** Multiplier for all bloom contributions >=0: off, 1(default), >1 brighter */
+	UPROPERTY(Interp, BlueprintReadWrite, Category="Lens|Bloom", meta=(ClampMin = "0.0", UIMax = "8.0", DisplayName = "Intensity"))
+	float Intensity;
 
 	/** Bloom gaussian sum method specific settings. */
 	UPROPERTY(Interp, BlueprintReadWrite, Category = "Gaussian Sum Method")
@@ -418,13 +427,27 @@ struct FLensBloomSettings
 
 	FLensBloomSettings()
 	{
+		Intensity = 0.675f;
 		Method = BM_SOG;
 	}
 
 	/* Exports to post process settings with overrides. */
 	ENGINE_API void ExportToPostProcessSettings(FPostProcessSettings* OutPostProcessSettings) const;
+	ENGINE_API bool Serialize(FArchive& Ar);
+	ENGINE_API void PostSerialize(const FArchive& Ar);
 };
 
+#if WITH_EDITORONLY_DATA
+template<> struct TStructOpsTypeTraits<FLensBloomSettings> : public TStructOpsTypeTraitsBase2<FLensBloomSettings>
+{
+	enum
+	{
+		WithSerializer = true,
+		WithPostSerialize = true,
+	};
+};
+#endif
+
 USTRUCT(BlueprintType)
 struct FLensImperfectionSettings
 {
@@ -772,6 +795,9 @@ struct FPostProcessSettings
 	UPROPERTY(EditAnywhere, BlueprintReadWrite, Category=Overrides, meta=(PinHiddenByDefault, InlineEditConditionToggle))
 	uint8 bOverride_BloomIntensity:1;
 
+	UPROPERTY(EditAnywhere, BlueprintReadWrite, Category=Overrides, meta=(PinHiddenByDefault, InlineEditConditionToggle))
+	uint8 bOverride_BloomGaussianIntensity:1;
+
 	UPROPERTY(EditAnywhere, BlueprintReadWrite, Category=Overrides, meta=(PinHiddenByDefault, InlineEditConditionToggle))
 	uint8 bOverride_BloomThreshold:1;
 
@@ -813,6 +839,9 @@ struct FPostProcessSettings
 
 	UPROPERTY(EditAnywhere, BlueprintReadWrite, Category=Overrides, meta=(PinHiddenByDefault, InlineEditConditionToggle))
 	uint8 bOverride_BloomSizeScale:1;
+	
+	UPROPERTY(EditAnywhere, BlueprintReadWrite, Category=Overrides, meta=(PinHiddenByDefault, InlineEditConditionToggle))
+	uint8 bOverride_BloomConvolutionIntensity:1;
 
 	UPROPERTY(EditAnywhere, BlueprintReadWrite, Category=Overrides, meta=(PinHiddenByDefault, InlineEditConditionToggle))
 	uint8 bOverride_BloomConvolutionTexture : 1;
@@ -1391,8 +1420,8 @@ struct FPostProcessSettings
 	UPROPERTY(EditAnywhere, BlueprintReadWrite, Category = "Lens|Mobile Depth of Field", meta = (editcondition = "bOverride_MobileHQGaussian", DisplayName = "High Quality Gaussian DoF on Mobile"))
 	uint8 bMobileHQGaussian:1;
 
-	/** Bloom algorithm */
-	UPROPERTY(EditAnywhere, BlueprintReadWrite, Category = "Lens|Bloom", meta = (editcondition = "bOverride_BloomMethod", DisplayName = "Method"))
+	/** Preferred bloom algorithm. A different method may be used if the preferred method is not available on the target platform, e.g. due to performance limits. */
+	UPROPERTY(EditAnywhere, BlueprintReadWrite, Category = "Lens|Bloom", meta = (editcondition = "bOverride_BloomMethod", DisplayName = "Preferred Method"))
 	TEnumAsByte<enum EBloomMethod> BloomMethod;
 
 	/** Luminance computation method */
@@ -1533,18 +1562,22 @@ struct FPostProcessSettings
 	/** Multiplier for all bloom contributions >=0: off, 1(default), >1 brighter */
 	UPROPERTY(interp, BlueprintReadWrite, Category="Lens|Bloom", meta=(ClampMin = "0.0", UIMax = "8.0", editcondition = "bOverride_BloomIntensity", DisplayName = "Intensity"))
 	float BloomIntensity;
+	
+	/** Multiplier for Gaussian bloom contributions >=0: off, 1(default), >1 brighter */
+	UPROPERTY(interp, BlueprintReadWrite, Category="Lens|Bloom|Gaussian", meta=(ClampMin = "0.0", UIMax = "8.0", editcondition = "bOverride_BloomGaussianIntensity", DisplayName = "Intensity"))
+	float BloomGaussianIntensity;
 
 	/**
 	 * minimum brightness the bloom starts having effect
 	 * -1:all pixels affect bloom equally (physically correct, faster as a threshold pass is omitted), 0:all pixels affect bloom brights more, 1(default), >1 brighter
 	 */
-	UPROPERTY(interp, BlueprintReadWrite, Category="Lens|Bloom", meta=(ClampMin = "-1.0", UIMax = "8.0", editcondition = "bOverride_BloomThreshold", DisplayName = "Threshold"))
+	UPROPERTY(interp, BlueprintReadWrite, Category="Lens|Bloom|Gaussian", meta=(ClampMin = "-1.0", UIMax = "8.0", editcondition = "bOverride_BloomThreshold", DisplayName = "Threshold"))
 	float BloomThreshold;
 
 	/**
 	 * Scale for all bloom sizes
 	 */
-	UPROPERTY(interp, BlueprintReadWrite, Category="Lens|Bloom", AdvancedDisplay, meta=(ClampMin = "0.0", UIMax = "64.0", editcondition = "bOverride_BloomSizeScale", DisplayName = "Size scale"))
+	UPROPERTY(interp, BlueprintReadWrite, Category="Lens|Bloom|Gaussian", AdvancedDisplay, meta=(ClampMin = "0.0", UIMax = "64.0", editcondition = "bOverride_BloomSizeScale", DisplayName = "Size scale"))
 	float BloomSizeScale;
 
 	/**
@@ -1552,77 +1585,81 @@ struct FPostProcessSettings
 	 * (is done in 1/2 resolution, larger values cost more performance, good for high frequency details)
 	 * >=0: can be clamped because of shader limitations
 	 */
-	UPROPERTY(interp, BlueprintReadWrite, Category="Lens|Bloom", AdvancedDisplay, meta=(ClampMin = "0.0", UIMax = "4.0", editcondition = "bOverride_Bloom1Size", DisplayName = "#1 Size"))
+	UPROPERTY(interp, BlueprintReadWrite, Category="Lens|Bloom|Gaussian", AdvancedDisplay, meta=(ClampMin = "0.0", UIMax = "4.0", editcondition = "bOverride_Bloom1Size", DisplayName = "#1 Size"))
 	float Bloom1Size;
 	/**
 	 * Diameter size for Bloom2 in percent of the screen width
 	 * (is done in 1/4 resolution, larger values cost more performance)
 	 * >=0: can be clamped because of shader limitations
 	 */
-	UPROPERTY(interp, BlueprintReadWrite, Category="Lens|Bloom", AdvancedDisplay, meta=(ClampMin = "0.0", UIMax = "8.0", editcondition = "bOverride_Bloom2Size", DisplayName = "#2 Size"))
+	UPROPERTY(interp, BlueprintReadWrite, Category="Lens|Bloom|Gaussian", AdvancedDisplay, meta=(ClampMin = "0.0", UIMax = "8.0", editcondition = "bOverride_Bloom2Size", DisplayName = "#2 Size"))
 	float Bloom2Size;
 	/**
 	 * Diameter size for Bloom3 in percent of the screen width
 	 * (is done in 1/8 resolution, larger values cost more performance)
 	 * >=0: can be clamped because of shader limitations
 	 */
-	UPROPERTY(interp, BlueprintReadWrite, Category="Lens|Bloom", AdvancedDisplay, meta=(ClampMin = "0.0", UIMax = "16.0", editcondition = "bOverride_Bloom3Size", DisplayName = "#3 Size"))
+	UPROPERTY(interp, BlueprintReadWrite, Category="Lens|Bloom|Gaussian", AdvancedDisplay, meta=(ClampMin = "0.0", UIMax = "16.0", editcondition = "bOverride_Bloom3Size", DisplayName = "#3 Size"))
 	float Bloom3Size;
 	/**
 	 * Diameter size for Bloom4 in percent of the screen width
 	 * (is done in 1/16 resolution, larger values cost more performance, best for wide contributions)
 	 * >=0: can be clamped because of shader limitations
 	 */
-	UPROPERTY(interp, BlueprintReadWrite, Category="Lens|Bloom", AdvancedDisplay, meta=(ClampMin = "0.0", UIMax = "32.0", editcondition = "bOverride_Bloom4Size", DisplayName = "#4 Size"))
+	UPROPERTY(interp, BlueprintReadWrite, Category="Lens|Bloom|Gaussian", AdvancedDisplay, meta=(ClampMin = "0.0", UIMax = "32.0", editcondition = "bOverride_Bloom4Size", DisplayName = "#4 Size"))
 	float Bloom4Size;
 	/**
 	 * Diameter size for Bloom5 in percent of the screen width
 	 * (is done in 1/32 resolution, larger values cost more performance, best for wide contributions)
 	 * >=0: can be clamped because of shader limitations
 	 */
-	UPROPERTY(interp, BlueprintReadWrite, Category="Lens|Bloom", AdvancedDisplay, meta=(ClampMin = "0.0", UIMax = "64.0", editcondition = "bOverride_Bloom5Size", DisplayName = "#5 Size"))
+	UPROPERTY(interp, BlueprintReadWrite, Category="Lens|Bloom|Gaussian", AdvancedDisplay, meta=(ClampMin = "0.0", UIMax = "64.0", editcondition = "bOverride_Bloom5Size", DisplayName = "#5 Size"))
 	float Bloom5Size;
 	/**
 	 * Diameter size for Bloom6 in percent of the screen width
 	 * (is done in 1/64 resolution, larger values cost more performance, best for wide contributions)
 	 * >=0: can be clamped because of shader limitations
 	 */
-	UPROPERTY(interp, BlueprintReadWrite, Category="Lens|Bloom", AdvancedDisplay, meta=(ClampMin = "0.0", UIMax = "128.0", editcondition = "bOverride_Bloom6Size", DisplayName = "#6 Size"))
+	UPROPERTY(interp, BlueprintReadWrite, Category="Lens|Bloom|Gaussian", AdvancedDisplay, meta=(ClampMin = "0.0", UIMax = "128.0", editcondition = "bOverride_Bloom6Size", DisplayName = "#6 Size"))
 	float Bloom6Size;
 
 	/** Bloom1 tint color */
-	UPROPERTY(interp, BlueprintReadWrite, Category="Lens|Bloom", AdvancedDisplay, meta=(editcondition = "bOverride_Bloom1Tint", DisplayName = "#1 Tint", HideAlphaChannel))
+	UPROPERTY(interp, BlueprintReadWrite, Category="Lens|Bloom|Gaussian", AdvancedDisplay, meta=(editcondition = "bOverride_Bloom1Tint", DisplayName = "#1 Tint", HideAlphaChannel))
 	FLinearColor Bloom1Tint;
 	/** Bloom2 tint color */
-	UPROPERTY(interp, BlueprintReadWrite, Category="Lens|Bloom", AdvancedDisplay, meta=(editcondition = "bOverride_Bloom2Tint", DisplayName = "#2 Tint", HideAlphaChannel))
+	UPROPERTY(interp, BlueprintReadWrite, Category="Lens|Bloom|Gaussian", AdvancedDisplay, meta=(editcondition = "bOverride_Bloom2Tint", DisplayName = "#2 Tint", HideAlphaChannel))
 	FLinearColor Bloom2Tint;
 	/** Bloom3 tint color */
-	UPROPERTY(interp, BlueprintReadWrite, Category="Lens|Bloom", AdvancedDisplay, meta=(editcondition = "bOverride_Bloom3Tint", DisplayName = "#3 Tint", HideAlphaChannel))
+	UPROPERTY(interp, BlueprintReadWrite, Category="Lens|Bloom|Gaussian", AdvancedDisplay, meta=(editcondition = "bOverride_Bloom3Tint", DisplayName = "#3 Tint", HideAlphaChannel))
 	FLinearColor Bloom3Tint;
 	/** Bloom4 tint color */
-	UPROPERTY(interp, BlueprintReadWrite, Category="Lens|Bloom", AdvancedDisplay, meta=(editcondition = "bOverride_Bloom4Tint", DisplayName = "#4 Tint", HideAlphaChannel))
+	UPROPERTY(interp, BlueprintReadWrite, Category="Lens|Bloom|Gaussian", AdvancedDisplay, meta=(editcondition = "bOverride_Bloom4Tint", DisplayName = "#4 Tint", HideAlphaChannel))
 	FLinearColor Bloom4Tint;
 	/** Bloom5 tint color */
-	UPROPERTY(interp, BlueprintReadWrite, Category="Lens|Bloom", AdvancedDisplay, meta=(editcondition = "bOverride_Bloom5Tint", DisplayName = "#5 Tint", HideAlphaChannel))
+	UPROPERTY(interp, BlueprintReadWrite, Category="Lens|Bloom|Gaussian", AdvancedDisplay, meta=(editcondition = "bOverride_Bloom5Tint", DisplayName = "#5 Tint", HideAlphaChannel))
 	FLinearColor Bloom5Tint;
 	/** Bloom6 tint color */
-	UPROPERTY(interp, BlueprintReadWrite, Category="Lens|Bloom", AdvancedDisplay, meta=(editcondition = "bOverride_Bloom6Tint", DisplayName = "#6 Tint", HideAlphaChannel))
+	UPROPERTY(interp, BlueprintReadWrite, Category="Lens|Bloom|Gaussian", AdvancedDisplay, meta=(editcondition = "bOverride_Bloom6Tint", DisplayName = "#6 Tint", HideAlphaChannel))
 	FLinearColor Bloom6Tint;
+	
+	/** Multiplier for Convolution bloom contributions >=0: off, 1(default), >1 brighter */
+	UPROPERTY(interp, BlueprintReadWrite, Category="Lens|Bloom|Convolution", meta=(ClampMin = "0.0", UIMax = "1.0", editcondition = "bOverride_BloomConvolutionIntensity", DisplayName = "Intensity"))
+	float BloomConvolutionIntensity;
 
 	/** Intensity multiplier on the scatter dispersion energy of the kernel. 1.0 means exactly use the same energy as the kernel scatter dispersion. */
-	UPROPERTY(interp, BlueprintReadWrite, Category = "Lens|Bloom", meta = (ClampMin = "0.0", UIMax = "20.0", editcondition = "bOverride_BloomConvolutionScatterDispersion", DisplayName = "Convolution Scatter Dispersion"))
+	UPROPERTY(interp, BlueprintReadWrite, Category = "Lens|Bloom|Convolution", meta = (ClampMin = "0.0", UIMax = "20.0", editcondition = "bOverride_BloomConvolutionScatterDispersion", DisplayName = "Scatter Dispersion"))
 	float BloomConvolutionScatterDispersion;
 
 	/** Relative size of the convolution kernel image compared to the minor axis of the viewport  */
-	UPROPERTY(interp, BlueprintReadWrite, Category = "Lens|Bloom", AdvancedDisplay, meta = (ClampMin = "0.0", UIMax = "1.0", editcondition = "bOverride_BloomConvolutionSize", DisplayName = "Convolution Scale"))
+	UPROPERTY(interp, BlueprintReadWrite, Category = "Lens|Bloom|Convolution", AdvancedDisplay, meta = (ClampMin = "0.0", UIMax = "1.0", editcondition = "bOverride_BloomConvolutionSize", DisplayName = "Scale"))
 	float BloomConvolutionSize;
 
 	/** Texture to replace default convolution bloom kernel */
-	UPROPERTY(EditAnywhere, BlueprintReadWrite, Category = "Lens|Bloom", meta = (editcondition = "bOverride_BloomConvolutionTexture", DisplayName = "Convolution Kernel"))
+	UPROPERTY(EditAnywhere, BlueprintReadWrite, Category = "Lens|Bloom|Convolution", meta = (editcondition = "bOverride_BloomConvolutionTexture", DisplayName = "Kernel"))
 	TObjectPtr<class UTexture2D> BloomConvolutionTexture;
 
 	/** The UV location of the center of the kernel.  Should be very close to (.5,.5) */
-	UPROPERTY(interp, BlueprintReadWrite, Category = "Lens|Bloom", AdvancedDisplay, meta = (editcondition = "bOverride_BloomConvolutionCenterUV", DisplayName = "Convolution Center"))
+	UPROPERTY(interp, BlueprintReadWrite, Category = "Lens|Bloom|Convolution", AdvancedDisplay, meta = (editcondition = "bOverride_BloomConvolutionCenterUV", DisplayName = "Center"))
 	FVector2D BloomConvolutionCenterUV;
 
 #if WITH_EDITORONLY_DATA
@@ -1631,19 +1668,19 @@ struct FPostProcessSettings
 #endif
 	
 	/** Boost intensity of select pixels  prior to computing bloom convolution (Min, Max, Multiplier).  Max < Min disables */
-	UPROPERTY(interp, BlueprintReadWrite, Category = "Lens|Bloom", AdvancedDisplay, meta = (editcondition = "bOverride_BloomConvolutionPreFilterMin", DisplayName = "Convolution Boost Min"))
+	UPROPERTY(interp, BlueprintReadWrite, Category = "Lens|Bloom|Convolution", AdvancedDisplay, meta = (editcondition = "bOverride_BloomConvolutionPreFilterMin", DisplayName = "Boost Min"))
 	float BloomConvolutionPreFilterMin;
 
 	/** Boost intensity of select pixels  prior to computing bloom convolution (Min, Max, Multiplier).  Max < Min disables */
-	UPROPERTY(interp, BlueprintReadWrite, Category = "Lens|Bloom", AdvancedDisplay, meta = (editcondition = "bOverride_BloomConvolutionPreFilterMax", DisplayName = "Convolution Boost Max"))
+	UPROPERTY(interp, BlueprintReadWrite, Category = "Lens|Bloom|Convolution", AdvancedDisplay, meta = (editcondition = "bOverride_BloomConvolutionPreFilterMax", DisplayName = "Boost Max"))
 	float BloomConvolutionPreFilterMax;
 
 	/** Boost intensity of select pixels  prior to computing bloom convolution (Min, Max, Multiplier).  Max < Min disables */
-	UPROPERTY(interp, BlueprintReadWrite, Category = "Lens|Bloom", AdvancedDisplay, meta = (editcondition = "bOverride_BloomConvolutionPreFilterMult", DisplayName = "Convolution Boost Mult"))
+	UPROPERTY(interp, BlueprintReadWrite, Category = "Lens|Bloom|Convolution", AdvancedDisplay, meta = (editcondition = "bOverride_BloomConvolutionPreFilterMult", DisplayName = "Boost Mult"))
 	float BloomConvolutionPreFilterMult;
 
 	/** Implicit buffer region as a fraction of the screen size to insure the bloom does not wrap across the screen.  Larger sizes have perf impact.*/
-	UPROPERTY(interp, BlueprintReadWrite, Category = "Lens|Bloom", AdvancedDisplay, meta = (ClampMin = "0.0", UIMax = "1.0", editcondition = "bOverride_BloomConvolutionBufferScale", DisplayName = "Convolution Buffer"))
+	UPROPERTY(interp, BlueprintReadWrite, Category = "Lens|Bloom|Convolution", AdvancedDisplay, meta = (ClampMin = "0.0", UIMax = "1.0", editcondition = "bOverride_BloomConvolutionBufferScale", DisplayName = "Buffer"))
 	float BloomConvolutionBufferScale;
 	
 	/**
diff --git a/Engine/Source/Runtime/Engine/Private/PostProcessUtils.cpp b/Engine/Source/Runtime/Engine/Private/PostProcessUtils.cpp
--- a/Engine/Source/Runtime/Engine/Private/PostProcessUtils.cpp
+++ b/Engine/Source/Runtime/Engine/Private/PostProcessUtils.cpp
@@ -63,6 +63,7 @@ bool FPostProcessUtils::OverridePostProcessSettings(FPostProcessSettings& ThisFr
 		UE_SET_PP(SceneFringeIntensity);
 		UE_SET_PP(ChromaticAberrationStartOffset);
 		UE_SET_PP(BloomIntensity);
+		UE_SET_PP(BloomGaussianIntensity);
 		UE_SET_PP(BloomThreshold);
 		UE_SET_PP(Bloom1Tint);
 		UE_SET_PP(BloomSizeScale);
@@ -79,6 +80,7 @@ bool FPostProcessUtils::OverridePostProcessSettings(FPostProcessSettings& ThisFr
 		UE_SET_PP(Bloom6Size);
 		UE_SET_PP(BloomDirtMaskIntensity);
 		UE_SET_PP(BloomDirtMaskTint);
+		UE_SET_PP(BloomConvolutionIntensity);
 		UE_SET_PP(BloomConvolutionScatterDispersion);
 		UE_SET_PP(BloomConvolutionSize);
 		UE_SET_PP(BloomConvolutionCenterUV);
@@ -366,6 +368,7 @@ bool FPostProcessUtils::BlendPostProcessSettings(FPostProcessSettings& ThisFrom,
 		UE_LERP_PP(SceneFringeIntensity);
 		UE_LERP_PP(ChromaticAberrationStartOffset);
 		UE_LERP_PP(BloomIntensity);
+		UE_LERP_PP(BloomGaussianIntensity);
 		UE_LERP_PP(BloomThreshold);
 		UE_LERP_PP(Bloom1Tint);
 		UE_LERP_PP(BloomSizeScale);
@@ -382,6 +385,7 @@ bool FPostProcessUtils::BlendPostProcessSettings(FPostProcessSettings& ThisFrom,
 		UE_LERP_PP(Bloom6Size);
 		UE_LERP_PP(BloomDirtMaskIntensity);
 		UE_LERP_PP(BloomDirtMaskTint);
+		UE_LERP_PP(BloomConvolutionIntensity);
 		UE_LERP_PP(BloomConvolutionScatterDispersion);
 		UE_LERP_PP(BloomConvolutionSize);
 		UE_LERP_PP(BloomConvolutionCenterUV);
diff --git a/Engine/Source/Runtime/Engine/Private/PostProcessVolume.cpp b/Engine/Source/Runtime/Engine/Private/PostProcessVolume.cpp
--- a/Engine/Source/Runtime/Engine/Private/PostProcessVolume.cpp
+++ b/Engine/Source/Runtime/Engine/Private/PostProcessVolume.cpp
@@ -193,38 +193,6 @@ bool APostProcessVolume::CanEditChange(const FProperty* InProperty) const
 			{
 				return Settings.AutoExposureMethod == EAutoExposureMethod::AEM_Histogram;
 			}
-
-			// Parameters that are only used for the Sum of Gaussian bloom / not the texture based fft bloom
-			if (PropertyName == GET_MEMBER_NAME_STRING_CHECKED(FPostProcessSettings, BloomThreshold) ||
-				PropertyName == GET_MEMBER_NAME_STRING_CHECKED(FPostProcessSettings, BloomSizeScale) ||
-				PropertyName == GET_MEMBER_NAME_STRING_CHECKED(FPostProcessSettings, Bloom1Size) ||
-				PropertyName == GET_MEMBER_NAME_STRING_CHECKED(FPostProcessSettings, Bloom2Size) ||
-				PropertyName == GET_MEMBER_NAME_STRING_CHECKED(FPostProcessSettings, Bloom3Size) ||
-				PropertyName == GET_MEMBER_NAME_STRING_CHECKED(FPostProcessSettings, Bloom4Size) ||
-				PropertyName == GET_MEMBER_NAME_STRING_CHECKED(FPostProcessSettings, Bloom5Size) ||
-				PropertyName == GET_MEMBER_NAME_STRING_CHECKED(FPostProcessSettings, Bloom6Size) ||
-				PropertyName == GET_MEMBER_NAME_STRING_CHECKED(FPostProcessSettings, Bloom1Tint) ||
-				PropertyName == GET_MEMBER_NAME_STRING_CHECKED(FPostProcessSettings, Bloom2Tint) ||
-				PropertyName == GET_MEMBER_NAME_STRING_CHECKED(FPostProcessSettings, Bloom3Tint) ||
-				PropertyName == GET_MEMBER_NAME_STRING_CHECKED(FPostProcessSettings, Bloom4Tint) ||
-				PropertyName == GET_MEMBER_NAME_STRING_CHECKED(FPostProcessSettings, Bloom5Tint) ||
-				PropertyName == GET_MEMBER_NAME_STRING_CHECKED(FPostProcessSettings, Bloom6Tint))
-			{
-				return (Settings.BloomMethod == EBloomMethod::BM_SOG);
-			}
-
-			// Parameters that are only of use with the bloom texture based fft
-			if (PropertyName == GET_MEMBER_NAME_STRING_CHECKED(FPostProcessSettings, BloomConvolutionTexture)      ||
-				PropertyName == GET_MEMBER_NAME_STRING_CHECKED(FPostProcessSettings, BloomConvolutionSize)         ||
-				PropertyName == GET_MEMBER_NAME_STRING_CHECKED(FPostProcessSettings, BloomConvolutionScatterDispersion) ||
-				PropertyName == GET_MEMBER_NAME_STRING_CHECKED(FPostProcessSettings, BloomConvolutionCenterUV)     ||
-				PropertyName == GET_MEMBER_NAME_STRING_CHECKED(FPostProcessSettings, BloomConvolutionPreFilterMin) ||
-				PropertyName == GET_MEMBER_NAME_STRING_CHECKED(FPostProcessSettings, BloomConvolutionPreFilterMax) ||
-				PropertyName == GET_MEMBER_NAME_STRING_CHECKED(FPostProcessSettings, BloomConvolutionPreFilterMult)||
-				PropertyName == GET_MEMBER_NAME_STRING_CHECKED(FPostProcessSettings, BloomConvolutionBufferScale))
-			{
-				return (Settings.BloomMethod == EBloomMethod::BM_FFT);
-			}
 			
 			if (PropertyName == GET_MEMBER_NAME_STRING_CHECKED(FPostProcessSettings, LumenRayLightingMode))
 			{
diff --git a/Engine/Source/Runtime/Engine/Private/Scene.cpp b/Engine/Source/Runtime/Engine/Private/Scene.cpp
--- a/Engine/Source/Runtime/Engine/Private/Scene.cpp
+++ b/Engine/Source/Runtime/Engine/Private/Scene.cpp
@@ -89,7 +89,7 @@ void FFilmStockSettings::ExportToPostProcessSettings(FPostProcessSettings* OutPo
 
 void FGaussianSumBloomSettings::ExportToPostProcessSettings(FPostProcessSettings* OutPostProcessSettings) const
 {
-	OutPostProcessSettings->bOverride_BloomIntensity = true;
+	OutPostProcessSettings->bOverride_BloomGaussianIntensity = true;
 	OutPostProcessSettings->bOverride_BloomThreshold = true;
 	OutPostProcessSettings->bOverride_BloomSizeScale = true;
 	OutPostProcessSettings->bOverride_Bloom1Tint = true;
@@ -105,7 +105,7 @@ void FGaussianSumBloomSettings::ExportToPostProcessSettings(FPostProcessSettings
 	OutPostProcessSettings->bOverride_Bloom6Tint = true;
 	OutPostProcessSettings->bOverride_Bloom6Size = true;
 
-	OutPostProcessSettings->BloomIntensity = Intensity;
+	OutPostProcessSettings->BloomGaussianIntensity = Intensity;
 	OutPostProcessSettings->BloomThreshold = Threshold;
 	OutPostProcessSettings->BloomSizeScale = SizeScale;
 	OutPostProcessSettings->Bloom1Tint = Filter1Tint;
@@ -124,6 +124,7 @@ void FGaussianSumBloomSettings::ExportToPostProcessSettings(FPostProcessSettings
 
 void FConvolutionBloomSettings::ExportToPostProcessSettings(FPostProcessSettings* OutPostProcessSettings) const
 {
+	OutPostProcessSettings->bOverride_BloomConvolutionIntensity = true;
 	OutPostProcessSettings->bOverride_BloomConvolutionTexture = true;
 	OutPostProcessSettings->bOverride_BloomConvolutionSize = true;
 	OutPostProcessSettings->bOverride_BloomConvolutionScatterDispersion = true;
@@ -133,6 +134,7 @@ void FConvolutionBloomSettings::ExportToPostProcessSettings(FPostProcessSettings
 	OutPostProcessSettings->bOverride_BloomConvolutionPreFilterMult = true;
 	OutPostProcessSettings->bOverride_BloomConvolutionBufferScale = true;
 
+	OutPostProcessSettings->BloomConvolutionIntensity = Intensity;
 	OutPostProcessSettings->BloomConvolutionTexture = Texture;
 	OutPostProcessSettings->BloomConvolutionSize = Size;
 	OutPostProcessSettings->BloomConvolutionScatterDispersion = ScatterDispersion;
@@ -145,6 +147,9 @@ void FConvolutionBloomSettings::ExportToPostProcessSettings(FPostProcessSettings
 
 void FLensBloomSettings::ExportToPostProcessSettings(FPostProcessSettings* OutPostProcessSettings) const
 {
+	OutPostProcessSettings->bOverride_BloomIntensity = true;
+	OutPostProcessSettings->BloomIntensity = Intensity;
+	
 	GaussianSum.ExportToPostProcessSettings(OutPostProcessSettings);
 	Convolution.ExportToPostProcessSettings(OutPostProcessSettings);
 
@@ -152,6 +157,25 @@ void FLensBloomSettings::ExportToPostProcessSettings(FPostProcessSettings* OutPo
 	OutPostProcessSettings->BloomMethod = Method;
 }
 
+bool FLensBloomSettings::Serialize(FArchive& Ar)
+{
+	Ar.UsingCustomVersion(FFortniteMainBranchObjectVersion::GUID);
+	return false;
+}
+
+void FLensBloomSettings::PostSerialize(const FArchive& Ar)
+{
+	if (Ar.IsLoading())
+	{
+		if (Ar.CustomVer(FFortniteMainBranchObjectVersion::GUID) < FFortniteMainBranchObjectVersion::SpecializeBloomIntensity)
+		{
+			Intensity = GaussianSum.Intensity;
+			GaussianSum.Intensity = 1.0f;
+			Convolution.Intensity = 1.0f;
+		}
+	}
+}
+
 void FLensImperfectionSettings::ExportToPostProcessSettings(FPostProcessSettings* OutPostProcessSettings) const
 {
 	OutPostProcessSettings->bOverride_BloomDirtMask = true;
@@ -246,20 +270,26 @@ void FCameraExposureSettings::ExportToPostProcessSettings(FPostProcessSettings*
 static void VerifyPostProcessingProperties(
 	const FString& PropertyPrefix,
 	const TArray<const UStruct*>& NewStructs,
+	const TArray<FString>& NewStructPrefixes,
 	const TMap<FString, FString>& RenameMap)
 {
+	check(NewStructs.Num() == NewStructPrefixes.Num());
 	const UStruct* LegacyStruct = FPostProcessSettings::StaticStruct();
 
 	TMap<FString, const FProperty*> NewPropertySet;
 
 	// Walk new struct and build list of property name.
-	for (const UStruct* NewStruct : NewStructs)
+	for (int Index = 0; Index < NewStructs.Num(); ++Index)
 	{
+		const UStruct* NewStruct = NewStructs[Index];
+		const FString& NewStructPrefix = NewStructPrefixes[Index];
+
 		for (FProperty* Property = NewStruct->PropertyLink; Property; Property = Property->PropertyLinkNext)
 		{
 			// Make sure there is no duplicate.
-			check(!NewPropertySet.Contains(Property->GetNameCPP()));
-			NewPropertySet.Add(Property->GetNameCPP(), Property);
+			FString Name = NewStructPrefix + Property->GetNameCPP();
+			check(!NewPropertySet.Contains(Name));
+			NewPropertySet.Add(Name, Property);
 		}
 	}
 
@@ -301,40 +331,33 @@ static void DoPostProcessSettingsSanityCheck()
 {
 	{
 		TMap<FString, FString> RenameMap;
-		RenameMap.Add(TEXT("Bloom1Size"), TEXT("Filter1Size"));
-		RenameMap.Add(TEXT("Bloom2Size"), TEXT("Filter2Size"));
-		RenameMap.Add(TEXT("Bloom3Size"), TEXT("Filter3Size"));
-		RenameMap.Add(TEXT("Bloom4Size"), TEXT("Filter4Size"));
-		RenameMap.Add(TEXT("Bloom5Size"), TEXT("Filter5Size"));
-		RenameMap.Add(TEXT("Bloom6Size"), TEXT("Filter6Size"));
-		RenameMap.Add(TEXT("Bloom1Tint"), TEXT("Filter1Tint"));
-		RenameMap.Add(TEXT("Bloom2Tint"), TEXT("Filter2Tint"));
-		RenameMap.Add(TEXT("Bloom3Tint"), TEXT("Filter3Tint"));
-		RenameMap.Add(TEXT("Bloom4Tint"), TEXT("Filter4Tint"));
-		RenameMap.Add(TEXT("Bloom5Tint"), TEXT("Filter5Tint"));
-		RenameMap.Add(TEXT("Bloom6Tint"), TEXT("Filter6Tint"));
-
-		RenameMap.Add(TEXT("BloomConvolutionTexture"), TEXT(""));
-		RenameMap.Add(TEXT("BloomConvolutionSize"), TEXT(""));
-		RenameMap.Add(TEXT("BloomConvolutionScatterDispersion"), TEXT(""));
-		RenameMap.Add(TEXT("BloomConvolutionCenterUV"), TEXT(""));
-		RenameMap.Add(TEXT("BloomConvolutionPreFilterMin"), TEXT(""));
-		RenameMap.Add(TEXT("BloomConvolutionPreFilterMax"), TEXT(""));
-		RenameMap.Add(TEXT("BloomConvolutionPreFilterMult"), TEXT(""));
-		RenameMap.Add(TEXT("BloomConvolutionBufferScale"), TEXT(""));
+		RenameMap.Add(TEXT("BloomThreshold"), TEXT("GaussianThreshold"));
+		RenameMap.Add(TEXT("BloomSizeScale"), TEXT("GaussianSizeScale"));
+		RenameMap.Add(TEXT("Bloom1Size"), TEXT("GaussianFilter1Size"));
+		RenameMap.Add(TEXT("Bloom2Size"), TEXT("GaussianFilter2Size"));
+		RenameMap.Add(TEXT("Bloom3Size"), TEXT("GaussianFilter3Size"));
+		RenameMap.Add(TEXT("Bloom4Size"), TEXT("GaussianFilter4Size"));
+		RenameMap.Add(TEXT("Bloom5Size"), TEXT("GaussianFilter5Size"));
+		RenameMap.Add(TEXT("Bloom6Size"), TEXT("GaussianFilter6Size"));
+		RenameMap.Add(TEXT("Bloom1Tint"), TEXT("GaussianFilter1Tint"));
+		RenameMap.Add(TEXT("Bloom2Tint"), TEXT("GaussianFilter2Tint"));
+		RenameMap.Add(TEXT("Bloom3Tint"), TEXT("GaussianFilter3Tint"));
+		RenameMap.Add(TEXT("Bloom4Tint"), TEXT("GaussianFilter4Tint"));
+		RenameMap.Add(TEXT("Bloom5Tint"), TEXT("GaussianFilter5Tint"));
+		RenameMap.Add(TEXT("Bloom6Tint"), TEXT("GaussianFilter6Tint"));
 
 		VerifyPostProcessingProperties(TEXT("Bloom"),
 			TArray<const UStruct*>({
 				FGaussianSumBloomSettings::StaticStruct(),
+				FConvolutionBloomSettings::StaticStruct(),
 				FLensBloomSettings::StaticStruct(),
 				FLensImperfectionSettings::StaticStruct()}),
-			RenameMap);
-	}
-	
-	{
-		TMap<FString, FString> RenameMap;
-		VerifyPostProcessingProperties(TEXT("BloomConvolution"),
-			TArray<const UStruct*>({FConvolutionBloomSettings::StaticStruct()}),
+			TArray<FString>({
+				TEXT("Gaussian"),
+				TEXT("Convolution"),
+				TEXT(""),
+				TEXT(""),
+			}),
 			RenameMap);
 	}
 
@@ -343,6 +366,9 @@ static void DoPostProcessSettingsSanityCheck()
 		VerifyPostProcessingProperties(TEXT("Exposure"),
 			TArray<const UStruct*>({
 				FCameraExposureSettings::StaticStruct()}),
+			TArray<FString>({
+				TEXT(""),
+			}),
 			RenameMap);
 	}
 
@@ -360,6 +386,9 @@ static void DoPostProcessSettingsSanityCheck()
 		RenameMap.Add(TEXT("FilmGrainTexture"), TEXT(""));
 		VerifyPostProcessingProperties(TEXT("Film"),
 			TArray<const UStruct*>({FFilmStockSettings::StaticStruct()}),
+			TArray<FString>({
+				TEXT(""),
+			}),
 			RenameMap);
 	}
 }
@@ -422,6 +451,7 @@ FPostProcessSettings::FPostProcessSettings()
 	BloomMethod = BM_SOG;
 	// next value might get overwritten by r.DefaultFeature.Bloom
 	BloomIntensity = 0.675f;
+	BloomGaussianIntensity = 1.0f;
 	BloomThreshold = -1.0f;
 	// default is 4 to maintain old settings after fixing something that caused a factor of 4
 	BloomSizeScale = 4.0;
@@ -437,6 +467,7 @@ FPostProcessSettings::FPostProcessSettings()
 	Bloom5Size = 30.0f;
 	Bloom6Tint = FLinearColor(0.061f, 0.061f, 0.061f);
 	Bloom6Size = 64.0f;
+	BloomConvolutionIntensity = 1.0f;
 	BloomConvolutionScatterDispersion = 1.f;
 	BloomConvolutionSize = 1.f;
 	BloomConvolutionCenterUV = FVector2D(0.5f, 0.5f);
diff --git a/Engine/Source/Runtime/Engine/Private/SceneView.cpp b/Engine/Source/Runtime/Engine/Private/SceneView.cpp
--- a/Engine/Source/Runtime/Engine/Private/SceneView.cpp
+++ b/Engine/Source/Runtime/Engine/Private/SceneView.cpp
@@ -1628,6 +1628,7 @@ void FSceneView::OverridePostProcessSettings(const FPostProcessSettings& Src, fl
 		LERP_PP(SceneFringeIntensity);
 		LERP_PP(ChromaticAberrationStartOffset);
 		LERP_PP(BloomIntensity);
+		LERP_PP(BloomGaussianIntensity);
 		LERP_PP(BloomThreshold);
 		LERP_PP(Bloom1Tint);
 		LERP_PP(BloomSizeScale);
@@ -1644,6 +1645,7 @@ void FSceneView::OverridePostProcessSettings(const FPostProcessSettings& Src, fl
 		LERP_PP(Bloom6Size);
 		LERP_PP(BloomDirtMaskIntensity);
 		LERP_PP(BloomDirtMaskTint);
+		LERP_PP(BloomConvolutionIntensity);
 		LERP_PP(BloomConvolutionScatterDispersion);
 		LERP_PP(BloomConvolutionSize);
 		LERP_PP(BloomConvolutionCenterUV);
diff --git a/Engine/Source/Runtime/Engine/Public/SceneView.h b/Engine/Source/Runtime/Engine/Public/SceneView.h
--- a/Engine/Source/Runtime/Engine/Public/SceneView.h
+++ b/Engine/Source/Runtime/Engine/Public/SceneView.h
@@ -2,6 +2,7 @@
 
 #pragma once
 
+#include "Templates/SharedPointer.h"
 #include "ConvexVolume.h"
 #include "DebugViewModeHelpers.h"
 #include "DynamicRenderScaling.h"
@@ -31,6 +32,10 @@ class FVolumetricFogViewResources;
 class ISceneRenderer;
 class ISpatialUpscaler;
 struct FMinimalViewInfo;
+namespace UE::RenderViewportFeedback
+{
+	class FCollector;
+}
 
 namespace UE::Renderer::Private
 {
@@ -1754,6 +1759,9 @@ public:
 	TStaticArray<FLinearColor, 6> AdditionalSelectionOutlineColors;
 	/** True if any components are selected in isolation (independent of actor selection) */
 	bool bHasSelectedComponents;
+
+	/** In editor, this is used to feed rendering info back to the viewport. May be null. */
+	TSharedPtr<UE::RenderViewportFeedback::FCollector> ViewportFeedback;
 #endif
 
 	/**
diff --git a/Engine/Source/Runtime/RenderCore/Internal/RenderViewportFeedback.h b/Engine/Source/Runtime/RenderCore/Internal/RenderViewportFeedback.h
new file mode 100644
--- /dev/null
+++ b/Engine/Source/Runtime/RenderCore/Internal/RenderViewportFeedback.h
@@ -0,0 +1,75 @@
+// Copyright Epic Games, Inc. All Rights Reserved.
+
+#pragma once
+
+#if WITH_EDITOR
+
+#include "Async/Mutex.h"
+#include "Templates/SharedPointer.h"
+#include "Templates/UniquePtr.h"
+#include "Engine/Scene.h"
+
+/** 
+  * Feedback from the renderer to the viewport, gathered while rendering. 
+  * Used to communicate lightweight renderer state to the user in the editor.
+  */
+class FRenderViewportFeedback
+{
+public:
+	// Bloom method that was used during rendering.
+	EBloomMethod BloomMethod;
+};
+
+namespace UE::RenderViewportFeedback
+{
+
+class FCollector;
+
+/** 
+  * Datastructure that owns FRenderViewportFeedback on the gamethread side
+  * and facilitates threadsafe updating from the renderer.
+  */
+class FReceiver : public TSharedFromThis<FReceiver, ESPMode::ThreadSafe>
+{
+	friend class FCollector;
+
+public:
+	RENDERCORE_API FReceiver();
+
+	// Call from game thread only
+	// Returns a copy, as the internal raw data may be updated at any time.
+	RENDERCORE_API FRenderViewportFeedback Data() const;
+
+	RENDERCORE_API TSharedPtr<FCollector> MakeCollector();
+
+private:
+	FRenderViewportFeedback InternalData;
+	mutable UE::FMutex Mutex;
+};
+
+/** Helper class used on the render thread to gather data for the viewport. */
+class FCollector
+{
+	friend class FReceiver;
+
+public:
+	FCollector(FReceiver& InReceiver);
+
+	// Modify from render thread only
+	FRenderViewportFeedback& Data()
+	{
+		return InternalData;
+	}
+	
+	// Data collection finished, push the data to the receiver.
+	RENDERCORE_API void EndFrameRenderThread();
+
+private:
+	// Weak pointer for in case the viewport is destroyed while the renderer is still active.
+	TWeakPtr<FReceiver> Receiver;
+	FRenderViewportFeedback InternalData;
+};
+
+};
+
+#endif //WITH_EDITOR
diff --git a/Engine/Source/Runtime/RenderCore/Private/RenderViewportFeedback.cpp b/Engine/Source/Runtime/RenderCore/Private/RenderViewportFeedback.cpp
new file mode 100644
--- /dev/null
+++ b/Engine/Source/Runtime/RenderCore/Private/RenderViewportFeedback.cpp
@@ -0,0 +1,44 @@
+// Copyright Epic Games, Inc. All Rights Reserved.
+
+#include "RenderViewportFeedback.h"
+#include "Misc/ScopeLock.h"
+
+#if WITH_EDITOR
+
+namespace UE::RenderViewportFeedback
+{
+
+FReceiver::FReceiver()
+: InternalData()
+, Mutex()
+{ }
+
+FRenderViewportFeedback FReceiver::Data() const
+{
+	UE::TScopeLock Lock(Mutex);
+	return InternalData;
+}
+
+TSharedPtr<FCollector> FReceiver::MakeCollector()
+{
+	return MakeShared<FCollector>(*this);
+}
+
+FCollector::FCollector(FReceiver& InReceiver)
+: Receiver(InReceiver.AsWeak())
+, InternalData()
+{ }
+
+void FCollector::EndFrameRenderThread()
+{
+	TSharedPtr<FReceiver> Feedback = Receiver.Pin();
+	if (Feedback)
+	{
+		UE::TScopeLock Lock(Feedback->Mutex);
+		Feedback->InternalData = MoveTemp(InternalData);
+	}
+}
+
+};
+
+#endif //WITH_EDITOR
diff --git a/Engine/Source/Runtime/RenderCore/RenderCore.Build.cs b/Engine/Source/Runtime/RenderCore/RenderCore.Build.cs
--- a/Engine/Source/Runtime/RenderCore/RenderCore.Build.cs
+++ b/Engine/Source/Runtime/RenderCore/RenderCore.Build.cs
@@ -38,7 +38,7 @@ public class RenderCore : ModuleRules
 
 		PrivateDependencyModuleNames.AddRange(new string[] { "Core", "Projects", "ApplicationCore", "TraceLog", "CookOnTheFly" });
 		
-		PublicIncludePathModuleNames.AddRange(new string[] { "RHI" });
+		PublicIncludePathModuleNames.AddRange(new string[] { "Engine", "RHI" });
 
 		if (Target.bBuildEditor == true)
 		{
diff --git a/Engine/Source/Runtime/Renderer/Private/PostProcess/PostProcessBloomSetup.cpp b/Engine/Source/Runtime/Renderer/Private/PostProcess/PostProcessBloomSetup.cpp
--- a/Engine/Source/Runtime/Renderer/Private/PostProcess/PostProcessBloomSetup.cpp
+++ b/Engine/Source/Runtime/Renderer/Private/PostProcess/PostProcessBloomSetup.cpp
@@ -250,7 +250,8 @@ FScreenPassTexture AddGaussianBloomPasses(FRDGBuilder& GraphBuilder, const FView
 		// Use bloom quality to select the number of downsample stages to use for bloom.
 		const uint32 BloomStageCount = BloomQualityToSceneDownsampleStage[BloomQualityIndex];
 
-		const float TintScale = (1.0f / BloomQualityCountMax) * Settings.BloomIntensity;
+		const float BloomIntensity = Settings.BloomIntensity * Settings.BloomGaussianIntensity;
+		const float TintScale = (1.0f / BloomQualityCountMax) * BloomIntensity;
 
 		for (uint32 StageIndex = 0, SourceIndex = BloomQualityCountMax - 1; StageIndex < BloomStageCount; ++StageIndex, --SourceIndex)
 		{
diff --git a/Engine/Source/Runtime/Renderer/Private/PostProcess/PostProcessFFTBloom.cpp b/Engine/Source/Runtime/Renderer/Private/PostProcess/PostProcessFFTBloom.cpp
--- a/Engine/Source/Runtime/Renderer/Private/PostProcess/PostProcessFFTBloom.cpp
+++ b/Engine/Source/Runtime/Renderer/Private/PostProcess/PostProcessFFTBloom.cpp
@@ -235,8 +235,6 @@ bool IsFFTBloomEnabled(const FViewInfo& View)
 {
 	const bool bUseFFTBloom = View.FinalPostProcessSettings.BloomMethod == EBloomMethod::BM_FFT && View.ViewState != nullptr && DoesPlatformSupportFFTBloom(View.GetShaderPlatform());
 
-	static bool bWarnAboutOldMetalFFTOnce = false;
-
 	if (bUseFFTBloom)
 	{
 		return View.FFTBloomKernelTexture != nullptr;
@@ -873,8 +871,10 @@ FFFTBloomOutput AddFFTBloomPass(
 			FRDGBufferDesc::CreateStructuredDesc(sizeof(FLinearColor), /* NumElements = */ 1),
 			TEXT("Bloom.FFT.FFTMulitplyParameters"));
 
+		const float BloomIntensity = View.FinalPostProcessSettings.BloomIntensity * View.FinalPostProcessSettings.BloomConvolutionIntensity;
+
 		FBloomFinalizeApplyConstantsCS::FParameters* PassParameters = GraphBuilder.AllocParameters<FBloomFinalizeApplyConstantsCS::FParameters>();
-		PassParameters->ScatterDispersionIntensity = View.FinalPostProcessSettings.BloomConvolutionScatterDispersion * View.FinalPostProcessSettings.BloomIntensity;
+		PassParameters->ScatterDispersionIntensity = View.FinalPostProcessSettings.BloomConvolutionScatterDispersion * BloomIntensity;
 		PassParameters->KernelConstantsBuffer = GraphBuilder.CreateSRV(KernelConstantsBuffer);
 		PassParameters->SceneColorApplyOutput = GraphBuilder.CreateUAV(BloomOutput.SceneColorApplyParameters);
 		PassParameters->FFTMulitplyOutput = GraphBuilder.CreateUAV(FFTMulitplyParameters);
diff --git a/Engine/Source/Runtime/Renderer/Private/PostProcess/PostProcessMobile.cpp b/Engine/Source/Runtime/Renderer/Private/PostProcess/PostProcessMobile.cpp
--- a/Engine/Source/Runtime/Renderer/Private/PostProcess/PostProcessMobile.cpp
+++ b/Engine/Source/Runtime/Renderer/Private/PostProcess/PostProcessMobile.cpp
@@ -1167,7 +1167,8 @@ FScreenPassTexture AddMobileSunMergePass(FRDGBuilder& GraphBuilder, const FViewI
 	SunColorVignetteIntensityParam.Z = MobileLightShaft.ColorApply.B;
 	SunColorVignetteIntensityParam.W = Settings.VignetteIntensity;
 
-	FLinearColor BloomColor = Settings.Bloom1Tint * Settings.BloomIntensity * 0.5;
+	const float BloomGaussianIntensity = Settings.BloomIntensity * Settings.BloomGaussianIntensity;
+	FLinearColor BloomColor = Settings.Bloom1Tint * BloomGaussianIntensity * 0.5;
 
 	FRHITexture* BloomDirtMaskTexture = GBlackTexture->TextureRHI;
 
diff --git a/Engine/Source/Runtime/Renderer/Private/PostProcess/PostProcessing.cpp b/Engine/Source/Runtime/Renderer/Private/PostProcess/PostProcessing.cpp
--- a/Engine/Source/Runtime/Renderer/Private/PostProcess/PostProcessing.cpp
+++ b/Engine/Source/Runtime/Renderer/Private/PostProcess/PostProcessing.cpp
@@ -82,6 +82,7 @@
 #include "GPUSkinCache.h"
 #include "VT/VirtualTextureFeedbackResource.h"
 #include "VT/VirtualTextureVisualizationData.h"
+#include "RenderViewportFeedback.h"
 
 bool IsMobileEyeAdaptationEnabled(const FViewInfo& View);
 
@@ -725,8 +726,6 @@ void AddPostProcessingPasses(
 			View.FinalPostProcessSettings.LocalExposureShadowContrastCurve ||
 			!FMath::IsNearlyEqual(View.FinalPostProcessSettings.LocalExposureDetailStrength, 1.0f);
 
-		const bool bBloomEnabled = View.FinalPostProcessSettings.BloomIntensity > 0.0f && !bVisualizeTSR;
-
 		// Whether separate translucency is composed in TSR.
 		bool bComposeSeparateTranslucencyInTSR = PostDOFTranslucencyResources.IsValid() && TAAConfig == EMainTAAPassConfig::TSR && ComposeSeparateTranslucencyInTSR(View);
 
@@ -775,8 +774,19 @@ void AddPostProcessingPasses(
 
 		PassSequence.Finalize();
 
+		bool bFFTBloomEnabled = IsFFTBloomEnabled(View);
+#if WITH_EDITOR
+		if (View.ViewportFeedback)
+		{
+			View.ViewportFeedback->Data().BloomMethod = bFFTBloomEnabled ? EBloomMethod::BM_FFT : EBloomMethod::BM_SOG;
+		}
+#endif
+		const bool bBloomEnabled = !bVisualizeTSR
+			&& View.FinalPostProcessSettings.BloomIntensity > 0.0f
+			&& (bFFTBloomEnabled ? View.FinalPostProcessSettings.BloomConvolutionIntensity : View.FinalPostProcessSettings.BloomGaussianIntensity) > 0.0f;
+		bFFTBloomEnabled = bFFTBloomEnabled && bBloomEnabled;
+
 		const bool bLensFlareEnabled = bBloomEnabled && IsLensFlaresEnabled(View);
-		const bool bFFTBloomEnabled = bBloomEnabled && IsFFTBloomEnabled(View);
 
 		const bool bBasicEyeAdaptationEnabled = bEyeAdaptationEnabled && (AutoExposureMethod == EAutoExposureMethod::AEM_Basic);
 		const bool bLocalExposureBlurredLum = bLocalExposureEnabled && View.FinalPostProcessSettings.LocalExposureMethod == ELocalExposureMethod::Bilateral && View.FinalPostProcessSettings.LocalExposureBlurredLuminanceBlend > 0.0f;
@@ -2902,6 +2912,7 @@ void AddMobilePostProcessingPasses(FRDGBuilder& GraphBuilder, FScene* Scene, con
 			};
 
 			const float BloomUpScale = 0.66f * 2.0f;
+			const float BloomGaussianIntensity = Settings.BloomIntensity * Settings.BloomGaussianIntensity;
 
 			const FLinearColor Tint[]{ Settings.Bloom1Tint, Settings.Bloom2Tint, Settings.Bloom3Tint, Settings.Bloom4Tint, Settings.Bloom5Tint, Settings.Bloom6Tint };
 
@@ -2911,8 +2922,8 @@ void AddMobilePostProcessingPasses(FRDGBuilder& GraphBuilder, FScene* Scene, con
 				const int IndexB = NumDownsamplePasses - 1;
 				FVector4f TintA{ Tint[IndexA], 0.0f };
 				FVector4f TintB{ Tint[IndexB], 0.0f };
-				TintA *= Settings.BloomIntensity;
-				TintB *= Settings.BloomIntensity;
+				TintA *= BloomGaussianIntensity;
+				TintB *= BloomGaussianIntensity;
 
 				BloomUpOutputs = AddBloomUpPass(PostProcessDownsample_Bloom[IndexA], PostProcessDownsample_Bloom[IndexB], BloomUpScale, TintA, TintB);
 			}
@@ -2921,7 +2932,7 @@ void AddMobilePostProcessingPasses(FRDGBuilder& GraphBuilder, FScene* Scene, con
 			{
 				// Upsample by 2
 				FVector4f TintA{ Tint[Index], 0.0f };
-				TintA *= Settings.BloomIntensity;
+				TintA *= BloomGaussianIntensity;
 				FVector4f TintB = FVector4f(1.0f, 1.0f, 1.0f, 0.0f);
 
 				BloomUpOutputs = AddBloomUpPass(PostProcessDownsample_Bloom[Index], BloomUpOutputs, BloomUpScale, TintA, TintB);
@@ -2930,7 +2941,7 @@ void AddMobilePostProcessingPasses(FRDGBuilder& GraphBuilder, FScene* Scene, con
 			// Upsample by 2
 			{
 				FVector4f TintA = FVector4f(Settings.Bloom2Tint.R, Settings.Bloom2Tint.G, Settings.Bloom2Tint.B, 0.0f);
-				TintA *= Settings.BloomIntensity;
+				TintA *= BloomGaussianIntensity;
 				// Scaling Bloom2 by extra factor to match filter area difference between PC default and mobile.
 				TintA *= 0.5;
 				FVector4f TintB = FVector4f(1.0f, 1.0f, 1.0f, 0.0f);
diff --git a/Engine/Source/Runtime/Renderer/Private/SceneRendering.cpp b/Engine/Source/Runtime/Renderer/Private/SceneRendering.cpp
--- a/Engine/Source/Runtime/Renderer/Private/SceneRendering.cpp
+++ b/Engine/Source/Runtime/Renderer/Private/SceneRendering.cpp
@@ -120,6 +120,7 @@
 #include "ShaderCompiler.h"
 #include "Quantization.h"
 #include "MaterialCache/MaterialCacheTagProvider.h"
+#include "RenderViewportFeedback.h"
 
 /*-----------------------------------------------------------------------------
 	Globals
@@ -2758,7 +2759,7 @@ FSceneRenderer::FSceneRenderer(const FSceneViewFamily* InViewFamily, FHitProxyCo
 			GEngine->LoadSMAATextures();
 		}
 
-		// Handle the FFT bloom kernel textire
+		// Handle the FFT bloom kernel texture
 		if (ViewInfo->FinalPostProcessSettings.BloomMethod == EBloomMethod::BM_FFT && ViewInfo->ViewState != nullptr)
 		{
 			UTexture2D* BloomConvolutionTexture = ViewInfo->FinalPostProcessSettings.BloomConvolutionTexture;
@@ -4984,6 +4985,13 @@ static void CleanupViewFamilies_RenderThread(FRHICommandListImmediate& RHICmdLis
 	{
 		for (FViewInfo& View : SceneRenderer->Views)
 		{
+#if WITH_EDITOR
+			if (View.ViewportFeedback)
+			{
+				View.ViewportFeedback->EndFrameRenderThread();
+			}
+#endif
+
 			// Copy relevant data from ViewInfo to ViewState->PrevFrameViewInfo
 			if (View.ViewState)
 			{
