World Partition Actor Component Descriptors: actor components can now create component descriptors and serialize them as part as actor descriptors.
Component descriptors properties can be delta serialized from their parent's values, without support for Blueprint component classes yet.

diff --git a/Engine/Source/Runtime/Core/Public/UObject/FortniteMainBranchObjectVersions.inl b/Engine/Source/Runtime/Core/Public/UObject/FortniteMainBranchObjectVersions.inl
--- a/Engine/Source/Runtime/Core/Public/UObject/FortniteMainBranchObjectVersions.inl
+++ b/Engine/Source/Runtime/Core/Public/UObject/FortniteMainBranchObjectVersions.inl
@@ -676,5 +676,8 @@ FORTNITE_MAIN_VERSION(DeduplicatedDebugNameSerializationInCVD, 220)
 // Add BloomGaussianIntensity and BloomConvolutionIntensity
 FORTNITE_MAIN_VERSION(SpecializeBloomIntensity, 221)
 
+// Add support for world partition actor component descriptors
+FORTNITE_MAIN_VERSION(WorldPartitionActorComponentDesc, 222)
+
 // -----<new versions can be added above this line>-------------------------------------------------
 #undef FORTNITE_MAIN_VERSION
\ No newline at end of file
diff --git a/Engine/Source/Runtime/Engine/Classes/Components/ActorComponent.h b/Engine/Source/Runtime/Engine/Classes/Components/ActorComponent.h
--- a/Engine/Source/Runtime/Engine/Classes/Components/ActorComponent.h
+++ b/Engine/Source/Runtime/Engine/Classes/Components/ActorComponent.h
@@ -425,6 +425,24 @@ private:
 	/** Calculate the UCS serialization index for a component that was saved before we started saving this data */
 	ENGINE_API void DetermineUCSSerializationIndexForLegacyComponent();
 
+#if WITH_EDITOR
+	/**
+	 * Creates an uninitialized component descriptor from this component. Meant to be called on the class CDO.
+	 */
+	ENGINE_API virtual TUniquePtr<class FWorldPartitionComponentDesc> CreateClassComponentDesc() const;
+
+public:
+	/**
+	 * Creates an initialized component descriptor from this component.
+	 */
+	ENGINE_API TUniquePtr<class FWorldPartitionComponentDesc> CreateComponentDesc() const;
+
+	/**
+	 * Creates an uninitialized component descriptor from a specific class.
+	 */
+	static ENGINE_API TUniquePtr<class FWorldPartitionComponentDesc> StaticCreateClassComponentDesc(const TSubclassOf<UActorComponent>& ComponentClass);
+#endif
+
 public:
 
 	/** Whether asynchronous physics state creation is running or not. */
diff --git a/Engine/Source/Runtime/Engine/Private/Components/ActorComponent.cpp b/Engine/Source/Runtime/Engine/Private/Components/ActorComponent.cpp
--- a/Engine/Source/Runtime/Engine/Private/Components/ActorComponent.cpp
+++ b/Engine/Source/Runtime/Engine/Private/Components/ActorComponent.cpp
@@ -51,6 +51,7 @@
 #include "Elements/Columns/TypedElementMiscColumns.h"
 #include "DataStorage/Features.h"
 #include "Elements/Interfaces/TypedElementDataStorageCompatibilityInterface.h"
+#include "WorldPartition/WorldPartitionActorDesc.h"
 #endif
 
 #if UE_WITH_IRIS
@@ -1008,6 +1009,28 @@ void UActorComponent::DetermineUCSSerializationIndexForLegacyComponent()
 #endif
 
 #if WITH_EDITOR
+TUniquePtr<FWorldPartitionComponentDesc> UActorComponent::CreateClassComponentDesc() const
+{
+	return nullptr;
+}
+
+TUniquePtr<FWorldPartitionComponentDesc> UActorComponent::CreateComponentDesc() const
+{
+	TUniquePtr<FWorldPartitionComponentDesc> ComponentDesc(CreateClassComponentDesc());
+
+	if (ComponentDesc.IsValid())
+	{
+		ComponentDesc->Init(this);
+	}
+
+	return ComponentDesc;
+}
+
+TUniquePtr<class FWorldPartitionComponentDesc> UActorComponent::StaticCreateClassComponentDesc(const TSubclassOf<UActorComponent>& ComponentClass)
+{
+	return CastChecked<UActorComponent>(ComponentClass->GetDefaultObject())->CreateClassComponentDesc();
+}
+
 void UActorComponent::CheckForErrors()
 {
 	if (AActor* MyOwner = GetOwner())
diff --git a/Engine/Source/Runtime/Engine/Private/WorldPartition/WorldPartitionActorDesc.cpp b/Engine/Source/Runtime/Engine/Private/WorldPartition/WorldPartitionActorDesc.cpp
--- a/Engine/Source/Runtime/Engine/Private/WorldPartition/WorldPartitionActorDesc.cpp
+++ b/Engine/Source/Runtime/Engine/Private/WorldPartition/WorldPartitionActorDesc.cpp
@@ -3,12 +3,15 @@
 #if WITH_EDITOR
 #include "WorldPartition/WorldPartitionActorDesc.h"
 #include "WorldPartition/WorldPartitionActorDescArchive.h"
+#include "WorldPartition/WorldPartitionActorDescUtils.h"
 
 #include "Misc/Paths.h"
 #include "Misc/ArchiveMD5.h"
 #include "UObject/MetaData.h"
 #include "AssetRegistry/AssetData.h"
 #include "AssetRegistry/AssetRegistryHelpers.h"
+#include "Engine/SimpleConstructionScript.h"
+#include "Engine/SCS_Node.h"
 #include "Serialization/MemoryReader.h"
 #include "Serialization/MemoryWriter.h"
 #include "UObject/UE5MainStreamObjectVersion.h"
@@ -43,6 +46,35 @@ static FGuid GetDefaultActorDescGuid(const FWorldPartitionActorDesc* ActorDesc)
 	return ArMD5.GetGuidFromHash();
 }
 
+FWorldPartitionComponentDesc::FWorldPartitionComponentDesc()
+{}
+
+void FWorldPartitionComponentDesc::Init(const UActorComponent* InComponent)
+{
+	check(IsValid(InComponent));
+
+	UClass* ComponentClass = InComponent->GetClass();
+
+	// Get the first native class in the hierarchy
+	ComponentNativeClass = GetParentNativeClass(ComponentClass);
+	NativeClass = *ComponentNativeClass->GetPathName();
+	
+	// For native class, don't set this
+	if (!ComponentNativeClass->IsNative())
+	{
+		BaseClass = *ComponentClass->GetPathName();
+	}
+
+	FString CleanComponentName = InComponent->GetName();
+	CleanComponentName.RemoveFromEnd(UActorComponent::ComponentTemplateNameSuffix);
+	ComponentName = *CleanComponentName;
+}
+
+void FWorldPartitionComponentDesc::Serialize(FArchive& Ar)
+{
+	Ar << BaseClass << ComponentName;		
+}
+
 FWorldPartitionActorDesc::FWorldPartitionActorDesc()
 	: ActorTransformRelative(FTransform::Identity)
 	, RuntimeBoundsRelative(ForceInit)
@@ -76,7 +108,7 @@ void FWorldPartitionActorDesc::Init(const AActor* InActor)
 	// For native class, don't set this
 	if (!ActorClass->IsNative())
 	{
-		BaseClass = *InActor->GetClass()->GetPathName();
+		BaseClass = *ActorClass->GetPathName();
 	}
 
 	if (InActor->HasAnyFlags(RF_ArchetypeObject | RF_ClassDefaultObject))
@@ -173,7 +205,7 @@ void FWorldPartitionActorDesc::Init(const AActor* InActor)
 	InActor->GetActorDescProperties(Properties);
 
 	ActorPackage = InActor->GetPackage()->GetFName();
-	ActorPath = bIsDefaultActorDesc ? *InActor->GetClass()->GetPathName() : *InActor->GetPathName();
+	ActorPath = bIsDefaultActorDesc ? *ActorClass->GetPathName() : *InActor->GetPathName();
 	ActorName = InActor->GetFName();
 	ActorNameString = ActorName.ToString();
 
@@ -219,6 +251,28 @@ void FWorldPartitionActorDesc::Init(const AActor* InActor)
 	ActorLabelString = ActorLabel.ToString();
 	ActorDisplayClassNameString = GetDisplayClassName().ToString();
 
+	if (!bIsDefaultActorDesc)
+	{
+		InActor->ForEachComponent(false, [this](const UActorComponent* ActorComponent)
+		{
+			if (TUniquePtr<FWorldPartitionComponentDesc> ComponentDesc = ActorComponent->CreateComponentDesc())
+			{
+				ComponentDescs.Add(MoveTemp(ComponentDesc));
+			}
+		});
+	}
+	else
+	{
+		AActor::ForEachComponentOfActorClassDefault<UActorComponent>(ActorClass, [this](const UActorComponent* ActorComponent)
+		{
+			if (TUniquePtr<FWorldPartitionComponentDesc> ComponentDesc = ActorComponent->CreateComponentDesc())
+			{
+				ComponentDescs.Add(MoveTemp(ComponentDesc));
+			}
+			return true;
+		});
+	}
+
 	Container = nullptr;
 }
 
@@ -369,6 +423,7 @@ bool FWorldPartitionActorDesc::Equals(const FWorldPartitionActorDesc* Other) con
 		FolderGuid == Other->FolderGuid &&
 		ParentActor == Other->ParentActor &&
 		ContentBundleGuid == Other->ContentBundleGuid &&
+		ComponentDescs == Other->ComponentDescs &&
 		CompareUnsortedArrays(DataLayers, Other->DataLayers) &&
 		CompareUnsortedArrays(References, Other->References) &&
 		CompareUnsortedArrays(EditorOnlyReferences, Other->EditorOnlyReferences) &&
@@ -553,6 +608,15 @@ FString FWorldPartitionActorDesc::ToString(EToStringMode Mode) const
 			AppendToString(TEXT("ActorPath"), ActorPath);
 		}
 
+		if (ComponentDescs.Num())
+		{
+			AppendFromString(TEXT("Components"), FString::JoinBy(ComponentDescs, TEXT(","), 
+				[&](const TUniquePtr<FWorldPartitionComponentDesc>& ComponentDesc)
+				{
+					return ComponentDesc->ToString();
+				}));
+		}
+
 		AppendToString(TEXT("Label"), GetActorLabel());
 		AppendFromBool(TEXT("SpatiallyLoaded"), bIsSpatiallyLoaded);
 		AppendToString(TEXT("EditorBounds"), EditorBounds);
@@ -590,17 +654,29 @@ FString FWorldPartitionActorDesc::ToString(EToStringMode Mode) const
 		{
 			if (References.Num())
 			{
-				AppendFromString(TEXT("References"), FString::JoinBy(References, TEXT(","), [&](const FGuid& ReferenceGuid) { return ReferenceGuid.ToString(); }));
+				AppendFromString(TEXT("References"), FString::JoinBy(References, TEXT(","), 
+					[&](const FGuid& ReferenceGuid)
+					{
+						return ReferenceGuid.ToString();
+					}));
 			}
 
 			if (EditorOnlyReferences.Num())
 			{
-				AppendFromString(TEXT("EditorOnlyReferences"), FString::JoinBy(EditorOnlyReferences, TEXT(","), [&](const FGuid& ReferenceGuid) { return ReferenceGuid.ToString(); }));
+				AppendFromString(TEXT("EditorOnlyReferences"), FString::JoinBy(EditorOnlyReferences, TEXT(","), 
+					[&](const FGuid& ReferenceGuid)
+					{
+						return ReferenceGuid.ToString();
+					}));
 			}
 
 			if (Tags.Num())
 			{
-				AppendFromString(TEXT("Tags"), FString::JoinBy(Tags, TEXT(","), [&](const FName& TagName) { return TagName.ToString(); }));
+				AppendFromString(TEXT("Tags"), FString::JoinBy(Tags, TEXT(","), 
+					[&](const FName& TagName)
+					{
+						return TagName.ToString();
+					}));
 			}
 
 			if (Properties.Num())
@@ -610,7 +686,11 @@ FString FWorldPartitionActorDesc::ToString(EToStringMode Mode) const
 
 			if (DataLayers.Num())
 			{
-				AppendFromString(TEXT("DataLayers"), FString::JoinBy(DataLayers, TEXT(","), [&](const FName& DataLayerName) { return DataLayerName.ToString(); }));
+				AppendFromString(TEXT("DataLayers"), FString::JoinBy(DataLayers, TEXT(","), 
+					[&](const FName& DataLayerName)
+					{
+						return DataLayerName.ToString();
+					}));
 			}
 
 			if (ExternalDataLayerAsset.IsValid())
@@ -879,6 +959,47 @@ void FWorldPartitionActorDesc::Serialize(FArchive& Ar)
 			UAssetRegistryHelpers::FixupRedirectedAssetPath(DataLayer);
 		}
 	}
+
+	if (Ar.CustomVer(FFortniteMainBranchObjectVersion::GUID) >= FFortniteMainBranchObjectVersion::WorldPartitionActorComponentDesc)
+	{
+		TArray<FTopLevelAssetPath> ComponentDescsType;
+
+		if (Ar.IsSaving())
+		{
+			ComponentDescsType.Reserve(ComponentDescs.Num());
+			for (const TUniquePtr<FWorldPartitionComponentDesc>& ComponentDesc : ComponentDescs)
+			{
+				ComponentDescsType.Add(ComponentDesc->NativeClass);
+			}
+		}
+
+		Ar << ComponentDescsType;
+
+		if (Ar.IsLoading())
+		{
+			for (const FTopLevelAssetPath& ComponentDescType : ComponentDescsType)
+			{
+				UClass* ComponentNativeClass = FWorldPartitionActorDescUtils::GetNativeClassFromClassName(*ComponentDescType.ToString());
+				check(ComponentNativeClass);
+
+				TUniquePtr<class FWorldPartitionComponentDesc> ComponentDesc = UActorComponent::StaticCreateClassComponentDesc(ComponentNativeClass);
+				ComponentDesc->NativeClass = *ComponentNativeClass->GetPathName();
+				ComponentDesc->ComponentNativeClass = ComponentNativeClass;
+
+				ComponentDescs.Add(MoveTemp(ComponentDesc));
+			}
+		}
+
+		for (const TUniquePtr<FWorldPartitionComponentDesc>& ComponentDesc : ComponentDescs)
+		{
+			FActorDescArchive& ActorDescAr = (FActorDescArchive&)Ar;
+			ActorDescAr.SetComponentDesc(ComponentDesc.Get());
+			
+			ComponentDesc->Serialize(Ar);
+
+			ActorDescAr.SetComponentDesc(nullptr);
+		}
+	}
 }
 
 FBox FWorldPartitionActorDesc::GetEditorBounds() const
@@ -1054,4 +1175,4 @@ void FWorldPartitionActorDesc::FixupStreamingBounds()
 
 #undef LOCTEXT_NAMESPACE
 
-#endif
\ No newline at end of file
+#endif
diff --git a/Engine/Source/Runtime/Engine/Private/WorldPartition/WorldPartitionActorDescArchive.cpp b/Engine/Source/Runtime/Engine/Private/WorldPartition/WorldPartitionActorDescArchive.cpp
--- a/Engine/Source/Runtime/Engine/Private/WorldPartition/WorldPartitionActorDescArchive.cpp
+++ b/Engine/Source/Runtime/Engine/Private/WorldPartition/WorldPartitionActorDescArchive.cpp
@@ -17,7 +17,11 @@
 FActorDescArchive::FActorDescArchive(FArchive& InArchive, FWorldPartitionActorDesc* InActorDesc, const FWorldPartitionActorDesc* InBaseActorDesc)
 	: FArchiveProxy(InArchive)
 	, ActorDesc(InActorDesc)
+	, ComponentDesc(nullptr)
 	, BaseDesc(InBaseActorDesc)
+	, BaseComponentDesc(nullptr)
+	, BaseDescSizeof(0)
+	, BaseComponentDescSizeof(0)
 	, bIsMissingBaseDesc(false)
 {
 	check(InArchive.IsPersistent());
@@ -26,6 +30,37 @@ FActorDescArchive::FActorDescArchive(FArchive& InArchive, FWorldPartitionActorDe
 	SetIsLoading(InArchive.IsLoading());
 }
 
+void FActorDescArchive::SetComponentDesc(FWorldPartitionComponentDesc* InComponentDesc)
+{
+	if (InComponentDesc)
+	{
+		check(!ComponentDesc);
+		check(!BaseComponentDesc);
+		check(!BaseComponentDescSizeof);
+
+		ComponentDesc = InComponentDesc;
+		BaseComponentDesc = nullptr;
+
+		for (const TUniquePtr<FWorldPartitionComponentDesc>& TargetBaseComponentDesc : BaseDesc->ComponentDescs)
+		{
+			if (TargetBaseComponentDesc->ComponentName == ComponentDesc->ComponentName)
+			{
+				BaseComponentDesc = TargetBaseComponentDesc.Get();
+				BaseComponentDescSizeof = BaseComponentDesc->GetSizeOf();
+				break;
+			}
+		}
+	}
+	else
+	{
+		check(ComponentDesc);
+
+		ComponentDesc = nullptr;
+		BaseComponentDesc = nullptr;
+		BaseComponentDescSizeof = 0;
+	}
+}
+
 void FActorDescArchive::Init(const FTopLevelAssetPath InClassPath)
 {
 	UsingCustomVersion(FUE5MainStreamObjectVersion::GUID);
diff --git a/Engine/Source/Runtime/Engine/Private/WorldPartition/WorldPartitionActorDescUtils.cpp b/Engine/Source/Runtime/Engine/Private/WorldPartition/WorldPartitionActorDescUtils.cpp
--- a/Engine/Source/Runtime/Engine/Private/WorldPartition/WorldPartitionActorDescUtils.cpp
+++ b/Engine/Source/Runtime/Engine/Private/WorldPartition/WorldPartitionActorDescUtils.cpp
@@ -38,22 +38,6 @@ namespace FWorldPartitionActorDescUtilsPrivate
 
 		return NewClassName.ToString();
 	}
-
-	static UClass* GetActorNativeClassFromClassName(FName InClassName)
-	{
-		// Avoid an assert when calling StaticFindObject during save to retrieve the actor's class.
-		// Since we are only looking for a native class, the call to StaticFindObject is legit.
-		TGuardValue<bool> GIsSavingPackageGuard(GIsSavingPackage, false);
-
-		// Look for a class redirectors
-		const FString ActorNativeClassName = FWorldPartitionActorDescUtilsPrivate::ResolveClassRedirector(InClassName.ToString());
-		
-		// Handle deprecated short class names
-		const FTopLevelAssetPath ClassPath = FAssetData::TryConvertShortClassNameToPathName(*ActorNativeClassName, ELogVerbosity::Log);
-
-		// Lookup the native class
-		return UClass::TryFindTypeSlow<UClass>(ClassPath.ToString(), EFindFirstObjectOptions::ExactClass);
-	}
 }
 
 FWorldPartitionActorDescUtils::FActorDescInitParams::FActorDescInitParams(const AActor* InActor)
@@ -131,7 +115,7 @@ UClass* FWorldPartitionActorDescUtils::GetActorNativeClassFromString(const FStri
 TUniquePtr<FWorldPartitionActorDesc> FWorldPartitionActorDescUtils::GetActorDescriptorFromInitParams(const FActorDescInitParams& InActorDescInitParams, FName InPackageName)
 {
 	FWorldPartitionActorDescInitData ActorDescInitData = FWorldPartitionActorDescInitData()
-		.SetNativeClass(FWorldPartitionActorDescUtilsPrivate::GetActorNativeClassFromClassName(InActorDescInitParams.NativeClassName))
+		.SetNativeClass(GetNativeClassFromClassName(InActorDescInitParams.NativeClassName))
 		.SetPackageName(InPackageName)
 		.SetActorPath(InActorDescInitParams.PathName.ToString());
 
@@ -286,4 +270,21 @@ bool FWorldPartitionActorDescUtils::FixupRedirectedAssetPath(FName& InOutAssetPa
 	UAssetRegistryHelpers::FixupRedirectedAssetPath(InOutAssetPath);
 	return true;
 }
+
+UClass* FWorldPartitionActorDescUtils::GetNativeClassFromClassName(FName InClassName)
+{
+	// Avoid an assert when calling StaticFindObject during save to retrieve the class.
+	// Since we are only looking for a native class, the call to StaticFindObject is legit.
+	TGuardValue<bool> GIsSavingPackageGuard(GIsSavingPackage, false);
+
+	// Look for a class redirectors
+	const FString NativeClassName = FWorldPartitionActorDescUtilsPrivate::ResolveClassRedirector(InClassName.ToString());
+		
+	// Handle deprecated short class names
+	const FTopLevelAssetPath ClassPath = FAssetData::TryConvertShortClassNameToPathName(*NativeClassName, ELogVerbosity::Log);
+
+	// Lookup the native class
+	return UClass::TryFindTypeSlow<UClass>(ClassPath.ToString(), EFindFirstObjectOptions::ExactClass);
+}
+
 #endif
\ No newline at end of file
diff --git a/Engine/Source/Runtime/Engine/Public/WorldPartition/WorldPartitionActorDesc.h b/Engine/Source/Runtime/Engine/Public/WorldPartition/WorldPartitionActorDesc.h
--- a/Engine/Source/Runtime/Engine/Public/WorldPartition/WorldPartitionActorDesc.h
+++ b/Engine/Source/Runtime/Engine/Public/WorldPartition/WorldPartitionActorDesc.h
@@ -18,6 +18,7 @@
 #include "WorldPartition/Filter/WorldPartitionActorFilter.h"
 
 class FActorDescArchive;
+class UActorComponent;
 
 // Struct used to create actor descriptor
 struct FWorldPartitionActorDescInitData
@@ -227,12 +228,41 @@ inline bool CompareUnsortedArrays(const TArray<FName>& Array1, const TArray<FNam
 {
 	return CompareUnsortedArrays(Array1, Array2, [](const FName& A, const FName& B) { return A.LexicalLess(B); });
 }
-
 #endif // WITH_EDITOR
 
 /**
  * Represents a Actor Metadata (editor-only)
  */
+class FWorldPartitionComponentDesc
+{
+#if WITH_EDITOR
+	friend class FActorDescArchive;
+	friend class FWorldPartitionActorDesc;
+
+public:
+	virtual ~FWorldPartitionComponentDesc() {}
+
+	virtual uint32 GetSizeOf() const = 0;
+	ENGINE_API virtual void Init(const UActorComponent* InComponent);
+	ENGINE_API virtual void Serialize(FArchive& Ar);
+	
+	UClass* GetComponentNativeClass() const { return ComponentNativeClass; }
+
+	FString ToString() const
+	{
+		return FString::Printf(TEXT("BaseClass=%s, NativeClass=%s, ComponentName=%s"), *BaseClass.ToString(), *NativeClass.ToString(), *ComponentName.ToString());
+	}
+
+protected:
+	ENGINE_API FWorldPartitionComponentDesc();
+
+	FTopLevelAssetPath BaseClass;
+	FTopLevelAssetPath NativeClass;
+	UClass* ComponentNativeClass;
+	FName ComponentName;
+#endif // WITH_EDITOR
+};
+
 class FWorldPartitionActorDesc
 {
 #if WITH_EDITOR
@@ -305,6 +335,7 @@ public:
 	inline void AddProperty(FName PropertyName, FName PropertyValue = NAME_None) { Properties.AddProperty(PropertyName, PropertyValue); }
 	inline bool GetProperty(FName PropertyName, FName* PropertyValue) const { return Properties.GetProperty(PropertyName, PropertyValue); }
 	inline bool HasProperty(FName PropertyName) const { return Properties.HasProperty(PropertyName); }
+	inline const TArray<TUniquePtr<FWorldPartitionComponentDesc>>& GetComponentDescs() const { return ComponentDescs; }
 
 	ENGINE_API FName GetActorName() const;
 	ENGINE_API FName GetActorLabelOrName() const;
@@ -563,6 +594,8 @@ protected:
 	FGuid							FolderGuid;
 	FGuid							ParentActor; // Used to validate settings against parent (to warn on layer/placement compatibility issues)
 	FGuid							ContentBundleGuid;
+
+	TArray<TUniquePtr<FWorldPartitionComponentDesc>> ComponentDescs;
 	
 	// Transient
 	UClass*							ActorNativeClass;
diff --git a/Engine/Source/Runtime/Engine/Public/WorldPartition/WorldPartitionActorDescArchive.h b/Engine/Source/Runtime/Engine/Public/WorldPartition/WorldPartitionActorDescArchive.h
--- a/Engine/Source/Runtime/Engine/Public/WorldPartition/WorldPartitionActorDescArchive.h
+++ b/Engine/Source/Runtime/Engine/Public/WorldPartition/WorldPartitionActorDescArchive.h
@@ -8,6 +8,7 @@
 #include "UObject/FortniteMainBranchObjectVersion.h"
 
 class FWorldPartitionActorDesc;
+class FWorldPartitionComponentDesc;
 struct FWorldPartitionAssetDataPatcher;
 
 class FActorDescArchive : public FArchiveProxy
@@ -16,6 +17,7 @@ public:
 	FActorDescArchive(FArchive& InArchive, FWorldPartitionActorDesc* InActorDesc, const FWorldPartitionActorDesc* InBaseActorDesc = nullptr);
 
 	void Init(const FTopLevelAssetPath InClassPath = FTopLevelAssetPath());
+	void SetComponentDesc(FWorldPartitionComponentDesc* InComponentDesc);
 
 	//~ Begin FArchive Interface
 	virtual FArchive& operator<<(FName& Value) override { FArchiveProxy::operator<<(Value); return *this; }
@@ -54,17 +56,16 @@ public:
 
 			Ar.UsingCustomVersion(FFortniteMainBranchObjectVersion::GUID);
 
-			auto GetBaseDefaultValue = [&V, &ActorDescAr]() -> const DestPropertyType*
+			auto GetBaseDefaultValue = [&V, &ActorDescAr](UPTRINT BasePtr, UPTRINT DefaultPtr, uint32 SizeOf) -> const DestPropertyType*
 			{
-				const UPTRINT PropertyOffset = (UPTRINT)&V.Value - *(UPTRINT*)&ActorDescAr.ActorDesc;
+				const UPTRINT PropertyOffset = (UPTRINT)&V.Value - BasePtr;
 
-				if ((PropertyOffset + sizeof(V.Value)) <= ActorDescAr.BaseDescSizeof)
+				if ((PropertyOffset + sizeof(V.Value)) <= SizeOf)
 				{
-					const DestPropertyType* RefValue = (const DestPropertyType*)(*(UPTRINT*)&ActorDescAr.BaseDesc + PropertyOffset);
+					const DestPropertyType* RefValue = (const DestPropertyType*)(DefaultPtr + PropertyOffset);
 					return RefValue;
 				}
 
-				check(ActorDescAr.bIsMissingBaseDesc);
 				return nullptr;
 			};
 
@@ -72,15 +73,32 @@ public:
 
 			if (Ar.CustomVer(FFortniteMainBranchObjectVersion::GUID) >= FFortniteMainBranchObjectVersion::WorldPartitionActorClassDescSerialize)
 			{
-				if (Ar.IsSaving() && ActorDescAr.BaseDesc)
+				if (Ar.IsSaving())
 				{
-					if constexpr (std::is_same_v<DestPropertyType, SourcePropertyType>)
+					if (ActorDescAr.ComponentDesc)
 					{
-						// When saving, we expect the class descriptor to be the exact type as what we are serializing.
-						const DestPropertyType* BaseDefaultValue = GetBaseDefaultValue();
-						check(BaseDefaultValue);
-
-						bSerialize = (V.Value != *BaseDefaultValue) ? 1 : 0;
+						if (ActorDescAr.BaseComponentDesc)
+						{
+							if constexpr (std::is_same_v<DestPropertyType, SourcePropertyType>)
+							{
+								// When saving, we expect the class descriptor to be the exact type as what we are serializing.
+								const DestPropertyType* BaseDefaultValue = GetBaseDefaultValue(*(UPTRINT*)&ActorDescAr.ComponentDesc, *(UPTRINT*)&ActorDescAr.BaseComponentDesc, ActorDescAr.BaseComponentDescSizeof);
+								check(BaseDefaultValue);
+
+								bSerialize = (V.Value != *BaseDefaultValue) ? 1 : 0;
+							}
+						}
+					}
+					else if (ActorDescAr.BaseDesc)
+					{
+						if constexpr (std::is_same_v<DestPropertyType, SourcePropertyType>)
+						{
+							// When saving, we expect the class descriptor to be the exact type as what we are serializing.
+							const DestPropertyType* BaseDefaultValue = GetBaseDefaultValue(*(UPTRINT*)&ActorDescAr.ActorDesc, *(UPTRINT*)&ActorDescAr.BaseDesc, ActorDescAr.BaseDescSizeof);
+							check(BaseDefaultValue);
+
+							bSerialize = (V.Value != *BaseDefaultValue) ? 1 : 0;
+						}
 					}
 				}
 
@@ -105,10 +123,21 @@ public:
 			else if (Ar.IsLoading())
 			{
 				// When loading, we need to handle a different class descriptor in case of missing classes, etc.
-				if (const DestPropertyType* BaseDefaultValue = GetBaseDefaultValue())
+				if (ActorDescAr.ComponentDesc)
 				{
+					if (const DestPropertyType* BaseDefaultValue = GetBaseDefaultValue(*(UPTRINT*)&ActorDescAr.ComponentDesc, *(UPTRINT*)&ActorDescAr.BaseComponentDesc, ActorDescAr.BaseComponentDescSizeof))
+					{
+						V.Value = *BaseDefaultValue;
+					}
+				}
+				else if (const DestPropertyType* BaseDefaultValue = GetBaseDefaultValue(*(UPTRINT*)&ActorDescAr.ActorDesc, *(UPTRINT*)&ActorDescAr.BaseDesc, ActorDescAr.BaseDescSizeof))
+				{				
 					V.Value = *BaseDefaultValue;
 				}
+				else
+				{
+					check(ActorDescAr.bIsMissingBaseDesc);
+				}
 			}
 
 			return Ar;
@@ -119,8 +148,11 @@ public:
 	};
 
 	FWorldPartitionActorDesc* ActorDesc;
+	FWorldPartitionComponentDesc* ComponentDesc;
 	const FWorldPartitionActorDesc* BaseDesc;
+	const FWorldPartitionComponentDesc* BaseComponentDesc;
 	uint32 BaseDescSizeof;
+	uint32 BaseComponentDescSizeof;
 	bool bIsMissingBaseDesc;
 };
 
@@ -152,4 +184,4 @@ private:
 
 template <typename DestPropertyType, typename SourcePropertyType = DestPropertyType>
 using TDeltaSerialize = FActorDescArchive::TDeltaSerializer<DestPropertyType, SourcePropertyType>;
-#endif
\ No newline at end of file
+#endif
diff --git a/Engine/Source/Runtime/Engine/Public/WorldPartition/WorldPartitionActorDescType.h b/Engine/Source/Runtime/Engine/Public/WorldPartition/WorldPartitionActorDescType.h
--- a/Engine/Source/Runtime/Engine/Public/WorldPartition/WorldPartitionActorDescType.h
+++ b/Engine/Source/Runtime/Engine/Public/WorldPartition/WorldPartitionActorDescType.h
@@ -7,6 +7,7 @@
 #include <type_traits>
 
 class AActor;
+class UActorComponent;
 
 #if WITH_EDITOR
 template <typename ActorType UE_REQUIRES(std::is_base_of_v<AActor, ActorType>)>
@@ -19,6 +20,19 @@ struct FWorldPartitionActorDescType<ActorType>			\
 {														\
 	typedef class ActorDescType Type;					\
 };
+
+template <typename ComponentType UE_REQUIRES(std::is_base_of_v<UActorComponent, ComponentType>)>
+struct FWorldPartitionComponentDescType
+{};
+
+#define DEFINE_COMPONENTDESC_TYPE(ComponentType, ComponentDescType) \
+template <>															\
+struct FWorldPartitionComponentDescType<ComponentType>				\
+{																	\
+	typedef class ComponentDescType Type;							\
+};
+
 #else
 #define DEFINE_ACTORDESC_TYPE(ActorType, ActorDescType)
+#define DEFINE_COMPONENTDESC_TYPE(ComponentType, ComponentDescType)
 #endif
diff --git a/Engine/Source/Runtime/Engine/Public/WorldPartition/WorldPartitionActorDescUtils.h b/Engine/Source/Runtime/Engine/Public/WorldPartition/WorldPartitionActorDescUtils.h
--- a/Engine/Source/Runtime/Engine/Public/WorldPartition/WorldPartitionActorDescUtils.h
+++ b/Engine/Source/Runtime/Engine/Public/WorldPartition/WorldPartitionActorDescUtils.h
@@ -121,5 +121,7 @@ struct FWorldPartitionActorDescUtils
 
 	UE_DEPRECATED(5.3, "Use FWorldPartitionHelpers::FixupRedirectedAssetPath instead")
 	static ENGINE_API bool FixupRedirectedAssetPath(FName& InOutAssetPath);
+
+	static ENGINE_API UClass* GetNativeClassFromClassName(FName InClassName);
 };
 #endif
