Chaos - fix per-BodyInstance iteration counts
- support a lower per-body iteration count than the project settings (as long as all bodies in the island are lower)
- add a flag to enable/disable per-body iteration counts
- allow bodies to specify 0 as their iteration counts
- fix iteration count cvar overrides

[RN] Added bOverrideIterationCounts flag to BodyInstance to enable use of the iteration count overides. The physics solver now handles per-body iteration counts lower than the project settings for islands where all bodies have lower counts. E.g., a pile of debris bodies with a 1 position iteration override will now only run 1 iteration, regardless of the project settings, until a body with more iterations comes into contact with the island. Previous behaviour used the maximum of the per-body and project iteration counts for the island.

diff --git a/Engine/Source/Runtime/Engine/Classes/PhysicsEngine/BodyInstance.h b/Engine/Source/Runtime/Engine/Classes/PhysicsEngine/BodyInstance.h
--- a/Engine/Source/Runtime/Engine/Classes/PhysicsEngine/BodyInstance.h
+++ b/Engine/Source/Runtime/Engine/Classes/PhysicsEngine/BodyInstance.h
@@ -329,16 +329,16 @@ struct FBodyInstance : public FBodyInstanceCore
 	/** When we are a body within a SkeletalMeshComponent, we cache the index of the bone we represent, to speed up sync'ing physics to anim. */
 	int16 InstanceBoneIndex;
 
-	/** This physics body's solver iteration count for position. Increasing this will be more CPU intensive, but better stabilized.  */
-	UPROPERTY(EditAnywhere, AdvancedDisplay, BlueprintReadOnly, Category=Physics, meta = (ClampMin = "1", UIMin = "1"))
+	/** Override the project default position iteration count for this body. Increasing this will be more CPU intensive, but better stabilized. 0 will lead to unresolved penetration and no static friction. */
+	UPROPERTY(EditAnywhere, AdvancedDisplay, BlueprintReadOnly, Category=Physics, meta = (editcondition = "bOverrideIterationCounts", InlineEditConditionToggle, ClampMin = "0", UIMin = "0"))
 	uint8 PositionSolverIterationCount;
 
-	/** This physics body's solver iteration count for velocity. Increasing this will be more CPU intensive, but better stabilized. */
-	UPROPERTY(EditAnywhere, AdvancedDisplay, BlueprintReadOnly, Category = Physics, meta = (ClampMin = "1", UIMin = "1"))
+	/** Override the project default velocity iteration count for this body. Increasing this will be more CPU intensive, but better stabilized. 0 will lead to unresolved restitution and lower quality dynamic friction. */
+	UPROPERTY(EditAnywhere, AdvancedDisplay, BlueprintReadOnly, Category = Physics, meta = (editcondition = "bOverrideIterationCounts", InlineEditConditionToggle, ClampMin = "0", UIMin = "0"))
 	uint8 VelocitySolverIterationCount;
 
-	/** This physics body's solver iteration count for projection. Increasing this will be more CPU intensive, but better stabilized. */
-	UPROPERTY(EditAnywhere, AdvancedDisplay, BlueprintReadOnly, Category = Physics, meta = (ClampMin = "1", UIMin = "1"))
+	/** Override the project default position iteration count for this body. Usually 0 or 1. */
+	UPROPERTY(EditAnywhere, AdvancedDisplay, BlueprintReadOnly, Category = Physics, meta = (editcondition = "bOverrideIterationCounts", InlineEditConditionToggle, ClampMin = "0", UIMin = "0"))
 	uint8 ProjectionSolverIterationCount;
 
 private:
@@ -479,6 +479,10 @@ protected:
 	UPROPERTY(EditAnywhere, Category = Physics, meta = (InlineEditConditionToggle))
 	uint8 bOverrideWalkableSlopeOnInstance : 1;
 
+	/** Enable per-body iterations counts */
+	UPROPERTY(EditAnywhere, Category = Physics, meta = (InlineEditConditionToggle))
+	uint8 bOverrideIterationCounts : 1;
+
 	/** 
 	 * Internal flag to allow us to quickly check whether we should interpolate when substepping 
 	 * e.g. kinematic bodies that are QueryOnly do not need to interpolate as we will not be querying them
@@ -1049,22 +1053,33 @@ public:
 	 */
 	ENGINE_API void SetAllowPartialIslandSleep(bool bInAllowPartialIslandSleep);
 
+	// Set the per-body position iterations (only used if bOverrideIterationsCounts is true)
 	ENGINE_API void SetPositionSolverIterationCount(uint8 PositionSolverIterationCountIn);
 
+	// Set the per-body velocity iterations (only used if bOverrideIterationsCounts is true)
 	ENGINE_API void SetVelocitySolverIterationCount(uint8 VelocitySolverIterationCountIn);
 
+	// Set the per-body projection iterations (only used if bOverrideIterationsCounts is true)
 	ENGINE_API void SetProjectionSolverIterationCount(uint8 ProjectionSolverIterationCountIn);
 
+	// Enable or disable the use of iteration count overrides
+	ENGINE_API void SetOverrideIterationCounts(bool bOverride);
+
+	// Get the per-body position iterations (or -1 if bOverrideIterationsCounts is false)
+	int32 GetPositionSolverIterationCount() const { return (!!bOverrideIterationCounts) ? PositionSolverIterationCount : -1; }
+
+	// Get the per-body velocity iterations (or -1 if bOverrideIterationsCounts is false)
+	int32 GetVelocitySolverIterationCount() const { return (!!bOverrideIterationCounts) ? VelocitySolverIterationCount : -1; }
+
+	// Get the per-body projection iterations (or -1 if bOverrideIterationsCounts is false)
+	int32 GetProjectionSolverIterationCount() const { return (!!bOverrideIterationCounts) ? ProjectionSolverIterationCount : -1; }
+
 	/** [EXPERIMENTAL] Whether Motion-Aware Collision Detection is enabled */
 	bool GetUseMACD() const { return bUseMACD != 0; }
 
 	/** [EXPERIMENTAL] Whether Partial Island Sleep is permitted for the island this body is in */
 	bool GetAllowPartialIslandSleep() const { return bAllowPartialIslandSleep != 0; }
 
-	uint8 GetPositionSolverIterationCount() const { return PositionSolverIterationCount; }
-	uint8 GetVelocitySolverIterationCount() const { return VelocitySolverIterationCount; }
-	uint8 GetProjectionSolverIterationCount() const { return ProjectionSolverIterationCount; }
-
 	/** Disable/Re-Enable this body in the solver,  when disable, the body won't be part of the simulation ( regardless if it's dynamic or kinematic ) and no collision will occur 
 	* this can be used for performance control situation for example
 	*/
diff --git a/Engine/Source/Runtime/Engine/Private/PhysicsEngine/BodyInstance.cpp b/Engine/Source/Runtime/Engine/Private/PhysicsEngine/BodyInstance.cpp
--- a/Engine/Source/Runtime/Engine/Private/PhysicsEngine/BodyInstance.cpp
+++ b/Engine/Source/Runtime/Engine/Private/PhysicsEngine/BodyInstance.cpp
@@ -386,6 +386,7 @@ FBodyInstance::FBodyInstance()
 	, bOverrideMaxAngularVelocity(false)
 	, bOverrideMaxDepenetrationVelocity(false)
 	, bOverrideWalkableSlopeOnInstance(false)
+	, bOverrideIterationCounts(false)
 	, bInterpolateWhenSubStepping(true)
 	, bPendingCollisionProfileSetup(false)
 	, bInertiaConditioning(true)	
@@ -1320,9 +1321,9 @@ void TInitBodiesHelperBase<InAllocatorType>::CreateActor_AssumesLocked(FBodyInst
 		FPhysicsInterface::SetSmoothEdgeCollisionsEnabled_AssumesLocked(ActorHandle, Instance->bSmoothEdgeCollisions);
 		FPhysicsInterface::SetInertiaConditioningEnabled_AssumesLocked(ActorHandle, Instance->IsInertiaConditioningEnabled());
 		FPhysicsInterface::SetGyroscopicTorqueEnabled_AssumesLocked(ActorHandle, Instance->bGyroscopicTorqueEnabled);
-		FPhysicsInterface::SetPositionSolverIterationCount_AssumesLocked(ActorHandle, Instance->PositionSolverIterationCount);
-		FPhysicsInterface::SetVelocitySolverIterationCount_AssumesLocked(ActorHandle, Instance->VelocitySolverIterationCount);
-		FPhysicsInterface::SetProjectionSolverIterationCount_AssumesLocked(ActorHandle, Instance->ProjectionSolverIterationCount);
+		FPhysicsInterface::SetPositionSolverIterationCount_AssumesLocked(ActorHandle, Instance->GetPositionSolverIterationCount());
+		FPhysicsInterface::SetVelocitySolverIterationCount_AssumesLocked(ActorHandle, Instance->GetVelocitySolverIterationCount());
+		FPhysicsInterface::SetProjectionSolverIterationCount_AssumesLocked(ActorHandle, Instance->GetProjectionSolverIterationCount());
 		FPhysicsInterface::SetGravityGroupIndex_AssumesLocked(ActorHandle, Instance->GravityGroupIndex);
 
 		// Set sleep event notification
@@ -2566,9 +2567,9 @@ void FBodyInstance::UpdateInstanceSimulatePhysics()
 		FPhysicsInterface::SetMACDEnabled_AssumesLocked(Actor, bUseMACD);
 		FPhysicsInterface::SetPartialIslandSleepAllowed_AssumesLocked(Actor, bAllowPartialIslandSleep);
 		FPhysicsInterface::SetOneWayInteraction_AssumesLocked(Actor, bOneWayInteraction);
-		FPhysicsInterface::SetPositionSolverIterationCount_AssumesLocked(Actor, PositionSolverIterationCount);
-		FPhysicsInterface::SetVelocitySolverIterationCount_AssumesLocked(Actor, VelocitySolverIterationCount);
-		FPhysicsInterface::SetProjectionSolverIterationCount_AssumesLocked(Actor, ProjectionSolverIterationCount);
+		FPhysicsInterface::SetPositionSolverIterationCount_AssumesLocked(Actor, GetPositionSolverIterationCount());
+		FPhysicsInterface::SetVelocitySolverIterationCount_AssumesLocked(Actor, GetVelocitySolverIterationCount());
+		FPhysicsInterface::SetProjectionSolverIterationCount_AssumesLocked(Actor, GetProjectionSolverIterationCount());
 
 		if(bSimulatePhysics)
 		{
@@ -4027,17 +4028,41 @@ void FBodyInstance::SetAllowPartialIslandSleep(bool bInAllowPartialIslandSleep)
 	}
 }
 
+void FBodyInstance::SetOverrideIterationCounts(bool bOverride)
+{
+	if (bOverride != bOverrideIterationCounts)
+	{
+		bOverrideIterationCounts = bOverride;
+
+		int32 PosCount = GetPositionSolverIterationCount();
+		int32 VelCount = GetVelocitySolverIterationCount();
+		int32 ProjCount = GetProjectionSolverIterationCount();
+
+		FPhysicsCommand::ExecuteWrite(GetPhysicsActor(), [this, PosCount, VelCount, ProjCount](const FPhysicsActorHandle& Actor)
+			{
+				if (FPhysicsInterface::IsValid(Actor) && FPhysicsInterface::IsRigidBody(Actor))
+				{
+					FPhysicsInterface::SetPositionSolverIterationCount_AssumesLocked(Actor, PosCount);
+					FPhysicsInterface::SetVelocitySolverIterationCount_AssumesLocked(Actor, VelCount);
+					FPhysicsInterface::SetProjectionSolverIterationCount_AssumesLocked(Actor, ProjCount);
+				}
+			});
+	}
+}
+
 void FBodyInstance::SetPositionSolverIterationCount(uint8 PositionSolverIterationCountIn)
 {
 	if (PositionSolverIterationCount != PositionSolverIterationCountIn)
 	{
 		PositionSolverIterationCount = PositionSolverIterationCountIn;
 
-		FPhysicsCommand::ExecuteWrite(GetPhysicsActor(), [this, PositionSolverIterationCountIn](const FPhysicsActorHandle& Actor)
+		int32 PosCount = GetPositionSolverIterationCount();
+
+		FPhysicsCommand::ExecuteWrite(GetPhysicsActor(), [this, PosCount](const FPhysicsActorHandle& Actor)
 			{
 				if (FPhysicsInterface::IsValid(Actor) && FPhysicsInterface::IsRigidBody(Actor))
 				{
-					FPhysicsInterface::SetPositionSolverIterationCount_AssumesLocked(Actor, PositionSolverIterationCount);
+					FPhysicsInterface::SetPositionSolverIterationCount_AssumesLocked(Actor, PosCount);
 				}
 			});
 	}
@@ -4049,11 +4074,13 @@ void FBodyInstance::SetVelocitySolverIterationCount(uint8 VelocitySolverIteratio
 	{
 		VelocitySolverIterationCount = VelocitySolverIterationCountIn;
 
-		FPhysicsCommand::ExecuteWrite(GetPhysicsActor(), [this, VelocitySolverIterationCountIn](const FPhysicsActorHandle& Actor)
+		int32 VelCount = GetVelocitySolverIterationCount();
+
+		FPhysicsCommand::ExecuteWrite(GetPhysicsActor(), [this, VelCount](const FPhysicsActorHandle& Actor)
 			{
 				if (FPhysicsInterface::IsValid(Actor) && FPhysicsInterface::IsRigidBody(Actor))
 				{
-					FPhysicsInterface::SetVelocitySolverIterationCount_AssumesLocked(Actor, VelocitySolverIterationCount);
+					FPhysicsInterface::SetVelocitySolverIterationCount_AssumesLocked(Actor, VelCount);
 				}
 			});
 	}
@@ -4065,11 +4092,13 @@ void FBodyInstance::SetProjectionSolverIterationCount(uint8 ProjectionSolverIter
 	{
 		ProjectionSolverIterationCount = ProjectionSolverIterationCountIn;
 
-		FPhysicsCommand::ExecuteWrite(GetPhysicsActor(), [this, ProjectionSolverIterationCountIn](const FPhysicsActorHandle& Actor)
+		int32 ProjCount = GetProjectionSolverIterationCount();
+
+		FPhysicsCommand::ExecuteWrite(GetPhysicsActor(), [this, ProjCount](const FPhysicsActorHandle& Actor)
 			{
 				if (FPhysicsInterface::IsValid(Actor) && FPhysicsInterface::IsRigidBody(Actor))
 				{
-					FPhysicsInterface::SetProjectionSolverIterationCount_AssumesLocked(Actor, ProjectionSolverIterationCount);
+					FPhysicsInterface::SetProjectionSolverIterationCount_AssumesLocked(Actor, ProjCount);
 				}
 			});
 	}
@@ -4601,12 +4630,9 @@ void FBodyInstance::InitDynamicProperties_AssumesLocked()
 		StabilizationThreshold *= StabilizationThresholdMultiplier;
 		FPhysicsInterface::SetStabilizationEnergyThreshold_AssumesLocked(GetPhysicsActor(), StabilizationThreshold);
 
-		uint32 PositionIterCount = FMath::Clamp<uint8>(PositionSolverIterationCount, 1, 255);
-		uint32 VelocityIterCount = FMath::Clamp<uint8>(VelocitySolverIterationCount, 1, 255);
-		uint32 ProjectionIterCount = FMath::Clamp<uint8>(ProjectionSolverIterationCount, 1, 255);
-		FPhysicsInterface::SetSolverPositionIterationCount_AssumesLocked(GetPhysicsActor(), PositionIterCount);
-		FPhysicsInterface::SetSolverVelocityIterationCount_AssumesLocked(GetPhysicsActor(), VelocityIterCount);
-		FPhysicsInterface::SetSolverProjectionIterationCount_AssumesLocked(GetPhysicsActor(), ProjectionIterCount);
+		FPhysicsInterface::SetSolverPositionIterationCount_AssumesLocked(GetPhysicsActor(), GetPositionSolverIterationCount());
+		FPhysicsInterface::SetSolverVelocityIterationCount_AssumesLocked(GetPhysicsActor(), GetVelocitySolverIterationCount());
+		FPhysicsInterface::SetSolverProjectionIterationCount_AssumesLocked(GetPhysicsActor(), GetProjectionSolverIterationCount());
 
 		CreateDOFLock();
 		if(FPhysicsInterface::IsInScene(GetPhysicsActor()) && !IsRigidBodyKinematic_AssumesLocked(GetPhysicsActor()))
diff --git a/Engine/Source/Runtime/Experimental/Chaos/Private/Chaos/Island/IslandGroup.cpp b/Engine/Source/Runtime/Experimental/Chaos/Private/Chaos/Island/IslandGroup.cpp
--- a/Engine/Source/Runtime/Experimental/Chaos/Private/Chaos/Island/IslandGroup.cpp
+++ b/Engine/Source/Runtime/Experimental/Chaos/Private/Chaos/Island/IslandGroup.cpp
@@ -5,6 +5,13 @@
 #include "Chaos/Island/IslandManager.h"
 #include "Chaos/PBDConstraintContainer.h"
 
+namespace Chaos::CVars
+{
+	extern int32 ChaosSolverPositionIterations;
+	extern int32 ChaosSolverVelocityIterations;
+	extern int32 ChaosSolverProjectionIterations;
+}
+
 namespace Chaos
 {
 	namespace Private
@@ -91,11 +98,49 @@ namespace Chaos
 
 		void FPBDIslandConstraintGroupSolver::SetIterationSettings(const FIterationSettings& InDefaultIterations)
 		{
-			Iterations = InDefaultIterations;
+			// Start with undefined iteration counts
+			Iterations = FIterationSettings::MakeEmpty();
+
+			// Here we set the iteration count for the island group to be the largest of any of the islands.
+			// @todo(chaos): we should support per-island iteration counts rather than pushing the max to all
+			// island solved in the same group.
 			for (FPBDIsland* Island : Islands)
 			{
 				Iterations = FIterationSettings::Merge(Iterations, Island->GetIterationSettings());
 			}
+
+			// If we still have undefined iterations counts, apply the defaults
+			if (Iterations.GetNumPositionIterations() < 0)
+			{
+				Iterations.SetNumPositionIterations(InDefaultIterations.GetNumPositionIterations());
+			}
+			if (Iterations.GetNumVelocityIterations() < 0)
+			{
+				Iterations.SetNumVelocityIterations(InDefaultIterations.GetNumVelocityIterations());
+			}
+			if (Iterations.GetNumProjectionIterations() < 0)
+			{
+				Iterations.SetNumProjectionIterations(InDefaultIterations.GetNumProjectionIterations());
+			}
+
+			// Apply cvar overrides
+			ApplyIterationSettingsOverrides();
+		}
+
+		void FPBDIslandConstraintGroupSolver::ApplyIterationSettingsOverrides()
+		{
+			if (CVars::ChaosSolverPositionIterations >= 0)
+			{
+				Iterations.SetNumPositionIterations(CVars::ChaosSolverPositionIterations);
+			}
+			if (CVars::ChaosSolverVelocityIterations >= 0)
+			{
+				Iterations.SetNumVelocityIterations(CVars::ChaosSolverVelocityIterations);
+			}
+			if (CVars::ChaosSolverProjectionIterations >= 0)
+			{
+				Iterations.SetNumProjectionIterations(CVars::ChaosSolverProjectionIterations);
+			}
 		}
 
 	}	// namespace Private
diff --git a/Engine/Source/Runtime/Experimental/Chaos/Private/Chaos/Island/IslandManager.cpp b/Engine/Source/Runtime/Experimental/Chaos/Private/Chaos/Island/IslandManager.cpp
--- a/Engine/Source/Runtime/Experimental/Chaos/Private/Chaos/Island/IslandManager.cpp
+++ b/Engine/Source/Runtime/Experimental/Chaos/Private/Chaos/Island/IslandManager.cpp
@@ -2286,20 +2286,25 @@ namespace Chaos::Private
 				Island->SleepCounter = 0;
 			}
 
+			// @todo(chaos): this won't handle updates to per-body iterations after being added to an island, until the island
+			// changes for some other reason. Not a big deal but should probably be fixed at some point.
 			if (Island->Flags.bItemsAdded || Island->Flags.bItemsRemoved)
 			{
-				//Compute max of iteration settings by iterating through all the nodes:
+				// Compute max of iteration settings by iterating through all the nodes:
+				// NOTE: It is valid for any or all iteration counts to be undefined (-1) - this means "use project settings"
 				SCOPE_CYCLE_COUNTER(STAT_IslandManager_ComputeIterationSettings);
 				if (Chaos::CVars::bChaosSolverComputeIterationSettings)
 				{
-					Island->SetIterationSettings(FIterationSettings(0, 0, 0));
+					// Default to undefined (-1) - use project settings
+					FIterationSettings IslandIterations = FIterationSettings();
 					for (FPBDIslandParticle* Node : Island->Nodes)
 					{
 						if (Node->Flags.bIsDynamic)
 						{
-							Island->SetIterationSettings(FIterationSettings::Merge(Island->GetIterationSettings(), Node->GetIterationSettings()));
+							IslandIterations = FIterationSettings::Merge(IslandIterations, Node->GetIterationSettings());
 						}
 					}
+					Island->SetIterationSettings(IslandIterations);
 				}
 			}
 
diff --git a/Engine/Source/Runtime/Experimental/Chaos/Public/Chaos/Island/IslandGroup.h b/Engine/Source/Runtime/Experimental/Chaos/Public/Chaos/Island/IslandGroup.h
--- a/Engine/Source/Runtime/Experimental/Chaos/Public/Chaos/Island/IslandGroup.h
+++ b/Engine/Source/Runtime/Experimental/Chaos/Public/Chaos/Island/IslandGroup.h
@@ -64,6 +64,8 @@ namespace Chaos
 			CHAOS_API virtual void GatherBodiesImpl(const FReal Dt, const int32 BeginBodyIndex, const int32 EndBodyIndex) override final;
 
 		private:
+			CHAOS_API void ApplyIterationSettingsOverrides();
+
 			FPBDIslandManager& IslandManager;
 
 			// Item counters used to initialize the solver containers
diff --git a/Engine/Source/Runtime/Experimental/Chaos/Public/Chaos/ParticleHandle.h b/Engine/Source/Runtime/Experimental/Chaos/Public/Chaos/ParticleHandle.h
--- a/Engine/Source/Runtime/Experimental/Chaos/Public/Chaos/ParticleHandle.h
+++ b/Engine/Source/Runtime/Experimental/Chaos/Public/Chaos/ParticleHandle.h
@@ -3514,19 +3514,19 @@ public:
 		MMiscData.Modify(true, MDirtyFlags, Proxy, [SolverIterationSettingsIn](auto& Data) { Data.SetIterationSettings(SolverIterationSettingsIn); });
 	}
 
-	void SetPositionSolverIterations(const uint32 PositionSolverIterationCount)
+	void SetPositionSolverIterations(const int32 PositionSolverIterationCount)
 	{
-		MMiscData.Modify(true, MDirtyFlags, Proxy, [PositionSolverIterationCount](auto& Data) { Data.SetPositionSolverIterationCount((int32)PositionSolverIterationCount); });
+		MMiscData.Modify(true, MDirtyFlags, Proxy, [PositionSolverIterationCount](auto& Data) { Data.SetPositionSolverIterationCount(PositionSolverIterationCount); });
 	}
 
-	void SetVelocitySolverIterations(const uint32 VelocitySolverIterationCount)
+	void SetVelocitySolverIterations(const int32 VelocitySolverIterationCount)
 	{
-		MMiscData.Modify(true, MDirtyFlags, Proxy, [VelocitySolverIterationCount](auto& Data) { Data.SetVelocitySolverIterationCount((int32)VelocitySolverIterationCount); });
+		MMiscData.Modify(true, MDirtyFlags, Proxy, [VelocitySolverIterationCount](auto& Data) { Data.SetVelocitySolverIterationCount(VelocitySolverIterationCount); });
 	}
 
-	void SetProjectionSolverIterations(const uint32 ProjectionSolverIterationCount)
+	void SetProjectionSolverIterations(const int32 ProjectionSolverIterationCount)
 	{
-		MMiscData.Modify(true, MDirtyFlags, Proxy, [ProjectionSolverIterationCount](auto& Data) { Data.SetProjectionSolverIterationCount((int32)ProjectionSolverIterationCount); });
+		MMiscData.Modify(true, MDirtyFlags, Proxy, [ProjectionSolverIterationCount](auto& Data) { Data.SetProjectionSolverIterationCount(ProjectionSolverIterationCount); });
 	}
 
 	bool InertiaConditioningEnabled() const { return MMiscData.Read().InertiaConditioningEnabled(); }
diff --git a/Engine/Source/Runtime/Experimental/Chaos/Public/PhysicsProxy/SingleParticlePhysicsProxy.h b/Engine/Source/Runtime/Experimental/Chaos/Public/PhysicsProxy/SingleParticlePhysicsProxy.h
--- a/Engine/Source/Runtime/Experimental/Chaos/Public/PhysicsProxy/SingleParticlePhysicsProxy.h
+++ b/Engine/Source/Runtime/Experimental/Chaos/Public/PhysicsProxy/SingleParticlePhysicsProxy.h
@@ -1441,30 +1441,30 @@ public:
 		return Private::FIterationSettings(0, 0, 0);
 	}
 
-	void SetPositionSolverIterationCount(uint32 PositionSolverIterationCountIn)
+	void SetPositionSolverIterationCount(int32 PositionSolverIterationCountIn)
 	{
 		VerifyContext();
-		if (TPBDRigidParticle<FReal, 3>*Rigid = GetParticle_LowLevel()->CastToRigidParticle())
+		if (TPBDRigidParticle<FReal, 3>* Rigid = GetParticle_LowLevel()->CastToRigidParticle())
 		{
 			Rigid->SetPositionSolverIterations(PositionSolverIterationCountIn);
 		}
 	}
 
 
-	void SetVelocitySolverIterationCount(uint32 VelocitySolverIterationCountIn)
+	void SetVelocitySolverIterationCount(int32 VelocitySolverIterationCountIn)
 	{
 		VerifyContext();
-		if (TPBDRigidParticle<FReal, 3>*Rigid = GetParticle_LowLevel()->CastToRigidParticle())
+		if (TPBDRigidParticle<FReal, 3>* Rigid = GetParticle_LowLevel()->CastToRigidParticle())
 		{
 			Rigid->SetVelocitySolverIterations(VelocitySolverIterationCountIn);
 		}
 	}
 
 
-	void SetProjectionSolverIterationCount(uint32 ProjectionSolverIterationCountIn)
+	void SetProjectionSolverIterationCount(int32 ProjectionSolverIterationCountIn)
 	{
 		VerifyContext();
-		if (TPBDRigidParticle<FReal, 3>*Rigid = GetParticle_LowLevel()->CastToRigidParticle())
+		if (TPBDRigidParticle<FReal, 3>* Rigid = GetParticle_LowLevel()->CastToRigidParticle())
 		{
 			Rigid->SetProjectionSolverIterations(ProjectionSolverIterationCountIn);
 		}
diff --git a/Engine/Source/Runtime/PhysicsCore/Private/ChaosEngineInterface.cpp b/Engine/Source/Runtime/PhysicsCore/Private/ChaosEngineInterface.cpp
--- a/Engine/Source/Runtime/PhysicsCore/Private/ChaosEngineInterface.cpp
+++ b/Engine/Source/Runtime/PhysicsCore/Private/ChaosEngineInterface.cpp
@@ -560,17 +560,17 @@ void FChaosEngineInterface::SetPartialIslandSleepAllowed_AssumesLocked(const FPh
 	InActorReference->GetGameThreadAPI().SetPartialIslandSleepAllowed(bIsPartialIslandSleepAllowed);
 }
 
-void FChaosEngineInterface::SetPositionSolverIterationCount_AssumesLocked(const FPhysicsActorHandle& InActorReference, uint8 PositionSolverIterationCount)
+void FChaosEngineInterface::SetPositionSolverIterationCount_AssumesLocked(const FPhysicsActorHandle& InActorReference, int32 PositionSolverIterationCount)
 {
 	InActorReference->GetGameThreadAPI().SetPositionSolverIterationCount(PositionSolverIterationCount);
 }
 
-void FChaosEngineInterface::SetVelocitySolverIterationCount_AssumesLocked(const FPhysicsActorHandle& InActorReference, uint8 VelocitySolverIterationCount)
+void FChaosEngineInterface::SetVelocitySolverIterationCount_AssumesLocked(const FPhysicsActorHandle& InActorReference, int32 VelocitySolverIterationCount)
 {
 	InActorReference->GetGameThreadAPI().SetVelocitySolverIterationCount(VelocitySolverIterationCount);
 }
 
-void FChaosEngineInterface::SetProjectionSolverIterationCount_AssumesLocked(const FPhysicsActorHandle& InActorReference, uint8 ProjectionSolverIterationCount)
+void FChaosEngineInterface::SetProjectionSolverIterationCount_AssumesLocked(const FPhysicsActorHandle& InActorReference, int32 ProjectionSolverIterationCount)
 {
 	InActorReference->GetGameThreadAPI().SetProjectionSolverIterationCount(ProjectionSolverIterationCount);
 }
@@ -1327,19 +1327,19 @@ void FChaosEngineInterface::SetStabilizationEnergyThreshold_AssumesLocked(const
 	// #todo : Implement
 }
 
-void FChaosEngineInterface::SetSolverPositionIterationCount_AssumesLocked(const FPhysicsActorHandle& InHandle,uint32 InSolverIterationCount)
+void FChaosEngineInterface::SetSolverPositionIterationCount_AssumesLocked(const FPhysicsActorHandle& InHandle, int32 InSolverIterationCount)
 {
 	Chaos::FRigidBodyHandle_External& Body_External = InHandle->GetGameThreadAPI();
 	Body_External.SetPositionSolverIterationCount(InSolverIterationCount);
 }
 
-void FChaosEngineInterface::SetSolverVelocityIterationCount_AssumesLocked(const FPhysicsActorHandle& InHandle, uint32 InSolverIterationCount)
+void FChaosEngineInterface::SetSolverVelocityIterationCount_AssumesLocked(const FPhysicsActorHandle& InHandle, int32 InSolverIterationCount)
 {
 	Chaos::FRigidBodyHandle_External& Body_External = InHandle->GetGameThreadAPI();
 	Body_External.SetVelocitySolverIterationCount(InSolverIterationCount);
 }
 
-void FChaosEngineInterface::SetSolverProjectionIterationCount_AssumesLocked(const FPhysicsActorHandle& InHandle, uint32 InSolverIterationCount)
+void FChaosEngineInterface::SetSolverProjectionIterationCount_AssumesLocked(const FPhysicsActorHandle& InHandle, int32 InSolverIterationCount)
 {
 	Chaos::FRigidBodyHandle_External& Body_External = InHandle->GetGameThreadAPI();
 	Body_External.SetProjectionSolverIterationCount(InSolverIterationCount);
diff --git a/Engine/Source/Runtime/PhysicsCore/Public/Chaos/ChaosEngineInterface.h b/Engine/Source/Runtime/PhysicsCore/Public/Chaos/ChaosEngineInterface.h
--- a/Engine/Source/Runtime/PhysicsCore/Public/Chaos/ChaosEngineInterface.h
+++ b/Engine/Source/Runtime/PhysicsCore/Public/Chaos/ChaosEngineInterface.h
@@ -415,9 +415,9 @@ public:
 	static PHYSICSCORE_API void SetCcdEnabled_AssumesLocked(const FPhysicsActorHandle& InActorReference,bool bIsCcdEnabled);
 	static PHYSICSCORE_API void SetMACDEnabled_AssumesLocked(const FPhysicsActorHandle& InActorReference, bool bIsMACDEnabled);
 	static PHYSICSCORE_API void SetPartialIslandSleepAllowed_AssumesLocked(const FPhysicsActorHandle& InActorReference, bool bIsPartialIslandSleepAllowed);
-	static PHYSICSCORE_API void SetPositionSolverIterationCount_AssumesLocked(const FPhysicsActorHandle& InActorReference, uint8 PositionSolverIterationCount);
-	static PHYSICSCORE_API void SetVelocitySolverIterationCount_AssumesLocked(const FPhysicsActorHandle& InActorReference, uint8 VelocitySolverIterationCount);
-	static PHYSICSCORE_API void SetProjectionSolverIterationCount_AssumesLocked(const FPhysicsActorHandle& InActorReference, uint8 ProjectionSolverIterationCount);
+	static PHYSICSCORE_API void SetPositionSolverIterationCount_AssumesLocked(const FPhysicsActorHandle& InActorReference, int32 PositionSolverIterationCount);
+	static PHYSICSCORE_API void SetVelocitySolverIterationCount_AssumesLocked(const FPhysicsActorHandle& InActorReference, int32 VelocitySolverIterationCount);
+	static PHYSICSCORE_API void SetProjectionSolverIterationCount_AssumesLocked(const FPhysicsActorHandle& InActorReference, int32 ProjectionSolverIterationCount);
 	static PHYSICSCORE_API void SetIgnoreAnalyticCollisions_AssumesLocked(const FPhysicsActorHandle& InActorReference,bool bIsCcdEnabled);
 
 	static PHYSICSCORE_API FTransform GetGlobalPose_AssumesLocked(const FPhysicsActorHandle& InActorReference);
@@ -496,9 +496,9 @@ public:
 
 	static PHYSICSCORE_API float GetStabilizationEnergyThreshold_AssumesLocked(const FPhysicsActorHandle& InHandle);
 	static PHYSICSCORE_API void SetStabilizationEnergyThreshold_AssumesLocked(const FPhysicsActorHandle& InHandle,float InThreshold);
-	static PHYSICSCORE_API void SetSolverPositionIterationCount_AssumesLocked(const FPhysicsActorHandle& InHandle,uint32 InSolverIterationCount);
-	static PHYSICSCORE_API void SetSolverVelocityIterationCount_AssumesLocked(const FPhysicsActorHandle& InHandle,uint32 InSolverIterationCount);
-	static PHYSICSCORE_API void SetSolverProjectionIterationCount_AssumesLocked(const FPhysicsActorHandle& InHandle, uint32 InSolverIterationCount);
+	static PHYSICSCORE_API void SetSolverPositionIterationCount_AssumesLocked(const FPhysicsActorHandle& InHandle, int32 InSolverIterationCount);
+	static PHYSICSCORE_API void SetSolverVelocityIterationCount_AssumesLocked(const FPhysicsActorHandle& InHandle, int32 InSolverIterationCount);
+	static PHYSICSCORE_API void SetSolverProjectionIterationCount_AssumesLocked(const FPhysicsActorHandle& InHandle, int32 InSolverIterationCount);
 	static PHYSICSCORE_API float GetWakeCounter_AssumesLocked(const FPhysicsActorHandle& InHandle);
 	static PHYSICSCORE_API void SetWakeCounter_AssumesLocked(const FPhysicsActorHandle& InHandle,float InWakeCounter);
 
