Update SpringMath to use unreal coding standards

diff --git a/Engine/Source/Runtime/Engine/Public/Animation/SpringMath.h b/Engine/Source/Runtime/Engine/Public/Animation/SpringMath.h
--- a/Engine/Source/Runtime/Engine/Public/Animation/SpringMath.h
+++ b/Engine/Source/Runtime/Engine/Public/Animation/SpringMath.h
@@ -10,9 +10,9 @@ namespace SpringMath
 {
 	// Reference https://theorangeduck.com/page/spring-roll-call 
 
-	constexpr float SpringHalfLifeToDamping(float HalfLife, float eps = 1e-5f)
+	constexpr float SpringHalfLifeToDamping(float HalfLife, float Eps = 1e-5f)
 	{
-		return (4.0f * UE_LN2) / (HalfLife + eps);
+		return (4.0f * UE_LN2) / (HalfLife + Eps);
 	}
 
 	constexpr float SpringDampingToHalfLife(float Damping)
@@ -30,19 +30,19 @@ namespace SpringMath
 		return HalfLife * UE_LN2;
 	}
 
-	constexpr float SpringSmoothingTimeToDamping(float SmoothingTime, float eps = 1e-5f)
+	constexpr float SpringSmoothingTimeToDamping(float SmoothingTime, float Eps = 1e-5f)
 	{
-		return 4.0f / FMath::Max(SmoothingTime, eps);
+		return 4.0f / FMath::Max(SmoothingTime, Eps);
 	}
 
-	constexpr float SpringSmoothingTimeToStrength(float SmoothingTime, float eps = 1e-5f)
+	constexpr float SpringSmoothingTimeToStrength(float SmoothingTime, float Eps = 1e-5f)
 	{
-		return 2.0f / FMath::Max(SmoothingTime, eps);
+		return 2.0f / FMath::Max(SmoothingTime, Eps);
 	}
 
-	constexpr float SpringStrengthToSmoothingTime(float Strength, float eps = 1e-5f)
+	constexpr float SpringStrengthToSmoothingTime(float Strength, float Eps = 1e-5f)
 	{
-		return 2.0f / FMath::Max(Strength, eps);
+		return 2.0f / FMath::Max(Strength, Eps);
 	}
 
 	constexpr float HalfLifeToLag(float HalfLife)
@@ -85,13 +85,13 @@ namespace SpringMath
 		float SmoothingTime,
 		float DeltaTime)
 	{
-		float y = SpringSmoothingTimeToDamping(SmoothingTime) / 2.0f;
-		T j0 = InOutX - TargetX;
-		T j1 = InOutV + j0 * y;
-		float eydt = FMath::InvExpApprox(y * DeltaTime);
+		float Y = SpringSmoothingTimeToDamping(SmoothingTime) / 2.0f;
+		T J0 = InOutX - TargetX;
+		T J1 = InOutV + J0 * Y;
+		float EyDt = FMath::InvExpApprox(Y * DeltaTime);
 
-		InOutX = eydt * (j0 + j1 * DeltaTime) + TargetX;
-		InOutV = eydt * (InOutV - j1 * y * DeltaTime);
+		InOutX = EyDt * (J0 + J1 * DeltaTime) + TargetX;
+		InOutV = EyDt * (InOutV - J1 * Y * DeltaTime);
 	}
 
 	/** Specialized quaternion version, uses FVector for angular velocity
@@ -109,17 +109,17 @@ namespace SpringMath
 		float SmoothingTime,
 		float DeltaTime)
 	{
-		float y = SpringSmoothingTimeToDamping(SmoothingTime) / 2.0f;
+		float Y = SpringSmoothingTimeToDamping(SmoothingTime) / 2.0f;
 
-		FQuat diff = InOutRotation * TargetRotation.Inverse();
-		diff.EnforceShortestArcWith(FQuat::Identity);
-		FVector j0 = QuatToScaledAngleAxis(diff);
-		FVector j1 = InOutAngularVelocityRadians + j0 * y;
+		FQuat Diff = InOutRotation * TargetRotation.Inverse();
+		Diff.EnforceShortestArcWith(FQuat::Identity);
+		FVector J0 = QuatToScaledAngleAxis(Diff);
+		FVector J1 = InOutAngularVelocityRadians + J0 * Y;
 
-		float eydt = FMath::InvExpApprox(y * DeltaTime);
+		float EyDt = FMath::InvExpApprox(Y * DeltaTime);
 
-		InOutRotation = QuatFromScaledAngleAxis(eydt * (j0 + j1 * DeltaTime)) * TargetRotation;
-		InOutAngularVelocityRadians = eydt * (InOutAngularVelocityRadians - j1 * y * DeltaTime);
+		InOutRotation = QuatFromScaledAngleAxis(EyDt * (J0 + J1 * DeltaTime)) * TargetRotation;
+		InOutAngularVelocityRadians = EyDt * (InOutAngularVelocityRadians - J1 * Y * DeltaTime);
 	}
 
 	/** A velocity spring will damp towards a target that follows a fixed linear target velocity, allowing control of the interpolation speed
@@ -135,7 +135,7 @@ namespace SpringMath
 	 * @param DeltaTime The timestep in seconds
 	 */
 	template <typename TFloat>
-	void VelocitySpringDamper(
+	void VelocitySpringDamperF(
 		TFloat& InOutX,
 		TFloat& InOutV,
 		TFloat& InOutXi,
@@ -144,17 +144,17 @@ namespace SpringMath
 		float SmoothingTime,
 		float DeltaTime)
 	{
-		float halflife = SpringSmoothingTimeToHalfLife(SmoothingTime);
+		float HalfLife = SpringSmoothingTimeToHalfLife(SmoothingTime);
 		static_assert(std::is_floating_point_v<TFloat>, "TFloat must be floating point");
 
-		TFloat x_diff = ((TargetX - InOutXi) > 0.0f ? 1.0f : -1.0f) * TargetV;
+		TFloat XDiff = ((TargetX - InOutXi) > 0.0f ? 1.0f : -1.0f) * TargetV;
 
-		float t_goal_future = HalfLifeToLag(halflife);
-		TFloat x_goal_future = fabs(TargetX - InOutXi) > t_goal_future * TargetV ? InOutXi + x_diff * t_goal_future : TargetX;
+		float TGoalFuture = HalfLifeToLag(HalfLife);
+		TFloat XGoalFuture = fabs(TargetX - InOutXi) > TGoalFuture * TargetV ? InOutXi + XDiff * TGoalFuture : TargetX;
 
-		SimpleSpringDamper(InOutX, InOutV, x_goal_future, halflife, DeltaTime);
+		SimpleSpringDamper(InOutX, InOutV, XGoalFuture, HalfLife, DeltaTime);
 
-		InOutXi = fabs(TargetX - InOutXi) > DeltaTime * TargetV ? InOutXi + x_diff * DeltaTime : TargetX;
+		InOutXi = fabs(TargetX - InOutXi) > DeltaTime * TargetV ? InOutXi + XDiff * DeltaTime : TargetX;
 	}
 
 	/** A velocity spring will damp towards a target that follows a fixed linear target velocity, allowing control of the interpolation speed
@@ -179,17 +179,17 @@ namespace SpringMath
 		float SmoothingTime,
 		float DeltaTime)
 	{
-		float halflife = SpringSmoothingTimeToHalfLife(SmoothingTime);
-		TVector x_diff = TargetX - InOutXi;
-		float x_diff_length = x_diff.Length();
-		TVector x_diff_dir = x_diff_length > FLT_EPSILON ? x_diff / x_diff_length : TVector::ZeroVector;
+		float HalfLife = SpringSmoothingTimeToHalfLife(SmoothingTime);
+		TVector XDiff = TargetX - InOutXi;
+		float XDiffLength = XDiff.Length();
+		TVector XDiffDir = XDiffLength > FLT_EPSILON ? XDiff / XDiffLength : TVector::ZeroVector;
 
-		float t_goal_future = HalfLifeToLag(halflife);
-		TVector x_goal_future = x_diff_length > t_goal_future * MaxV ? InOutXi + (x_diff_dir * MaxV) * t_goal_future : TargetX;
+		float TGoalFuture = HalfLifeToLag(HalfLife);
+		TVector XGoalFuture = XDiffLength > TGoalFuture * MaxV ? InOutXi + (XDiffDir * MaxV) * TGoalFuture : TargetX;
 
-		SimpleSpringDamper(InOutX, InOutV, x_goal_future, halflife, DeltaTime);
+		SimpleSpringDamper(InOutX, InOutV, XGoalFuture, HalfLife, DeltaTime);
 
-		InOutXi = x_diff_length > DeltaTime * MaxV ? InOutXi + x_diff_dir * MaxV * DeltaTime : TargetX;
+		InOutXi = XDiffLength > DeltaTime * MaxV ? InOutXi + XDiffDir * MaxV * DeltaTime : TargetX;
 	}
 
 	/** Update a position representing a character given a target velocity using a simple damped spring
@@ -213,15 +213,15 @@ namespace SpringMath
 		float VDeadzone = 1e-2,
 		float ADeadzone = 1e-4)
 	{
-		float y = SpringSmoothingTimeToDamping(SmoothingTime) / 2.0f;
-		FVector j0 = InOutVelocity - TargetVelocity;
-		FVector j1 = InOutAcceleration + j0 * y;
-		float eydt = FMath::InvExpApprox(y * DeltaTime);
+		float Y = SpringSmoothingTimeToDamping(SmoothingTime) / 2.0f;
+		FVector J0 = InOutVelocity - TargetVelocity;
+		FVector J1 = InOutAcceleration + J0 * Y;
+		float EyDt = FMath::InvExpApprox(Y * DeltaTime);
 
-		InOutPosition = eydt * (((-j1) / (y * y)) + ((-j0 - j1 * DeltaTime) / y)) +
-			(j1 / (y * y)) + j0 / y + TargetVelocity * DeltaTime + InOutPosition;
-		InOutVelocity = eydt * (j0 + j1 * DeltaTime) + TargetVelocity;
-		InOutAcceleration = eydt * (InOutAcceleration - j1 * y * DeltaTime);
+		InOutPosition = EyDt * (((-J1) / (Y * Y)) + ((-J0 - J1 * DeltaTime) / Y)) +
+			(J1 / (Y * Y)) + J0 / Y + TargetVelocity * DeltaTime + InOutPosition;
+		InOutVelocity = EyDt * (J0 + J1 * DeltaTime) + TargetVelocity;
+		InOutAcceleration = EyDt * (InOutAcceleration - J1 * Y * DeltaTime);
 
 		if ((TargetVelocity - InOutVelocity).SquaredLength() < FMath::Square(VDeadzone))
 		{
@@ -300,24 +300,24 @@ namespace SpringMath
 		float MacAccel,
 		float DeltaTime)
 	{
-		float halflife = SpringSmoothingTimeToHalfLife(SmoothingTime);
-		TVector v_diff = TargetVelocity - InOutVelocityIntermediate;
-		float v_diff_length = v_diff.Length();
-		TVector v_diff_dir = v_diff_length > 0.0001f ? v_diff / v_diff_length : TVector::ZeroVector;
-
-		float t_goal_future = HalfLifeToLag(halflife);
-		TVector max_v_future = v_diff_length > t_goal_future * MacAccel ? InOutVelocityIntermediate + (v_diff_dir * MacAccel) * t_goal_future : TargetVelocity;
-
-		float y = SpringHalfLifeToDamping(halflife) / 2.0f;
-		TVector j0 = InOutVelocity - max_v_future;
-		TVector j1 = InOutAccel + j0 * y;
-		float eydt = FMath::InvExpApprox(y * DeltaTime);
-
-		InOutPosition = eydt * (((-j1) / (y * y)) + ((-j0 - j1 * DeltaTime) / y)) +
-			(j1 / (y * y)) + j0 / y + max_v_future * DeltaTime + InOutPosition;
-		InOutVelocity = eydt * (j0 + j1 * DeltaTime) + max_v_future;
-		InOutAccel = eydt * (InOutAccel - j1 * y * DeltaTime);
-		InOutVelocityIntermediate = v_diff_length > DeltaTime * MacAccel ? InOutVelocityIntermediate + v_diff_dir * MacAccel * DeltaTime : TargetVelocity;
+		float HalfLife = SpringSmoothingTimeToHalfLife(SmoothingTime);
+		TVector VDiff = TargetVelocity - InOutVelocityIntermediate;
+		float VDiffLength = VDiff.Length();
+		TVector VDiffDir = VDiffLength > 0.0001f ? VDiff / VDiffLength : TVector::ZeroVector;
+
+		float TGoalFuture = HalfLifeToLag(HalfLife);
+		TVector MaxVFuture = VDiffLength > TGoalFuture * MacAccel ? InOutVelocityIntermediate + (VDiffDir * MacAccel) * TGoalFuture : TargetVelocity;
+
+		float Y = SpringHalfLifeToDamping(HalfLife) / 2.0f;
+		TVector J0 = InOutVelocity - MaxVFuture;
+		TVector J1 = InOutAccel + J0 * Y;
+		float EyDt = FMath::InvExpApprox(Y * DeltaTime);
+
+		InOutPosition = EyDt * (((-J1) / (Y * Y)) + ((-J0 - J1 * DeltaTime) / Y)) +
+			(J1 / (Y * Y)) + J0 / Y + MaxVFuture * DeltaTime + InOutPosition;
+		InOutVelocity = EyDt * (J0 + J1 * DeltaTime) + MaxVFuture;
+		InOutAccel = EyDt * (InOutAccel - J1 * Y * DeltaTime);
+		InOutVelocityIntermediate = VDiffLength > DeltaTime * MacAccel ? InOutVelocityIntermediate + VDiffDir * MacAccel * DeltaTime : TargetVelocity;
 	}
 
 
