Transform Provider Data Refactor - Remove hardcoding of anim banks, transform provider ID, lack of plugin extensibility, and fixed numerous order of operation issues with registering providers.

diff --git a/Engine/Plugins/Editor/EngineAssetDefinitions/Source/Private/Animation/AssetDefinition_AnimBank.h b/Engine/Plugins/Editor/EngineAssetDefinitions/Source/Private/Animation/AssetDefinition_AnimBank.h
--- a/Engine/Plugins/Editor/EngineAssetDefinitions/Source/Private/Animation/AssetDefinition_AnimBank.h
+++ b/Engine/Plugins/Editor/EngineAssetDefinitions/Source/Private/Animation/AssetDefinition_AnimBank.h
@@ -4,7 +4,7 @@
 
 #include "Animation/AnimBank.h"
 #include "Animation/AssetDefinition_AnimationAsset.h"
-
+#include "Animation/AssetDefinition_TransformProviderData.h"
 #include "AssetDefinition_AnimBank.generated.h"
 
 class UAssetDefinition_AnimationAsset;
@@ -32,3 +32,16 @@ public:
 
 	virtual TSharedPtr<SWidget> GetThumbnailOverlay(const FAssetData& AssetData) const override;
 };
+
+UCLASS()
+class UAssetDefinition_AnimBankData : public UAssetDefinition_TransformProviderData
+{
+	GENERATED_BODY()
+
+public:
+	// UAssetDefinition Begin
+	virtual FText GetAssetDisplayName() const override { return NSLOCTEXT("AssetTypeActions", "AssetTypeActions_AnimBankData", "Animation Bank Data"); }
+	virtual TSoftClassPtr<UObject> GetAssetClass() const override { return UAnimBankData::StaticClass(); }
+	// UAssetDefinition End
+};
+
diff --git a/Engine/Plugins/Editor/EngineAssetDefinitions/Source/Private/Animation/AssetDefinition_TransformProviderData.h b/Engine/Plugins/Editor/EngineAssetDefinitions/Source/Private/Animation/AssetDefinition_TransformProviderData.h
new file mode 100644
--- /dev/null
+++ b/Engine/Plugins/Editor/EngineAssetDefinitions/Source/Private/Animation/AssetDefinition_TransformProviderData.h
@@ -0,0 +1,18 @@
+// Copyright Epic Games, Inc. All Rights Reserved.
+
+#pragma once
+
+#include "Animation/TransformProviderData.h"
+#include "AssetDefinitionDefault.h"
+#include "AssetDefinition_TransformProviderData.generated.h"
+
+UCLASS()
+class UAssetDefinition_TransformProviderData : public UAssetDefinitionDefault
+{
+	GENERATED_BODY()
+
+public:
+	virtual FText GetAssetDisplayName() const override { return NSLOCTEXT("AssetTypeActions", "AssetTypeActions_TransformProviderData", "Transform Provider Data"); }
+	virtual FLinearColor GetAssetColor() const override { return FLinearColor(FColor(78, 40, 165)); }
+	virtual TSoftClassPtr<UObject> GetAssetClass() const override { return UTransformProviderData::StaticClass(); }
+};
diff --git a/Engine/Plugins/Experimental/FastGeoStreaming/Source/FastGeoStreaming/Private/FastGeoInstancedSkinnedMeshComponent.h b/Engine/Plugins/Experimental/FastGeoStreaming/Source/FastGeoStreaming/Private/FastGeoInstancedSkinnedMeshComponent.h
--- a/Engine/Plugins/Experimental/FastGeoStreaming/Source/FastGeoStreaming/Private/FastGeoInstancedSkinnedMeshComponent.h
+++ b/Engine/Plugins/Experimental/FastGeoStreaming/Source/FastGeoStreaming/Private/FastGeoInstancedSkinnedMeshComponent.h
@@ -49,9 +49,9 @@ protected:
 		return InstanceData.Num();
 	}
 
-	TArrayView<const struct FAnimBankItem> GetAnimBankItems() const
+	UTransformProviderData* GetTransformProvider() const
 	{
-		return TArrayView<const struct FAnimBankItem>();
+		return nullptr;
 	}
 	//~ End FFastGeoSkinnedMeshComponentBase interface
 
diff --git a/Engine/Plugins/Experimental/FastGeoStreaming/Source/FastGeoStreaming/Private/FastGeoWorldPartitionRuntimeCellTransformer.cpp b/Engine/Plugins/Experimental/FastGeoStreaming/Source/FastGeoStreaming/Private/FastGeoWorldPartitionRuntimeCellTransformer.cpp
--- a/Engine/Plugins/Experimental/FastGeoStreaming/Source/FastGeoStreaming/Private/FastGeoWorldPartitionRuntimeCellTransformer.cpp
+++ b/Engine/Plugins/Experimental/FastGeoStreaming/Source/FastGeoStreaming/Private/FastGeoWorldPartitionRuntimeCellTransformer.cpp
@@ -966,12 +966,10 @@ FFastGeoTransformResult UFastGeoWorldPartitionRuntimeCellTransformer::CanTransfo
 
 	if (UInstancedSkinnedMeshComponent* InstancedSkinnedMeshComponent = Cast<UInstancedSkinnedMeshComponent>(PrimitiveComponent))
 	{
-		for (const FAnimBankItem& AnimBankItem : InstancedSkinnedMeshComponent->GetAnimBankItems())
+		const UTransformProviderData* TransformProvider = InstancedSkinnedMeshComponent->GetTransformProvider();
+		if (TransformProvider != nullptr && TransformProvider->IsEnabled())
 		{
-			if (AnimBankItem.BankAsset != nullptr)
-			{
-				return { EFastGeoTransform::Reject, [PrimitiveComponent] { return FString::Printf(TEXT("Instanced skinned mesh component %s is animated"), *FastGeo::GetComponentShortName(PrimitiveComponent)); } };
-			}
+			return { EFastGeoTransform::Reject, [PrimitiveComponent] { return FString::Printf(TEXT("Instanced skinned mesh component %s is animated"), *FastGeo::GetComponentShortName(PrimitiveComponent)); } };
 		}
 
 		if (InstancedSkinnedMeshComponent->GetInstanceCount() == 0)
diff --git a/Engine/Plugins/PCG/Source/PCG/Private/Elements/PCGSkinnedMeshSpawner.cpp b/Engine/Plugins/PCG/Source/PCG/Private/Elements/PCGSkinnedMeshSpawner.cpp
--- a/Engine/Plugins/PCG/Source/PCG/Private/Elements/PCGSkinnedMeshSpawner.cpp
+++ b/Engine/Plugins/PCG/Source/PCG/Private/Elements/PCGSkinnedMeshSpawner.cpp
@@ -477,20 +477,20 @@ void FPCGSkinnedMeshSpawnerElement::SpawnSkinnedMeshInstances(FPCGSkinnedMeshSpa
 	ISKMC->SetNumCustomDataFloats(NumCustomDataFloats);
 
 	TArray<FTransform> Transforms;
-	TArray<int32> BankIndices;
+	TArray<int32> AnimationIndices;
 
 	Transforms.Reserve(NewInstanceCount);
-	BankIndices.Reserve(NewInstanceCount);
+	AnimationIndices.Reserve(NewInstanceCount);
 
 	// TODO: Remove allocs/copies
 	for (const FPCGSkinnedMeshInstance& Instance : InstanceList.Instances)
 	{
 		Transforms.Emplace(Instance.Transform);
-		BankIndices.Emplace(Instance.BankIndex);
+		AnimationIndices.Emplace(Instance.AnimationIndex);
 	}
 
 	// Populate the instances
-	TArray<FPrimitiveInstanceId> NewIds = ISKMC->AddInstances(Transforms, BankIndices, /*bShouldReturnIndices=*/true, /*bWorldSpace=*/true);
+	TArray<FPrimitiveInstanceId> NewIds = ISKMC->AddInstances(Transforms, AnimationIndices, /*bShouldReturnIndices=*/true, /*bWorldSpace=*/true);
 
 	// Copy new CustomData into the ABMC PerInstanceSMCustomData
 	if (NumCustomDataFloats > 0)
diff --git a/Engine/Plugins/PCG/Source/PCG/Private/MeshSelectors/PCGSkinnedMeshDescriptor.cpp b/Engine/Plugins/PCG/Source/PCG/Private/MeshSelectors/PCGSkinnedMeshDescriptor.cpp
--- a/Engine/Plugins/PCG/Source/PCG/Private/MeshSelectors/PCGSkinnedMeshDescriptor.cpp
+++ b/Engine/Plugins/PCG/Source/PCG/Private/MeshSelectors/PCGSkinnedMeshDescriptor.cpp
@@ -11,8 +11,8 @@ FPCGSoftSkinnedMeshComponentDescriptor::FPCGSoftSkinnedMeshComponentDescriptor()
 	// Override defaults from base class
 	ComponentClass = UInstancedSkinnedMeshComponent::StaticClass();
 	Mobility = EComponentMobility::Static;
-	// GW-TODO: bGenerateOverlapEvents = 0;
-	// GW-TODO: BodyInstance.SetCollisionProfileName(UCollisionProfile::NoCollision_ProfileName);
+	// TODO: bGenerateOverlapEvents = 0;
+	// TODO: BodyInstance.SetCollisionProfileName(UCollisionProfile::NoCollision_ProfileName);
 }
 
 FPCGSoftSkinnedMeshComponentDescriptor::FPCGSoftSkinnedMeshComponentDescriptor(const FSkinnedMeshComponentDescriptor& Other)
@@ -31,10 +31,12 @@ void FPCGSoftSkinnedMeshComponentDescriptor::InitComponent(UInstancedSkinnedMesh
 	Super::InitComponent(ABMComponent);
 }
 
-int32 FPCGSoftSkinnedMeshComponentDescriptor::GetOrAddBankIndex(const FSoftAnimBankItem& BankItem)
+int32 FPCGSoftSkinnedMeshComponentDescriptor::GetOrAddAnimationIndex(const FSoftAnimBankItem& BankItem)
 {
 	int32 BankIndex = INDEX_NONE;
 
+	// GW-TODO: 
+#if 0
 	for (int32 ExistingBankIndex = 0; ExistingBankIndex < BankItems.Num(); ++ExistingBankIndex)
 	{
 		if (BankItems[ExistingBankIndex] == BankItem)
@@ -49,6 +51,7 @@ int32 FPCGSoftSkinnedMeshComponentDescriptor::GetOrAddBankIndex(const FSoftAnimB
 		BankIndex = BankItems.Num();
 		BankItems.Emplace(BankItem);
 	}
+#endif
 
 	return BankIndex;
 }
diff --git a/Engine/Plugins/PCG/Source/PCG/Private/MeshSelectors/PCGSkinnedMeshSelector.cpp b/Engine/Plugins/PCG/Source/PCG/Private/MeshSelectors/PCGSkinnedMeshSelector.cpp
--- a/Engine/Plugins/PCG/Source/PCG/Private/MeshSelectors/PCGSkinnedMeshSelector.cpp
+++ b/Engine/Plugins/PCG/Source/PCG/Private/MeshSelectors/PCGSkinnedMeshSelector.cpp
@@ -34,8 +34,8 @@ namespace PCGSkinnedMeshSelector
 		for (FPCGSkinnedMeshInstanceList& InstanceList : InstanceLists)
 		{
 			if (InstanceList.Descriptor.SkinnedAsset == Asset &&
-				// GW-TODO: InstanceList.Descriptor.bReverseCulling == bReverseCulling &&
-				// GW-TODO: InstanceList.Descriptor.OverrideMaterials == MaterialOverrides &&
+				// TODO: InstanceList.Descriptor.bReverseCulling == bReverseCulling &&
+				// TODO: InstanceList.Descriptor.OverrideMaterials == MaterialOverrides &&
 				InstanceList.AttributePartitionIndex == AttributePartitionIndex)
 			{
 				return InstanceList;
@@ -44,8 +44,8 @@ namespace PCGSkinnedMeshSelector
 
 		FPCGSkinnedMeshInstanceList& NewInstanceList = InstanceLists.Emplace_GetRef(TemplateDescriptor);
 		NewInstanceList.Descriptor.SkinnedAsset = Asset;
-		// GW-TODO: NewInstanceList.Descriptor.OverrideMaterials = MaterialOverrides;
-		// GW-TODO: NewInstanceList.Descriptor.bReverseCulling = bReverseCulling;
+		// TODO: NewInstanceList.Descriptor.OverrideMaterials = MaterialOverrides;
+		// TODO: NewInstanceList.Descriptor.bReverseCulling = bReverseCulling;
 		NewInstanceList.AttributePartitionIndex = AttributePartitionIndex;
 		NewInstanceList.PointData = InPointData;
 
@@ -207,8 +207,8 @@ bool UPCGSkinnedMeshSelector::SelectInstances(
 			Overrides.Apply(AnyPointIndexOnThisPartition);
 		}
 
-		Context.PointBankValues.Reset();
-		Context.PointSequenceIndexValues.Reset();
+		Context.PointProviderValues.Reset();
+		Context.PointAnimationValues.Reset();
 
 		// Given partitioning is expensive, check if we're out of time for this frame
 		Context.bPartitionDone = true;
@@ -225,10 +225,10 @@ bool UPCGSkinnedMeshSelector::SelectInstances(
 	{
 		TRACE_CPUPROFILER_EVENT_SCOPE(UPCGSkinnedMeshSelector::SelectEntries::PushingPointsToInstanceLists);
 
-		if (Context.PointBankValues.IsEmpty() || Context.PointSequenceIndexValues.IsEmpty())
+		if (Context.PointProviderValues.IsEmpty() || Context.PointAnimationValues.IsEmpty())
 		{
-			PCGAttributeAccessorHelpers::ExtractAllValues(InPointData, BankAttribute, Context.PointBankValues, &Context);
-			PCGAttributeAccessorHelpers::ExtractAllValues(InPointData, SequenceIndexAttribute, Context.PointSequenceIndexValues, &Context);
+			PCGAttributeAccessorHelpers::ExtractAllValues(InPointData, BankAttribute, Context.PointProviderValues, &Context);
+			PCGAttributeAccessorHelpers::ExtractAllValues(InPointData, SequenceIndexAttribute, Context.PointAnimationValues, &Context);
 		}
 
 		// TODO: Revisit this when attribute partitioning is returned in a more optimized form
@@ -244,22 +244,6 @@ bool UPCGSkinnedMeshSelector::SelectInstances(
 				continue;
 			}
 
-			bool bSkip = false;
-
-			for (const FSoftAnimBankItem& BankItem : CurrentPartitionDescriptor.BankItems)
-			{
-				if (BankItem.BankAsset.IsNull())
-				{
-					bSkip = true;
-					break;
-				}
-			}
-
-			if (bSkip)
-			{
-				continue;
-			}
-
 			// Setup data for mesh bounds computation
 			if (OutPointData && Settings->bApplyMeshBoundsToPoints)
 			{
@@ -312,14 +296,15 @@ bool UPCGSkinnedMeshSelector::SelectInstances(
 							PartitionIndex
 						);
 
-						FSoftAnimBankItem BankItem;
-						{
-							BankItem.BankAsset = Context.PointBankValues[InstanceToPoint];
-							BankItem.SequenceIndex = Context.PointSequenceIndexValues[InstanceToPoint];
-						}
+						// GW-TODO: 
+						//FSoftAnimBankItem BankItem;
+						//{
+						//	BankItem.BankAsset = Context.PointBankValues[InstanceToPoint];
+						//	BankItem.SequenceIndex = Context.PointSequenceIndexValues[InstanceToPoint];
+						//}
 
 						FPCGSkinnedMeshInstance Instance;
-						Instance.BankIndex = InstanceList.Descriptor.GetOrAddBankIndex(BankItem);
+						Instance.AnimationIndex = INDEX_NONE;// InstanceList.Descriptor.GetOrAddAnimationIndex(BankItem);
 						Instance.Transform = Point.Transform;
 
 						InstanceList.Instances.Emplace(Instance);
@@ -352,11 +337,12 @@ bool UPCGSkinnedMeshSelector::SelectInstances(
 
 						FSoftAnimBankItem BankItem;
 						{
-							BankItem.BankAsset = Context.PointBankValues[InstanceToPoint];
-							BankItem.SequenceIndex = Context.PointSequenceIndexValues[InstanceToPoint];
+							//BankItem.BankAsset = Context.PointProviderValues[InstanceToPoint];
+							//BankItem.SequenceIndex = Context.PointAnimationValues[InstanceToPoint];
 						}
 
-						Instance.BankIndex = InstanceList.Descriptor.GetOrAddBankIndex(BankItem);
+						// GW-TODO: 
+						Instance.AnimationIndex = InstanceList.Descriptor.GetOrAddAnimationIndex(BankItem);
 						Instance.Transform = Points[InstanceToPoint].Transform;
 					}
 				}
diff --git a/Engine/Plugins/PCG/Source/PCG/Public/Elements/PCGSkinnedMeshSpawnerContext.h b/Engine/Plugins/PCG/Source/PCG/Public/Elements/PCGSkinnedMeshSpawnerContext.h
--- a/Engine/Plugins/PCG/Source/PCG/Public/Elements/PCGSkinnedMeshSpawnerContext.h
+++ b/Engine/Plugins/PCG/Source/PCG/Public/Elements/PCGSkinnedMeshSpawnerContext.h
@@ -56,8 +56,8 @@ struct FPCGSkinnedMeshSpawnerContext : public FPCGContext, public IPCGAsyncLoadi
 	// Keeping track of all touched resources to allow for correct cleanup on abort
 	TArray<TWeakObjectPtr<UPCGManagedISKMComponent>> TouchedResources;
 
-	TArray<FSoftObjectPath> PointBankValues;
-	TArray<int32> PointSequenceIndexValues;
+	TArray<FSoftObjectPath> PointProviderValues;
+	TArray<int32> PointAnimationValues;
 
 	void ResetInputIterationData();
 };
\ No newline at end of file
diff --git a/Engine/Plugins/PCG/Source/PCG/Public/MeshSelectors/PCGSkinnedMeshDescriptor.h b/Engine/Plugins/PCG/Source/PCG/Public/MeshSelectors/PCGSkinnedMeshDescriptor.h
--- a/Engine/Plugins/PCG/Source/PCG/Public/MeshSelectors/PCGSkinnedMeshDescriptor.h
+++ b/Engine/Plugins/PCG/Source/PCG/Public/MeshSelectors/PCGSkinnedMeshDescriptor.h
@@ -35,7 +35,8 @@ struct FPCGSoftSkinnedMeshComponentDescriptor : public FSoftSkinnedMeshComponent
 	PCG_API virtual void InitFrom(const UInstancedSkinnedMeshComponent* Component, bool bInitBodyInstance = true) override;
 	PCG_API virtual void InitComponent(UInstancedSkinnedMeshComponent* Component) const override;
 
-	PCG_API int32 GetOrAddBankIndex(const FSoftAnimBankItem& BankItem);
+	// GW-TODO: 
+	PCG_API int32 GetOrAddAnimationIndex(const FSoftAnimBankItem& BankItem);
 
 public:
 	UPROPERTY(EditAnywhere, Category="Component Settings|Tags", meta = (DisplayAfter = "ComponentClass"))
diff --git a/Engine/Plugins/PCG/Source/PCG/Public/MeshSelectors/PCGSkinnedMeshSelector.h b/Engine/Plugins/PCG/Source/PCG/Public/MeshSelectors/PCGSkinnedMeshSelector.h
--- a/Engine/Plugins/PCG/Source/PCG/Public/MeshSelectors/PCGSkinnedMeshSelector.h
+++ b/Engine/Plugins/PCG/Source/PCG/Public/MeshSelectors/PCGSkinnedMeshSelector.h
@@ -23,7 +23,7 @@ struct FPCGSkinnedMeshInstance
 	GENERATED_BODY()
 
 	FTransform Transform;
-	int32 BankIndex = 0;
+	int32 AnimationIndex = 0;
 };
 
 USTRUCT(BlueprintType)
diff --git a/Engine/Source/Developer/MeshUtilities/Private/MeshUtilities.cpp b/Engine/Source/Developer/MeshUtilities/Private/MeshUtilities.cpp
--- a/Engine/Source/Developer/MeshUtilities/Private/MeshUtilities.cpp
+++ b/Engine/Source/Developer/MeshUtilities/Private/MeshUtilities.cpp
@@ -4018,8 +4018,8 @@ public:
 		BeginSlowTask();
 
 		// Find wedge influences.
-		TArray<int32>	WedgeInfluenceIndices; // GW-TODO: Reserve
-		TMap<uint32, uint32> VertexIndexToInfluenceIndexMap; // GW-TODO: Reserve
+		TArray<int32>	WedgeInfluenceIndices; // TODO: Reserve
+		TMap<uint32, uint32> VertexIndexToInfluenceIndexMap; // TODO: Reserve
 
 		for (uint32 LookIdx = 0; LookIdx < (uint32)BuildData.Influences.Num(); LookIdx++)
 		{
diff --git a/Engine/Source/Editor/UnrealEd/Classes/Factories/AnimBankFactory.h b/Engine/Source/Editor/UnrealEd/Classes/Factories/AnimBankFactory.h
--- a/Engine/Source/Editor/UnrealEd/Classes/Factories/AnimBankFactory.h
+++ b/Engine/Source/Editor/UnrealEd/Classes/Factories/AnimBankFactory.h
@@ -5,6 +5,7 @@
 #include "CoreMinimal.h"
 #include "UObject/ObjectMacros.h"
 #include "Factories/Factory.h"
+#include "TransformProviderFactory.h"
 #include "AnimBankFactory.generated.h"
 
 struct FAssetData;
@@ -34,3 +35,14 @@ private:
 private:
 	TSharedPtr<SWindow> PickerWindow;
 };
+
+UCLASS(MinimalAPI)
+class UAnimBankDataFactory : public UTransformProviderDataFactory
+{
+	GENERATED_UCLASS_BODY()
+
+	// UFactory interface
+	virtual bool ConfigureProperties() override;
+	virtual bool ShouldShowInNewMenu() const override { return false; }
+	// End of UFactory interface
+};
\ No newline at end of file
diff --git a/Engine/Source/Editor/UnrealEd/Classes/Factories/TransformProviderFactory.h b/Engine/Source/Editor/UnrealEd/Classes/Factories/TransformProviderFactory.h
new file mode 100644
--- /dev/null
+++ b/Engine/Source/Editor/UnrealEd/Classes/Factories/TransformProviderFactory.h
@@ -0,0 +1,23 @@
+// Copyright Epic Games, Inc. All Rights Reserved.
+
+#pragma once
+
+#include "CoreMinimal.h"
+#include "UObject/ObjectMacros.h"
+#include "Factories/Factory.h"
+#include "TransformProviderFactory.generated.h"
+
+class UTransformProviderData;
+
+UCLASS(MinimalAPI)
+class UTransformProviderDataFactory : public UFactory
+{
+	GENERATED_UCLASS_BODY()
+
+	/** The type of transform provider data that will be created */
+	UPROPERTY(EditAnywhere, Category=CurveFactory)
+	TSubclassOf<UTransformProviderData> ProviderDataClass;
+
+	virtual bool ConfigureProperties() override;
+	virtual UObject* FactoryCreateNew(UClass* Class,UObject* InParent,FName Name,EObjectFlags Flags,UObject* Context,FFeedbackContext* Warn) override;
+};
\ No newline at end of file
diff --git a/Engine/Source/Editor/UnrealEd/Private/Factories/AnimBankFactory.cpp b/Engine/Source/Editor/UnrealEd/Private/Factories/AnimBankFactory.cpp
--- a/Engine/Source/Editor/UnrealEd/Private/Factories/AnimBankFactory.cpp
+++ b/Engine/Source/Editor/UnrealEd/Private/Factories/AnimBankFactory.cpp
@@ -85,4 +85,16 @@ void UAnimBankFactory::OnTargetSkeletonSelected(const FAssetData& SelectedAsset)
 	PickerWindow->RequestDestroyWindow();
 }
 
+UAnimBankDataFactory::UAnimBankDataFactory(const FObjectInitializer& ObjectInitializer)
+: Super(ObjectInitializer)
+{
+	SupportedClass = UAnimBankData::StaticClass();
+	ProviderDataClass = UAnimBankData::StaticClass();
+}
+
+bool UAnimBankDataFactory::ConfigureProperties()
+{
+	return true;
+}
+
 #undef LOCTEXT_NAMESPACE
diff --git a/Engine/Source/Editor/UnrealEd/Private/Factories/TransformProviderFactory.cpp b/Engine/Source/Editor/UnrealEd/Private/Factories/TransformProviderFactory.cpp
new file mode 100644
--- /dev/null
+++ b/Engine/Source/Editor/UnrealEd/Private/Factories/TransformProviderFactory.cpp
@@ -0,0 +1,122 @@
+// Copyright Epic Games, Inc. All Rights Reserved.
+
+#include "Factories/TransformProviderFactory.h"
+#include "Modules/ModuleManager.h"
+#include "Widgets/DeclarativeSyntaxSupport.h"
+#include "Widgets/SWindow.h"
+#include "Widgets/Layout/SBorder.h"
+#include "Styling/AppStyle.h"
+#include "Animation/Skeleton.h"
+#include "Animation/TransformProviderData.h"
+#include "Editor.h"
+
+#include "ClassViewerModule.h"
+#include "ClassViewerFilter.h"
+#include "Kismet2/SClassPickerDialog.h"
+#include "Kismet2/KismetEditorUtilities.h"
+
+#include "IContentBrowserSingleton.h"
+#include "ContentBrowserModule.h"
+
+#include UE_INLINE_GENERATED_CPP_BY_NAME(TransformProviderFactory)
+
+#define LOCTEXT_NAMESPACE "TransformProviderFactory"
+
+UTransformProviderDataFactory::UTransformProviderDataFactory(const FObjectInitializer& ObjectInitializer)
+: Super(ObjectInitializer)
+{
+	bCreateNew = true;
+	bEditAfterNew = true;
+	SupportedClass = UTransformProviderData::StaticClass();
+
+	ProviderDataClass = nullptr;
+}
+
+class FAssetClassParentFilter : public IClassViewerFilter
+{
+public:
+	FAssetClassParentFilter()
+	: DisallowedClassFlags(CLASS_None)
+	, bDisallowBlueprintBase(false)
+	{
+	}
+
+	/** All children of these classes will be included unless filtered out by another setting. */
+	TSet< const UClass* > AllowedChildrenOfClasses;
+
+	/** Disallowed class flags. */
+	EClassFlags DisallowedClassFlags;
+
+	/** Disallow blueprint base classes. */
+	bool bDisallowBlueprintBase;
+
+	virtual bool IsClassAllowed(const FClassViewerInitializationOptions& InInitOptions, const UClass* InClass, TSharedRef<FClassViewerFilterFuncs> InFilterFuncs) override
+	{
+		const bool bAllowed = !InClass->HasAnyClassFlags(DisallowedClassFlags)
+			&& InClass->CanCreateAssetOfClass()
+			&& InFilterFuncs->IfInChildOfClassesSet(AllowedChildrenOfClasses, InClass) != EFilterReturn::Failed;
+
+		if (bAllowed && bDisallowBlueprintBase)
+		{
+			if (FKismetEditorUtilities::CanCreateBlueprintOfClass(InClass))
+			{
+				return false;
+			}
+		}
+
+		return bAllowed;
+	}
+
+	virtual bool IsUnloadedClassAllowed(
+		const FClassViewerInitializationOptions& InInitOptions,
+		const TSharedRef<const IUnloadedBlueprintData> InUnloadedClassData,
+		TSharedRef<FClassViewerFilterFuncs> InFilterFuncs) override
+	{
+		return !InUnloadedClassData->HasAnyClassFlags(DisallowedClassFlags) && 
+			InFilterFuncs->IfInChildOfClassesSet(AllowedChildrenOfClasses, InUnloadedClassData) != EFilterReturn::Failed;
+	}
+};
+
+bool UTransformProviderDataFactory::ConfigureProperties()
+{
+	// Null the ProviderDataClass so we can get a clean class
+	ProviderDataClass = nullptr;
+
+	// Load the class viewer module to display a class picker
+	FClassViewerModule& ClassViewerModule = FModuleManager::LoadModuleChecked<FClassViewerModule>("ClassViewer");
+
+	// Fill in options
+	FClassViewerInitializationOptions Options;
+	Options.Mode = EClassViewerMode::ClassPicker;
+
+	TSharedPtr<FAssetClassParentFilter> Filter = MakeShareable(new FAssetClassParentFilter);
+	Options.ClassFilters.Add(Filter.ToSharedRef());
+
+	Filter->DisallowedClassFlags = CLASS_Abstract | CLASS_Deprecated | CLASS_NewerVersionExists;
+	Filter->AllowedChildrenOfClasses.Add(UTransformProviderData::StaticClass());
+
+	const FText TitleText = LOCTEXT("CreateTransformProviderDataOptions", "Pick Transform Provider Data Class");
+	UClass* ChosenClass = nullptr;
+	const bool bPressedOk = SClassPickerDialog::PickClass(TitleText, Options, ChosenClass, UTransformProviderData::StaticClass());
+
+	if (bPressedOk)
+	{
+		ProviderDataClass = ChosenClass;
+	}
+
+	return bPressedOk;
+}
+
+UObject* UTransformProviderDataFactory::FactoryCreateNew(UClass* Class, UObject* InParent, FName Name, EObjectFlags Flags, UObject* Context, FFeedbackContext* Warn)
+{
+	UTransformProviderData* NewProviderData = nullptr;
+	
+	if (ProviderDataClass != nullptr)
+	{
+		NewProviderData = NewObject<UTransformProviderData>(InParent, ProviderDataClass, Name, Flags);
+	}
+
+	return NewProviderData;
+}
+
+#undef LOCTEXT_NAMESPACE
diff --git a/Engine/Source/Runtime/Engine/Classes/Animation/AnimBank.h b/Engine/Source/Runtime/Engine/Classes/Animation/AnimBank.h
--- a/Engine/Source/Runtime/Engine/Classes/Animation/AnimBank.h
+++ b/Engine/Source/Runtime/Engine/Classes/Animation/AnimBank.h
@@ -3,6 +3,7 @@
 #pragma once
 
 #include "AnimationAsset.h"
+#include "Animation/TransformProviderData.h"
 #include "Experimental/Containers/RobinHoodHashTable.h"
 #include "HLOD/HLODBatchingPolicy.h"
 #include "Interfaces/Interface_AsyncCompilation.h"
@@ -21,6 +22,7 @@ class FQueuedThreadPool;
 enum class EQueuedWorkPriority : uint8;
 struct FPropertyChangedEvent;
 
+struct FRenderBounds;
 struct FReferenceSkeleton;
 class USkinnedAsset;
 class UInstancedSkinnedMeshComponent;
@@ -486,7 +488,7 @@ public:
 	TObjectPtr<class USkinnedAsset> SkinnedAsset;
 
 	UPROPERTY(EditAnywhere, Category = "Component Settings", meta = (DisplayAfter = "SkinnedAsset"))
-	TArray<FAnimBankItem> BankItems;
+	TObjectPtr<class UTransformProviderData> TransformProvider;
 };
 
 USTRUCT()
@@ -530,7 +532,7 @@ public:
 	TSoftObjectPtr<class USkinnedAsset> SkinnedAsset;
 
 	UPROPERTY()
-	TArray<FSoftAnimBankItem> BankItems;
+	TSoftObjectPtr<class UTransformProviderData> TransformProvider;
 };
 
 struct FAnimBankDesc
@@ -608,6 +610,52 @@ struct FAnimBankDescKeyFuncs : TDefaultMapHashableKeyFuncs<FAnimBankDesc, FAnimB
 	}
 };
 
+UCLASS(config=Engine, hidecategories=Object, MinimalAPI, BlueprintType)
+class UAnimBankData : public UTransformProviderData
+{
+	GENERATED_BODY()
+
+public:
+	virtual bool IsEnabled() const override;
+	virtual const FGuid& GetTransformProviderID() const override;
+	virtual const uint32 GetUniqueAnimationCount() const override;
+	virtual bool HasAnimationBounds() const override;
+	virtual bool GetAnimationBounds(uint32 AnimationIndex, FRenderBounds& OutBounds) const override;
+	virtual FTransformProviderRenderProxy* CreateRenderThreadResources(FSkeletalMeshSceneProxyBase* SceneProxy, FSceneInterface& Scene, FRHICommandListBase& RHICmdList) override;
+	virtual void DestroyRenderThreadResources(FTransformProviderRenderProxy* ProviderProxy) override;
+	virtual bool IsCompiling() const override;
+
+#if WITH_EDITOR
+	void BeginCacheForCookedPlatformData(const ITargetPlatform* TargetPlatform) override;
+	bool IsCachedCookedPlatformDataLoaded(const ITargetPlatform* TargetPlatform) override;
+#endif
+
+public:
+	UPROPERTY(EditAnywhere, BlueprintReadOnly, Category = TransformProvider, meta = (ShowOnlyInnerProperties))
+	TArray<struct FAnimBankItem> AnimBankItems;
+};
+
+class FAnimBankDataRenderProxy : public FTransformProviderRenderProxy
+{
+	friend class UAnimBankData;
+
+public:
+	ENGINE_API FAnimBankDataRenderProxy(UAnimBankData* BankData, FSkeletalMeshSceneProxyBase* SceneProxy, FSceneInterface& Scene);
+	ENGINE_API virtual ~FAnimBankDataRenderProxy();
+	ENGINE_API virtual void CreateRenderThreadResources(FRHICommandListBase& RHICmdList) override;
+	ENGINE_API virtual void DestroyRenderThreadResources() override;
+	ENGINE_API virtual const TConstArrayView<uint64> GetProviderData(bool& bOutValid) const override;
+
+private:
+	FSkeletalMeshSceneProxyBase* SceneProxy = nullptr;
+	TObjectPtr<const class USkinnedAsset> SkinnedAsset = nullptr;
+	FSceneInterface& Scene;
+	TArray<uint64> AnimBankIds;
+	TArray<FAnimBankRecordHandle> AnimBankHandles;
+	TArray<FAnimBankItem> AnimBankItems;
+	uint32 UniqueAnimationCount = 0;
+};
+
 using FAnimBankRecordMap = Experimental::TRobinHoodHashMap<FAnimBankDesc, FAnimBankRecord, FAnimBankDescKeyFuncs>;
 
 namespace UE::AnimBank
diff --git a/Engine/Source/Runtime/Engine/Classes/Components/InstancedSkinnedMeshComponent.h b/Engine/Source/Runtime/Engine/Classes/Components/InstancedSkinnedMeshComponent.h
--- a/Engine/Source/Runtime/Engine/Classes/Components/InstancedSkinnedMeshComponent.h
+++ b/Engine/Source/Runtime/Engine/Classes/Components/InstancedSkinnedMeshComponent.h
@@ -28,24 +28,24 @@ struct FSkinnedMeshInstanceData
 	FTransform3f Transform;
 
 	UPROPERTY(EditAnywhere, Category = Animation)
-	uint32 BankIndex;
+	uint32 AnimationIndex;
 
 	FSkinnedMeshInstanceData()
 	: Transform(FTransform3f::Identity)
-	, BankIndex(0)
+	, AnimationIndex(0)
 	{
 	}
 
-	FSkinnedMeshInstanceData(const FTransform3f& InTransform, uint32 InBankIndex)
+	FSkinnedMeshInstanceData(const FTransform3f& InTransform, uint32 InAnimationIndex)
 	: Transform(InTransform)
-	, BankIndex(InBankIndex)
+	, AnimationIndex(InAnimationIndex)
 	{
 	}
 
 	friend FArchive& operator<<(FArchive& Ar, FSkinnedMeshInstanceData& InstanceData)
 	{
 		Ar << InstanceData.Transform;
-		Ar << InstanceData.BankIndex;
+		Ar << InstanceData.AnimationIndex;
 		return Ar;
 	}
 };
@@ -57,6 +57,9 @@ class UInstancedSkinnedMeshComponent : public USkinnedMeshComponent
 	GENERATED_UCLASS_BODY()
 	
 protected:
+	UPROPERTY(EditAnywhere, BlueprintReadOnly, Category = Animation)
+	TObjectPtr<class UTransformProviderData> TransformProvider;
+
 	/** Array of instances, bulk serialized. */
 	UPROPERTY(EditAnywhere, SkipSerialization, DisplayName="Instances", Category=Instances, meta=(MakeEditWidget=true, EditFixedOrder))
 	TArray<FSkinnedMeshInstanceData> InstanceData;
@@ -94,12 +97,6 @@ protected:
 	UPROPERTY()
 	uint8 bInheritPerInstanceData : 1;
 
-	UPROPERTY(EditAnywhere, BlueprintReadOnly, Category = Animation, meta = (ShowOnlyInnerProperties))
-	TArray<struct FAnimBankItem> AnimBankItems;
-
-	UPROPERTY(EditAnywhere, BlueprintReadOnly, Category = Animation)
-	bool bAnimBankEnabled = false;
-
 	UPROPERTY()
 	FBox PrimitiveBoundsOverride;
 
@@ -170,11 +167,11 @@ public:
 
 	/** Add an instance to this component. Transform is given in local space of this component unless bWorldSpace is set. */
 	UFUNCTION(BlueprintCallable, Category = "Components|InstancedSkinnedMesh")
-	ENGINE_API FPrimitiveInstanceId AddInstance(const FTransform& InstanceTransform, int32 BankIndex, bool bWorldSpace = false);
+	ENGINE_API FPrimitiveInstanceId AddInstance(const FTransform& InstanceTransform, int32 AnimationIndex, bool bWorldSpace = false);
 
 	/** Add multiple instances to this component. Transform is given in local space of this component unless bWorldSpace is set. */
 	UFUNCTION(BlueprintCallable, Category = "Components|InstancedSkinnedMesh")
-	ENGINE_API TArray<FPrimitiveInstanceId> AddInstances(const TArray<FTransform>& Transforms, const TArray<int32>& BankIndices, bool bShouldReturnIds, bool bWorldSpace = false);
+	ENGINE_API TArray<FPrimitiveInstanceId> AddInstances(const TArray<FTransform>& Transforms, const TArray<int32>& AnimationIndices, bool bShouldReturnIds, bool bWorldSpace = false);
 
 	/** Update custom data for specific instance */
 	UFUNCTION(BlueprintCallable, Category = "Components|InstancedSkinnedMesh")
@@ -194,8 +191,8 @@ public:
 	UFUNCTION(BlueprintCallable, Category = "Components|InstancedSkinnedMesh")
 	ENGINE_API bool GetInstanceTransform(FPrimitiveInstanceId InstanceId, FTransform& OutInstanceTransform, bool bWorldSpace = false) const;
 
-	/** Get the bank index for the instance specified. Returns True on success. */
-	ENGINE_API bool GetInstanceBankIndex(FPrimitiveInstanceId InstanceId, int32& OutBankIndex) const;
+	/** Get the animation index for the instance specified. Returns True on success. */
+	ENGINE_API bool GetInstanceAnimationIndex(FPrimitiveInstanceId InstanceId, int32& OutAnimationIndex) const;
 
 	/** Remove the instance specified. Returns True on success. */
 	UFUNCTION(BlueprintCallable, Category = "Components|InstancedSkinnedMesh")
@@ -210,7 +207,10 @@ public:
 	ENGINE_API void ClearInstances();
 
 	UFUNCTION(BlueprintCallable, Category = "Components|InstancedSkinnedMesh")
-	ENGINE_API int32 AddAnimBankItem(const FAnimBankItem& BankItem);
+	ENGINE_API class UTransformProviderData* GetTransformProvider() const;
+
+	UFUNCTION(BlueprintCallable, Category = "Components|InstancedSkinnedMesh")
+	ENGINE_API void SetTransformProvider(class UTransformProviderData* InTransformProvider);
 
 	/**
 	 * Optimize the instance data by performing sorting according to spatial hash on the _source_ data.
@@ -223,10 +223,6 @@ public:
 	ENGINE_API bool IsEnabled() const;
 	ENGINE_API int32 GetInstanceCount() const;
 
-	ENGINE_API TArrayView<const struct FAnimBankItem> GetAnimBankItems() const;
-
-	ENGINE_API void SetAnimBankItems(TArrayView<const struct FAnimBankItem> InAnimBankItems);
-
 	int32 GetNumCustomDataFloats() const { return NumCustomDataFloats; }
 	const TArray<float>& GetInstanceCustomData() const { return InstanceCustomData; }
 	const TArray<FSkinnedMeshInstanceData>& GetInstanceData() const { return InstanceData; }
@@ -257,7 +253,7 @@ protected:
 
 	/** 
 	 */
-	FPrimitiveInstanceId AddInstanceInternal(int32 InstanceIndex, const FTransform& InstanceTransform, int32 BankIndex, bool bWorldSpace);
+	FPrimitiveInstanceId AddInstanceInternal(int32 InstanceIndex, const FTransform& InstanceTransform, int32 AnimationIndex, bool bWorldSpace);
 
 	/** 
 	 */
@@ -271,12 +267,12 @@ private:
 	void SetInstanceDataGPUOnly(bool bInInstancesGPUOnly);
 
 	/** Sets up new instance data to sensible defaults, creates physics counterparts if possible. */
-	void SetupNewInstanceData(FSkinnedMeshInstanceData& InOutNewInstanceData, int32 InInstanceIndex, const FTransform3f& InInstanceTransform, int32 BankIndex);
+	void SetupNewInstanceData(FSkinnedMeshInstanceData& InOutNewInstanceData, int32 InInstanceIndex, const FTransform3f& InInstanceTransform, int32 AnimationIndex);
 
 	static ENGINE_API bool ShouldForceRefPose();
-	static ENGINE_API bool ShouldUseSampledBounds();
+	static ENGINE_API bool ShouldUseAnimationBounds();
 
-	static ENGINE_API FSkeletalMeshObject* CreateMeshObject(const FInstancedSkinnedMeshSceneProxyDesc& InMeshDesc, TConstArrayView<FAnimBankItem> InAnimBankItems, FSkeletalMeshRenderData* InRenderData, ERHIFeatureLevel::Type InFeatureLevel);
+	static ENGINE_API FSkeletalMeshObject* CreateMeshObject(const FInstancedSkinnedMeshSceneProxyDesc& InMeshDesc, TObjectPtr<UTransformProviderData> InTransformProvider, FSkeletalMeshRenderData* InRenderData, ERHIFeatureLevel::Type InFeatureLevel);
 	static ENGINE_API FPrimitiveSceneProxy* CreateSceneProxy(const FInstancedSkinnedMeshSceneProxyDesc& Desc, bool bHideSkin, bool bShouldNaniteSkin, bool bIsEnabled, int32 MinLODIndex);
 
 	TSharedPtr<FInstanceDataSceneProxy, ESPMode::ThreadSafe> GetOrCreateInstanceDataSceneProxy();
diff --git a/Engine/Source/Runtime/Engine/Private/Animation/AnimBank.cpp b/Engine/Source/Runtime/Engine/Private/Animation/AnimBank.cpp
--- a/Engine/Source/Runtime/Engine/Private/Animation/AnimBank.cpp
+++ b/Engine/Source/Runtime/Engine/Private/Animation/AnimBank.cpp
@@ -18,7 +18,9 @@
 #include "Misc/App.h"
 #include "ReferenceSkeleton.h"
 #include "RenderUtils.h"
+#include "SceneInterface.h"
 #include "SkeletalRenderPublic.h"
+#include "SkeletalMeshSceneProxyBase.h"
 #include "UObject/FrameworkObjectVersion.h"
 #include "UObject/ObjectSaveContext.h"
 #include "UObject/UE5MainStreamObjectVersion.h"
@@ -1200,6 +1202,196 @@ void UAnimBank::NotifyOnGPUDataChanged()
 
 #endif
 
+bool UAnimBankData::IsEnabled() const
+{
+	for (const FAnimBankItem& BankItem : AnimBankItems)
+	{
+		if (BankItem.BankAsset == nullptr)
+		{
+			return false;
+		}
+
+	#if WITH_EDITORONLY_DATA
+		if (BankItem.BankAsset->Asset == nullptr)
+		{
+			return false;
+		}
+
+		if (BankItem.SequenceIndex >= BankItem.BankAsset->Sequences.Num())
+		{
+			return false;
+		}
+	#endif
+
+		if (BankItem.SequenceIndex >= BankItem.BankAsset->GetData().Entries.Num())
+		{
+			return false;
+		}
+	}
+
+	return bEnabled;
+}
+
+const FGuid& UAnimBankData::GetTransformProviderID() const
+{
+	static FGuid AnimBankGPUProviderId(ANIM_BANK_GPU_TRANSFORM_PROVIDER_GUID);
+	static FGuid AnimBankCPUProviderId(ANIM_BANK_CPU_TRANSFORM_PROVIDER_GUID);
+
+	static const auto AnimBankGPUVar = IConsoleManager::Get().FindTConsoleVariableDataBool(TEXT("r.AnimBank.GPU"));
+	if (AnimBankGPUVar && AnimBankGPUVar->GetValueOnAnyThread() == 1)
+	{
+		return AnimBankGPUProviderId;
+	}
+	else
+	{
+		return AnimBankCPUProviderId;
+	}
+}
+
+const uint32 UAnimBankData::GetUniqueAnimationCount() const
+{
+	return AnimBankItems.Num();
+}
+
+bool UAnimBankData::HasAnimationBounds() const
+{
+	return true;
+}
+
+bool UAnimBankData::GetAnimationBounds(uint32 AnimationIndex, FRenderBounds& OutBounds) const
+{
+	if (AnimationIndex < uint32(AnimBankItems.Num()))
+	{
+		const FAnimBankItem& BankItem = AnimBankItems[AnimationIndex];
+		if (BankItem.BankAsset != nullptr)
+		{
+			const FAnimBankData& BankData = BankItem.BankAsset->GetData();
+			if (BankItem.SequenceIndex < BankData.Entries.Num())
+			{
+				OutBounds = BankData.Entries[BankItem.SequenceIndex].SampledBounds;
+				return true;
+			}
+		}
+	}
+
+	return false;
+}
+
+FTransformProviderRenderProxy* UAnimBankData::CreateRenderThreadResources(FSkeletalMeshSceneProxyBase* SceneProxy, FSceneInterface& Scene, FRHICommandListBase& RHICmdList)
+{
+	FAnimBankDataRenderProxy* ProviderProxy = new FAnimBankDataRenderProxy(this, SceneProxy, Scene);
+	ProviderProxy->CreateRenderThreadResources(RHICmdList);
+	return ProviderProxy;
+}
+
+void UAnimBankData::DestroyRenderThreadResources(FTransformProviderRenderProxy* ProviderProxy)
+{
+	ProviderProxy->DestroyRenderThreadResources();
+	delete ProviderProxy;
+}
+
+bool UAnimBankData::IsCompiling() const
+{
+	for (const FAnimBankItem& BankItem : AnimBankItems)
+	{
+		if (BankItem.BankAsset && BankItem.BankAsset->IsCompiling())
+		{
+			return true;
+		}
+	}
+
+	return false;
+}
+
+#if WITH_EDITOR
+
+void UAnimBankData::BeginCacheForCookedPlatformData(const ITargetPlatform* TargetPlatform)
+{
+	for (FAnimBankItem& BankItem : AnimBankItems)
+	{
+		if (BankItem.BankAsset)
+		{
+			BankItem.BankAsset->BeginCacheForCookedPlatformData(TargetPlatform);
+		}
+	}
+}
+
+bool UAnimBankData::IsCachedCookedPlatformDataLoaded(const ITargetPlatform* TargetPlatform)
+{
+	return !IsCompiling();
+}
+
+#endif
+
+FAnimBankDataRenderProxy::FAnimBankDataRenderProxy(UAnimBankData* BankData, FSkeletalMeshSceneProxyBase* InSceneProxy, FSceneInterface& InScene)
+: SceneProxy(InSceneProxy)
+, Scene(InScene)
+{
+	UniqueAnimationCount = SceneProxy->GetUniqueAnimationCount();
+	SkinnedAsset = SceneProxy->GetSkinnedAsset();
+	AnimBankItems = BankData->AnimBankItems;
+	check(SkinnedAsset != nullptr);
+}
+
+FAnimBankDataRenderProxy::~FAnimBankDataRenderProxy()
+{
+}
+
+void FAnimBankDataRenderProxy::CreateRenderThreadResources(FRHICommandListBase& RHICmdList)
+{
+	TArray<FAnimBankDesc> Descs;
+	Descs.Reserve(AnimBankItems.Num());
+
+	for (const FAnimBankItem& Item : AnimBankItems)
+	{
+		if (Item.BankAsset == nullptr || SkinnedAsset == nullptr)
+		{
+			Descs.Emplace(FAnimBankDesc());
+			continue;
+		}
+
+		const FAnimBankData& BankData = Item.BankAsset->GetData();
+
+		if (Item.SequenceIndex >= BankData.Entries.Num())
+		{
+			Descs.Emplace(FAnimBankDesc());
+			continue;
+		}
+
+		const FAnimBankEntry& BankEntry = BankData.Entries[Item.SequenceIndex];
+
+		FAnimBankDesc& Desc = Descs.Emplace_GetRef();
+
+		Desc.BankAsset = MakeWeakObjectPtr(Item.BankAsset);
+		Desc.SequenceIndex = Item.SequenceIndex;
+		Desc.Asset = MakeWeakObjectPtr(SkinnedAsset);
+		Desc.Position = BankEntry.Position;
+		Desc.PlayRate = BankEntry.PlayRate;
+		Desc.bLooping = BankEntry.IsLooping() ? 1 : 0;
+		Desc.bAutoStart = BankEntry.IsAutoStart() ? 1 : 0;
+	}
+
+	AnimBankHandles = Scene.RegisterAnimBank(Descs);
+	AnimBankIds.Reset(AnimBankHandles.Num());
+	for (const FAnimBankRecordHandle& Handle : AnimBankHandles)
+	{
+		AnimBankIds.Emplace((uint64)Handle.Id);
+	}
+}
+
+void FAnimBankDataRenderProxy::DestroyRenderThreadResources()
+{
+	Scene.UnregisterAnimBank(AnimBankHandles);
+	AnimBankHandles.Reset();
+	AnimBankIds.Reset();
+}
+
+const TConstArrayView<uint64> FAnimBankDataRenderProxy::GetProviderData(bool& bOutValid) const
+{
+	bOutValid = (AnimBankIds.Num() == UniqueAnimationCount);
+	return AnimBankIds;
+}
+
 void FAnimBankSequence::ValidatePosition()
 {
 	float Min = 0.0f, Max = 0.0f;
@@ -1539,7 +1731,7 @@ FSkinnedMeshComponentDescriptor::FSkinnedMeshComponentDescriptor(const FSoftSkin
 	: FSkinnedMeshComponentDescriptorBase(InDescriptor)
 {
 	SkinnedAsset = InDescriptor.SkinnedAsset.LoadSynchronous();
-	Algo::Transform(InDescriptor.BankItems, BankItems, [](const FSoftAnimBankItem& BankItem) { return FAnimBankItem(BankItem); });
+	TransformProvider = InDescriptor.TransformProvider.LoadSynchronous();
 	// AB-TODO: Algo::Transform(InDescriptor.OverrideMaterials, OverrideMaterials, [](TSoftObjectPtr<UMaterialInterface> Material) { return Material.LoadSynchronous(); });
 	// AB-TODO: OverlayMaterial = InDescriptor.OverlayMaterial.LoadSynchronous();
 	// AB-TODO: Algo::Transform(InDescriptor.RuntimeVirtualTextures, RuntimeVirtualTextures, [](TSoftObjectPtr<URuntimeVirtualTexture> RVT) { return RVT.LoadSynchronous(); });
@@ -1558,7 +1750,7 @@ void FSkinnedMeshComponentDescriptor::InitFrom(const UInstancedSkinnedMeshCompon
 	Super::InitFrom(Template, bInitBodyInstance);
 
 	SkinnedAsset = Template->GetSkinnedAsset();
-	BankItems = Template->GetAnimBankItems();
+	TransformProvider = Template->GetTransformProvider();
 }
 
 uint32 FSkinnedMeshComponentDescriptor::ComputeHash() const
@@ -1577,7 +1769,7 @@ void FSkinnedMeshComponentDescriptor::InitComponent(UInstancedSkinnedMeshCompone
 {
 	Super::InitComponent(Component);
 	Component->SetSkinnedAsset(SkinnedAsset);
-	Component->SetAnimBankItems(BankItems);
+	Component->SetTransformProvider(TransformProvider);
 }
 
 void FSkinnedMeshComponentDescriptor::PostLoadFixup(UObject* Loader)
@@ -1595,7 +1787,7 @@ bool FSkinnedMeshComponentDescriptor::operator==(const FSkinnedMeshComponentDesc
 	return
 		(Hash == 0 || Other.Hash == 0 || Hash == Other.Hash) && // Check hash first, other checks are in case of Hash collision
 		SkinnedAsset == Other.SkinnedAsset &&
-		BankItems == Other.BankItems &&
+		TransformProvider == Other.TransformProvider &&
 		Super::operator==(Other);
 }
 
@@ -1614,7 +1806,7 @@ FSoftSkinnedMeshComponentDescriptor::FSoftSkinnedMeshComponentDescriptor(const F
 	: FSkinnedMeshComponentDescriptorBase(InDescriptor)
 {
 	SkinnedAsset = InDescriptor.SkinnedAsset;
-	Algo::Transform(InDescriptor.BankItems, BankItems, [](const FAnimBankItem& BankItem) { return FSoftAnimBankItem(BankItem); });
+	TransformProvider = InDescriptor.TransformProvider;
 	// AB-TODO: Algo::Transform(InDescriptor.OverrideMaterials, OverrideMaterials, [](TObjectPtr<UMaterialInterface> Material) { return Material; });
 	// AB-TODO: OverlayMaterial = InDescriptor.OverlayMaterial;
 	// AB-TODO: Algo::Transform(InDescriptor.RuntimeVirtualTextures, RuntimeVirtualTextures, [](TObjectPtr<URuntimeVirtualTexture> RVT) { return RVT; });
@@ -1632,7 +1824,7 @@ void FSoftSkinnedMeshComponentDescriptor::InitFrom(const UInstancedSkinnedMeshCo
 {
 	Super::InitFrom(Template, bInitBodyInstance);
 	SkinnedAsset = Template->GetSkinnedAsset();
-	BankItems = Template->GetAnimBankItems();
+	TransformProvider = Template->GetTransformProvider();
 }
 
 uint32 FSoftSkinnedMeshComponentDescriptor::ComputeHash() const
@@ -1651,12 +1843,7 @@ void FSoftSkinnedMeshComponentDescriptor::InitComponent(UInstancedSkinnedMeshCom
 	Super::InitComponent(Component);
 
 	Component->SetSkinnedAsset(SkinnedAsset.LoadSynchronous());
-	TArray<FAnimBankItem> AnimBankItems;
-	for (const FSoftAnimBankItem& BankItem : BankItems)
-	{
-		AnimBankItems.Emplace(BankItem);
-	}
-	Component->SetAnimBankItems(AnimBankItems);
+	Component->SetTransformProvider(TransformProvider.LoadSynchronous());
 }
 
 void FSoftSkinnedMeshComponentDescriptor::PostLoadFixup(UObject* Loader)
@@ -1674,7 +1861,7 @@ bool FSoftSkinnedMeshComponentDescriptor::operator==(const FSoftSkinnedMeshCompo
 	return
 		(Hash == 0 || Other.Hash == 0 || Hash == Other.Hash) && // Check hash first, other checks are in case of Hash collision
 		SkinnedAsset == Other.SkinnedAsset &&
-		BankItems == Other.BankItems &&
+		TransformProvider == Other.TransformProvider &&
 		Super::operator==(Other);
 }
 
diff --git a/Engine/Source/Runtime/Engine/Private/Animation/AnimBankCompiler.cpp b/Engine/Source/Runtime/Engine/Private/Animation/AnimBankCompiler.cpp
--- a/Engine/Source/Runtime/Engine/Private/Animation/AnimBankCompiler.cpp
+++ b/Engine/Source/Runtime/Engine/Private/Animation/AnimBankCompiler.cpp
@@ -219,9 +219,27 @@ void FAnimBankCompilingManager::PostCompilation(UAnimBank* AnimBank)
 		AnimBank->InitResources();
 
 		FObjectCacheContextScope ObjectCacheScope;
-		for (UInstancedSkinnedMeshComponent* Component : ObjectCacheScope.GetContext().GetInstancedSkinnedMeshComponents(AnimBank))
+
+		for (UInstancedSkinnedMeshComponent* Component : ObjectCacheScope.GetContext().GetInstancedSkinnedMeshComponents())
 		{
-			Component->PostAssetCompilation();
+			UTransformProviderData* ProviderData = Component->GetTransformProvider();
+			if (UAnimBankData* AnimBankData = Cast<UAnimBankData>(ProviderData))
+			{
+				bool bMatched = false;
+				for (const FAnimBankItem& BankItem : AnimBankData->AnimBankItems)
+				{
+					if (BankItem.BankAsset == AnimBank)
+					{
+						bMatched = true;
+						break;
+					}
+				}
+
+				if (bMatched)
+				{
+					Component->PostAssetCompilation();
+				}
+			}
 		}
 
 		// Calling this delegate during app exit might be quite dangerous and lead to crash
diff --git a/Engine/Source/Runtime/Engine/Private/Components/InstancedSkinnedMeshComponent.cpp b/Engine/Source/Runtime/Engine/Private/Components/InstancedSkinnedMeshComponent.cpp
--- a/Engine/Source/Runtime/Engine/Private/Components/InstancedSkinnedMeshComponent.cpp
+++ b/Engine/Source/Runtime/Engine/Private/Components/InstancedSkinnedMeshComponent.cpp
@@ -41,10 +41,10 @@ static TAutoConsoleVariable<int32> CVarInstancedSkinnedMeshesForceRefPose(
 	ECVF_RenderThreadSafe
 );
 
-static TAutoConsoleVariable<int32> CVarInstancedSkinnedMeshesSampledBounds(
-	TEXT("r.InstancedSkinnedMeshes.SampledBounds"),
+static TAutoConsoleVariable<int32> CVarInstancedSkinnedMeshesAnimationBounds(
+	TEXT("r.InstancedSkinnedMeshes.AnimationBounds"),
 	1,
-	TEXT("Whether to use sampled bounds for anim bank meshes"),
+	TEXT("Whether to use animation bounds for instanced skinned meshes"),
 	FConsoleVariableDelegate::CreateLambda([](IConsoleVariable* InVariable)
 	{
 		FGlobalComponentRecreateRenderStateContext Context;
@@ -64,8 +64,8 @@ public:
 
 	FInstancedSkinnedMeshSceneProxy(const Nanite::FMaterialAudit& MaterialAudit, const FInstancedSkinnedMeshSceneProxyDesc& InMeshDesc, FSkeletalMeshRenderData* InRenderData)
 		: Super(MaterialAudit, InMeshDesc, InRenderData, false /* bAllowScale */)
+		, TransformProviderData(InMeshDesc.TransformProvider)
 		, AnimationMinScreenSize(InMeshDesc.AnimationMinScreenSize)
-		, bAnimBankEnabled(InMeshDesc.bAnimBankEnabled)
 	{
 #if WITH_EDITOR
 		const bool bSupportInstancePicking = HasPerInstanceHitProxies() && SMInstanceElementDataUtil::SMInstanceElementsEnabled();
@@ -81,8 +81,6 @@ public:
 			}
 		}
 #endif
-		const bool bForceRefPose = CVarInstancedSkinnedMeshesForceRefPose.GetValueOnAnyThread() != 0;
-		const bool bUseAnimBank = !bForceRefPose && bAnimBankEnabled && InMeshDesc.AnimBankItems.Num() > 0;
 
 		InstanceMinDrawDistance = InMeshDesc.InstanceMinDrawDistance;
 		InstanceStartCullDistance = InMeshDesc.InstanceStartCullDistance;
@@ -100,97 +98,52 @@ public:
 		// ISKM doesn't currently support skinning in ray tracing geometry
 		bDynamicRayTracingGeometry = false;
 
-		if (bUseAnimBank)
-		{
-			static FGuid AnimBankGPUProviderId(ANIM_BANK_GPU_TRANSFORM_PROVIDER_GUID);
-			static FGuid AnimBankCPUProviderId(ANIM_BANK_CPU_TRANSFORM_PROVIDER_GUID);
-
-			static const auto AnimBankGPUVar = IConsoleManager::Get().FindTConsoleVariableDataBool(TEXT("r.AnimBank.GPU"));
-			if (AnimBankGPUVar && AnimBankGPUVar->GetValueOnAnyThread() == 1)
-			{
-				TransformProviderId = AnimBankGPUProviderId;
-			}
-			else
-			{
-				TransformProviderId = AnimBankCPUProviderId;
-			}
+		UniqueAnimationCount = 1; // Ref Pose
+		static FGuid RefPoseProviderId(REF_POSE_TRANSFORM_PROVIDER_GUID);
+		TransformProviderId = RefPoseProviderId;
 
-			AnimBankItems = InMeshDesc.AnimBankItems;
-			AnimBankIds.Reset(AnimBankItems.Num());
-			AnimBankHandles.Reset(AnimBankItems.Num());
-			UniqueAnimationCount = AnimBankItems.Num();
-		}
-		else
+		const bool bForceRefPose = CVarInstancedSkinnedMeshesForceRefPose.GetValueOnAnyThread() != 0;
+		if (!bForceRefPose && GetSkinnedAsset() && TransformProviderData != nullptr && TransformProviderData->IsEnabled())
 		{
-			UniqueAnimationCount = 1; // Ref Pose
-			static FGuid RefPoseProviderId(REF_POSE_TRANSFORM_PROVIDER_GUID);
-			TransformProviderId = RefPoseProviderId;
+			TransformProviderId = TransformProviderData->GetTransformProviderID();
+			UniqueAnimationCount = TransformProviderData->GetUniqueAnimationCount();
 		}
 	}
 
 	virtual void CreateRenderThreadResources(FRHICommandListBase& RHICmdList) override
 	{
 		Super::CreateRenderThreadResources(RHICmdList);
-
-		if (bAnimBankEnabled)
+		if (TransformProviderData != nullptr && TransformProviderData->GetTransformProviderID() == TransformProviderId)
 		{
-			TArray<FAnimBankDesc> Descs;
-			Descs.Reserve(AnimBankItems.Num());
-
-			for (const FAnimBankItem& Item : AnimBankItems)
-			{
-				if (Item.BankAsset == nullptr || SkinnedAsset == nullptr)
-				{
-					Descs.Emplace(FAnimBankDesc());
-					continue;
-				}
-
-				const FAnimBankData& BankData = Item.BankAsset->GetData();
-
-				if (Item.SequenceIndex >= BankData.Entries.Num())
-				{
-					Descs.Emplace(FAnimBankDesc());
-					continue;
-				}
-
-				const FAnimBankEntry& BankEntry = BankData.Entries[Item.SequenceIndex];
-
-				FAnimBankDesc& Desc = Descs.Emplace_GetRef();
-
-				Desc.BankAsset		= MakeWeakObjectPtr(Item.BankAsset);
-				Desc.SequenceIndex	= Item.SequenceIndex;
-				Desc.Asset			= MakeWeakObjectPtr(SkinnedAsset);
-				Desc.Position		= BankEntry.Position;
-				Desc.PlayRate		= BankEntry.PlayRate;
-				Desc.bLooping		= BankEntry.IsLooping()   ? 1 : 0;
-				Desc.bAutoStart		= BankEntry.IsAutoStart() ? 1 : 0;
-			}
-
-			AnimBankHandles = GetScene().RegisterAnimBank(Descs);
-			AnimBankIds.Reset(AnimBankHandles.Num());
-			for (const FAnimBankRecordHandle& Handle : AnimBankHandles)
-			{
-				AnimBankIds.Emplace((uint64)Handle.Id);
-			}
+			TransformProviderProxy = TransformProviderData->CreateRenderThreadResources(this, GetScene(), RHICmdList);
+		}
+		else
+		{
+			TransformProviderProxy = nullptr;
 		}
 	}
 
 	virtual void DestroyRenderThreadResources() override
 	{
-		if (bAnimBankEnabled)
+		if (TransformProviderProxy != nullptr)
 		{
-			GetScene().UnregisterAnimBank(AnimBankHandles);
-			AnimBankHandles.Reset();
-			AnimBankIds.Reset();
+			TransformProviderData->DestroyRenderThreadResources(TransformProviderProxy);
+			TransformProviderProxy = nullptr;
 		}
+
 		Super::DestroyRenderThreadResources();
 	}
 
 	// FSkinnedSceneProxy interface
 	virtual const TConstArrayView<uint64> GetAnimationProviderData(bool& bOutValid) const override
 	{
-		bOutValid = (AnimBankIds.Num() == UniqueAnimationCount);
-		return AnimBankIds; // TODO: Rename generic
+		if (TransformProviderProxy != nullptr)
+		{
+			return TransformProviderProxy->GetProviderData(bOutValid);
+		}
+
+		bOutValid = false;
+		return {};
 	}
 
 #if RHI_RAYTRACING
@@ -243,17 +196,13 @@ public:
 	virtual float GetAnimationMinScreenSize() const  override { return AnimationMinScreenSize; }
 
 private:
+	TObjectPtr<class UTransformProviderData> TransformProviderData;
+	FTransformProviderRenderProxy* TransformProviderProxy = nullptr;
 	TSharedPtr<FInstanceDataSceneProxy, ESPMode::ThreadSafe> InstanceDataSceneProxy; 
-	TArray<uint64> AnimBankIds;
-	TArray<FAnimBankRecordHandle> AnimBankHandles;
-	TArray<FAnimBankItem> AnimBankItems;
-	uint64 BankRegistration = 0;
 	float AnimationMinScreenSize = 0.0f;
 	uint32 InstanceMinDrawDistance = 0;
 	uint32 InstanceStartCullDistance = 0;
 	uint32 InstanceEndCullDistance = 0;
-
-	uint8 bAnimBankEnabled : 1 = false;
 };
 
 /** 
@@ -265,7 +214,7 @@ class FInstancedSkinnedMeshData
 public:
 	FInstancedSkinnedMeshData(
 		const USkinnedAsset* InSkinnedAsset,
-		TConstArrayView<FAnimBankItem> AnimBankItems,
+		TObjectPtr<UTransformProviderData> InTransformProviderData,
 		int32 InLODIndex,
 		EPreviousBoneTransformUpdateMode InPreviousBoneTransformUpdateMode
 	);
@@ -300,7 +249,7 @@ private:
 
 FInstancedSkinnedMeshData::FInstancedSkinnedMeshData(
 	const USkinnedAsset* InSkinnedAsset,
-	TConstArrayView<FAnimBankItem> AnimBankItems,
+	TObjectPtr<UTransformProviderData> InTransformProviderData,
 	int32 InLODIndex,
 	EPreviousBoneTransformUpdateMode InPreviousBoneTransformUpdateMode
 )
@@ -358,9 +307,9 @@ void FInstancedSkinnedMeshData::GetResourceSizeEx(FResourceSizeEx& CumulativeRes
 class FInstancedSkinnedMeshObject : public FSkeletalMeshObject
 {
 public:
-	FInstancedSkinnedMeshObject(const FSkinnedMeshSceneProxyDesc& InMeshDesc, TConstArrayView<FAnimBankItem> InAnimBankItems, FSkeletalMeshRenderData* InRenderData, ERHIFeatureLevel::Type InFeatureLevel)
+	FInstancedSkinnedMeshObject(const FSkinnedMeshSceneProxyDesc& InMeshDesc, TObjectPtr<UTransformProviderData> InTransformProvider, FSkeletalMeshRenderData* InRenderData, ERHIFeatureLevel::Type InFeatureLevel)
 		: FSkeletalMeshObject(InMeshDesc, InRenderData, InFeatureLevel)
-		, AnimBankItems(InAnimBankItems)
+		, TransformProvider(InTransformProvider)
 	{
 		for (int32 LODIndex = 0; LODIndex < InRenderData->LODRenderData.Num(); ++LODIndex)
 		{
@@ -371,7 +320,7 @@ public:
 	}
 
 	FInstancedSkinnedMeshObject(USkinnedMeshComponent* InComponent, FSkeletalMeshRenderData* InRenderData, ERHIFeatureLevel::Type InFeatureLevel)
-	: FInstancedSkinnedMeshObject(FSkinnedMeshSceneProxyDesc(InComponent), CastChecked<UInstancedSkinnedMeshComponent>(InComponent)->AnimBankItems, InRenderData, InFeatureLevel)
+	: FInstancedSkinnedMeshObject(FSkinnedMeshSceneProxyDesc(InComponent), CastChecked<UInstancedSkinnedMeshComponent>(InComponent)->TransformProvider, InRenderData, InFeatureLevel)
 	{
 	}
 
@@ -420,7 +369,7 @@ public:
 	{
 		// Create the new dynamic data for use by the rendering thread
 		// this data is only deleted when another update is sent
-		TUniquePtr<FInstancedSkinnedMeshData> NewDynamicData = MakeUnique<FInstancedSkinnedMeshData>(InSkinnedAsset, AnimBankItems, LODIndex, PreviousBoneTransformUpdateMode);
+		TUniquePtr<FInstancedSkinnedMeshData> NewDynamicData = MakeUnique<FInstancedSkinnedMeshData>(InSkinnedAsset, TransformProvider, LODIndex, PreviousBoneTransformUpdateMode);
 
 		uint64 FrameNumberToPrepare = GFrameCounter;
 		uint32 RevisionNumber = 0;
@@ -657,7 +606,7 @@ private:
 		}
 	};
 
-	TConstArrayView<FAnimBankItem> AnimBankItems;
+	TObjectPtr<UTransformProviderData> TransformProvider;
 	TArray<FInstancedSkinnedMeshObjectLOD> LODs;
 
 	mutable int32 CachedLOD;
@@ -696,29 +645,29 @@ bool UInstancedSkinnedMeshComponent::ShouldForceRefPose()
 	return CVarInstancedSkinnedMeshesForceRefPose.GetValueOnAnyThread() != 0;
 }
 
-bool UInstancedSkinnedMeshComponent::ShouldUseSampledBounds()
+bool UInstancedSkinnedMeshComponent::ShouldUseAnimationBounds()
 {
-	return CVarInstancedSkinnedMeshesSampledBounds.GetValueOnAnyThread() != 0;
+	return CVarInstancedSkinnedMeshesAnimationBounds.GetValueOnAnyThread() != 0;
 }
 
 struct FSkinnedMeshInstanceData_Deprecated
 {
 	FMatrix Transform;
-	uint32 BankIndex;
+	uint32 AnimationIndex;
 	uint32 Padding[3]; // Need to respect 16 byte alignment for bulk-serialization
 
 	FSkinnedMeshInstanceData_Deprecated()
 	: Transform(FMatrix::Identity)
-	, BankIndex(0)
+	, AnimationIndex(0)
 	{
 		Padding[0] = 0;
 		Padding[1] = 0;
 		Padding[2] = 0;
 	}
 
-	FSkinnedMeshInstanceData_Deprecated(const FMatrix& InTransform, uint32 InBankIndex)
+	FSkinnedMeshInstanceData_Deprecated(const FMatrix& InTransform, uint32 InAnimationIndex)
 	: Transform(InTransform)
-	, BankIndex(InBankIndex)
+	, AnimationIndex(InAnimationIndex)
 	{
 		Padding[0] = 0;
 		Padding[1] = 0;
@@ -730,7 +679,7 @@ struct FSkinnedMeshInstanceData_Deprecated
 		// @warning BulkSerialize: FSkinnedMeshInstanceData is serialized as memory dump
 		// See TArray::BulkSerialize for detailed description of implied limitations.
 		Ar << InstanceData.Transform;
-		Ar << InstanceData.BankIndex;
+		Ar << InstanceData.AnimationIndex;
 		Ar << InstanceData.Padding[0];
 		Ar << InstanceData.Padding[1];
 		Ar << InstanceData.Padding[2];
@@ -772,7 +721,7 @@ void UInstancedSkinnedMeshComponent::Serialize(FArchive& Ar)
 				TempInstanceData.Reserve(TempInstanceData_Deprecated.Num());
 				for (const auto& Item : TempInstanceData_Deprecated)
 				{
-					TempInstanceData.Emplace(FTransform3f(FMatrix44f(Item.Transform)), Item.BankIndex);
+					TempInstanceData.Emplace(FTransform3f(FMatrix44f(Item.Transform)), Item.AnimationIndex);
 				}
 			}
 			else
@@ -869,14 +818,14 @@ int32 UInstancedSkinnedMeshComponent::GetInstanceCount() const
 	return bIsInstanceDataGPUOnly ? NumInstancesGPUOnly : InstanceData.Num();
 }
 
-TArrayView<const struct FAnimBankItem> UInstancedSkinnedMeshComponent::GetAnimBankItems() const
+UTransformProviderData* UInstancedSkinnedMeshComponent::GetTransformProvider() const
 {
-	return MakeArrayView(AnimBankItems);
+	return TransformProvider.Get();
 }
 
-void UInstancedSkinnedMeshComponent::SetAnimBankItems(TArrayView<const struct FAnimBankItem> InAnimBankItems)
+void UInstancedSkinnedMeshComponent::SetTransformProvider(UTransformProviderData* InTransformProvider)
 {
-	AnimBankItems = InAnimBankItems;
+	TransformProvider = InTransformProvider;
 	// We use the transform dirty state to drive the update of the animation data (to defer the need to add more bits), so we mark those as dirty here.
 	InstanceDataManager.TransformsChangedAll();
 	MarkRenderStateDirty();
@@ -946,10 +895,10 @@ void UInstancedSkinnedMeshComponent::SetInstanceDataGPUOnly(bool bInInstancesGPU
 	}
 }
 
-void UInstancedSkinnedMeshComponent::SetupNewInstanceData(FSkinnedMeshInstanceData& InOutNewInstanceData, int32 InInstanceIndex, const FTransform3f& InInstanceTransform, int32 InBankIndex)
+void UInstancedSkinnedMeshComponent::SetupNewInstanceData(FSkinnedMeshInstanceData& InOutNewInstanceData, int32 InInstanceIndex, const FTransform3f& InInstanceTransform, int32 InAnimationIndex)
 {
 	InOutNewInstanceData.Transform = InInstanceTransform;
-	InOutNewInstanceData.BankIndex = InBankIndex;
+	InOutNewInstanceData.AnimationIndex = InAnimationIndex;
 
 	if (bPhysicsStateCreated)
 	{
@@ -990,7 +939,6 @@ FInstanceDataManagerSourceDataDesc UInstancedSkinnedMeshComponent::GetComponentD
 {
 	FInstanceDataManagerSourceDataDesc ComponentDesc;
 
-
 	ComponentDesc.PrimitiveMaterialDesc = GetUsedMaterialPropertyDesc(GetFeatureLevelShaderPlatform_Checked(FeatureLevel));
 
 	FInstanceDataFlags Flags;
@@ -1001,14 +949,15 @@ FInstanceDataManagerSourceDataDesc UInstancedSkinnedMeshComponent::GetComponentD
 #endif
 	
 	const bool bForceRefPose = CVarInstancedSkinnedMeshesForceRefPose.GetValueOnAnyThread() != 0;
-	const bool bUseAnimBank = !bForceRefPose && bAnimBankEnabled && AnimBankItems.Num() > 0;
+	const bool bValidTransformProvider = !bForceRefPose && TransformProvider != nullptr && TransformProvider->IsEnabled();
+
+	const uint32 UniqueAnimationCount = TransformProvider ? TransformProvider->GetUniqueAnimationCount() : 1;
 
 	Flags.bHasPerInstanceHierarchyOffset = false;
-	Flags.bHasPerInstanceLocalBounds = bUseAnimBank && AnimBankItems.Num() > 1;
+	Flags.bHasPerInstanceLocalBounds = TransformProvider ? TransformProvider->HasAnimationBounds() : false;
 	Flags.bHasPerInstanceDynamicData = false;
 	Flags.bHasPerInstanceSkinningData = true;
 
-
 	Flags.bHasPerInstanceLMSMUVBias = false;//IsStaticLightingAllowed();
 
 	ComponentDesc.Flags = Flags;
@@ -1032,37 +981,36 @@ FInstanceDataManagerSourceDataDesc UInstancedSkinnedMeshComponent::GetComponentD
 
 		ChangeSet.GetSkinningDataWriter().Gather(
 			[&](int32 InstanceIndex) -> uint32 
-			{ 
-				return bUseAnimBank ? (InstanceData[InstanceIndex].BankIndex * MaxBoneTransformCount * 2u) : 0;
+			{
+				const uint32 AnimationIndex = InstanceData[InstanceIndex].AnimationIndex;
+				if (bValidTransformProvider && AnimationIndex < UniqueAnimationCount)
+				{
+					return AnimationIndex * MaxBoneTransformCount * 2u;
+				}
+
+				return 0u;
 			});
 
 		ChangeSet.GetLocalBoundsWriter().Gather(
 			[&](int32 InstanceIndex) -> FRenderBounds 
 			{
-				if (bUseAnimBank)
+				const uint32 AnimationIndex = InstanceData[InstanceIndex].AnimationIndex;
+				if (bValidTransformProvider && Flags.bHasPerInstanceLocalBounds)
 				{
-					uint32 BankIndex = InstanceData[InstanceIndex].BankIndex;
-					if (BankIndex < uint32(AnimBankItems.Num()))
+					FRenderBounds AnimationBounds;
+					if (TransformProvider->GetAnimationBounds(AnimationIndex, AnimationBounds))
 					{
-						const FAnimBankItem& BankItem = AnimBankItems[BankIndex];
-						if (BankItem.BankAsset != nullptr)
-						{
-							const FAnimBankData& BankData = BankItem.BankAsset->GetData();
-							if (BankItem.SequenceIndex < BankData.Entries.Num())
-							{
-								return BankData.Entries[BankItem.SequenceIndex].SampledBounds;
-							}
-						}
+						return AnimationBounds;
 					}
 				}
-
+				
 				return MeshBounds;
 			});
 
 #if WITH_EDITOR
 	if (ChangeSet.Flags.bHasPerInstanceEditorData)
 	{
-		// TODO: the way hit proxies are managed seems daft, why don't we just add them when needed and store them in an array alonside the instances?
+		// TODO: the way hit proxies are managed seems daft, why don't we just add them when needed and store them in an array alongside the instances?
 		//       this will always force us to update all the hit proxy data for every instances.
 		TArray<TRefCountPtr<HHitProxy>> HitProxies;
 		CreateHitProxyData(HitProxies);
@@ -1150,16 +1098,14 @@ void UInstancedSkinnedMeshComponent::ComputeHLODHash(FHLODHashBuilder& HashBuild
 	for (const FSkinnedMeshInstanceData& SkinnedMeshInstanceData : InstanceData)
 	{
 		HashBuilder << FTransform(SkinnedMeshInstanceData.Transform);
-		HashBuilder << SkinnedMeshInstanceData.BankIndex;
+		HashBuilder << SkinnedMeshInstanceData.AnimationIndex;
 	}
 	HashBuilder << FHLODHashContext(GET_MEMBER_NAME_CHECKED(UInstancedSkinnedMeshComponent, InstanceData));
 
-	for (FAnimBankItem AnimBankItem : AnimBankItems)
 	{
-		HashBuilder << AnimBankItem.BankAsset;
-		HashBuilder << AnimBankItem.SequenceIndex;
+		UObject* TransformProviderObject = GetTransformProvider();
+		HashBuilder << TransformProviderObject << FHLODHashContext(TEXT("TransformProvider"));
 	}
-	HashBuilder << FHLODHashContext(GET_MEMBER_NAME_CHECKED(UInstancedSkinnedMeshComponent, AnimBankItems));
 
 	HashBuilder << InstanceCustomData << FHLODHashContext(GET_MEMBER_NAME_CHECKED(UInstancedSkinnedMeshComponent, InstanceCustomData));
 	HashBuilder << InstanceMinDrawDistance << FHLODHashContext(GET_MEMBER_NAME_CHECKED(UInstancedSkinnedMeshComponent, InstanceMinDrawDistance));
@@ -1202,9 +1148,9 @@ FPrimitiveSceneProxy* UInstancedSkinnedMeshComponent::CreateSceneProxy()
 		return nullptr;
 	}
 
-	for (FAnimBankItem& BankItem : AnimBankItems)
+	if (TransformProvider != nullptr && TransformProvider->IsEnabled())
 	{
-		if (BankItem.BankAsset && BankItem.BankAsset->IsCompiling())
+		if (TransformProvider->IsCompiling())
 		{
 			return nullptr;
 		}
@@ -1259,7 +1205,7 @@ void UInstancedSkinnedMeshComponent::PostEditChangeChainProperty(FPropertyChange
 				AddInstanceInternal(
 					AddedAtIndex,
 					PropertyChangedEvent.ChangeType == EPropertyChangeType::ArrayAdd ? FTransform::Identity : FTransform(InstanceData[AddedAtIndex].Transform),
-					PropertyChangedEvent.ChangeType == EPropertyChangeType::ArrayAdd ? 0 : InstanceData[AddedAtIndex].BankIndex,
+					PropertyChangedEvent.ChangeType == EPropertyChangeType::ArrayAdd ? 0 : InstanceData[AddedAtIndex].AnimationIndex,
 					/*bWorldSpace*/false
 				);
 
@@ -1283,7 +1229,7 @@ void UInstancedSkinnedMeshComponent::PostEditChangeChainProperty(FPropertyChange
 			MarkRenderStateDirty();
 		}
 		else if (PropertyChangedEvent.Property->GetFName() == GET_MEMBER_NAME_CHECKED(FSkinnedMeshInstanceData, Transform)
-			 || PropertyChangedEvent.Property->GetFName() == GET_MEMBER_NAME_CHECKED(FSkinnedMeshInstanceData, BankIndex))
+			 || PropertyChangedEvent.Property->GetFName() == GET_MEMBER_NAME_CHECKED(FSkinnedMeshInstanceData, AnimationIndex))
 		{
 			MarkRenderStateDirty();
 		}
@@ -1300,37 +1246,8 @@ void UInstancedSkinnedMeshComponent::PostEditChangeChainProperty(FPropertyChange
 			}
 			MarkRenderStateDirty();
 		}
-		else if (
-			(PropertyChangedEvent.Property->GetFName() == GET_MEMBER_NAME_CHECKED(UInstancedSkinnedMeshComponent, AnimBankItems)) ||
-			(PropertyChangedEvent.Property->GetFName() == GET_MEMBER_NAME_CHECKED(FAnimBankItem, BankAsset))
-			)
+		else if (PropertyChangedEvent.Property->GetFName() == GET_MEMBER_NAME_CHECKED(UInstancedSkinnedMeshComponent, TransformProvider))
 		{
-			for (FAnimBankItem& BankItem : AnimBankItems)
-			{
-				// TODO: BankItem.ValidatePosition();
-
-				// Make sure the animation skeleton is valid
-				bool bInvalid = true;
-				if (BankItem.BankAsset && BankItem.BankAsset->Asset && BankItem.BankAsset->Asset->GetSkeleton())
-				{
-					USkeletalMesh* SkeletalMeshPtr = Cast<USkeletalMesh>(GetSkinnedAsset());
-					if (SkeletalMeshPtr)
-					{
-						// Make sure the skeletons match!
-						if (BankItem.BankAsset->Asset->GetSkeleton() == SkeletalMeshPtr->GetSkeleton())
-						{
-							bInvalid = false;
-						}
-					}
-				}
-
-				if (bInvalid)
-				{
-					UE_LOG(LogAnimation, Warning, TEXT("Invalid animation skeleton"));
-					BankItem.BankAsset = nullptr;
-				}
-			}
-
 			MarkRenderStateDirty();
 		}
 	}
@@ -1345,21 +1262,17 @@ void UInstancedSkinnedMeshComponent::PostEditUndo()
 
 void UInstancedSkinnedMeshComponent::BeginCacheForCookedPlatformData(const ITargetPlatform* TargetPlatform)
 {
-	for (FAnimBankItem& BankItem : AnimBankItems)
+	if (TransformProvider != nullptr && TransformProvider->IsEnabled())
 	{
-		if (BankItem.BankAsset)
-		{
-			BankItem.BankAsset->BeginCacheForCookedPlatformData(TargetPlatform);
-		}
+		TransformProvider->BeginCacheForCookedPlatformData(TargetPlatform);
 	}
-
 }
 
-bool UInstancedSkinnedMeshComponent::IsCachedCookedPlatformDataLoaded( const ITargetPlatform* TargetPlatform ) 
+bool UInstancedSkinnedMeshComponent::IsCachedCookedPlatformDataLoaded(const ITargetPlatform* TargetPlatform)
 {
-	for (FAnimBankItem& BankItem : AnimBankItems)
+	if (TransformProvider != nullptr)
 	{
-		if (BankItem.BankAsset->IsCompiling())
+		if (!TransformProvider->IsCachedCookedPlatformDataLoaded(TargetPlatform))
 		{
 			return false;
 		}
@@ -1493,15 +1406,15 @@ void UInstancedSkinnedMeshComponent::SetNumGPUInstances(int32 InCount)
 	NumInstancesGPUOnly = InCount;
 }
 
-FPrimitiveInstanceId UInstancedSkinnedMeshComponent::AddInstance(const FTransform& InstanceTransform, int32 BankIndex, bool bWorldSpace)
+FPrimitiveInstanceId UInstancedSkinnedMeshComponent::AddInstance(const FTransform& InstanceTransform, int32 AnimationIndex, bool bWorldSpace)
 {
-	return AddInstanceInternal(InstanceData.Num(), InstanceTransform, BankIndex, bWorldSpace);
+	return AddInstanceInternal(InstanceData.Num(), InstanceTransform, AnimationIndex, bWorldSpace);
 }
 
-TArray<FPrimitiveInstanceId> UInstancedSkinnedMeshComponent::AddInstances(const TArray<FTransform>& Transforms, const TArray<int32>& BankIndices, bool bShouldReturnIds, bool bWorldSpace)
+TArray<FPrimitiveInstanceId> UInstancedSkinnedMeshComponent::AddInstances(const TArray<FTransform>& Transforms, const TArray<int32>& AnimationIndices, bool bShouldReturnIds, bool bWorldSpace)
 {
 	TArray<FPrimitiveInstanceId> NewInstanceIds;
-	if (Transforms.IsEmpty() || (Transforms.Num() != BankIndices.Num()))
+	if (Transforms.IsEmpty() || (Transforms.Num() != AnimationIndices.Num()))
 	{
 		return NewInstanceIds;
 	}
@@ -1526,8 +1439,8 @@ TArray<FPrimitiveInstanceId> UInstancedSkinnedMeshComponent::AddInstances(const
 	for (int32 AddIndex = 0; AddIndex < NumToAdd; ++AddIndex)
 	{
 		const FTransform& Transform = Transforms[AddIndex];
-		const int32 BankIndex = BankIndices[AddIndex];
-		FPrimitiveInstanceId InstanceId = AddInstanceInternal(InstanceData.Num(), Transform, BankIndex, bWorldSpace);
+		const int32 AnimationIndex = AnimationIndices[AddIndex];
+		FPrimitiveInstanceId InstanceId = AddInstanceInternal(InstanceData.Num(), Transform, AnimationIndex, bWorldSpace);
 		if (bShouldReturnIds)
 		{
 			NewInstanceIds[AddIndex] = InstanceId;
@@ -1620,7 +1533,7 @@ bool UInstancedSkinnedMeshComponent::GetInstanceTransform(FPrimitiveInstanceId I
 	return true;
 }
 
-bool UInstancedSkinnedMeshComponent::GetInstanceBankIndex(FPrimitiveInstanceId InstanceId, int32& OutBankIndex) const
+bool UInstancedSkinnedMeshComponent::GetInstanceAnimationIndex(FPrimitiveInstanceId InstanceId, int32& OutAnimationIndex) const
 {
 	int32 InstanceIndex = InstanceDataManager.IdToIndex(InstanceId);
 	if (!InstanceData.IsValidIndex(InstanceIndex))
@@ -1628,7 +1541,7 @@ bool UInstancedSkinnedMeshComponent::GetInstanceBankIndex(FPrimitiveInstanceId I
 		return false;
 	}
 
-	OutBankIndex = InstanceData[InstanceIndex].BankIndex;
+	OutAnimationIndex = InstanceData[InstanceIndex].AnimationIndex;
 	return true;
 }
 
@@ -1668,13 +1581,6 @@ void UInstancedSkinnedMeshComponent::ClearInstances()
 	InstanceDataManager.ClearInstances();
 }
 
-int32 UInstancedSkinnedMeshComponent::AddAnimBankItem(const FAnimBankItem& BankItem)
-{
-	int32 ItemIndex = AnimBankItems.Num();
-	AnimBankItems.Emplace(BankItem);
-	return ItemIndex;
-}
-
 struct HSkinnedMeshInstance : public HHitProxy
 {
 	TObjectPtr<UInstancedSkinnedMeshComponent> Component;
@@ -1743,7 +1649,7 @@ void UInstancedSkinnedMeshComponent::CreateHitProxyData(TArray<TRefCountPtr<HHit
 	}
 }
 
-FPrimitiveInstanceId UInstancedSkinnedMeshComponent::AddInstanceInternal(int32 InstanceIndex, const FTransform& InstanceTransform, int32 BankIndex, bool bWorldSpace)
+FPrimitiveInstanceId UInstancedSkinnedMeshComponent::AddInstanceInternal(int32 InstanceIndex, const FTransform& InstanceTransform, int32 AnimationIndex, bool bWorldSpace)
 {
 	// This happens because the editor modifies the InstanceData array _before_ callbacks. If we could change the UI to not do that we could remove this ugly hack.
 	if (!InstanceData.IsValidIndex(InstanceIndex))
@@ -1755,7 +1661,7 @@ FPrimitiveInstanceId UInstancedSkinnedMeshComponent::AddInstanceInternal(int32 I
 	FPrimitiveInstanceId InstanceId = InstanceDataManager.Add(InstanceIndex);
 
 	const FTransform3f LocalTransform = FTransform3f(bWorldSpace ? InstanceTransform.GetRelativeTransform(GetComponentTransform()) : InstanceTransform);
-	SetupNewInstanceData(InstanceData[InstanceIndex], InstanceIndex, LocalTransform, BankIndex);
+	SetupNewInstanceData(InstanceData[InstanceIndex], InstanceIndex, LocalTransform, AnimationIndex);
 
 	// Add custom data to instance
 	InstanceCustomData.AddZeroed(NumCustomDataFloats);
@@ -1796,9 +1702,9 @@ bool UInstancedSkinnedMeshComponent::RemoveInstanceInternal(int32 InstanceIndex,
 	return true;
 }
 
-FSkeletalMeshObject* UInstancedSkinnedMeshComponent::CreateMeshObject(const FInstancedSkinnedMeshSceneProxyDesc& InMeshDesc, TConstArrayView<FAnimBankItem> InAnimBankItems, FSkeletalMeshRenderData* InRenderData, ERHIFeatureLevel::Type InFeatureLevel)
+FSkeletalMeshObject* UInstancedSkinnedMeshComponent::CreateMeshObject(const FInstancedSkinnedMeshSceneProxyDesc& InMeshDesc, TObjectPtr<UTransformProviderData> InTransformProvider, FSkeletalMeshRenderData* InRenderData, ERHIFeatureLevel::Type InFeatureLevel)
 {
-	return ::new FInstancedSkinnedMeshObject(InMeshDesc, InAnimBankItems, InRenderData, InFeatureLevel);
+	return ::new FInstancedSkinnedMeshObject(InMeshDesc, InTransformProvider, InRenderData, InFeatureLevel);
 }
 
 FPrimitiveSceneProxy* UInstancedSkinnedMeshComponent::CreateSceneProxy(const FInstancedSkinnedMeshSceneProxyDesc& Desc, bool bHideSkin, bool bShouldNaniteSkin, bool bIsEnabled, int32 MinLODIndex)
diff --git a/Engine/Source/Runtime/Engine/Private/ISMPartition/ISMComponentBatcher.cpp b/Engine/Source/Runtime/Engine/Private/ISMPartition/ISMComponentBatcher.cpp
--- a/Engine/Source/Runtime/Engine/Private/ISMPartition/ISMComponentBatcher.cpp
+++ b/Engine/Source/Runtime/Engine/Private/ISMPartition/ISMComponentBatcher.cpp
@@ -181,7 +181,7 @@ void FISMComponentBatcher::AddInternal(const UActorComponent* InComponent, TOpti
 				}
 				InstancesTransformsWS.Add(InstanceTransformWS);
 
-				ISKMC->GetInstanceBankIndex(InstanceId, BankIndices.AddZeroed_GetRef());
+				ISKMC->GetInstanceAnimationIndex(InstanceId, AnimationIndices.AddZeroed_GetRef());
 
 				// Add per instance custom data, if any
 				if (NumCustomDataFloats > 0)
@@ -212,7 +212,7 @@ void FISMComponentBatcher::InitComponent(UInstancedStaticMeshComponent* ISMCompo
 void FISMComponentBatcher::InitComponent(UInstancedSkinnedMeshComponent* ISMComponent) const
 {
 	ISMComponent->SetNumCustomDataFloats(NumCustomDataFloats);
-	ISMComponent->AddInstances(InstancesTransformsWS, BankIndices, /*bShouldReturnIndices*/false, /*bWorldSpace*/true);
+	ISMComponent->AddInstances(InstancesTransformsWS, AnimationIndices, /*bShouldReturnIndices*/false, /*bWorldSpace*/true);
 
 	if (NumCustomDataFloats)
 	{
@@ -238,7 +238,7 @@ void FISMComponentBatcher::ComputeHash() const
 	Ar << This.InstancesCustomData;
 	Ar << This.InstancingRandomSeed;
 	Ar << This.AdditionalRandomSeeds;
-	Ar << This.BankIndices;
+	Ar << This.AnimationIndices;
 
 	Hash = Ar.GetCrc();
 }
diff --git a/Engine/Source/Runtime/Engine/Private/InstancedSkinnedMeshSceneProxyDesc.cpp b/Engine/Source/Runtime/Engine/Private/InstancedSkinnedMeshSceneProxyDesc.cpp
--- a/Engine/Source/Runtime/Engine/Private/InstancedSkinnedMeshSceneProxyDesc.cpp
+++ b/Engine/Source/Runtime/Engine/Private/InstancedSkinnedMeshSceneProxyDesc.cpp
@@ -12,6 +12,8 @@ void FInstancedSkinnedMeshSceneProxyDesc::InitializeFromInstancedSkinnedMeshComp
 {
 	InitializeFromSkinnedMeshComponent(InComponent);
 
+	TransformProvider = InComponent->TransformProvider;
+
 	AnimationMinScreenSize = InComponent->AnimationMinScreenSize;
 	InstanceMinDrawDistance = InComponent->InstanceMinDrawDistance;
 	InstanceStartCullDistance = InComponent->InstanceStartCullDistance;
@@ -20,8 +22,6 @@ void FInstancedSkinnedMeshSceneProxyDesc::InitializeFromInstancedSkinnedMeshComp
 	SelectedInstances = InComponent->SelectedInstances;
 #endif
 
-	AnimBankItems = InComponent->AnimBankItems;
-	bAnimBankEnabled = InComponent->bAnimBankEnabled;
 	bAllowAlwaysVisible = true;
 
 	InstanceDataSceneProxy = InComponent->GetInstanceDataSceneProxy();
diff --git a/Engine/Source/Runtime/Engine/Private/ObjectCacheContext.cpp b/Engine/Source/Runtime/Engine/Private/ObjectCacheContext.cpp
--- a/Engine/Source/Runtime/Engine/Private/ObjectCacheContext.cpp
+++ b/Engine/Source/Runtime/Engine/Private/ObjectCacheContext.cpp
@@ -837,12 +837,11 @@ TObjectCacheIterator<UInstancedSkinnedMeshComponent> FObjectCacheContext::GetIns
 	return InstancedSkinnedMeshCache.Get<TObjectIterator<UInstancedSkinnedMeshComponent>>();
 }
 
-TObjectCacheIterator<UInstancedSkinnedMeshComponent> FObjectCacheContext::GetInstancedSkinnedMeshComponents(UAnimBank* InAnimBank)
+TObjectCacheIterator<UInstancedSkinnedMeshComponent> FObjectCacheContext::GetInstancedSkinnedMeshComponents(UTransformProviderData* InTransformProvider)
 {
-	return InstancedSkinnedMeshCache.Get<TObjectIterator<UInstancedSkinnedMeshComponent>>(InAnimBank,
-		[](UInstancedSkinnedMeshComponent* Component){ return Component->GetAnimBankItems().Num();  },
-		[](UInstancedSkinnedMeshComponent* Component, int32 Index)->const UAnimBank* { return Component->GetAnimBankItems()[Index].BankAsset; }
-	);
+	return InstancedSkinnedMeshCache.Get<TObjectIterator<UInstancedSkinnedMeshComponent>>(InTransformProvider,
+		[](UInstancedSkinnedMeshComponent* Component) { return 1; },
+		[](UInstancedSkinnedMeshComponent* Component, int32) { return Component->GetTransformProvider(); });
 }
 
 TObjectCacheIterator<UMaterialInterface> FObjectCacheContext::GetUsedMaterials(IPrimitiveComponent* Component)
diff --git a/Engine/Source/Runtime/Engine/Public/Animation/TransformProviderData.h b/Engine/Source/Runtime/Engine/Public/Animation/TransformProviderData.h
new file mode 100644
--- /dev/null
+++ b/Engine/Source/Runtime/Engine/Public/Animation/TransformProviderData.h
@@ -0,0 +1,95 @@
+// Copyright Epic Games, Inc. All Rights Reserved.
+
+#pragma once
+
+#include "UObject/ObjectMacros.h"
+#include "UObject/Object.h"
+#include "Misc/Guid.h"
+
+#include "TransformProviderData.generated.h"
+
+struct FRenderBounds;
+class FSkeletalMeshSceneProxyBase;
+class FSceneInterface;
+class FRHICommandListBase;
+class ITargetPlatform;
+
+class FTransformProviderRenderProxy
+{
+public:
+	ENGINE_API FTransformProviderRenderProxy()
+	{
+	}
+
+	ENGINE_API virtual ~FTransformProviderRenderProxy()
+	{
+	}
+
+	virtual void CreateRenderThreadResources(FRHICommandListBase& RHICmdList) = 0;
+
+	virtual void DestroyRenderThreadResources() = 0;
+
+	virtual const TConstArrayView<uint64> GetProviderData(bool& bOutValid) const = 0;
+};
+
+UCLASS(Abstract, config=Engine, hidecategories=Object, MinimalAPI, BlueprintType)
+class UTransformProviderData : public UObject
+{
+	GENERATED_BODY()
+
+public:
+	virtual bool IsEnabled() const
+	{
+		return bEnabled;
+	}
+
+	virtual const FGuid& GetTransformProviderID() const
+	{
+		static FGuid InvalidID(0, 0, 0, 0);
+		return InvalidID;
+	}
+
+	virtual const uint32 GetUniqueAnimationCount() const
+	{
+		return 1u;
+	}
+
+	virtual bool HasAnimationBounds() const
+	{
+		return false;
+	}
+
+	virtual bool GetAnimationBounds(uint32 AnimationIndex, FRenderBounds& OutBounds) const
+	{
+		return false;
+	}
+
+	virtual FTransformProviderRenderProxy* CreateRenderThreadResources(FSkeletalMeshSceneProxyBase* SceneProxy, FSceneInterface& Scene, FRHICommandListBase& RHICmdList)
+	{
+		return nullptr;
+	}
+
+	virtual void DestroyRenderThreadResources(FTransformProviderRenderProxy* ProviderProxy)
+	{
+	}
+
+	virtual bool IsCompiling() const
+	{
+		return false;
+	}
+
+#if WITH_EDITOR
+	void BeginCacheForCookedPlatformData(const ITargetPlatform* TargetPlatform)
+	{
+	}
+
+	bool IsCachedCookedPlatformDataLoaded(const ITargetPlatform* TargetPlatform)
+	{
+		return true;
+	}
+#endif
+
+public:
+	UPROPERTY(EditAnywhere, Category = TransformProvider)
+	bool bEnabled;
+};
diff --git a/Engine/Source/Runtime/Engine/Public/ISMPartition/ISMComponentBatcher.h b/Engine/Source/Runtime/Engine/Public/ISMPartition/ISMComponentBatcher.h
--- a/Engine/Source/Runtime/Engine/Public/ISMPartition/ISMComponentBatcher.h
+++ b/Engine/Source/Runtime/Engine/Public/ISMPartition/ISMComponentBatcher.h
@@ -140,5 +140,5 @@ protected:
 	TArray<FInstancedStaticMeshRandomSeed> AdditionalRandomSeeds; 
 
 	// For ISKMC only
-	TArray<int32> BankIndices;
+	TArray<int32> AnimationIndices;
 };
diff --git a/Engine/Source/Runtime/Engine/Public/InstancedSkinnedMeshComponentHelper.h b/Engine/Source/Runtime/Engine/Public/InstancedSkinnedMeshComponentHelper.h
--- a/Engine/Source/Runtime/Engine/Public/InstancedSkinnedMeshComponentHelper.h
+++ b/Engine/Source/Runtime/Engine/Public/InstancedSkinnedMeshComponentHelper.h
@@ -33,7 +33,6 @@ FInstanceDataManagerSourceDataDesc FInstancedSkinnedMeshComponentHelper::GetComp
 {
 	FInstanceDataManagerSourceDataDesc ComponentDesc;
 
-
 	ComponentDesc.PrimitiveMaterialDesc = FPrimitiveComponentHelper::GetUsedMaterialPropertyDesc(InComponent, ShaderPlatform);
 
 	FInstanceDataFlags Flags;
@@ -47,13 +46,16 @@ FInstanceDataManagerSourceDataDesc FInstancedSkinnedMeshComponentHelper::GetComp
 #endif
 
 	USkinnedAsset* SkinnedAsset = InComponent.GetSkinnedAsset();
-	TArrayView<const struct FAnimBankItem> AnimBankItems = InComponent.GetAnimBankItems();
+
+	const UTransformProviderData* TransformProvider = InComponent.GetTransformProvider();
 
 	const bool bForceRefPose = UInstancedSkinnedMeshComponent::ShouldForceRefPose();
-	const bool bUseAnimBank = !bForceRefPose && AnimBankItems.Num() > 0;
+	const bool bValidTransformProvider = !bForceRefPose && TransformProvider != nullptr && TransformProvider->IsEnabled();
+
+	const uint32 UniqueAnimationCount = TransformProvider ? TransformProvider->GetUniqueAnimationCount() : 1;
 
 	Flags.bHasPerInstanceHierarchyOffset = false;
-	Flags.bHasPerInstanceLocalBounds = bUseAnimBank && AnimBankItems.Num() > 1;
+	Flags.bHasPerInstanceLocalBounds = TransformProvider ? TransformProvider->HasAnimationBounds() : false;
 	Flags.bHasPerInstanceDynamicData = false;
 	Flags.bHasPerInstanceSkinningData = true;
 
@@ -82,25 +84,28 @@ FInstanceDataManagerSourceDataDesc FInstancedSkinnedMeshComponentHelper::GetComp
 		ChangeSet.GetSkinningDataWriter().Gather(
 			[&](int32 InstanceIndex) -> uint32
 			{
-				return InComponent.InstanceData[InstanceIndex].BankIndex * MaxBoneTransformCount * 2u;
+				const uint32 AnimationIndex = InComponent.InstanceData[InstanceIndex].AnimationIndex;
+				if (bValidTransformProvider && AnimationIndex < UniqueAnimationCount)
+				{
+					return AnimationIndex * MaxBoneTransformCount * 2u;
+				}
+
+				return 0u;
 			});
 
 		ChangeSet.GetLocalBoundsWriter().Gather(
 			[&](int32 InstanceIndex) -> FRenderBounds
 			{
-				uint32 BankIndex = InComponent.InstanceData[InstanceIndex].BankIndex;
-				if (BankIndex < uint32(AnimBankItems.Num()))
+				const uint32 AnimationIndex = InComponent.InstanceData[InstanceIndex].AnimationIndex;
+				if (bValidTransformProvider && Flags.bHasPerInstanceLocalBounds)
 				{
-					const FAnimBankItem& BankItem = AnimBankItems[BankIndex];
-					if (BankItem.BankAsset != nullptr)
+					FRenderBounds AnimationBounds;
+					if (TransformProvider->GetAnimationBounds(AnimationIndex, AnimationBounds))
 					{
-						const FAnimBankData& BankData = BankItem.BankAsset->GetData();
-						if (BankItem.SequenceIndex < BankData.Entries.Num())
-						{
-							return BankData.Entries[BankItem.SequenceIndex].SampledBounds;
-						}
+						return AnimationBounds;
 					}
 				}
+
 				return MeshBounds;
 			});
 
@@ -109,7 +114,7 @@ FInstanceDataManagerSourceDataDesc FInstancedSkinnedMeshComponentHelper::GetComp
 		{
 			if (ChangeSet.Flags.bHasPerInstanceEditorData)
 			{
-				// TODO: the way hit proxies are managed seems daft, why don't we just add them when needed and store them in an array alonside the instances?
+				// TODO: the way hit proxies are managed seems daft, why don't we just add them when needed and store them in an array alongside the instances?
 				//       this will always force us to update all the hit proxy data for every instances.
 				TArray<TRefCountPtr<HHitProxy>> HitProxies;
 				InComponent.CreateHitProxyData(HitProxies);
@@ -133,27 +138,25 @@ FBoxSphereBounds FInstancedSkinnedMeshComponentHelper::CalcBounds(const T& InCom
 		const FMatrix BoundTransformMatrix = LocalToWorld.ToMatrixWithScale();
 
 		FBoxSphereBounds::Builder BoundsBuilder;
-		
-		TArrayView<const struct FAnimBankItem> AnimBankItems = InComponent.GetAnimBankItems();
-		const bool bUseSampledBounds = UInstancedSkinnedMeshComponent::ShouldUseSampledBounds();
-		
-		if (bUseSampledBounds && AnimBankItems.Num() > 0)
+
+		const bool bUseAnimationBounds = UInstancedSkinnedMeshComponent::ShouldUseAnimationBounds();
+		const UTransformProviderData* TransformProvider = InComponent.GetTransformProvider();
+		if (TransformProvider != nullptr && TransformProvider->IsEnabled() && TransformProvider->HasAnimationBounds() && !TransformProvider->IsCompiling())
 		{
-			// Trade per sequence bounds (tigher fitting) for faster builds with high instance counts.
+			const uint32 AnimationCount = TransformProvider->GetUniqueAnimationCount();
+
+			// Trade per sequence bounds (tighter fitting) for faster builds with high instance counts.
 			const bool bFastBuild = false;
 			if (bFastBuild)
 			{
 				FBox MergedBounds;
 
-				for (const FAnimBankItem& BankItem : AnimBankItems)
+				for (uint32 AnimationIndex = 0; AnimationIndex < AnimationCount; ++AnimationIndex)
 				{
-					if (BankItem.BankAsset != nullptr)
+					FRenderBounds AnimationBounds;
+					if (TransformProvider->GetAnimationBounds(AnimationIndex, AnimationBounds))
 					{
-						const FAnimBankData& BankData = BankItem.BankAsset->GetData();
-						for (const FAnimBankEntry& BankEntry : BankData.Entries)
-						{
-							MergedBounds += BankEntry.SampledBounds.GetBox();
-						}
+						MergedBounds += AnimationBounds.ToBox();
 					}
 				}
 
@@ -171,33 +174,21 @@ FBoxSphereBounds FInstancedSkinnedMeshComponentHelper::CalcBounds(const T& InCom
 				for (int32 InstanceIndex = 0; InstanceIndex < InComponent.InstanceData.Num(); InstanceIndex++)
 				{
 					const FSkinnedMeshInstanceData& Instance = InComponent.InstanceData[InstanceIndex];
-					if (Instance.BankIndex < uint32(AnimBankItems.Num()))
+					FRenderBounds AnimationBounds;
+					if (TransformProvider->GetAnimationBounds(Instance.AnimationIndex, AnimationBounds))
 					{
-						const FAnimBankItem& BankItem = AnimBankItems[Instance.BankIndex];
-						if (BankItem.BankAsset != nullptr
-#if WITH_EDITOR
-							&& !BankItem.BankAsset->IsCompiling()
-#endif
-							)
-						{
-							const FAnimBankData& BankData = BankItem.BankAsset->GetData();
-							if (BankItem.SequenceIndex < BankData.Entries.Num())
-							{
-								const FBox BankBounds = BankData.Entries[BankItem.SequenceIndex].SampledBounds.GetBox();
-								BoundsBuilder += BankBounds.TransformBy(FTransform(Instance.Transform) * LocalToWorld);
-							}
-						}
+						BoundsBuilder += AnimationBounds.ToBox().TransformBy(FTransform(Instance.Transform) * LocalToWorld);
 					}
 				}
 			}
 
-			// Only use bounds if valid, else continue with implementation not using AnimBankItems
+			// Only use bounds if valid, otherwise use the skinned asset bounds in ref pose.
 			if (BoundsBuilder.IsValid())
 			{
 				return BoundsBuilder;
 			}
 		}
-		
+
 		const FBox InstanceBounds = SkinnedAssetPtr->GetBounds().GetBox();
 		if (InstanceBounds.IsValid)
 		{
@@ -216,7 +207,7 @@ FBoxSphereBounds FInstancedSkinnedMeshComponentHelper::CalcBounds(const T& InCom
 template <class T>
 FSkeletalMeshObject* FInstancedSkinnedMeshComponentHelper::CreateMeshObject(const T& InComponent, const FInstancedSkinnedMeshSceneProxyDesc& InSceneProxyDesc)
 {
-	return UInstancedSkinnedMeshComponent::CreateMeshObject(InSceneProxyDesc, InComponent.GetAnimBankItems(), FSkinnedMeshComponentHelper::GetSkeletalMeshRenderData(InComponent), InComponent.GetScene()->GetFeatureLevel());
+	return UInstancedSkinnedMeshComponent::CreateMeshObject(InSceneProxyDesc, InComponent.GetTransformProvider(), FSkinnedMeshComponentHelper::GetSkeletalMeshRenderData(InComponent), InComponent.GetScene()->GetFeatureLevel());
 }
 
 template <class T>
diff --git a/Engine/Source/Runtime/Engine/Public/InstancedSkinnedMeshSceneProxyDesc.h b/Engine/Source/Runtime/Engine/Public/InstancedSkinnedMeshSceneProxyDesc.h
--- a/Engine/Source/Runtime/Engine/Public/InstancedSkinnedMeshSceneProxyDesc.h
+++ b/Engine/Source/Runtime/Engine/Public/InstancedSkinnedMeshSceneProxyDesc.h
@@ -16,6 +16,8 @@ struct FInstancedSkinnedMeshSceneProxyDesc : public FSkinnedMeshSceneProxyDesc
 
 	TSharedPtr<FInstanceDataSceneProxy, ESPMode::ThreadSafe> InstanceDataSceneProxy;
 
+	TObjectPtr<class UTransformProviderData> TransformProvider;
+
 	// Properties
 	float AnimationMinScreenSize = 0.0f;
 	int32 InstanceMinDrawDistance = 0;
@@ -26,7 +28,4 @@ struct FInstancedSkinnedMeshSceneProxyDesc : public FSkinnedMeshSceneProxyDesc
 	/** One bit per instance if the instance is selected. */
 	TBitArray<> SelectedInstances;
 #endif
-
-	TArrayView<const struct FAnimBankItem> AnimBankItems;
-	uint8 bAnimBankEnabled : 1 = false;
 };
diff --git a/Engine/Source/Runtime/Engine/Public/ObjectCacheContext.h b/Engine/Source/Runtime/Engine/Public/ObjectCacheContext.h
--- a/Engine/Source/Runtime/Engine/Public/ObjectCacheContext.h
+++ b/Engine/Source/Runtime/Engine/Public/ObjectCacheContext.h
@@ -21,7 +21,7 @@ class USkinnedAsset;
 class IPrimitiveComponent;
 class IStaticMeshComponent;
 class UInstancedSkinnedMeshComponent;
-class UAnimBank;
+class UTransformProviderData;
 
 /**
  *  Class to abstract implementation details of the containers used inside the
@@ -76,7 +76,7 @@ public:
 	ENGINE_API TObjectCacheIterator<USkinnedMeshComponent> GetSkinnedMeshComponents();
 	ENGINE_API TObjectCacheIterator<USkinnedMeshComponent> GetSkinnedMeshComponents(USkinnedAsset* InSkinnedAsset);
 	ENGINE_API TObjectCacheIterator<UInstancedSkinnedMeshComponent> GetInstancedSkinnedMeshComponents();
-	ENGINE_API TObjectCacheIterator<UInstancedSkinnedMeshComponent> GetInstancedSkinnedMeshComponents(UAnimBank* InAnimBank);
+	ENGINE_API TObjectCacheIterator<UInstancedSkinnedMeshComponent> GetInstancedSkinnedMeshComponents(UTransformProviderData* InTransformProvider);
 	ENGINE_API TObjectCacheIterator<IStaticMeshComponent>  GetStaticMeshComponents(UStaticMesh* InStaticMesh);
 	ENGINE_API TObjectCacheIterator<UMaterialInterface>    GetMaterialsAffectedByTexture(UTexture* InTexture);
 	ENGINE_API TObjectCacheIterator<UMaterialInterface>    GetMaterialsAffectedByTextureCollection(UTextureCollection* InTextureCollection);
@@ -120,7 +120,7 @@ private:
 	};
 
 	TComponentTypeCache<USkinnedMeshComponent, USkinnedAsset> SkinnedMeshCache;
-	TComponentTypeCache<UInstancedSkinnedMeshComponent, UAnimBank> InstancedSkinnedMeshCache;
+	TComponentTypeCache<UInstancedSkinnedMeshComponent, UTransformProviderData> InstancedSkinnedMeshCache;
 	TOptional<TArray<IPrimitiveComponent*>> PrimitiveComponents;
 	TComponentTypeCache<IStaticMeshComponent, UStaticMesh> StaticMeshCache;
 };
