World Partition: fix runtime grid inheritance through level instances bug introduced in 39839079.
Also updated World Partition unit tests to have this case validated.

#okdorgithub public

#virtualized

diff --git a/Engine/Source/Runtime/Engine/Private/WorldPartition/WorldPartitionStreamingGeneration.cpp b/Engine/Source/Runtime/Engine/Private/WorldPartition/WorldPartitionStreamingGeneration.cpp
--- a/Engine/Source/Runtime/Engine/Private/WorldPartition/WorldPartitionStreamingGeneration.cpp
+++ b/Engine/Source/Runtime/Engine/Private/WorldPartition/WorldPartitionStreamingGeneration.cpp
@@ -753,10 +753,6 @@ class FWorldPartitionStreamingGenerator
 		{
 			ActorDescView.SetForcedNoRuntimeGrid();
 		}
-		else if (ActorDescView.GetRuntimeGrid().IsNone())
-		{			
-			ActorDescView.SetRuntimeGrid(DefaultGrid);
-		}
 	}
 
 	void ResolveRuntimeDataLayers(FStreamingGenerationActorDescView& ActorDescView, const TArray<const FWorldDataLayersActorDesc*>& InDataLayerResolvers)
