Fix instanced stereo and mobile multi-view rendering for debug draws taking the FSimpleElementCollector + FSimpleElementVS path, which includes:
- r.VisualizeOccludedPrimitives
- ShowFlag.Bones
- Most Draw Debug calls (e.g. Draw Debug Box, does not include Draw Debug String which is still broken)

This includes PR #13528, which covers the MMV side of the fix.

diff --git a/Engine/Shaders/Private/SimpleElementVertexShader.usf b/Engine/Shaders/Private/SimpleElementVertexShader.usf
--- a/Engine/Shaders/Private/SimpleElementVertexShader.usf
+++ b/Engine/Shaders/Private/SimpleElementVertexShader.usf
@@ -10,9 +10,9 @@
 #define ENABLE_LWC 0
 #endif
 
-float4x4 Transform;
+float4x4 Transform[2];
 #if ENABLE_LWC
-float3 TransformPositionHigh;
+float3 TransformPositionHigh[2];
 #endif
 
 void Main(
@@ -21,19 +21,36 @@ void Main(
 	in float2 InTextureCoordinate : ATTRIBUTE2,
 	in float4 InColor : ATTRIBUTE3,
 	in float4 InHitProxyId : ATTRIBUTE4,
+#if INSTANCED_STEREO
+    in uint InstanceId : SV_InstanceID,
+#elif MOBILE_MULTI_VIEW
+	in nointerpolation uint ViewId : SV_ViewID,
+#endif
 	out float2 OutTextureCoordinate : TEXCOORD0,
 	out float4 OutColor : TEXCOORD1,
 	out float4 OutHitProxyId : TEXCOORD2,
 	out float4 OutPosition : SV_POSITION
+#if INSTANCED_STEREO
+	, out uint ViewportIndex : SV_ViewPortArrayIndex
+#endif
 	)
 {
+#if INSTANCED_STEREO
+    const uint EyeIndex = InstanceId;
+	ViewportIndex = EyeIndex;
+#elif MOBILE_MULTI_VIEW
+	const uint EyeIndex = ViewId;
+#else
+	const uint EyeIndex = 0;	
+#endif
+
 #if ENABLE_LWC
 	FDFVector4 InPosition = MakeDFVector4(InPositionHigh, InPositionLow);
-	FDFInverseMatrix WorldToClip = MakeDFInverseMatrix(TransformPositionHigh, Transform);
+	FDFInverseMatrix WorldToClip = MakeDFInverseMatrix(TransformPositionHigh[EyeIndex], Transform[EyeIndex]);
 	OutPosition = DFMultiplyDemote(InPosition, WorldToClip);
 	OutPosition.z = max(OutPosition.z, 0);
 #else
-	OutPosition = mul(InPositionHigh, Transform);
+	OutPosition = mul(InPositionHigh, Transform[EyeIndex]);
 #endif
 
 	OutTextureCoordinate = InTextureCoordinate;
diff --git a/Engine/Source/Runtime/Engine/Private/BatchedElements.cpp b/Engine/Source/Runtime/Engine/Private/BatchedElements.cpp
--- a/Engine/Source/Runtime/Engine/Private/BatchedElements.cpp
+++ b/Engine/Source/Runtime/Engine/Private/BatchedElements.cpp
@@ -819,11 +819,14 @@ void FBatchedElements::PrepareShaders(
 		}
 
 		// Set the simple element vertex shader parameters
-		SetShaderParametersLegacyVS(RHICmdList, VertexShader, ViewMatrices);
+		const FSceneView* InstancedView = View->GetInstancedSceneView();
+		SetShaderParametersLegacyVS(RHICmdList, VertexShader, ViewMatrices,
+			InstancedView ? FDFRelativeViewMatrices::Create(InstancedView->ViewMatrices)
+						  : ViewMatrices);
 	}
 }
 
-void FBatchedElements::DrawPointElements(FRHICommandList& RHICmdList, const FMatrix& Transform, const uint32 ViewportSizeX, const uint32 ViewportSizeY, const FVector& CameraX, const FVector& CameraY) const
+void FBatchedElements::DrawPointElements(FRHICommandList& RHICmdList, const FMatrix& Transform, const uint32 ViewportSizeX, const uint32 ViewportSizeY, const FVector& CameraX, const FVector& CameraY, const uint32 NumInstances) const
 {
 	// Draw the point elements.
 	if( Points.Num() > 0 )
@@ -873,7 +876,7 @@ void FBatchedElements::DrawPointElements(FRHICommandList& RHICmdList, const FMat
 		FBufferRHIRef VertexBufferRHI = Initializer.Finalize();
 
 		RHICmdList.SetStreamSource(0, VertexBufferRHI, 0);
-		RHICmdList.DrawPrimitive(0, NumTris, 1);
+		RHICmdList.DrawPrimitive(0, NumTris, NumInstances);
 	}
 }
 
@@ -888,7 +891,7 @@ FSceneView FBatchedElements::CreateProxySceneView(const FMatrix& ProjectionMatri
 	return FSceneView(ProxyViewInitOptions);
 }
 
-bool FBatchedElements::Draw(FRHICommandList& RHICmdList, const FMeshPassProcessorRenderState& DrawRenderState, ERHIFeatureLevel::Type FeatureLevel, const FSceneView& View, bool bHitTesting, float Gamma /* = 1.0f */, EBlendModeFilter::Type Filter /* = EBlendModeFilter::All */) const
+bool FBatchedElements::Draw(FRHICommandList& RHICmdList, const FMeshPassProcessorRenderState& DrawRenderState, ERHIFeatureLevel::Type FeatureLevel, const FSceneView& View, bool bHitTesting, float Gamma /* = 1.0f */, EBlendModeFilter::Type Filter /* = EBlendModeFilter::All */, uint32 NumInstances /* = 1 */) const
 {
 	const FDFRelativeViewMatrices RelativeMatrices = FDFRelativeViewMatrices::Create(View.ViewMatrices);
 	const FMatrix& WorldToClip = View.ViewMatrices.GetViewProjectionMatrix();
@@ -953,7 +956,7 @@ bool FBatchedElements::Draw(FRHICommandList& RHICmdList, const FMeshPassProcesso
 				while( MinVertex < TotalVerts )
 				{
 					int32 NumLinePrims = FMath::Min( MaxVerticesAllowed, TotalVerts - MinVertex ) / 2;
-					RHICmdList.DrawPrimitive(MinVertex, NumLinePrims, 1);
+					RHICmdList.DrawPrimitive(MinVertex, NumLinePrims, NumInstances);
 					MinVertex += NumLinePrims * 2;
 				}
 
@@ -966,7 +969,7 @@ bool FBatchedElements::Draw(FRHICommandList& RHICmdList, const FMeshPassProcesso
 			PrepareShaders(RHICmdList, GraphicsPSOInit, StencilRef, FeatureLevel, SE_BLEND_Opaque, RelativeMatrices, BatchedElementParameters, GWhiteTexture, bHitTesting, Gamma, NULL, &View);
 
 			// Draw points
-			DrawPointElements(RHICmdList, WorldToClip, ViewportSizeX, ViewportSizeY, CameraX, CameraY);
+			DrawPointElements(RHICmdList, WorldToClip, ViewportSizeX, ViewportSizeY, CameraX, CameraY, NumInstances);
 
 			if ( ThickLines.Num() > 0 )
 			{
@@ -1234,7 +1237,7 @@ bool FBatchedElements::Draw(FRHICommandList& RHICmdList, const FMeshPassProcesso
 					if (AddedTris > 0)
 					{
 						RHICmdList.SetStreamSource(0, VertexBufferRHI, 0);
-						RHICmdList.DrawPrimitive(0, AddedTris, 1);
+						RHICmdList.DrawPrimitive(0, AddedTris, NumInstances);
 					}
 				}
 
@@ -1287,7 +1290,7 @@ bool FBatchedElements::Draw(FRHICommandList& RHICmdList, const FMeshPassProcesso
 					PrepareShaders(RHICmdList, GraphicsPSOInit, StencilRef, FeatureLevel, SE_BLEND_Opaque, RelativeMatrices, BatchedElementParameters, GWhiteTexture, bHitTesting, Gamma, NULL, &View);
 
 					int32 NumTris = MaxTri - MinTri;
-					RHICmdList.DrawPrimitive(MinTri * 3, NumTris, 1);
+					RHICmdList.DrawPrimitive(MinTri * 3, NumTris, NumInstances);
 					MinTri = MaxTri;
 				}
 				VertexBufferRHI.SafeRelease();
@@ -1397,7 +1400,7 @@ bool FBatchedElements::Draw(FRHICommandList& RHICmdList, const FMeshPassProcesso
 						const int32 PrimCount = SpriteNum * 2;
 						PrepareShaders(RHICmdList, GraphicsPSOInit, StencilRef, FeatureLevel, CurrentBlendMode, RelativeMatrices, BatchedElementParameters, CurrentTexture, bHitTesting, Gamma, NULL, &View, CurrentOpacityMask);
 						RHICmdList.SetStreamSource(0, VertexBufferRHI, 0);
-						RHICmdList.DrawPrimitive(BaseVertex, PrimCount, 1);
+						RHICmdList.DrawPrimitive(BaseVertex, PrimCount, NumInstances);
 
 						// begin the next batch
 						if (SpriteIndex < ValidSpriteCount)
@@ -1448,7 +1451,7 @@ bool FBatchedElements::Draw(FRHICommandList& RHICmdList, const FMeshPassProcesso
 
 					// Draw the mesh.
 					RHICmdList.SetStreamSource(0, VertexBufferRHI, MeshElement.MinVertex * sizeof(FSimpleElementVertex));
-					RHICmdList.DrawIndexedPrimitive(IndexBufferRHI, 0, 0, MeshElement.MaxVertex - MeshElement.MinVertex + 1, 0, MeshElement.Indices.Num() / 3, 1);
+					RHICmdList.DrawIndexedPrimitive(IndexBufferRHI, 0, 0, MeshElement.MaxVertex - MeshElement.MinVertex + 1, 0, MeshElement.Indices.Num() / 3, NumInstances);
 
 					IndexBufferRHI.SafeRelease();
 				}
diff --git a/Engine/Source/Runtime/Engine/Private/CubemapUnwrapUtils.cpp b/Engine/Source/Runtime/Engine/Private/CubemapUnwrapUtils.cpp
--- a/Engine/Source/Runtime/Engine/Private/CubemapUnwrapUtils.cpp
+++ b/Engine/Source/Runtime/Engine/Private/CubemapUnwrapUtils.cpp
@@ -134,7 +134,8 @@ bool FCubemapTexturePropertiesVS::ShouldCompilePermutation(const FGlobalShaderPe
 
 void FCubemapTexturePropertiesVS::SetParameters(FRHIBatchedShaderParameters& BatchedParameters, const FMatrix& TransformValue )
 {
-	SetShaderValue(BatchedParameters, Transform, (FMatrix44f)TransformValue);
+	FMatrix44f TransformValues[2] = { FMatrix44f(TransformValue), FMatrix44f(TransformValue) };
+	SetShaderValue(BatchedParameters, Transform, TransformValues);
 }
 
 FCubemapTexturePropertiesPS::FCubemapTexturePropertiesPS() = default;
diff --git a/Engine/Source/Runtime/Engine/Private/SceneManagement.cpp b/Engine/Source/Runtime/Engine/Private/SceneManagement.cpp
--- a/Engine/Source/Runtime/Engine/Private/SceneManagement.cpp
+++ b/Engine/Source/Runtime/Engine/Private/SceneManagement.cpp
@@ -272,7 +272,7 @@ void FSimpleElementCollector::RegisterDynamicResource(FDynamicPrimitiveResource*
 	DynamicResource->InitPrimitiveResource(FRHICommandListImmediate::Get());
 }
 
-void FSimpleElementCollector::DrawBatchedElements(FRHICommandList& RHICmdList, const FMeshPassProcessorRenderState& DrawRenderState, const FSceneView& InView, EBlendModeFilter::Type Filter, ESceneDepthPriorityGroup DepthPriorityGroup) const
+void FSimpleElementCollector::DrawBatchedElements(FRHICommandList& RHICmdList, const FMeshPassProcessorRenderState& DrawRenderState, const FSceneView& InView, EBlendModeFilter::Type Filter, ESceneDepthPriorityGroup DepthPriorityGroup, uint32 NumInstances) const
 {
 	const FBatchedElements& Elements = DepthPriorityGroup == SDPG_World ? BatchedElements : TopBatchedElements;
 
@@ -284,7 +284,8 @@ void FSimpleElementCollector::DrawBatchedElements(FRHICommandList& RHICmdList, c
 		InView,
 		InView.Family->EngineShowFlags.HitProxies,
 		1.0f,
-		Filter
+		Filter,
+		NumInstances
 		);
 }
 
diff --git a/Engine/Source/Runtime/Engine/Private/SimpleElementShaders.cpp b/Engine/Source/Runtime/Engine/Private/SimpleElementShaders.cpp
--- a/Engine/Source/Runtime/Engine/Private/SimpleElementShaders.cpp
+++ b/Engine/Source/Runtime/Engine/Private/SimpleElementShaders.cpp
@@ -26,14 +26,18 @@ FSimpleElementVS::FSimpleElementVS(const ShaderMetaType::CompiledShaderInitializ
 
 void FSimpleElementVS::SetParameters(FRHIBatchedShaderParameters& BatchedParameters, const FMatrix& WorldToClipMatrix)
 {
-	SetShaderValue(BatchedParameters, RelativeTransform, FMatrix44f(WorldToClipMatrix));
-	SetShaderValue(BatchedParameters, TransformPositionHigh, FVector3f::ZeroVector);
+	const FMatrix44f FinalMats[2] = { FMatrix44f(WorldToClipMatrix) , FMatrix44f(WorldToClipMatrix) };
+	const FVector4f PositionsHigh[2] = { FVector3f::ZeroVector, FVector3f::ZeroVector };
+	SetShaderValue(BatchedParameters, RelativeTransform, FinalMats);
+	SetShaderValue(BatchedParameters, TransformPositionHigh, PositionsHigh);
 }
 
-void FSimpleElementVS::SetParameters(FRHIBatchedShaderParameters& BatchedParameters, const FDFRelativeViewMatrices& Matrices)
+void FSimpleElementVS::SetParameters(FRHIBatchedShaderParameters& BatchedParameters, const FDFRelativeViewMatrices& Matrices, const FDFRelativeViewMatrices& InstancedViewMatrices)
 {
-	SetShaderValue(BatchedParameters, RelativeTransform, Matrices.RelativeWorldToClip);
-	SetShaderValue(BatchedParameters, TransformPositionHigh, Matrices.PositionHigh);
+	const FMatrix44f FinalMats[2] = { Matrices.RelativeWorldToClip , InstancedViewMatrices.RelativeWorldToClip };
+	const FVector4f PositionsHigh[2] = { Matrices.PositionHigh, InstancedViewMatrices.PositionHigh };
+	SetShaderValue(BatchedParameters, RelativeTransform, FinalMats);
+	SetShaderValue(BatchedParameters, TransformPositionHigh, PositionsHigh);
 }
 
 void FSimpleElementVS::ModifyCompilationEnvironment(const FGlobalShaderPermutationParameters& Parameters, FShaderCompilerEnvironment& OutEnvironment)
diff --git a/Engine/Source/Runtime/Engine/Public/BatchedElements.h b/Engine/Source/Runtime/Engine/Public/BatchedElements.h
--- a/Engine/Source/Runtime/Engine/Public/BatchedElements.h
+++ b/Engine/Source/Runtime/Engine/Public/BatchedElements.h
@@ -198,7 +198,7 @@ public:
 	 * @param Gamma			Optional gamma override
 	 * @param DepthTexture	DepthTexture for manual depth testing with editor compositing in the pixel shader
 	 */
-	ENGINE_API bool Draw(FRHICommandList& RHICmdList, const FMeshPassProcessorRenderState& DrawRenderState, ERHIFeatureLevel::Type FeatureLevel, const FSceneView& View, bool bHitTesting, float Gamma = 1.0f, EBlendModeFilter::Type Filter = EBlendModeFilter::All) const;
+	ENGINE_API bool Draw(FRHICommandList& RHICmdList, const FMeshPassProcessorRenderState& DrawRenderState, ERHIFeatureLevel::Type FeatureLevel, const FSceneView& View, bool bHitTesting, float Gamma = 1.0f, EBlendModeFilter::Type Filter = EBlendModeFilter::All, uint32 NumInstances = 1) const;
 
 	/**
 	 * Creates a proxy FSceneView for operations that are not tied directly to a scene but still require batched elements to be drawn.
@@ -268,7 +268,7 @@ private:
 	 * @param	CameraX		Local space normalized view direction X vector
 	 * @param	CameraY		Local space normalized view direction Y vector
 	 */
-	ENGINE_API void DrawPointElements(FRHICommandList& RHICmdList, const FMatrix& Transform, const uint32 ViewportSizeX, const uint32 ViewportSizeY, const FVector& CameraX, const FVector& CameraY) const;
+	ENGINE_API void DrawPointElements(FRHICommandList& RHICmdList, const FMatrix& Transform, const uint32 ViewportSizeX, const uint32 ViewportSizeY, const FVector& CameraX, const FVector& CameraY, const uint32 NumInstances) const;
 
 	TArray<FSimpleElementVertex> LineVertices;
 
diff --git a/Engine/Source/Runtime/Engine/Public/SceneManagement.h b/Engine/Source/Runtime/Engine/Public/SceneManagement.h
--- a/Engine/Source/Runtime/Engine/Public/SceneManagement.h
+++ b/Engine/Source/Runtime/Engine/Public/SceneManagement.h
@@ -1386,7 +1386,7 @@ public:
 		return 0;
 	}
 
-	ENGINE_API void DrawBatchedElements(FRHICommandList& RHICmdList, const FMeshPassProcessorRenderState& DrawRenderState, const FSceneView& InView, EBlendModeFilter::Type Filter, ESceneDepthPriorityGroup DPG) const;
+	ENGINE_API void DrawBatchedElements(FRHICommandList& RHICmdList, const FMeshPassProcessorRenderState& DrawRenderState, const FSceneView& InView, EBlendModeFilter::Type Filter, ESceneDepthPriorityGroup DPG, uint32 NumInstances = 1) const;
 
 	class FAllocationInfo
 	{
diff --git a/Engine/Source/Runtime/Engine/Public/SimpleElementShaders.h b/Engine/Source/Runtime/Engine/Public/SimpleElementShaders.h
--- a/Engine/Source/Runtime/Engine/Public/SimpleElementShaders.h
+++ b/Engine/Source/Runtime/Engine/Public/SimpleElementShaders.h
@@ -28,7 +28,11 @@ public:
 	FSimpleElementVS() {}
 
 	ENGINE_API void SetParameters(FRHIBatchedShaderParameters& BatchedParameters, const FMatrix& WorldToClipMatrix);
-	ENGINE_API void SetParameters(FRHIBatchedShaderParameters& BatchedParameters, const FDFRelativeViewMatrices& Matrices);
+	ENGINE_API void SetParameters(FRHIBatchedShaderParameters& BatchedParameters, const FDFRelativeViewMatrices& Matrices, const FDFRelativeViewMatrices& InstancedViewMatrices);
+	ENGINE_API void SetParameters(FRHIBatchedShaderParameters& BatchedParameters, const FDFRelativeViewMatrices& Matrices)
+	{
+		SetParameters(BatchedParameters, Matrices, Matrices);
+	}
 
 	static void ModifyCompilationEnvironment(const FGlobalShaderPermutationParameters& Parameters, FShaderCompilerEnvironment& OutEnvironment);
 
diff --git a/Engine/Source/Runtime/Renderer/Private/BasePassRendering.cpp b/Engine/Source/Runtime/Renderer/Private/BasePassRendering.cpp
--- a/Engine/Source/Runtime/Renderer/Private/BasePassRendering.cpp
+++ b/Engine/Source/Runtime/Renderer/Private/BasePassRendering.cpp
@@ -1347,9 +1347,9 @@ static void RenderEditorPrimitivesForDPG(
 		ERDGPassFlags::Raster,
 		[&View, DrawRenderState, DepthPriorityGroup](FRHICommandList& RHICmdList)
 	{
-		RHICmdList.SetViewport(View.ViewRect.Min.X, View.ViewRect.Min.Y, 0.0f, View.ViewRect.Max.X, View.ViewRect.Max.Y, 1.0f);
+		FSceneRenderer::SetStereoViewport(RHICmdList, View, 1.0f);
 
-		View.SimpleElementCollector.DrawBatchedElements(RHICmdList, DrawRenderState, View, EBlendModeFilter::OpaqueAndMasked, DepthPriorityGroup);
+		View.SimpleElementCollector.DrawBatchedElements(RHICmdList, DrawRenderState, View, EBlendModeFilter::OpaqueAndMasked, DepthPriorityGroup, View.GetStereoPassInstanceFactor());
 
 		if (!View.Family->EngineShowFlags.CompositeEditorPrimitives)
 		{
@@ -1396,7 +1396,7 @@ static void RenderEditorPrimitivesForDPG(
 			});
 
 			// Draw the view's batched simple elements(lines, sprites, etc).
-			View.TopBatchedViewElements.Draw(RHICmdList, DrawRenderState, View.GetFeatureLevel(), View, false);
+			View.TopBatchedViewElements.Draw(RHICmdList, DrawRenderState, View.GetFeatureLevel(), View, false, 1.0f, EBlendModeFilter::All, View.GetStereoPassInstanceFactor());
 		}
 	});
 }
diff --git a/Engine/Source/Runtime/Renderer/Private/DeferredShadingRenderer.cpp b/Engine/Source/Runtime/Renderer/Private/DeferredShadingRenderer.cpp
--- a/Engine/Source/Runtime/Renderer/Private/DeferredShadingRenderer.cpp
+++ b/Engine/Source/Runtime/Renderer/Private/DeferredShadingRenderer.cpp
@@ -3946,6 +3946,19 @@ void FDeferredShadingSceneRenderer::Render(FRDGBuilder& GraphBuilder, const FSce
 							TSRFlickeringInput = TSRFlickeringInputTextures[ViewIndex];
 						}
 
+						// If we're using instanced stereo, only the primary view simple element collectors will be populated with elements.
+						// However, since post processing is always rendered per-view, we need to mirror the collectors to any instanced secondary views.
+						if (View.bIsSinglePassStereo && View.StereoPass == EStereoscopicPass::eSSP_SECONDARY)
+						{
+							const FViewInfo& PrimaryView = Views[View.PrimaryViewIndex];
+
+							View.SimpleElementCollector = PrimaryView.SimpleElementCollector;
+							View.EditorSimpleElementCollector = PrimaryView.EditorSimpleElementCollector;
+#if UE_ENABLE_DEBUG_DRAWING
+							View.DebugSimpleElementCollector = PrimaryView.DebugSimpleElementCollector;
+#endif
+						}
+
 						AddPostProcessingPasses(
 							GraphBuilder,
 							View, ViewIndex,
diff --git a/Engine/Source/Runtime/Renderer/Private/DepthRendering.cpp b/Engine/Source/Runtime/Renderer/Private/DepthRendering.cpp
--- a/Engine/Source/Runtime/Renderer/Private/DepthRendering.cpp
+++ b/Engine/Source/Runtime/Renderer/Private/DepthRendering.cpp
@@ -451,10 +451,10 @@ static void RenderPrePassEditorPrimitives(
 	{
 		const bool bRespectUseAsOccluderFlag = true;
 
-		RHICmdList.SetViewport(View.ViewRect.Min.X, View.ViewRect.Min.Y, 0.0f, View.ViewRect.Max.X, View.ViewRect.Max.Y, 1.0f);
+		FSceneRenderer::SetStereoViewport(RHICmdList, View, 1.0f);
 
-		View.SimpleElementCollector.DrawBatchedElements(RHICmdList, DrawRenderState, View, EBlendModeFilter::OpaqueAndMasked, SDPG_World);
-		View.SimpleElementCollector.DrawBatchedElements(RHICmdList, DrawRenderState, View, EBlendModeFilter::OpaqueAndMasked, SDPG_Foreground);
+		View.SimpleElementCollector.DrawBatchedElements(RHICmdList, DrawRenderState, View, EBlendModeFilter::OpaqueAndMasked, SDPG_World, View.GetStereoPassInstanceFactor());
+		View.SimpleElementCollector.DrawBatchedElements(RHICmdList, DrawRenderState, View, EBlendModeFilter::OpaqueAndMasked, SDPG_Foreground, View.GetStereoPassInstanceFactor());
 
 		if (!View.Family->EngineShowFlags.CompositeEditorPrimitives)
 		{
@@ -482,7 +482,7 @@ static void RenderPrePassEditorPrimitives(
 			});
 
 			// Draw the view's batched simple elements(lines, sprites, etc).
-			View.BatchedViewElements.Draw(RHICmdList, DrawRenderState, View.FeatureLevel, View, false);
+			View.BatchedViewElements.Draw(RHICmdList, DrawRenderState, View.FeatureLevel, View, false, 1.0f, EBlendModeFilter::All, View.GetStereoPassInstanceFactor());
 
 			DrawDynamicMeshPass(View, RHICmdList, [&](FDynamicPassMeshDrawListContext* DynamicMeshPassContext)
 			{
@@ -508,7 +508,7 @@ static void RenderPrePassEditorPrimitives(
 			});
 
 			// Draw the view's batched simple elements(lines, sprites, etc).
-			View.TopBatchedViewElements.Draw(RHICmdList, DrawRenderState, View.FeatureLevel, View, false);
+			View.TopBatchedViewElements.Draw(RHICmdList, DrawRenderState, View.FeatureLevel, View, false, 1.0f, EBlendModeFilter::All, View.GetStereoPassInstanceFactor());
 		}
 	});
 }
diff --git a/Engine/Source/Runtime/Renderer/Private/SceneRenderingUtils.h b/Engine/Source/Runtime/Renderer/Private/SceneRenderingUtils.h
--- a/Engine/Source/Runtime/Renderer/Private/SceneRenderingUtils.h
+++ b/Engine/Source/Runtime/Renderer/Private/SceneRenderingUtils.h
@@ -15,10 +15,10 @@ inline void AddSimpleElementCollectorPass(const FSimpleElementCollector& SimpleE
 		ERDGPassFlags::Raster,
 		[&View, &SimpleElementCollector, DrawRenderState](FRHICommandListImmediate& RHICmdList)
 		{
-			RHICmdList.SetViewport(View.ViewRect.Min.X, View.ViewRect.Min.Y, 0.0f, View.ViewRect.Max.X, View.ViewRect.Max.Y, 1.0f);
+			FSceneRenderer::SetStereoViewport(RHICmdList, View, 1.0f);
 
-			SimpleElementCollector.DrawBatchedElements(RHICmdList, DrawRenderState, View, EBlendModeFilter::OpaqueAndMasked, SDPG_World);
-			SimpleElementCollector.DrawBatchedElements(RHICmdList, DrawRenderState, View, EBlendModeFilter::OpaqueAndMasked, SDPG_Foreground);
+			SimpleElementCollector.DrawBatchedElements(RHICmdList, DrawRenderState, View, EBlendModeFilter::OpaqueAndMasked, SDPG_World, View.GetStereoPassInstanceFactor());
+			SimpleElementCollector.DrawBatchedElements(RHICmdList, DrawRenderState, View, EBlendModeFilter::OpaqueAndMasked, SDPG_Foreground, View.GetStereoPassInstanceFactor());
 		}
 	);
 }
@@ -32,9 +32,9 @@ inline void AddSimpleElementCollectorPass(const FSimpleElementCollector& SimpleE
 		ERDGPassFlags::Raster,
 		[&View, &SimpleElementCollector, DrawRenderState, SceneDepthPriorityGroup](FRHICommandListImmediate& RHICmdList)
 		{
-			RHICmdList.SetViewport(View.ViewRect.Min.X, View.ViewRect.Min.Y, 0.0f, View.ViewRect.Max.X, View.ViewRect.Max.Y, 1.0f);
+			FSceneRenderer::SetStereoViewport(RHICmdList, View, 1.0f);
 
-			SimpleElementCollector.DrawBatchedElements(RHICmdList, DrawRenderState, View, EBlendModeFilter::OpaqueAndMasked, SceneDepthPriorityGroup);
+			SimpleElementCollector.DrawBatchedElements(RHICmdList, DrawRenderState, View, EBlendModeFilter::OpaqueAndMasked, SceneDepthPriorityGroup, View.GetStereoPassInstanceFactor());
 		}
 	);
 }
