Fix simple logic error in ULocalPlayerSaveGame::IsSaveInProgress where it always returned true

diff --git a/Engine/Source/Runtime/Engine/Private/GameFramework/SaveGame.cpp b/Engine/Source/Runtime/Engine/Private/GameFramework/SaveGame.cpp
--- a/Engine/Source/Runtime/Engine/Private/GameFramework/SaveGame.cpp
+++ b/Engine/Source/Runtime/Engine/Private/GameFramework/SaveGame.cpp
@@ -283,7 +283,8 @@ bool ULocalPlayerSaveGame::WasLoaded() const
 
 bool ULocalPlayerSaveGame::IsSaveInProgress() const
 {
-	return (CurrentSaveRequest > LastErrorSaveRequest) || (CurrentSaveRequest > LastSuccessfulSaveRequest);
+	// If the current request is greater than both the last error and last success, that must mean it is in progress
+	return (CurrentSaveRequest > LastErrorSaveRequest) && (CurrentSaveRequest > LastSuccessfulSaveRequest);
 }
 
 bool ULocalPlayerSaveGame::WasLastSaveSuccessful() const
