Improve BodyInstance comments

#rnx

diff --git a/Engine/Source/Runtime/Engine/Classes/PhysicsEngine/BodyInstance.h b/Engine/Source/Runtime/Engine/Classes/PhysicsEngine/BodyInstance.h
--- a/Engine/Source/Runtime/Engine/Classes/PhysicsEngine/BodyInstance.h
+++ b/Engine/Source/Runtime/Engine/Classes/PhysicsEngine/BodyInstance.h
@@ -329,15 +329,15 @@ struct FBodyInstance : public FBodyInstanceCore
 	/** When we are a body within a SkeletalMeshComponent, we cache the index of the bone we represent, to speed up sync'ing physics to anim. */
 	int16 InstanceBoneIndex;
 
-	/** Override the project default position iteration count for this body. Increasing this will be more CPU intensive, but better stabilized. 0 will lead to unresolved penetration and no static friction. */
+	/** Override the project-settings position iteration count for this body (if bOverrideIterationCounts is set). Increasing this will be more CPU intensive, but better stabilized. 0 will lead to unresolved penetration and no static friction. */
 	UPROPERTY(EditAnywhere, AdvancedDisplay, BlueprintReadOnly, Category=Physics, meta = (editcondition = "bOverrideIterationCounts", InlineEditConditionToggle, ClampMin = "0", UIMin = "0"))
 	uint8 PositionSolverIterationCount;
 
-	/** Override the project default velocity iteration count for this body. Increasing this will be more CPU intensive, but better stabilized. 0 will lead to unresolved restitution and lower quality dynamic friction. */
+	/** Override the project-settings velocity iteration count for this body (if bOverrideIterationCounts is set). Increasing this will be more CPU intensive, but better stabilized. 0 will lead to unresolved restitution and lower quality dynamic friction. */
 	UPROPERTY(EditAnywhere, AdvancedDisplay, BlueprintReadOnly, Category = Physics, meta = (editcondition = "bOverrideIterationCounts", InlineEditConditionToggle, ClampMin = "0", UIMin = "0"))
 	uint8 VelocitySolverIterationCount;
 
-	/** Override the project default position iteration count for this body. Usually 0 or 1. */
+	/** Override the project-settings projection iteration count for this body (if bOverrideIterationCounts is set). Usually 0 or 1. */
 	UPROPERTY(EditAnywhere, AdvancedDisplay, BlueprintReadOnly, Category = Physics, meta = (editcondition = "bOverrideIterationCounts", InlineEditConditionToggle, ClampMin = "0", UIMin = "0"))
 	uint8 ProjectionSolverIterationCount;
 
@@ -1053,25 +1053,25 @@ public:
 	 */
 	ENGINE_API void SetAllowPartialIslandSleep(bool bInAllowPartialIslandSleep);
 
-	// Set the per-body position iterations (only used if bOverrideIterationsCounts is true)
+	// Set the per-body position iterations (only used if bOverrideIterationCounts is true)
 	ENGINE_API void SetPositionSolverIterationCount(uint8 PositionSolverIterationCountIn);
 
-	// Set the per-body velocity iterations (only used if bOverrideIterationsCounts is true)
+	// Set the per-body velocity iterations (only used if bOverrideIterationCounts is true)
 	ENGINE_API void SetVelocitySolverIterationCount(uint8 VelocitySolverIterationCountIn);
 
-	// Set the per-body projection iterations (only used if bOverrideIterationsCounts is true)
+	// Set the per-body projection iterations (only used if bOverrideIterationCounts is true)
 	ENGINE_API void SetProjectionSolverIterationCount(uint8 ProjectionSolverIterationCountIn);
 
 	// Enable or disable the use of iteration count overrides
 	ENGINE_API void SetOverrideIterationCounts(bool bOverride);
 
-	// Get the per-body position iterations (or -1 if bOverrideIterationsCounts is false)
+	// Get the per-body position iterations (or -1 if bOverrideIterationCounts is false)
 	int32 GetPositionSolverIterationCount() const { return (!!bOverrideIterationCounts) ? PositionSolverIterationCount : -1; }
 
-	// Get the per-body velocity iterations (or -1 if bOverrideIterationsCounts is false)
+	// Get the per-body velocity iterations (or -1 if bOverrideIterationCounts is false)
 	int32 GetVelocitySolverIterationCount() const { return (!!bOverrideIterationCounts) ? VelocitySolverIterationCount : -1; }
 
-	// Get the per-body projection iterations (or -1 if bOverrideIterationsCounts is false)
+	// Get the per-body projection iterations (or -1 if bOverrideIterationCounts is false)
 	int32 GetProjectionSolverIterationCount() const { return (!!bOverrideIterationCounts) ? ProjectionSolverIterationCount : -1; }
 
 	/** [EXPERIMENTAL] Whether Motion-Aware Collision Detection is enabled */
