Actor: Remove the message logged when a non-valid AActor (MarkAsGarbage was called on it) is executing PostLoad. The message is causing log spam during the cook of large projects (17k out of 250k log lines) due to its occurrence whenUChildActorComponent::Serialize marks a duplicated Actor as garbage to suppress the use of the duplicated object. The bug that originally prompted the log message (per CL 67832) was fixed many years ago.
#rnx

diff --git a/Engine/Source/Runtime/Engine/Private/Actor.cpp b/Engine/Source/Runtime/Engine/Private/Actor.cpp
--- a/Engine/Source/Runtime/Engine/Private/Actor.cpp
+++ b/Engine/Source/Runtime/Engine/Private/Actor.cpp
@@ -1158,12 +1158,6 @@ void AActor::PostLoad()
 	{
 		// Propagate the hidden at editor startup flag to the transient hidden flag
 		bHiddenEdTemporary = bHiddenEd;
-
-		// Check/warning when loading actors in editor. Should never load invalid Actors!
-		if ( !IsValidChecked(this) )
-		{
-			UE_LOG(LogActor, Log,  TEXT("Loaded Actor (%s) with IsValid() == false"), *GetName() );
-		}
 	}
 
 	// Fixup DataLayers
