Mobile: Added access to a GBuffer textures in a mobile deferred shading in a multi-pass case

diff --git a/Engine/Shaders/Private/MaterialTemplate.ush b/Engine/Shaders/Private/MaterialTemplate.ush
--- a/Engine/Shaders/Private/MaterialTemplate.ush
+++ b/Engine/Shaders/Private/MaterialTemplate.ush
@@ -3124,6 +3124,10 @@ MaterialFloat4 MobileSceneTextureLookup(inout FMaterialPixelParameters Parameter
 		return 1;
 	case PPI_CustomStencil:
 		return GBuffer.CustomStencil;
+	case PPI_StoredBaseColor:
+		return MaterialFloat4(GBuffer.StoredBaseColor, 0);
+	case PPI_StoredSpecular:
+		return MaterialFloat4(GBuffer.StoredSpecular.rrr, 0);
 	case PPI_WorldTangent:
 		return MaterialFloat4(GBuffer.WorldTangent, 0);
 #ifdef POST_PROCESS_PROPAGATE_ALPHA_INPUT
diff --git a/Engine/Source/Runtime/Engine/Private/Materials/HLSLMaterialTranslator.cpp b/Engine/Source/Runtime/Engine/Private/Materials/HLSLMaterialTranslator.cpp
--- a/Engine/Source/Runtime/Engine/Private/Materials/HLSLMaterialTranslator.cpp
+++ b/Engine/Source/Runtime/Engine/Private/Materials/HLSLMaterialTranslator.cpp
@@ -8477,13 +8477,16 @@ void FHLSLMaterialTranslator::UseSceneTextureId(ESceneTextureId SceneTextureId,
 
 	if (bNeedsGBuffer)
 	{
-		if (IsForwardShadingEnabled(Platform) || (IsMobilePlatform(Platform) && !IsMobileDeferredShadingEnabled(Platform)))
+		const bool bMobilePlatform = IsMobilePlatform(Platform);
+		
+		if ((bMobilePlatform && !IsMobileDeferredShadingEnabled(Platform)) || 
+			(!bMobilePlatform && IsForwardShadingEnabled(Platform)))
 		{
 			Errorf(TEXT("GBuffer scene textures not available with forward shading (platform id %d)."), Platform);
 		}
 		
 		// Post-process can't access memoryless GBuffer on mobile
-		if (IsMobilePlatform(Platform))
+		if (bMobilePlatform && (!IsMobileDeferredShadingEnabled(Platform) || MobileAllowFramebufferFetch(Platform)))
 		{
 			if (Material->GetMaterialDomain() == MD_PostProcess)
 			{
