Remove unmaintained substrate material layers code.
#rnx
#tests Rendering Engine Tests, Cook Tests

diff --git a/Engine/Source/Editor/DetailCustomizations/Private/MaterialAttributePropertyDetails.cpp b/Engine/Source/Editor/DetailCustomizations/Private/MaterialAttributePropertyDetails.cpp
--- a/Engine/Source/Editor/DetailCustomizations/Private/MaterialAttributePropertyDetails.cpp
+++ b/Engine/Source/Editor/DetailCustomizations/Private/MaterialAttributePropertyDetails.cpp
@@ -40,10 +40,8 @@ void FMaterialAttributePropertyDetails::CustomizeDetails(IDetailLayoutBuilder& D
 	for (const TPair<FString, FGuid>& NameGUIDPair : AttributeNameToIDList)
 	{
 		// We remove unwanted elements from the list of material attributes that can be used in the material layer workflow.
-		// We do not want to blend Substrate BSDF which could increase the complexity of the material (BSDF / Slab count) for each blend operation.
-		// Only parameters are allowed to finally be transform into Substrate at the end of the chain.
 		// This filtering is ran only when constructing the UI element.
-		if(Substrate::IsMaterialLayeringSupportEnabled() || NameGUIDPair.Key != ExcludedFrontMaterialName)
+		if(NameGUIDPair.Key != ExcludedFrontMaterialName)
 		{
 			AttributeDisplayNameList.Add(MakeShareable(new FString(NameGUIDPair.Key)));
 		}
diff --git a/Engine/Source/Editor/MaterialEditor/Private/MaterialEditor.cpp b/Engine/Source/Editor/MaterialEditor/Private/MaterialEditor.cpp
--- a/Engine/Source/Editor/MaterialEditor/Private/MaterialEditor.cpp
+++ b/Engine/Source/Editor/MaterialEditor/Private/MaterialEditor.cpp
@@ -782,193 +782,57 @@ void FMaterialEditor::InitMaterialEditor( const EToolkitMode::Type Mode, const T
 				// We can check the usage here and add the appropriate inputs too (e.g. Layer==1MA, Blend==2MA)
 				if (MaterialFunction->GetMaterialFunctionUsage() == EMaterialFunctionUsage::MaterialLayer)
 				{
-					if (Substrate::IsMaterialLayeringSupportEnabled())
+					UMaterialExpression* Input = CreateNewMaterialExpression(UMaterialExpressionFunctionInput::StaticClass(), FVector2D(-350, 300), false, true);
+					if (Input)
 					{
-						UMaterialExpression* Input = CreateNewMaterialExpression(UMaterialExpressionFunctionInput::StaticClass(), FVector2D(-800, 300), false, true);
-						if (Input)
-						{
-							UMaterialExpressionFunctionInput* BaseAttributesInput = Cast<UMaterialExpressionFunctionInput>(Input);
-							BaseAttributesInput->InputType = FunctionInput_MaterialAttributes;
-							BaseAttributesInput->InputName = TEXT("Material Attributes");
-							BaseAttributesInput->bUsePreviewValueAsDefault = true;
-							BaseAttributesInput->PreviewValue = FLinearColor::MakeRandomColor();
-						}
+						UMaterialExpressionFunctionInput* BaseAttributesInput = Cast<UMaterialExpressionFunctionInput>(Input);
+						BaseAttributesInput->InputType = FunctionInput_MaterialAttributes;
+						BaseAttributesInput->InputName = TEXT("Material Attributes");
+						BaseAttributesInput->bUsePreviewValueAsDefault = true;
+					}
 
-						UMaterialExpressionSetMaterialAttributes* SetMaterialAttributes = Cast<UMaterialExpressionSetMaterialAttributes>(CreateNewMaterialExpression(UMaterialExpressionSetMaterialAttributes::StaticClass(), FVector2D(0, 300), false, true));
-						UMaterialExpressionGetMaterialAttributes* GetMaterialAttributes = Cast<UMaterialExpressionGetMaterialAttributes>(CreateNewMaterialExpression(UMaterialExpressionGetMaterialAttributes::StaticClass(), FVector2D(-400, 300), false, true));
-						if (Input && SetMaterialAttributes && GetMaterialAttributes)
+					if (GetDefault<UEditorExperimentalSettings>()->bExampleLayersAndBlends)
+					{
+						UMaterialExpression* SetMaterialAttributes = CreateNewMaterialExpression(UMaterialExpressionSetMaterialAttributes::StaticClass(), FVector2D(40, 300), false, true);
+						if (Input && SetMaterialAttributes)
 						{
-							UMaterialEditingLibrary::ConnectMaterialExpressions(Input, FString(), GetMaterialAttributes, FString());
-						
-							SetMaterialAttributes->ConnectInputAttribute(MP_FrontMaterial, GetMaterialAttributes, GetMaterialAttributes->CreateOrGetOutputAttribute(MP_FrontMaterial));
-							SetMaterialAttributes->ConnectInputAttribute(MP_MaterialAttributes, GetMaterialAttributes, GetMaterialAttributes->CreateOrGetOutputAttribute(MP_MaterialAttributes));
-
+						UMaterialEditingLibrary::ConnectMaterialExpressions(Input, FString(), SetMaterialAttributes, FString());
 							UMaterialEditingLibrary::ConnectMaterialExpressions(SetMaterialAttributes, FString(), Expression, FString());
 							bMaterialDirty = true;
 						}
-
-					}
-					else
-					{
-						UMaterialExpression* Input = CreateNewMaterialExpression(UMaterialExpressionFunctionInput::StaticClass(), FVector2D(-350, 300), false, true);
-						if (Input)
-						{
-							UMaterialExpressionFunctionInput* BaseAttributesInput = Cast<UMaterialExpressionFunctionInput>(Input);
-							BaseAttributesInput->InputType = FunctionInput_MaterialAttributes;
-							BaseAttributesInput->InputName = TEXT("Material Attributes");
-							BaseAttributesInput->bUsePreviewValueAsDefault = true;
-						}
-
-						if (GetDefault<UEditorExperimentalSettings>()->bExampleLayersAndBlends)
-						{
-							UMaterialExpression* SetMaterialAttributes = CreateNewMaterialExpression(UMaterialExpressionSetMaterialAttributes::StaticClass(), FVector2D(40, 300), false, true);
-							if (Input && SetMaterialAttributes)
-							{
-							UMaterialEditingLibrary::ConnectMaterialExpressions(Input, FString(), SetMaterialAttributes, FString());
-								UMaterialEditingLibrary::ConnectMaterialExpressions(SetMaterialAttributes, FString(), Expression, FString());
-								bMaterialDirty = true;
-							}
-						}
 					}
 				}
 				else if (MaterialFunction->GetMaterialFunctionUsage() == EMaterialFunctionUsage::MaterialLayerBlend)
 				{
-					if (Substrate::IsMaterialLayeringSupportEnabled())
+					// "Top layer" should be below "bottom layer" on the graph, to align with B on blend nodes
+					UMaterialExpression* InputTop = CreateNewMaterialExpression(UMaterialExpressionFunctionInput::StaticClass(), FVector2D(-300, 400), false, true);
+					if (InputTop)
 					{
-						{
-							UMaterialExpression* InputBottom = CreateNewMaterialExpression(UMaterialExpressionFunctionInput::StaticClass(), FVector2D(-1100, 200), false, true);
-							if (InputBottom)
-							{
-								UMaterialExpressionFunctionInput* BaseAttributesInput = Cast<UMaterialExpressionFunctionInput>(InputBottom);
-								BaseAttributesInput->InputType = FunctionInput_MaterialAttributes;
-								BaseAttributesInput->InputName = TEXT("Background Layer");
-								BaseAttributesInput->bUsePreviewValueAsDefault = true;
-							}
-
-							// "Top layer" should be below "bottom layer" on the graph, to align with Foreground/B on blend nodes
-							UMaterialExpression* InputTop = CreateNewMaterialExpression(UMaterialExpressionFunctionInput::StaticClass(), FVector2D(-1100, 400), false, true);
-							if (InputTop)
-							{
-								UMaterialExpressionFunctionInput* BaseAttributesInput = Cast<UMaterialExpressionFunctionInput>(InputTop);
-								BaseAttributesInput->InputType = FunctionInput_MaterialAttributes;
-								BaseAttributesInput->InputName = TEXT("Foreground Layer");
-								BaseAttributesInput->bUsePreviewValueAsDefault = true;
-								BaseAttributesInput->PreviewValue = FLinearColor::White;
-							}
-
-							if(InputTop && InputBottom && GetDefault<UEditorExperimentalSettings>()->bExampleLayersAndBlends)
-							{
-								static TObjectPtr<UMaterialFunction> DefaultBlendFunction = FindObject<UMaterialFunction>(nullptr, DEFAULT_SUBSTRATE_MATERIALLAYERBLEND_PATH);
-								if (!DefaultBlendFunction)
-								{
-									DefaultBlendFunction = LoadObject<UMaterialFunction>(nullptr, DEFAULT_SUBSTRATE_MATERIALLAYERBLEND_PATH);
-								}
-							
-								if (DefaultBlendFunction)
-								{
-									if(UMaterialExpressionMaterialFunctionCall* BlendFunctionCall = Cast<UMaterialExpressionMaterialFunctionCall>(CreateNewMaterialExpression(UMaterialExpressionMaterialFunctionCall::StaticClass(), FVector2D(-100, 300), false, false)))
-									{
-										BlendFunctionCall->Function = MaterialFunction;
-										InputTop->MaterialExpressionEditorX = -800;
-										InputBottom->MaterialExpressionEditorX = -800;
-
-										if (BlendFunctionCall->SetMaterialFunction(DefaultBlendFunction))
-										{
-											if (BlendFunctionCall->FunctionInputs.Num() >= 8 && BlendFunctionCall->FunctionOutputs.Num() > 0)
-											{
-												BlendFunctionCall->FunctionInputs[0].Input.Connect(0, InputBottom);
-												BlendFunctionCall->FunctionInputs[1].Input.Connect(0, InputTop);
-												UMaterialEditingLibrary::ConnectMaterialExpressions(BlendFunctionCall, FString(), Expression, FString());
-
-												int32 StartIndex = 3;
-												int32 EndIndex = 7;
-												int32 BoolYPosition = 500;
-												for (int32 FunctionCallIndex = StartIndex; FunctionCallIndex <= EndIndex; FunctionCallIndex++)
-												{
-													UMaterialExpressionStaticBool* ThisBool = Cast<UMaterialExpressionStaticBool>(CreateNewMaterialExpression(UMaterialExpressionStaticBool::StaticClass(), FVector2D(-450, BoolYPosition), false, false));
-													if (ThisBool)
-													{
-														BlendFunctionCall->FunctionInputs[FunctionCallIndex].Input.Connect(0, ThisBool);
-														BoolYPosition += 80;		
-													}
-												}
-
-												bMaterialDirty = true;
-											}
-										}
-									}
-								}
+						UMaterialExpressionFunctionInput* BaseAttributesInput = Cast<UMaterialExpressionFunctionInput>(InputTop);
+						BaseAttributesInput->InputType = FunctionInput_MaterialAttributes;
+						BaseAttributesInput->InputName = TopMaterialBlendInputName;
+						BaseAttributesInput->bUsePreviewValueAsDefault = true;
+						BaseAttributesInput->BlendInputRelevance = EBlendInputRelevance::Top;
+					}
 
-								if (!bMaterialDirty)
-								{
-									UMaterialExpressionSetMaterialAttributes* SetMaterialAttributes = Cast<UMaterialExpressionSetMaterialAttributes>(CreateNewMaterialExpression(UMaterialExpressionSetMaterialAttributes::StaticClass(), FVector2D(-75, 300), false, true));
-									UMaterialExpression* HorizontalMixingNode = CreateNewMaterialExpression(UMaterialExpressionSubstrateHorizontalMixing::StaticClass(), FVector2D(-350, 200), false, true);
-									UMaterialExpressionGetMaterialAttributes* GetBottomMaterialAttributes = Cast<UMaterialExpressionGetMaterialAttributes>(CreateNewMaterialExpression(UMaterialExpressionGetMaterialAttributes::StaticClass(), FVector2D(-750, 200), false, true));
-									int32 BottomFrontMaterialIndex = GetBottomMaterialAttributes ? GetBottomMaterialAttributes->CreateOrGetOutputAttribute(MP_FrontMaterial) : INDEX_NONE;
-
-									UMaterialExpression* LegacyBlendNode = CreateNewMaterialExpression(UMaterialExpressionBlendMaterialAttributes::StaticClass(), FVector2D(-350, 400), false, true);
-									UMaterialExpressionGetMaterialAttributes* GetTopMaterialAttributes = Cast< UMaterialExpressionGetMaterialAttributes>(CreateNewMaterialExpression(UMaterialExpressionGetMaterialAttributes::StaticClass(), FVector2D(-750, 400), false, true));
-									int32 TopFrontMaterialIndex = GetTopMaterialAttributes ? GetTopMaterialAttributes->CreateOrGetOutputAttribute(MP_FrontMaterial) : INDEX_NONE;
-
-									if (SetMaterialAttributes 
-										&& HorizontalMixingNode 
-										&& BottomFrontMaterialIndex > 0 
-										&& LegacyBlendNode 
-										&& TopFrontMaterialIndex > 0)
-									{
-										const FString FrontMaterialName = FMaterialAttributeDefinitionMap::GetAttributeName(MP_FrontMaterial);
-										const FString MaterialAttributesName = FMaterialAttributeDefinitionMap::GetAttributeName(MP_MaterialAttributes);
-
-										UMaterialEditingLibrary::ConnectMaterialExpressions(InputBottom, FString(), GetBottomMaterialAttributes, FString());
-										UMaterialEditingLibrary::ConnectMaterialExpressions(GetBottomMaterialAttributes, FrontMaterialName, HorizontalMixingNode, FString("Background"));
-										UMaterialEditingLibrary::ConnectMaterialExpressions(GetBottomMaterialAttributes, MaterialAttributesName, LegacyBlendNode, FString("A"));
-
-										UMaterialEditingLibrary::ConnectMaterialExpressions(InputTop, FString(), GetTopMaterialAttributes, FString());
-										UMaterialEditingLibrary::ConnectMaterialExpressions(GetTopMaterialAttributes, FrontMaterialName, HorizontalMixingNode, FString("Foreground"));
-										UMaterialEditingLibrary::ConnectMaterialExpressions(GetTopMaterialAttributes, MaterialAttributesName, LegacyBlendNode, FString("B"));
-
-										SetMaterialAttributes->ConnectInputAttribute(MP_FrontMaterial, HorizontalMixingNode);
-										SetMaterialAttributes->ConnectInputAttribute(MP_MaterialAttributes, LegacyBlendNode);
-
-										UMaterialEditingLibrary::ConnectMaterialExpressions(SetMaterialAttributes, FString(), Expression, FString());
-										bMaterialDirty = true;
-									}
-								}
-							}
-						}
+					UMaterialExpression* InputBottom = CreateNewMaterialExpression(UMaterialExpressionFunctionInput::StaticClass(), FVector2D(-300, 200), false, true);
+					if (InputBottom)
+					{
+						UMaterialExpressionFunctionInput* BaseAttributesInput = Cast<UMaterialExpressionFunctionInput>(InputBottom);
+						BaseAttributesInput->InputType = FunctionInput_MaterialAttributes;
+						BaseAttributesInput->InputName = BottomMaterialBlendInputName;
+						BaseAttributesInput->bUsePreviewValueAsDefault = true;
+						BaseAttributesInput->BlendInputRelevance = EBlendInputRelevance::Bottom;
 					}
-					else
+					if (GetDefault<UEditorExperimentalSettings>()->bExampleLayersAndBlends)
 					{
-						// "Top layer" should be below "bottom layer" on the graph, to align with B on blend nodes
-						UMaterialExpression* InputTop = CreateNewMaterialExpression(UMaterialExpressionFunctionInput::StaticClass(), FVector2D(-300, 400), false, true);
-						if (InputTop)
-						{
-							UMaterialExpressionFunctionInput* BaseAttributesInput = Cast<UMaterialExpressionFunctionInput>(InputTop);
-							BaseAttributesInput->InputType = FunctionInput_MaterialAttributes;
-							BaseAttributesInput->InputName = TopMaterialBlendInputName;
-							BaseAttributesInput->bUsePreviewValueAsDefault = true;
-							BaseAttributesInput->BlendInputRelevance = EBlendInputRelevance::Top;
-						}
-
-						UMaterialExpression* InputBottom = CreateNewMaterialExpression(UMaterialExpressionFunctionInput::StaticClass(), FVector2D(-300, 200), false, true);
-						if (InputBottom)
-						{
-							UMaterialExpressionFunctionInput* BaseAttributesInput = Cast<UMaterialExpressionFunctionInput>(InputBottom);
-							BaseAttributesInput->InputType = FunctionInput_MaterialAttributes;
-							BaseAttributesInput->InputName = BottomMaterialBlendInputName;
-							BaseAttributesInput->bUsePreviewValueAsDefault = true;
-							BaseAttributesInput->BlendInputRelevance = EBlendInputRelevance::Bottom;
-						}
-						if (GetDefault<UEditorExperimentalSettings>()->bExampleLayersAndBlends)
+						UMaterialExpression* BlendMaterialAttributes = CreateNewMaterialExpression(UMaterialExpressionBlendMaterialAttributes::StaticClass(), FVector2D(40, 300), false, true);
+						if (InputTop && InputBottom && BlendMaterialAttributes)
 						{
-							UMaterialExpression* BlendMaterialAttributes = CreateNewMaterialExpression(UMaterialExpressionBlendMaterialAttributes::StaticClass(), FVector2D(40, 300), false, true);
-							if (InputTop && InputBottom && BlendMaterialAttributes)
-							{
-								UMaterialEditingLibrary::ConnectMaterialExpressions(InputBottom, FString(), BlendMaterialAttributes, FString(TEXT("A")));
-								UMaterialEditingLibrary::ConnectMaterialExpressions(InputTop, FString(), BlendMaterialAttributes, FString(TEXT("B")));
-								UMaterialEditingLibrary::ConnectMaterialExpressions(BlendMaterialAttributes, FString(), Expression, FString());
-								bMaterialDirty = true;
-							}
+							UMaterialEditingLibrary::ConnectMaterialExpressions(InputBottom, FString(), BlendMaterialAttributes, FString(TEXT("A")));
+							UMaterialEditingLibrary::ConnectMaterialExpressions(InputTop, FString(), BlendMaterialAttributes, FString(TEXT("B")));
+							UMaterialEditingLibrary::ConnectMaterialExpressions(BlendMaterialAttributes, FString(), Expression, FString());
+							bMaterialDirty = true;
 						}
 					}
 				}
@@ -1826,18 +1690,9 @@ void FMaterialEditor::CreateInternalWidgets()
 		.InGenerator(Generator)
 		.InMaterialDetailsView(MaterialDetailsView);
 
-	if(Substrate::IsMaterialLayeringSupportEnabled())
-	{
-		MaterialLayersFunctionsInstance = SNew(SMaterialLayersFunctionsInstanceWrapper)
-			.InMaterialEditorInstance(MaterialEditorInstance)
-			.InGenerator(Generator);
-	}
-	else
-	{
-		MaterialLayersFunctionsInstance = SNew(SMaterialLayersFunctionsMaterialWrapper)
-			.InMaterialEditorInstance(MaterialEditorInstance)
-			.InGenerator(Generator);
-	}
+	MaterialLayersFunctionsInstance = SNew(SMaterialLayersFunctionsMaterialWrapper)
+		.InMaterialEditorInstance(MaterialEditorInstance)
+		.InGenerator(Generator);
 
 	Palette = SNew(SMaterialPalette, SharedThis(this));
 
@@ -1951,7 +1806,7 @@ void FMaterialEditor::OnChangeBreadCrumbGraph(UEdGraph* InGraph)
 void FMaterialEditor::GeneratorRowsRefreshed()
 {
 	MaterialParametersOverviewWidget->Refresh();
-	if (!Substrate::IsMaterialLayeringSupportEnabled() && MaterialLayersFunctionsInstance)
+	if (MaterialLayersFunctionsInstance)
 	{
 		MaterialLayersFunctionsInstance->Refresh();
 	}
diff --git a/Engine/Source/Editor/MaterialEditor/Private/MaterialEditorDetailCustomization.cpp b/Engine/Source/Editor/MaterialEditor/Private/MaterialEditorDetailCustomization.cpp
--- a/Engine/Source/Editor/MaterialEditor/Private/MaterialEditorDetailCustomization.cpp
+++ b/Engine/Source/Editor/MaterialEditor/Private/MaterialEditorDetailCustomization.cpp
@@ -1070,11 +1070,6 @@ void FMaterialExpressionLayersParameterDetails::CustomizeDetails(IDetailLayoutBu
 	// for expression parameters all their properties are in one category based on their class name.
 
 	FName LayerCategory = FName("Layers");
-	if(Substrate::IsMaterialLayeringSupportEnabled())
-	{
-		DetailLayout.HideCategory(LayerCategory);
-	}
-
 	IDetailCategoryBuilder& Category = DetailLayout.EditCategory(LayerCategory);
 	// Get a handle to the property we are about to edit
 	GroupPropertyHandle = DetailLayout.GetProperty("DefaultLayers");
diff --git a/Engine/Source/Editor/MaterialEditor/Private/MaterialEditorInstanceDetailCustomization.cpp b/Engine/Source/Editor/MaterialEditor/Private/MaterialEditorInstanceDetailCustomization.cpp
--- a/Engine/Source/Editor/MaterialEditor/Private/MaterialEditorInstanceDetailCustomization.cpp
+++ b/Engine/Source/Editor/MaterialEditor/Private/MaterialEditorInstanceDetailCustomization.cpp
@@ -14,7 +14,6 @@
 #include "Materials/MaterialInterface.h"
 #include "MaterialEditor/DEditorParameterCollectionParameterValue.h"
 #include "MaterialEditor/MaterialEditorInstanceConstant.h"
-#include "SMaterialSubstrateTree.h"
 #include "Materials/MaterialInstance.h"
 #include "Materials/MaterialParameterCollection.h"
 #include "MaterialShared.h"
diff --git a/Engine/Source/Editor/MaterialEditor/Private/MaterialEditorUtilities.cpp b/Engine/Source/Editor/MaterialEditor/Private/MaterialEditorUtilities.cpp
--- a/Engine/Source/Editor/MaterialEditor/Private/MaterialEditorUtilities.cpp
+++ b/Engine/Source/Editor/MaterialEditor/Private/MaterialEditorUtilities.cpp
@@ -583,13 +583,6 @@ void FMaterialEditorUtilities::GetVisibleMaterialParametersFromExpression(
 				TArray<TObjectPtr<UMaterialExpressionMaterialFunctionCall>>& LayerCallers = LayersExpression->LayerCallers;
 				TArray<TObjectPtr<UMaterialExpressionMaterialFunctionCall>>& BlendCallers = LayersExpression->BlendCallers;
 				
-				if(Substrate::IsMaterialLayeringSupportEnabled())
-				{
-					TSharedPtr< FMaterialLayersFunctionsRuntimeGraphCache > LayerRuntimeTree = LayersValue.RuntimeGraphCache;
-					LayerCallers = LayerRuntimeTree->LayerCallers;
-					BlendCallers = LayerRuntimeTree->BlendCallers;
-				}
-
 				for (auto& Layer : LayerCallers)
 				{
 					// Possible that Layer->FunctionOutputs will be empty if this is a newly create layer
diff --git a/Engine/Source/Editor/MaterialEditor/Private/MaterialInstanceEditor.cpp b/Engine/Source/Editor/MaterialEditor/Private/MaterialInstanceEditor.cpp
--- a/Engine/Source/Editor/MaterialEditor/Private/MaterialInstanceEditor.cpp
+++ b/Engine/Source/Editor/MaterialEditor/Private/MaterialInstanceEditor.cpp
@@ -261,9 +261,8 @@ void FMaterialInstanceEditor::RegisterTabSpawners(const TSharedRef<class FTabMan
 
 	if (!bIsFunctionPreviewMaterial)
 	{
-		FText DisplayName = Substrate::IsMaterialLayeringSupportEnabled() ? LOCTEXT("MaterialLayersTab", "Material Layers") : LOCTEXT("LayerPropertiesTab", "Layer Parameters");
 		InTabManager->RegisterTabSpawner(FMaterialInstanceEditorTabs::LayerPropertiesTabId, FOnSpawnTab::CreateSP(this, &FMaterialInstanceEditor::SpawnTab_LayerProperties))
-			.SetDisplayName(DisplayName)
+			.SetDisplayName(LOCTEXT("LayerPropertiesTab", "Layer Parameters"))
 			.SetGroup(WorkspaceMenuCategoryRef)
 			.SetIcon(FSlateIcon(FAppStyle::GetAppStyleSetName(), "LevelEditor.Tabs.Layers"));
 	}
@@ -273,19 +272,6 @@ void FMaterialInstanceEditor::RegisterTabSpawners(const TSharedRef<class FTabMan
 		.SetGroup(WorkspaceMenuCategoryRef)
 		.SetIcon(FSlateIcon(FAppStyle::GetAppStyleSetName(), "LevelEditor.Tabs.Details"));
 
-	if(Substrate::IsMaterialLayeringSupportEnabled())
-	{
-		InTabManager->RegisterTabSpawner(FMaterialInstanceEditorTabs::AssetBrowserTabId, FOnSpawnTab::CreateSP(this, &FMaterialInstanceEditor::SpawnTab_AssetBrowser))
-			.SetDisplayName(LOCTEXT("AssetBrowserTab", "Asset Browser"))
-			.SetGroup(WorkspaceMenuCategoryRef)
-			.SetIcon(FSlateIcon(FAppStyle::GetAppStyleSetName(), "LevelEditor.Tabs.Details"));
-
-		InTabManager->RegisterTabSpawner(FMaterialInstanceEditorTabs::SubstrateTabId, FOnSpawnTab::CreateSP(this, &FMaterialInstanceEditor::SpawnTab_Substrate))
-			.SetDisplayName(LOCTEXT("SubstrateTab", "Substrate"))
-			.SetGroup(WorkspaceMenuCategoryRef)
-			.SetIcon(FSlateIcon(FAppStyle::GetAppStyleSetName(), "Kismet.Tabs.Palette"));	// SUBSTRATE_TODO a Substrate icon
-	}
-
 	MaterialStatsManager->RegisterTabs();
 
 	OnRegisterTabSpawners().Broadcast(InTabManager);
@@ -302,11 +288,6 @@ void FMaterialInstanceEditor::UnregisterTabSpawners(const TSharedRef<class FTabM
 		InTabManager->UnregisterTabSpawner(FMaterialInstanceEditorTabs::LayerPropertiesTabId);
 	}
 	InTabManager->UnregisterTabSpawner( FMaterialInstanceEditorTabs::PreviewSettingsTabId );
-	if(Substrate::IsMaterialLayeringSupportEnabled())
-	{
-		InTabManager->UnregisterTabSpawner( FMaterialInstanceEditorTabs::AssetBrowserTabId );
-		InTabManager->UnregisterTabSpawner(FMaterialInstanceEditorTabs::SubstrateTabId);
-	}
 
 	MaterialStatsManager->UnregisterTabs();
 
@@ -484,73 +465,28 @@ void FMaterialInstanceEditor::InitMaterialInstanceEditor( const EToolkitMode::Ty
 
 	if (!bIsFunctionPreviewMaterial)
 	{
-		
-		if(Substrate::IsMaterialLayeringSupportEnabled())
-		{
-			StandaloneDefaultLayout = FTabManager::NewLayout("Standalone_MaterialInstanceEditor_Layout_v9")
-				->AddArea
+		StandaloneDefaultLayout = FTabManager::NewLayout("Standalone_MaterialInstanceEditor_Layout_v8")
+			->AddArea
+			(
+				FTabManager::NewPrimaryArea()->SetOrientation(Orient_Vertical)
+				->Split
 				(
-					FTabManager::NewPrimaryArea()->SetOrientation(Orient_Vertical)
+					FTabManager::NewSplitter()->SetOrientation(Orient_Horizontal)->SetSizeCoefficient(0.9f)
 					->Split
 					(
-						FTabManager::NewSplitter()->SetOrientation(Orient_Horizontal)->SetSizeCoefficient(0.9f)
-						->Split
-						(
-							FTabManager::NewStack()->SetSizeCoefficient(0.70f)->SetHideTabWell(true)
-							->AddTab(FMaterialInstanceEditorTabs::PreviewTabId, ETabState::OpenedTab)
-							->AddTab(FMaterialInstanceEditorTabs::PreviewSettingsTabId, ETabState::ClosedTab)
-						)
-						->Split
-						(
-							FTabManager::NewSplitter()->SetOrientation(Orient_Vertical)->SetSizeCoefficient(0.7f)
-							->Split
-							(
-						FTabManager::NewSplitter()->SetOrientation(Orient_Horizontal)->SetSizeCoefficient(0.3f)
-								->Split
-								(
-									FTabManager::NewStack()/*->SetSizeCoefficient(0.30f)*/
-									->AddTab(FMaterialInstanceEditorTabs::LayerPropertiesTabId, ETabState::OpenedTab)
-								)
-								->Split
-								(
-									FTabManager::NewStack()/*->SetSizeCoefficient(0.30f)*/
-									->AddTab(FMaterialInstanceEditorTabs::AssetBrowserTabId, ETabState::OpenedTab)
-								)
-							)
-							->Split
-							(
-							FTabManager::NewStack()->SetSizeCoefficient(0.30f)
-									->AddTab(FMaterialInstanceEditorTabs::PropertiesTabId, ETabState::OpenedTab)
-							)
-						)
+						FTabManager::NewStack()->SetSizeCoefficient(0.70f)->SetHideTabWell(true)
+						->AddTab(FMaterialInstanceEditorTabs::PreviewTabId, ETabState::OpenedTab)
+						->AddTab(FMaterialInstanceEditorTabs::PreviewSettingsTabId, ETabState::ClosedTab)
 					)
-				);
-		}
-		else
-		{
-			StandaloneDefaultLayout = FTabManager::NewLayout("Standalone_MaterialInstanceEditor_Layout_v8")
-				->AddArea
-				(
-					FTabManager::NewPrimaryArea()->SetOrientation(Orient_Vertical)
 					->Split
 					(
-						FTabManager::NewSplitter()->SetOrientation(Orient_Horizontal)->SetSizeCoefficient(0.9f)
-						->Split
-						(
-							FTabManager::NewStack()->SetSizeCoefficient(0.70f)->SetHideTabWell(true)
-							->AddTab(FMaterialInstanceEditorTabs::PreviewTabId, ETabState::OpenedTab)
-							->AddTab(FMaterialInstanceEditorTabs::PreviewSettingsTabId, ETabState::ClosedTab)
-						)
-						->Split
-						(
-							FTabManager::NewStack()->SetSizeCoefficient(0.30f)
-							->AddTab(FMaterialInstanceEditorTabs::PropertiesTabId, ETabState::OpenedTab)
-							->AddTab(FMaterialInstanceEditorTabs::LayerPropertiesTabId, ETabState::OpenedTab)
-							->SetForegroundTab(FMaterialInstanceEditorTabs::PropertiesTabId)
-						)
+						FTabManager::NewStack()->SetSizeCoefficient(0.30f)
+						->AddTab(FMaterialInstanceEditorTabs::PropertiesTabId, ETabState::OpenedTab)
+						->AddTab(FMaterialInstanceEditorTabs::LayerPropertiesTabId, ETabState::OpenedTab)
+						->SetForegroundTab(FMaterialInstanceEditorTabs::PropertiesTabId)
 					)
-				);
-		}
+				)
+			);
 	}
 	
 	const bool bCreateDefaultStandaloneMenu = true;
@@ -1126,32 +1062,18 @@ TSharedRef<SDockTab> FMaterialInstanceEditor::SpawnTab_LayerProperties(const FSp
 {
 	check(Args.GetTabId().TabType == FMaterialInstanceEditorTabs::LayerPropertiesTabId);
 
-	if(Substrate::IsMaterialLayeringSupportEnabled())
-	{
-		TSharedRef<SDockTab> SpawnedTab = SNew(SDockTab)
-		.Label(LOCTEXT("MaterialLayerPropertiesTitle", "Layer Parameters"))
-		[
-			MaterialLayersFunctionsInstance.ToSharedRef()
-		];
-
-		AddToSpawnedToolPanels(Args.GetTabId().TabType, SpawnedTab);
-		return SpawnedTab;
-	}
-	else
-	{
-		TSharedRef<SDockTab> SpawnedTab = SNew(SDockTab)
-		.Label(LOCTEXT("MaterialLayerPropertiesTitle", "Layer Parameters"))
-		[
-			SNew(SBorder)
-				.Padding(4)
-				[
-					MaterialLayersFunctionsInstance.ToSharedRef()
-				]
-		];
+	TSharedRef<SDockTab> SpawnedTab = SNew(SDockTab)
+	.Label(LOCTEXT("MaterialLayerPropertiesTitle", "Layer Parameters"))
+	[
+		SNew(SBorder)
+			.Padding(4)
+			[
+				MaterialLayersFunctionsInstance.ToSharedRef()
+			]
+	];
 
-		AddToSpawnedToolPanels(Args.GetTabId().TabType, SpawnedTab);
-		return SpawnedTab;
-	}
+	AddToSpawnedToolPanels(Args.GetTabId().TabType, SpawnedTab);
+	return SpawnedTab;
 }
 
 TSharedRef<SDockTab> FMaterialInstanceEditor::SpawnTab_PreviewSettings(const FSpawnTabArgs& Args)
@@ -1848,11 +1770,6 @@ void FMaterialInstanceEditor::Refresh()
 	UpdatePropertyWindow();
 
 	RefreshOnScreenTextureWarnings();
-
-	if (Substrate::IsMaterialLayeringSupportEnabled() && SubstrateWidget)
-	{
-		SubstrateWidget->UpdateFromMaterial();
-	}
 }
 
 void FMaterialInstanceEditor::PostUndo( bool bSuccess )
diff --git a/Engine/Source/Editor/MaterialEditor/Private/MaterialLayersFunctionsCustomization.cpp b/Engine/Source/Editor/MaterialEditor/Private/MaterialLayersFunctionsCustomization.cpp
--- a/Engine/Source/Editor/MaterialEditor/Private/MaterialLayersFunctionsCustomization.cpp
+++ b/Engine/Source/Editor/MaterialEditor/Private/MaterialLayersFunctionsCustomization.cpp
@@ -109,12 +109,6 @@ void FMaterialLayersFunctionsCustomization::RebuildChildren()
 
 void FMaterialLayersFunctionsCustomization::GenerateHeaderRowContent(FDetailWidgetRow& NodeRow)
 {
-	if(Substrate::IsMaterialLayeringSupportEnabled())
-	{
-		// early out if Substrate layering is enabled as we only want to edit FMaterialLayersFunctions through the new Substrate tree UI
-		return;
-	}
-	
 	TAttribute<bool>::FGetter IsResetButtonEnabledDelegate = TAttribute<bool>::FGetter::CreateSP(this, &FMaterialLayersFunctionsCustomization::IsResetToDefaultVisible);
 	TAttribute<bool> IsEnabledAttribute = TAttribute<bool>::Create(IsResetButtonEnabledDelegate);
 
@@ -151,12 +145,6 @@ void FMaterialLayersFunctionsCustomization::GenerateHeaderRowContent(FDetailWidg
 
 void FMaterialLayersFunctionsCustomization::GenerateChildContent(IDetailChildrenBuilder& ChildrenBuilder)
 {
-	if(Substrate::IsMaterialLayeringSupportEnabled())
-	{
-		// early out if the new prototype is enabled as we only want to edit FMaterialLayersFunctions through the new Substrate tree UI
-		return;
-	}
-	
 	uint32 LayerChildren;
 	LayerHandle->GetNumChildren(LayerChildren);
 	uint32 BlendChildren;
diff --git a/Engine/Source/Editor/MaterialEditor/Private/MaterialPropertyHelpers.cpp b/Engine/Source/Editor/MaterialEditor/Private/MaterialPropertyHelpers.cpp
--- a/Engine/Source/Editor/MaterialEditor/Private/MaterialPropertyHelpers.cpp
+++ b/Engine/Source/Editor/MaterialEditor/Private/MaterialPropertyHelpers.cpp
@@ -142,18 +142,13 @@ void FMaterialPropertyHelpers::OnMaterialLayerAssetChanged(const struct FAssetDa
 			{
 				InMaterialFunction->Layers[Index] = LayerFunction;
 				InMaterialFunction->UnlinkLayerFromParent(Index);
-				if(Substrate::IsMaterialLayeringSupportEnabled())
-				{
-					// update name of the layer
-					InMaterialFunction->EditorOnly.LayerNames[Index] = FText::FromName(InAssetData.AssetName);
-				}
 			}
 			break;
 		case EMaterialParameterAssociation::BlendParameter:
 			if (InMaterialFunction->Blends[Index] != LayerFunction)
 			{
 				InMaterialFunction->Blends[Index] = LayerFunction;
-				InMaterialFunction->UnlinkLayerFromParent(Substrate::IsMaterialLayeringSupportEnabled() ? Index : Index + 1); // Blend indices are offset by 1 in legacy, no blend for base layer
+				InMaterialFunction->UnlinkLayerFromParent(Index + 1); // Blend indices are offset by 1, no blend for base layer
 			}
 			break;
 		default:
diff --git a/Engine/Source/Editor/MaterialEditor/Private/SMaterialEditorTopologyWidget.cpp b/Engine/Source/Editor/MaterialEditor/Private/SMaterialEditorTopologyWidget.cpp
--- a/Engine/Source/Editor/MaterialEditor/Private/SMaterialEditorTopologyWidget.cpp
+++ b/Engine/Source/Editor/MaterialEditor/Private/SMaterialEditorTopologyWidget.cpp
@@ -244,32 +244,18 @@ static const TSharedRef<SWidget> InternalProcessOperatorAsThumbnails(
 		}
 		case SUBSTRATE_OPERATOR_BSDF_LEGACY:	// legacy BSDF should have been converted to BSDF already.
 		case SUBSTRATE_OPERATOR_BSDF:
-			{
-				int32 ExpressionIndex = INDEX_NONE;
-				const float ThumbnailSize = 40.0f;
-				UObject* ThumbnailObject = nullptr;
-
-				if(Substrate::IsMaterialLayeringSupportEnabled())
-				{
-					const TSharedPtr<FMaterialLayersFunctionsRuntimeGraphCache> LayerTreeCache = LayersFunctions->RuntimeGraphCache;
-					if (LayerTreeCache)
-					{
-						ExpressionIndex = LayerTreeCache->FindExpressionIndex(Op.MaterialExpressionGuid, LayersFunctions);
-					}
-					if (ExpressionIndex >= 0)
-					{
-						ThumbnailObject = LayerTreeCache->NodePreviewMaterials[ExpressionIndex];
-					}
-				}
+		{
+			int32 ExpressionIndex = INDEX_NONE;
+			const float ThumbnailSize = 40.0f;
+			UObject* ThumbnailObject = nullptr;
 
-				const TSharedPtr<FAssetThumbnail> AssetThumbnail = MakeShareable(new FAssetThumbnail(ThumbnailObject, ThumbnailSize, ThumbnailSize, UThumbnailManager::Get().GetSharedThumbnailPool()));
-				TSharedRef<SWidget> ThumbnailWidget = SNew(SBorder)
-					.Padding(2.0f)
-					[
-						AssetThumbnail->MakeThumbnailWidget()
-					];
+			const TSharedPtr<FAssetThumbnail> AssetThumbnail = MakeShareable(new FAssetThumbnail(ThumbnailObject, ThumbnailSize, ThumbnailSize, UThumbnailManager::Get().GetSharedThumbnailPool()));
+			TSharedRef<SWidget> ThumbnailWidget = SNew(SBorder)
+				.Padding(2.0f)
+				[
+					AssetThumbnail->MakeThumbnailWidget()
+				];
 			AssetThumbnail->SetRealTime(true);
-			//	ThumbnailWidget->SetOnMouseDoubleClick(FPointerEventHandler::CreateSP(this, &SMaterialSubstrateTree::OnThumbnailDoubleClick, InAssociation, InIndex));
 			return ThumbnailWidget;
 		}
 	}
diff --git a/Engine/Source/Editor/MaterialEditor/Private/SMaterialEditorViewport.cpp b/Engine/Source/Editor/MaterialEditor/Private/SMaterialEditorViewport.cpp
--- a/Engine/Source/Editor/MaterialEditor/Private/SMaterialEditorViewport.cpp
+++ b/Engine/Source/Editor/MaterialEditor/Private/SMaterialEditorViewport.cpp
@@ -276,10 +276,10 @@ void SMaterialEditor3DPreviewViewport::Construct(const FArguments& InArgs)
 	PostProcessVolumeActor = nullptr;
 
 	UMaterialInterface* Material = GetPreviewingMaterialInterface();
-	if (Material)
-	{
-		SetPreviewMaterial(Material);
-	}
+		if (Material)
+		{
+			SetPreviewMaterial(Material);
+		}
 
 	SetPreviewAsset( GUnrealEd->GetThumbnailManager()->EditorSphere );
 
@@ -725,7 +725,7 @@ void SMaterialEditor3DPreviewViewport::BindCommands()
 
 	if(MaterialEditorPtr.IsValid())
 	{
-		CommandList->Append(MaterialEditorPtr.Pin()->GetToolkitCommands());
+	CommandList->Append(MaterialEditorPtr.Pin()->GetToolkitCommands());
 	}
 
 	// Add the commands to the toolkit command list so that the toolbar buttons can find them
@@ -1010,31 +1010,14 @@ TSharedPtr<IPreviewProfileController> SMaterialEditor3DPreviewViewport::CreatePr
 
 void SMaterialEditor3DPreviewViewport::PopulateViewportOverlays(TSharedRef<class SOverlay> Overlay)
 {
-	// only show topology widget for material instances
-	if (Substrate::IsMaterialLayeringSupportEnabled() && MaterialEditorPtr.IsValid() && Cast<UMaterialInstance>(MaterialEditorPtr.Pin()->GetMaterialInterface()))
-	{
-		Overlay->AddSlot()
-			// clang-format off
-			.VAlign(VAlign_Top)
-			.HAlign(HAlign_Right)
-			.Padding(FMargin(5.0f, 5.0f, 50.0f, 5.0f))
-			[
-				SAssignNew(SubstrateWidget, SMaterialEditorTopologyWidget, MaterialEditorPtr)
-			];
-		// clang-format on
-	}
-
-	// clang-format on
 	// add the feature level display widget
 	Overlay->AddSlot()
-		// clang-format off
 		.VAlign(VAlign_Top)
 		.HAlign(HAlign_Right)
 		.Padding(5.0f)
 		[
 			BuildShaderPlatformWidget()
 		];
-		// clang-format on
 }
 
 EVisibility SMaterialEditor3DPreviewViewport::OnGetViewportContentVisibility() const
diff --git a/Engine/Source/Editor/MaterialEditor/Private/SMaterialLayersFunctionsTree.cpp b/Engine/Source/Editor/MaterialEditor/Private/SMaterialLayersFunctionsTree.cpp
--- a/Engine/Source/Editor/MaterialEditor/Private/SMaterialLayersFunctionsTree.cpp
+++ b/Engine/Source/Editor/MaterialEditor/Private/SMaterialLayersFunctionsTree.cpp
@@ -1,7 +1,6 @@
 // Copyright Epic Games, Inc. All Rights Reserved.
 
 #include "SMaterialLayersFunctionsTree.h"
-#include "SMaterialSubstrateTree.h"
 #include "MaterialEditor/DEditorMaterialLayersParameterValue.h"
 #include "MaterialEditor/MaterialEditorInstanceConstant.h"
 #include "Materials/Material.h"
@@ -1448,13 +1447,6 @@ void SMaterialLayersFunctionsInstanceWrapper::Refresh()
 	TSharedRef<SWidget> ThumbnailWidget = SNullWidget::NullWidget;
 	TSharedPtr<SBox> ThumbnailBox = nullptr;
 	FText MaterialName = FText::GetEmpty();
-	if (Substrate::IsMaterialLayeringSupportEnabled())
-	{
-		UObject* ThumbnailObject = MaterialEditorInstance->GetMaterialInterface().Get();
-		const TSharedPtr<FAssetThumbnail> AssetThumbnail = MakeShareable(new FAssetThumbnail(ThumbnailObject, ThumbnailSize, ThumbnailSize, NestedTree->GetTreeThumbnailPool()));
-		ThumbnailWidget = AssetThumbnail->MakeThumbnailWidget();
-		MaterialName = FText::FromName(MaterialEditorInstance->GetMaterialInterface()->GetFName());
-	}
 	
 	if (LayerParameter != nullptr)
 	{
@@ -1464,63 +1456,16 @@ void SMaterialLayersFunctionsInstanceWrapper::Refresh()
 				return NestedTree->RelinkLayersToParent();
 			});
 		
-		if (Substrate::IsMaterialLayeringSupportEnabled())
-		{
-			HeaderBox = SNew(SHorizontalBox)
+		HeaderBox = SNew(SHorizontalBox)
 			+ SHorizontalBox::Slot()
-				.AutoWidth()
-					.HAlign(HAlign_Center)
-					.VAlign(VAlign_Center)
-					.Padding(4.0f)
-					.MaxWidth(ThumbnailSize)
-					[
-						SAssignNew(ThumbnailBox, SBox)
-						.MaxDesiredWidth(ThumbnailSize)
-						.MaxDesiredWidth(ThumbnailSize)
-						.MaxDesiredHeight(ThumbnailSize)
-						.MinDesiredHeight(ThumbnailSize)
-						[
-							ThumbnailWidget
-						]
-					]
-
-				+ SHorizontalBox::Slot()
-				.VAlign(VAlign_Center)
-				.AutoWidth()
-				.Padding(8.0f, 0.0f)
-				[
-					SNew(SVerticalBox)
-					+ SVerticalBox::Slot()
-					.AutoHeight()
-					[
-						SNew(STextBlock)
-						.Text(MaterialName)
-						.TextStyle(FAppStyle::Get(), "LargeText")
-					]
-
-					+ SVerticalBox::Slot()
-					.Padding(0, 4)
-					.AutoHeight()
-					[
-						SNew(STextBlock)
-						.Text(LOCTEXT("MaterialInstance", "Material Instance"))
-						.TextStyle(FAppStyle::Get(), "NormalText.Important")
-					]
-				];
-		}
-		else
-		{
-			HeaderBox = SNew(SHorizontalBox)
-				+ SHorizontalBox::Slot()
-				.Padding(FMargin(4.0f, 0.0f))
-				.HAlign(HAlign_Left)
-				.AutoWidth()
-				.VAlign(VAlign_Center)
-				[
-					SNew(STextBlock)
-					.Text(LOCTEXT("MaterialLayers", "Material Layers"))
-				];
-		}
+			.Padding(FMargin(4.0f, 0.0f))
+			.HAlign(HAlign_Left)
+			.AutoWidth()
+			.VAlign(VAlign_Center)
+			[
+				SNew(STextBlock)
+				.Text(LOCTEXT("MaterialLayers", "Material Layers"))
+			];
 
 		this->ChildSlot
 			[
@@ -1538,120 +1483,58 @@ void SMaterialLayersFunctionsInstanceWrapper::Refresh()
 				]
 		];
 		if (FMaterialPropertyHelpers::IsOverriddenExpression(NestedTree->GetFunctionParameter()))
-		{
-			if (Substrate::IsMaterialLayeringSupportEnabled())
-			{
-				FMenuBuilder DetailViewOptions( true, nullptr);
-				DetailViewOptions.AddMenuEntry(
-				LOCTEXT("AddMaterialSlab", "Add Material Slab"),
-				LOCTEXT("AddMaterialSlab_ToolTip", "Add new Material Slab to the Substrate tree"),
-				FSlateIcon(),
-				FUIAction(FExecuteAction::CreateSPLambda(this, [this]()
-					{
-						NestedTree->AddRootNodeLayer();
-					})));
-
-				if (NestedTree->GetRelinkLayersToParentVisibility() == EVisibility::Visible)
-				{
-					DetailViewOptions.AddMenuEntry(
-					LOCTEXT("Relink", "Relink"),
-					LOCTEXT("RelinkToParentLayers", "Relink to Parent Layers and Blends"),
-					FSlateIcon(),
-					FUIAction(FExecuteAction::CreateSPLambda(this, [this]()
-					{
-						NestedTree->RelinkLayersToParent();
-					})));
-				}
-				DetailViewOptions.AddMenuEntry(
-				LOCTEXT("SaveSibling", "Save Sibling"),
-				LOCTEXT("SaveToSiblingInstance", "Save to Sibling Instance"),
-				FSlateIcon(),
-				FUIAction(FExecuteAction::CreateLambda([this, OnSiblingButtonClicked]()	{ OnSiblingButtonClicked.Execute(); })));
-
-				DetailViewOptions.AddMenuEntry(
-				LOCTEXT("SaveChild", "Save Child"),
-				LOCTEXT("SaveToChildInstance", "Save To Child Instance"),
-				FSlateIcon(),
-				FUIAction(FExecuteAction::CreateLambda([this, OnChildButtonClicked]()	{ OnChildButtonClicked.Execute(); })));
-
-				HeaderBox->AddSlot()
-				.FillWidth(1.0f)
-				[
-					SNullWidget::NullWidget
-				];
-				HeaderBox->AddSlot()
-					.HAlign(HAlign_Right)
-					.AutoWidth()
-					.Padding(2.0f)
-					.VAlign(VAlign_Center)
-					[
-						SNew(SComboButton)
-							.ContentPadding(0.f)
-							.ComboButtonStyle(&FAppStyle::Get().GetWidgetStyle<FComboButtonStyle>("StatusBar.StatusBarEllipsisComboButton"))
-							.HasDownArrow(true)
-							.MenuContent()
-							[
-								DetailViewOptions.MakeWidget()
-							]
-					];
-			}
-			else
-			{
-				HeaderBox->AddSlot()
-					.HAlign(HAlign_Left)
-					.AutoWidth()
-					.VAlign(VAlign_Center)
-					[
-					PropertyCustomizationHelpers::MakeAddButton(FSimpleDelegate::CreateSPLambda(this, [this]()
-						{
-							NestedTree->AddRootNodeLayer();
-						}))
-					];
-			}
-		}
-
-		if (!Substrate::IsMaterialLayeringSupportEnabled())
 		{
 			HeaderBox->AddSlot()
-				.FillWidth(1.0f)
-				[
-					SNullWidget::NullWidget
-				];
-			HeaderBox->AddSlot()
+				.HAlign(HAlign_Left)
 				.AutoWidth()
-				.Padding(2.0f)
+				.VAlign(VAlign_Center)
 				[
-					SNew(SButton)
-					.Text(LOCTEXT("Relink", "Relink"))
-					.HAlign(HAlign_Center)
-					.OnClicked(OnRelinkToParent)
-					.ToolTipText(LOCTEXT("RelinkToParentLayers", "Relink to Parent Layers and Blends"))
-					.Visibility_Lambda([this]()
+				PropertyCustomizationHelpers::MakeAddButton(FSimpleDelegate::CreateSPLambda(this, [this]()
 					{
-						return NestedTree->GetRelinkLayersToParentVisibility();
-					})
+						NestedTree->AddRootNodeLayer();
+					}))
 				];
-			HeaderBox->AddSlot()
-					.AutoWidth()
-					.Padding(2.0f)
-					[
-						SNew(SButton)
-						.Text(LOCTEXT("SaveSibling", "Save Sibling"))
-						.HAlign(HAlign_Center)
-						.OnClicked(OnSiblingButtonClicked)
-						.ToolTipText(LOCTEXT("SaveToSiblingInstance", "Save to Sibling Instance"))
-					];
-			HeaderBox->AddSlot()
+		}
+
+		HeaderBox->AddSlot()
+			.FillWidth(1.0f)
+			[
+				SNullWidget::NullWidget
+			];
+		HeaderBox->AddSlot()
+			.AutoWidth()
+			.Padding(2.0f)
+			[
+				SNew(SButton)
+				.Text(LOCTEXT("Relink", "Relink"))
+				.HAlign(HAlign_Center)
+				.OnClicked(OnRelinkToParent)
+				.ToolTipText(LOCTEXT("RelinkToParentLayers", "Relink to Parent Layers and Blends"))
+				.Visibility_Lambda([this]()
+				{
+					return NestedTree->GetRelinkLayersToParentVisibility();
+				})
+			];
+		HeaderBox->AddSlot()
 				.AutoWidth()
 				.Padding(2.0f)
 				[
 					SNew(SButton)
-					.Text(LOCTEXT("SaveChild", "Save Child"))
+					.Text(LOCTEXT("SaveSibling", "Save Sibling"))
 					.HAlign(HAlign_Center)
-					.OnClicked(OnChildButtonClicked)
-					.ToolTipText(LOCTEXT("SaveToChildInstance", "Save To Child Instance"))
+					.OnClicked(OnSiblingButtonClicked)
+					.ToolTipText(LOCTEXT("SaveToSiblingInstance", "Save to Sibling Instance"))
 				];
-		}
+		HeaderBox->AddSlot()
+			.AutoWidth()
+			.Padding(2.0f)
+			[
+				SNew(SButton)
+				.Text(LOCTEXT("SaveChild", "Save Child"))
+				.HAlign(HAlign_Center)
+				.OnClicked(OnChildButtonClicked)
+				.ToolTipText(LOCTEXT("SaveToChildInstance", "Save To Child Instance"))
+			];
 	}
 	else
 	{
@@ -1671,29 +1554,16 @@ void SMaterialLayersFunctionsInstanceWrapper::Refresh()
 
 void SMaterialLayersFunctionsInstanceWrapper::Construct(const FArguments& InArgs)
 {
-	if (Substrate::IsMaterialLayeringSupportEnabled())
-	{
-		NestedTree = SNew(SMaterialSubstrateTree)
-			.InMaterialEditorInstance(InArgs._InMaterialEditorInstance)
-			.InGenerator(InArgs._InGenerator)
-			.InWrapper(this)
-			.InShowHiddenDelegate(InArgs._InShowHiddenDelegate);
-	}
-	else
-	{
-		NestedTree = SNew(SMaterialLayersFunctionsInstanceTree)
-			.InMaterialEditorInstance(Cast<UMaterialEditorInstanceConstant>(InArgs._InMaterialEditorInstance))
-			.InWrapper(this)
-			.InShowHiddenDelegate(InArgs._InShowHiddenDelegate)
-			.InNotifyHook(InArgs._InNotifyHook);
-	}
-	
+	NestedTree = SNew(SMaterialLayersFunctionsInstanceTree)
+		.InMaterialEditorInstance(Cast<UMaterialEditorInstanceConstant>(InArgs._InMaterialEditorInstance))
+		.InWrapper(this)
+		.InShowHiddenDelegate(InArgs._InShowHiddenDelegate)
+		.InNotifyHook(InArgs._InNotifyHook);
 
 	LayerParameter = NestedTree->GetFunctionParameter();
 
 	MaterialEditorInstance = InArgs._InMaterialEditorInstance;
 	Refresh();
-
 }
 
 void SMaterialLayersFunctionsInstanceWrapper::SetEditorInstance(UMaterialEditorParameters* InMaterialEditorInstance)
diff --git a/Engine/Source/Editor/MaterialEditor/Private/SMaterialSubstrateTree.cpp b/Engine/Source/Editor/MaterialEditor/Private/SMaterialSubstrateTree.cpp
deleted file mode 100644
diff --git a/Engine/Source/Editor/MaterialEditor/Private/SMaterialSubstrateTree.h b/Engine/Source/Editor/MaterialEditor/Private/SMaterialSubstrateTree.h
deleted file mode 100644
diff --git a/Engine/Source/Editor/MaterialEditor/Public/MaterialPropertyHelpers.h b/Engine/Source/Editor/MaterialEditor/Public/MaterialPropertyHelpers.h
--- a/Engine/Source/Editor/MaterialEditor/Public/MaterialPropertyHelpers.h
+++ b/Engine/Source/Editor/MaterialEditor/Public/MaterialPropertyHelpers.h
@@ -119,11 +119,6 @@ public:
 
 	FLayerDragDropOp(TSharedPtr<IDraggableItem> InOwningStack)
 	{
-		if(Substrate::IsMaterialLayeringSupportEnabled())
-		{
-			MouseCursor = EMouseCursor::GrabHandClosed;
-		}
-
 		OwningStack = InOwningStack;
 		DecoratorWidget = SNew(SBorder)
 			.BorderImage(FAppStyle::GetBrush("Graph.ConnectorFeedback.Border"))
diff --git a/Engine/Source/Runtime/Engine/Private/Materials/Material.cpp b/Engine/Source/Runtime/Engine/Private/Materials/Material.cpp
--- a/Engine/Source/Runtime/Engine/Private/Materials/Material.cpp
+++ b/Engine/Source/Runtime/Engine/Private/Materials/Material.cpp
@@ -3830,61 +3830,6 @@ bool UMaterial::ConvertMaterialToSubstrateMaterial(bool bAllowEmptyMaterialUpdat
 		bool bRelinkCustomOutputNodes = false;
 		UMaterialExpressionSubstrateShadingModels* ConvertNode = nullptr;
 		// Connect all the legacy pin into the conversion node
-		//First check if this material contains a layer node connection, as that requires different upgrade behavior to convert to the Substrate material layer system
-		bool bIsSubstrateLayersMaterial = false;
-		if (Substrate::IsMaterialLayeringSupportEnabled() && bUseMaterialAttributes && EditorOnly->MaterialAttributes.Expression && !EditorOnly->FrontMaterial.IsConnected())
-		{
-			UMaterialExpression* FinalNode = nullptr;
-			UMaterialExpressionMaterialAttributeLayers* LayersNode = nullptr;
-			for(UMaterialExpression* Expression : GetExpressions())
-			{
-				if(Expression->IsA<UMaterialExpressionMaterialAttributeLayers>())
-				{
-					LayersNode = Cast<UMaterialExpressionMaterialAttributeLayers>(Expression);
-					FinalNode = EditorOnly->MaterialAttributes.Expression;
-					bIsSubstrateLayersMaterial = true;
-					break;
-				}
-			}
-
-			if(FinalNode && LayersNode)
-			{
-				//Separate the layers/blends MAs to Substrate/Non-Substrate respectively and connect to their relevant material output
-				UMaterialExpressionGetMaterialAttributes* GetAttributesNode = NewObject<UMaterialExpressionGetMaterialAttributes>(this);
-				GetAttributesNode->Material = this;
-				GetAttributesNode->MaterialAttributes.Connect(0, FinalNode);
-				EditorOnly->FrontMaterial.Connect(GetAttributesNode->CreateOrGetOutputAttribute(MP_FrontMaterial), GetAttributesNode);
-				EditorOnly->MaterialAttributes.Connect(GetAttributesNode->CreateOrGetOutputAttribute(MP_MaterialAttributes), GetAttributesNode);
-				SetPosXAndMoveReferenceToTheRight(GetAttributesNode);
-
-				//If a non-Substrate default input is connected to the MLA node, convert it.
-				if (LayersNode->Input.IsConnected() && !LayersNode->Input.Expression->IsResultSubstrateMaterial(LayersNode->Input.OutputIndex))
-				{
-					UMaterialExpressionSubstrateConvertMaterialAttributes* ConvertAttributeNode = NewObject<UMaterialExpressionSubstrateConvertMaterialAttributes>(this);
-					ConvertAttributeNode->Material = this;
-					ConvertAttributeNode->MaterialAttributes.Connect(LayersNode->Input.OutputIndex, LayersNode->Input.Expression);
-					ConvertAttributeNode->ShadingModelOverride = MSM_DefaultLit;
-
-					UMaterialExpressionSetMaterialAttributes* SetAttributesNode = NewObject<UMaterialExpressionSetMaterialAttributes>(this);
-					SetAttributesNode->Material = this;
-					SetAttributesNode->ConnectInputAttribute(MP_FrontMaterial, ConvertAttributeNode);
-					SetAttributesNode->ConnectInputAttribute(MP_MaterialAttributes, LayersNode->Input.Expression, LayersNode->Input.OutputIndex);
-
-					LayersNode->Input.Connect(0, SetAttributesNode);
-
-					SetPosXAndMoveReferenceToTheRight(GetAttributesNode);
-					ReplaceNodeAndMoveToTheRight(LayersNode, ConvertAttributeNode);
-					SetPosXAndMoveReferenceToTheRight(GetAttributesNode);
-					ReplaceNodeAndMoveToTheRight(LayersNode, SetAttributesNode);			
-				}
-
-				BlendMode = ConvertLegacyBlendMode(BlendMode, ShadingModels);
-				RefractionCoverageMode = RCM_CoverageIgnored;
-				bInvalidateShader = true;
-			}
-		}
-	
-		if(!bIsSubstrateLayersMaterial)
 		{
 			const bool bHasAnySubstrateNodes = HasAnyExpressionsInMaterialAndFunctionsOfType<UMaterialExpressionSubstrateBSDF>();
 			if (bUseMaterialAttributes && EditorOnly->MaterialAttributes.Expression && !EditorOnly->FrontMaterial.IsConnected() && !EditorOnly->MaterialAttributes.Expression->IsResultSubstrateMaterial(EditorOnly->MaterialAttributes.OutputIndex) && !bHasAnySubstrateNodes) // M_Rifle cause issues there
diff --git a/Engine/Source/Runtime/Engine/Private/Materials/MaterialExpressions.cpp b/Engine/Source/Runtime/Engine/Private/Materials/MaterialExpressions.cpp
--- a/Engine/Source/Runtime/Engine/Private/Materials/MaterialExpressions.cpp
+++ b/Engine/Source/Runtime/Engine/Private/Materials/MaterialExpressions.cpp
@@ -7036,192 +7036,11 @@ FSubstrateOperator* UMaterialExpressionBlendMaterialAttributes::SubstrateGenerat
 UMaterialExpressionLegacyBlendMaterialAttributes::UMaterialExpressionLegacyBlendMaterialAttributes(const FObjectInitializer& ObjectInitializer)
 	: Super(ObjectInitializer)
 {
-	if(!Substrate::IsMaterialLayeringSupportEnabled())
-	{
-		Outputs.Reset();
-		bShowInputs = false;
-		bShowOutputs = false;
-		bCollapsed = true;
-	}
-}
-
-#if WITH_EDITOR
-void UMaterialExpressionLegacyBlendMaterialAttributes::GetCaption(TArray<FString>& OutCaptions) const
-{
-	if(Substrate::IsMaterialLayeringSupportEnabled())
-	{
-		OutCaptions.Add(TEXT("LegacyBlendMaterialAttributes"));
-	}
-	else
-	{
-		OutCaptions.Add(TEXT("InvalidLegacyBlend_ForSubstrateLayeringOnly"));
-	}
-}
-bool UMaterialExpressionLegacyBlendMaterialAttributes::CanEditChange(const FProperty* InProperty) const
-{
-	bool bIsEditable = Super::CanEditChange(InProperty);
-	if(InProperty != nullptr)
-	{
-		FName PropertyFName = InProperty->GetFName();
-		if (PropertyFName == GET_MEMBER_NAME_CHECKED(UMaterialExpressionBlendMaterialAttributes, VertexAttributeBlendType))
-		{
-			bIsEditable = !(VertexAttribute_UseA.IsConnected() || VertexAttribute_UseB.IsConnected());
-		}
-		else if (PropertyFName == GET_MEMBER_NAME_CHECKED(UMaterialExpressionBlendMaterialAttributes, PixelAttributeBlendType))
-		{
-			bIsEditable = !(PixelAttribute_UseA.IsConnected() || PixelAttribute_UseB.IsConnected());
-		}
-	}
-	return bIsEditable;
-}
-FExpressionInput* UMaterialExpressionLegacyBlendMaterialAttributes::GetInput(int32 InputIndex)
-{
-	switch(InputIndex)
-	{
-		case 3: return &VertexAttribute_UseA;
-		case 4: return &VertexAttribute_UseB;
-		case 5: return &PixelAttribute_UseA;
-		case 6: return &PixelAttribute_UseB;
-		default: return Super::GetInput(InputIndex);
-	}
-}
-FName UMaterialExpressionLegacyBlendMaterialAttributes::GetInputName(int32 InputIndex) const
-{
-	FName Name;
-
-	switch (InputIndex)
-	{
-		case 3: Name = TEXT("Vertex Attributes Use A"); break;
-		case 4: Name = TEXT("Vertex Attributes Use B"); break;
-		case 5: Name = TEXT("Pixel Attributes Use A"); break;
-		case 6: Name = TEXT("Pixel Attributes Use B"); break;
-		default: Name = Super::GetInputName(InputIndex);
-	}
-
-	return Name;
-}
-bool UMaterialExpressionLegacyBlendMaterialAttributes::IsInputConnectionRequired(int32 InputIndex) const
-{
-	switch (InputIndex)
-	{
-		case 3:
-		case 4:
-		case 5:
-		case 6: return false;
-		default: return Super::IsInputConnectionRequired(InputIndex);
-	}
-}
-
-EMaterialValueType UMaterialExpressionLegacyBlendMaterialAttributes::GetInputValueType(int32 InputIndex)
-{
-	switch (InputIndex)
-	{
-		case 3:
-		case 4:
-		case 5:
-		case 6: return MCT_StaticBool;
-		default: return Super::GetInputValueType(InputIndex);
-	}
-}
-
-bool GetStaticBoolExpressionOutput(FMaterialCompiler* Compiler, FExpressionInput& Input)
-{
-	if(Input.GetTracedInput().Expression)
-	{
-		bool bSucceeded = false;
-		bool ResultBool = Compiler->GetStaticBoolValue(Input.Compile(Compiler), bSucceeded);
-		if(bSucceeded)
-		{
-			return ResultBool;
-		}
-	}
-	return false;
-}
-
-int32 UMaterialExpressionLegacyBlendMaterialAttributes::Compile(FMaterialCompiler* Compiler, int32 OutputIndex)
-{
-	if (!Substrate::IsMaterialLayeringSupportEnabled())
-	{
-		return Compiler->Errorf(TEXT("LegacyBlendMaterialAttributes is only usable with Substrate material layering enabled."));
-	}
-
-	VertexAttributeBlendType = EMaterialAttributeBlend::Blend;
-	PixelAttributeBlendType = EMaterialAttributeBlend::Blend;	
-	
-	bool bVertexUseA = GetStaticBoolExpressionOutput(Compiler, VertexAttribute_UseA);
-	bool bVertexUseB = GetStaticBoolExpressionOutput(Compiler, VertexAttribute_UseB);
-	if (bVertexUseA && !bVertexUseB)
-	{
-		VertexAttributeBlendType = EMaterialAttributeBlend::UseA;
-	}
-	if (!bVertexUseA && bVertexUseB)
-	{
-		VertexAttributeBlendType = EMaterialAttributeBlend::UseB;
-	}
-
-	bool bPixelUseA = GetStaticBoolExpressionOutput(Compiler, PixelAttribute_UseA);
-	bool bPixelUseB = GetStaticBoolExpressionOutput(Compiler, PixelAttribute_UseB);
-	if (bPixelUseA && !bPixelUseB)
-	{
-		PixelAttributeBlendType = EMaterialAttributeBlend::UseA;
-	}
-	if (!bPixelUseA && bPixelUseB)
-	{
-		PixelAttributeBlendType = EMaterialAttributeBlend::UseB;
-	}
-
-	const FGuid AttributeID = Compiler->GetMaterialAttribute();
-
-	// Blending is optional, can skip on a per-node basis
-	EMaterialAttributeBlend::Type BlendType;
-	EShaderFrequency AttributeFrequency = FMaterialAttributeDefinitionMap::GetShaderFrequency(AttributeID);
-
-	switch (AttributeFrequency)
-	{
-	case SF_Vertex:	BlendType = VertexAttributeBlendType;	break;
-	case SF_Pixel:	BlendType = PixelAttributeBlendType;	break;
-	default:
-		return Compiler->Errorf(TEXT("Attribute blending for shader frequency %i not implemented."), AttributeFrequency);
-	}
-
-	switch (BlendType)
-	{
-	case EMaterialAttributeBlend::UseA:
-	{
-		Compiler->SubstrateTreeStackPush(this, 0);
-		int32 CodeChunk = A.CompileWithDefault(Compiler, AttributeID);
-		Compiler->SubstrateTreeStackPop();
-		return CodeChunk;
-	}
-	case EMaterialAttributeBlend::UseB:
-	{
-		Compiler->SubstrateTreeStackPush(this, 1);
-		int32 CodeChunk = B.CompileWithDefault(Compiler, AttributeID);
-		Compiler->SubstrateTreeStackPop();
-		return CodeChunk;
-	}
-	default:
-		check(BlendType == EMaterialAttributeBlend::Blend);
-	}
-
-	// Allow custom blends or fallback to standard interpolation
-	Compiler->SubstrateTreeStackPush(this, 0);
-	int32 ResultA = A.CompileWithDefault(Compiler, AttributeID);
-	Compiler->SubstrateTreeStackPop();
-	Compiler->SubstrateTreeStackPush(this, 1);
-	int32 ResultB = B.CompileWithDefault(Compiler, AttributeID);
-	Compiler->SubstrateTreeStackPop();
-	int32 ResultAlpha = Alpha.Compile(Compiler);
-
-	MaterialAttributeBlendFunction BlendFunction = FMaterialAttributeDefinitionMap::GetBlendFunction(AttributeID);
-	if (BlendFunctionType == EMaterialAttributeBlendFunction::Type::Horizontal && BlendFunction)
-	{
-		return BlendFunction(Compiler, ResultA, ResultB, ResultAlpha);
-	}	
-
-	return Compiler->Lerp(ResultA, ResultB, ResultAlpha);
+	Outputs.Reset();
+	bShowInputs = false;
+	bShowOutputs = false;
+	bCollapsed = true;
 }
-#endif // WITH_EDITOR
 
 //
 //	UMaterialExpressionMaterialAttributeLayers
@@ -7273,192 +7092,6 @@ void UMaterialExpressionMaterialAttributeLayers::PostEditChangeProperty(FPropert
 #if WITH_EDITOR
 void UMaterialExpressionMaterialAttributeLayers::RebuildLayerGraph(bool bReportErrors)
 {
-	if(Substrate::IsMaterialLayeringSupportEnabled())
-	{
-		const FMaterialLayersFunctions* LayersTree = (ParamLayers ? ParamLayers : &DefaultLayers);
-		const TArray<UMaterialFunctionInterface*>& Layers = GetLayers();
-		const TArray<UMaterialFunctionInterface*>& Blends = GetBlends();
-		const TArray<bool>& LayerStates = GetLayerStates();
-
-		if (!LayersTree->RuntimeGraphCache)
-		{
-			LayersTree->RuntimeGraphCache = MakeShareable(new FMaterialLayersFunctionsRuntimeGraphCache());
-		}
-		TSharedPtr< FMaterialLayersFunctionsRuntimeGraphCache> LayersTreeCache = LayersTree->RuntimeGraphCache;
-
-		// Pre-populate callers, we maintain these transient objects to avoid
-		// heavy UObject recreation as the graphs are frequently rebuilt
-		while (LayersTreeCache->LayerCallers.Num() < Layers.Num())
-		{
-			LayersTreeCache->LayerCallers.Add(NewObject<UMaterialExpressionMaterialFunctionCall>(GetTransientPackage()));
-		}
-		while (LayersTreeCache->BlendCallers.Num() < Blends.Num())
-		{
-			LayersTreeCache->BlendCallers.Add(NewObject<UMaterialExpressionMaterialFunctionCall>(GetTransientPackage()));
-		}
-
-		// Reset graph connectivity
-		bIsLayerGraphBuilt = false;
-		NumActiveLayerCallers = 0;
-		NumActiveBlendCallers = 0;
-
-		if (ValidateLayerConfiguration(nullptr, bReportErrors))
-		{
-			// Initialize layer function callers
-			for (int32 LayerIndex = 0; LayerIndex < Layers.Num(); ++LayerIndex)
-			{
-				if (Layers[LayerIndex] && LayerStates[LayerIndex])
-				{
-					LayersTreeCache->LayerCallers[LayerIndex]->MaterialFunction = Layers[LayerIndex];
-					LayersTreeCache->LayerCallers[LayerIndex]->FunctionParameterInfo.Association = EMaterialParameterAssociation::LayerParameter;
-					LayersTreeCache->LayerCallers[LayerIndex]->FunctionParameterInfo.Index = LayerIndex;
-
-					Layers[LayerIndex]->GetInputsAndOutputs(LayersTreeCache->LayerCallers[LayerIndex]->FunctionInputs, LayersTreeCache->LayerCallers[LayerIndex]->FunctionOutputs);
-				
-					for (FFunctionExpressionOutput& FunctionOutput : LayersTreeCache->LayerCallers[LayerIndex]->FunctionOutputs)
-					{
-						LayersTreeCache->LayerCallers[LayerIndex]->Outputs.Add(FunctionOutput.Output);
-					}
-
-					// Optional: Single material attributes input, the base input to the stack
-					if (LayersTreeCache->LayerCallers[LayerIndex]->FunctionInputs.Num() > 0)
-					{
-						if (Input.GetTracedInput().Expression)
-						{
-							LayersTreeCache->LayerCallers[LayerIndex]->FunctionInputs[0].Input = Input;
-						}
-					}
-
-					// Recursively run through internal functions to allow connection of inputs/outputs
-					LayersTreeCache->LayerCallers[LayerIndex]->UpdateFromFunctionResource();
-
-					++NumActiveLayerCallers;
-				}
-				else
-				{
-					// Empty entries for opaque layers
-					LayersTreeCache->LayerCallers[LayerIndex]->MaterialFunction = nullptr;
-				}
-			}
-
-			for (int32 BlendIndex = 0; BlendIndex < Blends.Num(); ++BlendIndex)
-			{
-				if (Blends[BlendIndex])
-				{
-					LayersTreeCache->BlendCallers[BlendIndex]->MaterialFunction = Blends[BlendIndex];
-					LayersTreeCache->BlendCallers[BlendIndex]->FunctionParameterInfo.Association = EMaterialParameterAssociation::BlendParameter;
-					LayersTreeCache->BlendCallers[BlendIndex]->FunctionParameterInfo.Index = BlendIndex;
-
-					Blends[BlendIndex]->GetInputsAndOutputs(LayersTreeCache->BlendCallers[BlendIndex]->FunctionInputs, LayersTreeCache->BlendCallers[BlendIndex]->FunctionOutputs);
-					for (FFunctionExpressionOutput& FunctionOutput : LayersTreeCache->BlendCallers[BlendIndex]->FunctionOutputs)
-					{
-						LayersTreeCache->BlendCallers[BlendIndex]->Outputs.Add(FunctionOutput.Output);
-					}
-
-					// Recursively run through internal functions to allow connection of inputs/ouputs
-					LayersTreeCache->BlendCallers[BlendIndex]->UpdateFromFunctionResource();
-
-					++NumActiveBlendCallers;
-				}
-				else
-				{
-					// Empty entries for opaque layers
-					LayersTreeCache->BlendCallers[BlendIndex]->MaterialFunction = nullptr;
-				}
-			}
-
-			// Empty out unused callers
-			for (int32 CallerIndex = Layers.Num(); CallerIndex < LayersTreeCache->LayerCallers.Num(); ++CallerIndex)
-			{
-				LayersTreeCache->LayerCallers[CallerIndex]->MaterialFunction = nullptr;
-			}
-
-			for (int32 CallerIndex = Blends.Num(); CallerIndex < LayersTreeCache->BlendCallers.Num(); ++CallerIndex)
-			{
-				LayersTreeCache->BlendCallers[CallerIndex]->MaterialFunction = nullptr;
-			}
-
-
-			LayersTreeCache->NodeMaterialGraphExpressions.Empty(LayersTree->Tree.Nodes.Num());
-			LayersTreeCache->NodeMaterialGraphExpressions.InsertZeroed(0, LayersTree->Tree.Nodes.Num());
-			LayersTreeCache->NodeMaterialGraphExpressionsBlends.Empty(LayersTree->Tree.Nodes.Num());
-			LayersTreeCache->NodeMaterialGraphExpressionsBlends.InsertZeroed(0, LayersTree->Tree.Nodes.Num());
-			LayersTreeCache->NodePreviewMaterials.Empty(LayersTree->Tree.Nodes.Num());
-			LayersTreeCache->NodePreviewMaterials.InsertZeroed(0, LayersTree->Tree.Nodes.Num());
-			struct TreeGraphBuilder
-			{
-				UMaterialExpressionMaterialAttributeLayers* This = nullptr;
-				const FMaterialLayersFunctions* Tree = nullptr;
-				TSharedPtr< FMaterialLayersFunctionsRuntimeGraphCache> TreeCache;
-
-				UMaterialExpressionMaterialFunctionCall* AssembleBranchGraph(int32 InSourceId)
-				{
-					// Fetch the Layer function of this node as the base background layer
-					// Always null for the root node.	
-					UMaterialExpressionMaterialFunctionCall* PreviousLayerInput = nullptr;
-					if (InSourceId != -1)
-					{
-						int32 BackgroundLayerFuncIdx = Tree->GetLayerFuncIndex(InSourceId);
-						PreviousLayerInput = (TreeCache->LayerCallers.IsValidIndex(BackgroundLayerFuncIdx) ? TreeCache->LayerCallers[BackgroundLayerFuncIdx] : nullptr);
-						PreviousLayerInput = (PreviousLayerInput && PreviousLayerInput->MaterialFunction ? PreviousLayerInput : nullptr);
-						if (PreviousLayerInput)
-							TreeCache->ExpressionCollection.AddExpression(PreviousLayerInput);
-					}
-
-					// For child of this node, try to assemble the blend/layer tree 
-					auto NodeChildrenIds = Tree->GetNodeChildren(InSourceId);
-					for (auto NodeId : NodeChildrenIds)
-					{
-						// Assemble the branch of layer / blend under this sub node by calling recursively  into that sub node
-						UMaterialExpressionMaterialFunctionCall* CurrentLayerInput = AssembleBranchGraph(NodeId);
-
-
-						// Fetch the blend function of the node
-						int32 BlendFuncIdx = Tree->GetBlendFuncIndex(NodeId);
-						UMaterialExpressionMaterialFunctionCall* CurrentBlendCaller = (TreeCache->BlendCallers.IsValidIndex(BlendFuncIdx) ? TreeCache->BlendCallers[BlendFuncIdx] : nullptr);
-						CurrentBlendCaller = (CurrentBlendCaller && CurrentBlendCaller->MaterialFunction ? CurrentBlendCaller : nullptr);
-
-						TreeCache->NodeMaterialGraphExpressionsBlends[NodeId] = CurrentBlendCaller;
-					
-						if (CurrentBlendCaller) // Blend is the new current layer if valid and connect the previous layer and current layer is valid
-						{
-							if (PreviousLayerInput)
-								TreeCache->ExpressionCollection.AddExpression(CurrentBlendCaller);
-
-							if (PreviousLayerInput)
-								CurrentBlendCaller->FunctionInputs[0].Input.Connect(0, PreviousLayerInput);
-							if (CurrentLayerInput)
-								CurrentBlendCaller->FunctionInputs[1].Input.Connect(0, CurrentLayerInput);
-
-							PreviousLayerInput = CurrentBlendCaller;
-						}
-						else if (CurrentLayerInput) // No Blend then the new current layer if valid
-						{
-							PreviousLayerInput = ToRawPtr(CurrentLayerInput);
-						}
-						//else this node's layer and and blend functions aren't valid, ignore and move on to the next sub node
-					}
-
-					if (InSourceId != -1)
-					{
-						TreeCache->NodeMaterialGraphExpressions[InSourceId] = (PreviousLayerInput);
-					}
-
-					return PreviousLayerInput;
-				}
-			} BuildTreeGraph = { this, LayersTree, LayersTreeCache };
-
-			OutputCaller = BuildTreeGraph.AssembleBranchGraph(-1);
-
-			if (OutputCaller && (OutputCaller->MaterialFunction == nullptr))
-			{	
-				OutputCaller = nullptr;
-			}
-
-			bIsLayerGraphBuilt = true;
-		}
-	}
-	else
 	{
 		const TArray<UMaterialFunctionInterface*>& Layers = GetLayers();
 		const TArray<UMaterialFunctionInterface*>& Blends = GetBlends();
@@ -7654,7 +7287,7 @@ bool UMaterialExpressionMaterialAttributeLayers::ValidateLayerConfiguration(FMat
 					bIsValid = false;
 				}
 			}
-			else if (!Substrate::IsMaterialLayeringSupportEnabled())
+			else
 			{
 				TArray<UMaterialExpressionFunctionInput*> InputExpressions;
 				Layer->GetAllExpressionsOfType<UMaterialExpressionFunctionInput>(InputExpressions, false);
@@ -7696,7 +7329,7 @@ bool UMaterialExpressionMaterialAttributeLayers::ValidateLayerConfiguration(FMat
 				TArray<UMaterialExpressionFunctionInput*> InputExpressions;
 				Blend->GetAllExpressionsOfType<UMaterialExpressionFunctionInput>(InputExpressions, false);
 
-				if (Substrate::IsMaterialLayeringSupportEnabled() ? InputExpressions.Num() < AcceptableNumBlendMAInputs : InputExpressions.Num() != AcceptableNumBlendMAInputs)
+				if (InputExpressions.Num() != AcceptableNumBlendMAInputs)
 				{
 					COMPILER_OR_LOG_ERROR(TEXT("Blend %i, %s, must have two MaterialAttributes inputs."), BlendIndex, *Blend->GetName());
 					bIsValid = false;
@@ -7705,24 +7338,16 @@ bool UMaterialExpressionMaterialAttributeLayers::ValidateLayerConfiguration(FMat
 		}
 		
 		// Null blends signify an opaque layer so count as valid for the sake of graph validation
-		if (Substrate::IsMaterialLayeringSupportEnabled() ? Layers[BlendIndex] && LayerStates[BlendIndex] : Layers[BlendIndex+1] && LayerStates[BlendIndex+1])
+		if (Layers[BlendIndex+1] && LayerStates[BlendIndex+1])
 		{
 			++NumActiveBlends;
 		}
 	}
 
-	bool bValidGraphLayout = false;
-	if(Substrate::IsMaterialLayeringSupportEnabled())
-	{
-		bValidGraphLayout = (NumActiveLayers == 0 && NumActiveBlends == 0)		// Pass-through
-			|| (NumActiveLayers == NumActiveBlends);			    // Single layer
-	}
-	else
-	{
-		bValidGraphLayout = (NumActiveLayers == 0 && NumActiveBlends == 0)		// Pass-through
-			|| (NumActiveLayers == 1 && NumActiveBlends == 0)						// Single layer
-			|| (NumActiveLayers >= 2 && NumActiveBlends == NumActiveLayers - 1);	// Blend graph
-	}
+	bool bValidGraphLayout = 
+		(NumActiveLayers == 0 && NumActiveBlends == 0)		// Pass-through
+		|| (NumActiveLayers == 1 && NumActiveBlends == 0)						// Single layer
+		|| (NumActiveLayers >= 2 && NumActiveBlends == NumActiveLayers - 1);	// Blend graph
 
 	if (!bValidGraphLayout)
 	{
@@ -7827,11 +7452,7 @@ int32 UMaterialExpressionMaterialAttributeLayers::Compile(FMaterialCompiler* Com
 
 	if (ValidateLayerConfiguration(Compiler, true) && bIsLayerGraphBuilt)
 	{
-		if (Substrate::IsMaterialLayeringSupportEnabled() && OutputCaller)
-		{
-			Result = OutputCaller->Compile(Compiler, 0);
-		}
-		else if (NumActiveBlendCallers > 0 && BlendCallers[NumActiveBlendCallers-1]->MaterialFunction)
+		if (NumActiveBlendCallers > 0 && BlendCallers[NumActiveBlendCallers-1]->MaterialFunction)
 		{
 			// Multiple blended layers
 			Result = BlendCallers[NumActiveBlendCallers-1]->Compile(Compiler, 0);
@@ -7871,45 +7492,6 @@ int32 UMaterialExpressionMaterialAttributeLayers::Compile(FMaterialCompiler* Com
 	return Result;
 }
 
-bool UMaterialExpressionMaterialAttributeLayers::IsResultSubstrateMaterial(int32 OutputIndex)
-{
-	return Substrate::IsMaterialLayeringSupportEnabled();
-}
-
-FSubstrateOperator* UMaterialExpressionMaterialAttributeLayers::SubstrateGenerateMaterialTopologyTree(FMaterialCompiler* Compiler, UMaterialExpression* Parent, int32 OutputIndex)
-{
-	FSubstrateOperator* OutOperator = nullptr;
-
-	if(Substrate::IsMaterialLayeringSupportEnabled())
-	{
-		const FMaterialLayersFunctions* OverrideLayers = Compiler->GetMaterialLayers();
-		OverrideLayerGraph(OverrideLayers);
-
-		//Behaviour mirrors the behaviour of the existing compile function
-		if (ValidateLayerConfiguration(Compiler, true) && bIsLayerGraphBuilt)
-		{
-			if (OutputCaller)
-			{
-				OutOperator = OutputCaller->SubstrateGenerateMaterialTopologyTree(Compiler, Parent, 0);
-				Compiler->FeedbackMaterialLayersInstancedGraphFromCompilation(OverrideLayers);
-			}
-		
-			if (!OutOperator)
-			{
-				//If this is reached, compile a default operator to avoid a crash
-				FSubstrateOperator& DefaultOperator = Compiler->SubstrateCompilationRegisterOperator(SUBSTRATE_OPERATOR_BSDF, Compiler->SubstrateTreeStackGetPathUniqueId(), this->MaterialExpressionGuid, Parent, Compiler->SubstrateTreeStackGetParentPathUniqueId());
-				DefaultOperator.BSDFType = SUBSTRATE_BSDF_TYPE_SLAB;
-				DefaultOperator.ThicknessIndex = Compiler->SubstrateThicknessStackGetThicknessIndex();
-				OutOperator = &DefaultOperator;
-			}
-		}
-
-		OverrideLayerGraph(nullptr);
-	}
-
-	return OutOperator;
-}
-
 void UMaterialExpressionMaterialAttributeLayers::GetCaption(TArray<FString>& OutCaptions) const
 {
 	OutCaptions.Add(TEXT("Material Attribute Layers"));
@@ -9439,7 +9021,7 @@ bool UMaterialExpressionFeatureLevelSwitch::IsInputConnectionRequired(int32 Inpu
 
 bool UMaterialExpressionFeatureLevelSwitch::IsResultMaterialAttributes(int32 OutputIndex)
 {
-	check(Substrate::IsMaterialLayeringSupportEnabled() || OutputIndex == 0);
+	check(OutputIndex == 0);
 	for (FExpressionInputIterator It{ this }; It; ++It)
 	{
 		// If there is a loop anywhere in this expression's inputs then we can't risk checking them
@@ -13997,443 +13579,9 @@ void UMaterialFunction::GetAllCustomOutputExpressions(TArray<class UMaterialExpr
 	}
 }
 
-void UMaterialFunction::ConvertExpressionsBetweenLegacyAndSubstrate()
-{
-	if (!Substrate::IsMaterialLayeringSupportEnabled() || GetMaterialFunctionUsage() == EMaterialFunctionUsage::Default)
-	{
-		return;
-	}
-
-	UMaterialFunctionEditorOnlyData* EditorOnly = GetEditorOnlyData();
-	if (!EditorOnly)
-	{
-		UE_LOG(LogMaterial, Warning, TEXT("Material Layer/Blend Function %s editor only data missing."), *GetFullName());
-		return;
-	}
-
-	TArray<TObjectPtr<UMaterialExpression>>& Expressions = EditorOnly->ExpressionCollection.Expressions;
-	if (Expressions.Num() < 1)
-	{
-		UE_LOG(LogMaterial, Warning, TEXT("Material Layer/Blend Function %s contains no expressions."), *GetFullName());
-		return;
-	}
-
-	UMaterialExpressionFunctionOutput* OutputNode = Cast<UMaterialExpressionFunctionOutput>(Expressions[0]);
-	if (!OutputNode || !OutputNode->A.IsConnected() || !OutputNode->A.Expression->IsResultMaterialAttributes(OutputNode->A.OutputIndex))
-	{
-		UE_LOG(LogMaterial, Log, TEXT("Material Layer/Blend Function %s is not valid for Substrate upgrade path."), *GetFullName());
-		return;
-	}
-
-	//Perform a full sweep to check for embedded Substrate expressions, checking if conversion is necessary
-	TArray<UMaterialFunctionInterface*> AllFunctions;
-	GetDependentFunctions(AllFunctions);
-	AllFunctions.Add(this);	
-	for (UMaterialFunctionInterface* Function : AllFunctions)
-	{
-		for (UMaterialExpression* Expression : Function->GetExpressions())
-		{
-			if (!Expression || Expression->IsA<UMaterialExpressionSubstrateBSDF>()
-				|| Expression->IsA<UMaterialExpressionSubstrateUtilityBase>())
-			{
-				//If Substrate nodes are present, do not convert this function, assume previously converted or Substrate MF
-				return;
-			}
-
-			//Check if we have already applied FrontMaterial logic to the Set/GetMaterialAttributes nodes, and if so, skip further processing.
-			if(UMaterialExpressionSetMaterialAttributes* SetAttributes = Cast<UMaterialExpressionSetMaterialAttributes>(Expression))
-			{
-				TArrayView<FExpressionInput*> InputsArray = SetAttributes->GetInputsView();
-				//0 index is always MaterialAttributes so no need to check that entry
-				for(int32 InputIndex = InputsArray.Num() - 1; InputIndex > 0; InputIndex--)
-				{
-					if(SetAttributes->GetInputValueType(InputIndex) == MCT_Substrate)
-					{
-						return;
-					}
-				}
-			}
-
-			if (UMaterialExpressionGetMaterialAttributes* GetAttributes = Cast<UMaterialExpressionGetMaterialAttributes>(Expression))
-			{
-				//0 index is always MaterialAttributes so no need to check that entry
-				for(int32 OutputIndex = GetAttributes->GetOutputs().Num() - 1; OutputIndex > 0; OutputIndex--)
-				{
-					if (GetAttributes->GetOutputValueType(OutputIndex) == MCT_Substrate)
-					{
-						return;
-					}
-				}
-			}
-		}
-	}
-
-	auto MoveNodeInHorizonalAxis = [](UMaterialExpression* NodeToMove, int32 HorizontalOffset = 0)
-	{
-		NodeToMove->MaterialExpressionEditorX += HorizontalOffset;
-	};
-
-	auto PlaceNodeInLocation = [](UMaterialExpression* ExistingNode, UMaterialExpression* NewNode)
-	{
-		NewNode->MaterialExpressionEditorX = ExistingNode->MaterialExpressionEditorX;
-		NewNode->MaterialExpressionEditorY = ExistingNode->MaterialExpressionEditorY;
-	};
-
-	auto PlaceBelowNode= [](UMaterialExpression* ExistingNode, UMaterialExpression* NewNode, int32 DownOffset = 100)
-	{
-		NewNode->MaterialExpressionEditorX = ExistingNode->MaterialExpressionEditorX;
-		NewNode->MaterialExpressionEditorY = ExistingNode->MaterialExpressionEditorY + DownOffset;
-	};
-
-	auto ReplaceNodeAndMoveToTheRight = [](UMaterialExpression* ExistingNode, UMaterialExpression* NewNode, int32 RightOffset = 300)
-	{
-		NewNode->MaterialExpressionEditorX = ExistingNode->MaterialExpressionEditorX;
-		NewNode->MaterialExpressionEditorY = ExistingNode->MaterialExpressionEditorY;
-		ExistingNode->MaterialExpressionEditorX = NewNode->MaterialExpressionEditorX + RightOffset;
-	};
-
-	bool bBlendConverted = false;
-	if (GetMaterialFunctionUsage() == EMaterialFunctionUsage::MaterialLayerBlend)
-	{
-		uint32 InputCount = 0;
-		UMaterialExpressionFunctionInput* BottomInput = nullptr;
-		UMaterialExpressionFunctionInput* TopInput = nullptr;
-		UMaterialExpressionBlendMaterialAttributes* BlendNodeForInputs = nullptr;
-		for (UMaterialExpression* Expression : Expressions)
-		{
-			if (UMaterialExpressionFunctionInput* InputNode = Cast<UMaterialExpressionFunctionInput>(Expression))
-			{
-				if (!BottomInput)
-				{
-					BottomInput = InputNode;
-				}
-				else if (!TopInput)
-				{
-					TopInput = InputNode;
-				}
-				InputCount++;
-			}
-			else if(!BlendNodeForInputs)
-			{
-				BlendNodeForInputs = Cast<UMaterialExpressionBlendMaterialAttributes>(Expression);
-			}
-
-			if(BottomInput && TopInput && BlendNodeForInputs)
-			{
-				break;
-			}
-		}
-
-		//Legacy blend nodes should have exactly 2 inputs, if not, assume invalid for auto-upgrade.
-		if (!TopInput || !BottomInput || InputCount > 2)
-		{
-			UE_LOG(LogMaterial, Warning, TEXT("Material Blend Function %s does not have exactly 2 inputs, so cannot be automatically upgraded to support Substrate."), *GetFullName());
-			return;
-		}
-
-		bool bInputsPresumedCorrect = false;
-		if(BlendNodeForInputs)
-		{
-			auto RecurseBlendInputs= [&](UMaterialExpression* Expression, auto&& RecurseBlendInputs) -> UMaterialExpressionFunctionInput*
-				{
-					//If null or if the expression has already been recursed, skip.
-					if (!Expression)
-					{
-						return nullptr;
-					}
-
-					if(UMaterialExpressionFunctionInput* InputExpression = Cast<UMaterialExpressionFunctionInput>(Expression))
-					{
-						return InputExpression;
-					}
-
-					//If the expression is valid, iterate the connected expressions for upgrade to Substrate, then recurse via each valid expression
-					for (FExpressionInputIterator It{ Expression }; It; ++It)
-					{
-						if(UMaterialExpressionFunctionInput* InputExpression = RecurseBlendInputs(It.Input->Expression, RecurseBlendInputs))
-						{
-							return InputExpression;
-						}
-					}
-
-					return nullptr;
-				};
-
-			UMaterialExpressionFunctionInput* RecursedBottomInput = RecurseBlendInputs(BlendNodeForInputs->A.Expression, RecurseBlendInputs);
-			UMaterialExpressionFunctionInput* RecursedTopInput = RecurseBlendInputs(BlendNodeForInputs->B.Expression, RecurseBlendInputs);
-
-			if(RecursedBottomInput && RecursedTopInput && RecursedBottomInput != RecursedTopInput)
-			{
-				BottomInput = RecursedBottomInput;
-				TopInput = RecursedTopInput;
-				bInputsPresumedCorrect = true;
-			}
-		}
-
-		auto SwapInputs = [&](UMaterialExpressionFunctionInput* BottomInput, UMaterialExpressionFunctionInput*  TopInput) -> void
-			{
-				UMaterialExpressionFunctionInput* SwapInput = BottomInput;
-				TopInput = BottomInput;
-				BottomInput = SwapInput;
-			};
-
-		if(!bInputsPresumedCorrect)
-		{
-			if (BottomInput->InputName == TEXT("Bottom Layer") || BottomInput->InputName == TEXT("Background Layer"))
-			{
-				SwapInputs(BottomInput, TopInput);
-				if(TopInput->InputName == TEXT("Top Layer") || TopInput->InputName == TEXT("Foreground Layer"))
-				{
-					bInputsPresumedCorrect = true;
-				}
-			}
-		}
-
-		if(!bInputsPresumedCorrect)
-		{
-			if(BottomInput->MaterialExpressionEditorY > TopInput->MaterialExpressionEditorY)
-			{
-				SwapInputs(BottomInput, TopInput);
-				bInputsPresumedCorrect = true;
-			}
-		}
-		
-		if(!bInputsPresumedCorrect)
-		{
-			UE_LOG(LogMaterial, Warning, TEXT("Material Blend Function %s inputs could not be used to upgrade for Substrate."), *GetFullName());
-			return;
-		}
-
-		if(BottomInput->InputName != TEXT("Background Layer"))
-		{
-			BottomInput->InputName = TEXT("Background Layer");
-		}
-
-		if(TopInput->InputName == TEXT("Foreground Layer"))
-		{
-			TopInput->InputName = TEXT("Background Layer");
-		}
-
-		if (BottomInput->PreviewValue == FVector4f(FLinearColor::Black) && TopInput->PreviewValue == FVector4f(FLinearColor::Black))
-		{
-			//Only occurs during the update process and if inputs are set to black default so the blend has distinct preview, further user defined colours will not be altered.
-			TopInput->PreviewValue = FLinearColor::White;
-		}
-
-		/**
-		 * Blend node conversion occurs recursively. We can't resolve which nodes could have multiple inputs already connected to a node's output from the expression itself,
-		 * so we have to traverse from the output function to map the existing graph correctly.
-		 */
-		TMap<UMaterialExpression*, UMaterialExpression*> ReplacementNodeMapping;
-		TArray<UMaterialExpression*> RecursedExpressions;
-		auto ConvertBlendExpression = [&](FExpressionInput* Input, UMaterialExpression* InputExpression) -> void
-		{
-			if(!Input || !InputExpression)
-			{
-				return;
-			}
-		
-			if (UMaterialExpressionBlendMaterialAttributes* BlendNode = Cast<UMaterialExpressionBlendMaterialAttributes>(InputExpression))
-			{
-				//Load the existing substrate default blend function which will replace the legacy blend node.
-				static TObjectPtr<UMaterialFunction> DefaultBlendFunction = FindObject<UMaterialFunction>(GetTransientPackage(), DEFAULT_SUBSTRATE_MATERIALLAYERBLEND_PATH);
-				if (!DefaultBlendFunction)
-				{
-					DefaultBlendFunction = LoadObject<UMaterialFunction>(GetTransientPackage(), DEFAULT_SUBSTRATE_MATERIALLAYERBLEND_PATH);
-				}
-
-				if(DefaultBlendFunction)
-				{
-					UMaterialExpressionMaterialFunctionCall* BlendFunctionCall = ReplacementNodeMapping.Contains(BlendNode) ? Cast<UMaterialExpressionMaterialFunctionCall>(*ReplacementNodeMapping.Find(BlendNode)) : nullptr;
-					if (!BlendFunctionCall)
-					{
-						//Store the mapping of a blend node to it's replacement function call, which means we only create a call once per individual blend node.
-						BlendFunctionCall = Cast<UMaterialExpressionMaterialFunctionCall>(ReplacementNodeMapping.Add(BlendNode, NewObject<UMaterialExpressionMaterialFunctionCall>(this, NAME_None, RF_Transactional)));
-						BlendFunctionCall->Function = this;
-						BlendFunctionCall->SetMaterialFunction(DefaultBlendFunction);
-
-						PlaceNodeInLocation(BlendNode, BlendFunctionCall);
-						InputExpression = BlendFunctionCall;
-						Expressions.Add(BlendFunctionCall);
-						RecursedExpressions.Add(BlendFunctionCall);
-
-						if (BlendFunctionCall && BlendFunctionCall->FunctionInputs.Num() >= 8)
-						{
-							//These should match the sort priority of the Default MLB
-							BlendFunctionCall->FunctionInputs[0].Input.Connect(BlendNode->A.OutputIndex, BlendNode->A.Expression); //Background
-							BlendFunctionCall->FunctionInputs[1].Input.Connect(BlendNode->B.OutputIndex, BlendNode->B.Expression); //Foreground
-							BlendFunctionCall->FunctionInputs[2].Input.Connect(BlendNode->Alpha.OutputIndex, BlendNode->Alpha.Expression);
-
-							TArray<bool> StaticBoolValues;
-							int32 StartIndex = StaticBoolValues.Add(true); //UseParameterBlend
-							StaticBoolValues.Add(BlendNode->VertexAttributeBlendType == EMaterialAttributeBlend::Type::UseA); //VertexUseA
-							StaticBoolValues.Add(BlendNode->VertexAttributeBlendType == EMaterialAttributeBlend::Type::UseB); //VertexUseB
-							StaticBoolValues.Add(BlendNode->PixelAttributeBlendType == EMaterialAttributeBlend::Type::UseA); //PixelUseA
-							int32 EndIndex = StaticBoolValues.Add(BlendNode->PixelAttributeBlendType == EMaterialAttributeBlend::Type::UseB); //PixelUseB
-
-							TArray<UMaterialExpressionStaticBool*> StaticBoolArray;
-							StartIndex += 3; 
-							EndIndex += 3;
-							for(int32 FunctionCallIndex = StartIndex; FunctionCallIndex <= EndIndex; FunctionCallIndex++)
-							{
-								int32 ArrayIndex = StaticBoolArray.Add(NewObject<UMaterialExpressionStaticBool>(this, NAME_None, RF_Transactional));
-								UMaterialExpressionStaticBool* ThisBool = StaticBoolArray[ArrayIndex];
-								if(ThisBool)
-								{
-									bool ThisValue = StaticBoolValues[ArrayIndex];
-									ThisBool->Value = ThisValue;
-									BlendFunctionCall->FunctionInputs[FunctionCallIndex].Input.Connect(0, ThisBool);
-									Expressions.Add(ThisBool);
-									RecursedExpressions.Add(ThisBool);
-
-									if(FunctionCallIndex == StartIndex)
-									{
-										PlaceBelowNode(BlendFunctionCall, ThisBool, 300);
-									}
-									else
-									{
-										PlaceBelowNode(StaticBoolArray[ArrayIndex-1], ThisBool, 80);
-									}
-								}
-							}
-						}
-					}
-
-					if (BlendFunctionCall)
-					{
-						Input->Connect(0, BlendFunctionCall);
-						bBlendConverted = true;
-					}
-				}
-			}
-			else if (InputExpression->IsA<UMaterialExpressionSetMaterialAttributes>() || InputExpression->IsA<UMaterialExpressionMakeMaterialAttributes>())
-			{
-				//Same behaviour as above, but for now extends the Make/Set nodes with conversion, so we don't remove the existing node in this case.
-				UMaterialExpressionSetMaterialAttributes* SetAttributesNode = ReplacementNodeMapping.Contains(InputExpression) ? Cast<UMaterialExpressionSetMaterialAttributes>(*ReplacementNodeMapping.Find(InputExpression)) : nullptr;
-				if(!SetAttributesNode)
-				{
-					SetAttributesNode = Cast<UMaterialExpressionSetMaterialAttributes>(ReplacementNodeMapping.Add(InputExpression, NewObject<UMaterialExpressionSetMaterialAttributes>(this, NAME_None, RF_Transactional)));
-					SetAttributesNode->Function = this;
-					SetAttributesNode->ConnectInputAttribute(MP_MaterialAttributes, InputExpression);
-					Expressions.Add(SetAttributesNode);
-					RecursedExpressions.Add(SetAttributesNode);
-
-					UMaterialExpressionSubstrateConvertMaterialAttributes* ConvertNode = NewObject<UMaterialExpressionSubstrateConvertMaterialAttributes>(this, NAME_None, RF_Transactional);
-					ConvertNode->Function = this;
-					ConvertNode->MaterialAttributes.Connect(0, InputExpression);
-					ConvertNode->ShadingModelOverride = MSM_DefaultLit;
-					Expressions.Add(ConvertNode);
-					RecursedExpressions.Add(ConvertNode);
-
-					SetAttributesNode->ConnectInputAttribute(MP_FrontMaterial, ConvertNode);
-
-					PlaceBelowNode(InputExpression, ConvertNode);
-					PlaceBelowNode(ConvertNode, SetAttributesNode, 250);
-				}
-
-				if(SetAttributesNode)
-				{
-					Input->Connect(0, SetAttributesNode);
-					bBlendConverted = true;
-				}
-			}
-		};
-
-		auto RecurseBlendFunction = [&](UMaterialExpression* Expression, auto&& RecurseBlendFunction) -> void
-		{
-			//If null or if the expression has already been recursed, skip.
-			if(!Expression || RecursedExpressions.Contains(Expression))
-			{
-				return;
-			}
-			RecursedExpressions.Add(Expression);
-
-			//If the expression is valid, iterate the connected expressions for upgrade to Substrate, then recurse via each valid expression
-			for (FExpressionInputIterator It{ Expression }; It; ++It)
-			{
-				ConvertBlendExpression(It.Input, It.Input->Expression);
-				RecurseBlendFunction(It.Input->Expression, RecurseBlendFunction);	
-			}
-		};
-		//Begin recursion process from the output node.
-		RecurseBlendFunction(OutputNode, RecurseBlendFunction);
-
-		//Once the nodes have been replaced and all connections remapped, we can remove the replaced nodes from the function.
-		for (auto& Mapping : ReplacementNodeMapping)
-		{
-			if(UMaterialExpressionBlendMaterialAttributes* BlendExpression = Cast<UMaterialExpressionBlendMaterialAttributes>(Mapping.Key))
-			{
-				for (FExpressionInputIterator It{ BlendExpression }; It; ++It)
-				{
-					It.Input->Expression = nullptr;
-				}
-				Expressions.Remove(BlendExpression);
-			}
-		}
-	}
-
-	//If we are converting a standard layer function, or if the blend node has not had blend functionality replaced directly, add basic conversion logic to the node.
-	if (!bBlendConverted || GetMaterialFunctionUsage() == EMaterialFunctionUsage::MaterialLayer)
-	{
-		//Layer assets are simpler than their blend counterparts because we only need to convert the legacy MAs, similar to the behaviour of a material upgrade.
-		UMaterialExpressionSubstrateConvertMaterialAttributes* ConvertAttributesNode = NewObject<UMaterialExpressionSubstrateConvertMaterialAttributes>(this);
-		ConvertAttributesNode->Function = this;
-		ConvertAttributesNode->MaterialAttributes.Connect(0, OutputNode->A.Expression);
-		ConvertAttributesNode->ShadingModelOverride = MSM_DefaultLit;
-		ReplaceNodeAndMoveToTheRight(OutputNode, ConvertAttributesNode);
-		PlaceBelowNode(ConvertAttributesNode, ConvertAttributesNode);
-
-		//Add Custom logic connections
-		TArray<class UMaterialExpressionCustomOutput*> CustomOutputExpressions;
-		GetAllCustomOutputExpressions(CustomOutputExpressions);
-		check(ConvertAttributesNode);
-		for (UMaterialExpressionCustomOutput* Expression : CustomOutputExpressions)
-		{
-			// Gather custom output for single layer water
-			if (UMaterialExpressionSingleLayerWaterMaterialOutput* SingleLayerWaterOutput = Cast<UMaterialExpressionSingleLayerWaterMaterialOutput>(Expression))
-			{	
-				if (SingleLayerWaterOutput->ScatteringCoefficients.Expression)
-				{
-					ConvertAttributesNode->WaterScatteringCoefficients.Connect(SingleLayerWaterOutput->ScatteringCoefficients.OutputIndex, SingleLayerWaterOutput->ScatteringCoefficients.Expression);
-				}
-				if (SingleLayerWaterOutput->AbsorptionCoefficients.Expression)
-				{
-					ConvertAttributesNode->WaterAbsorptionCoefficients.Connect(SingleLayerWaterOutput->AbsorptionCoefficients.OutputIndex, SingleLayerWaterOutput->AbsorptionCoefficients.Expression);
-				}
-				if (SingleLayerWaterOutput->PhaseG.Expression)
-				{
-					ConvertAttributesNode->WaterPhaseG.Connect(SingleLayerWaterOutput->PhaseG.OutputIndex, SingleLayerWaterOutput->PhaseG.Expression);
-				}
-				if (SingleLayerWaterOutput->ColorScaleBehindWater.Expression)
-				{
-					ConvertAttributesNode->ColorScaleBehindWater.Connect(SingleLayerWaterOutput->ColorScaleBehindWater.OutputIndex, SingleLayerWaterOutput->ColorScaleBehindWater.Expression);
-				}
-
-				break;
-			}
-		}
-
-		/**
-		 * Layer MFs differ from materials in that they only have 1 output,
-		 * so we use the Set Substrate Attributes node to collect the Front Material and MAs to pass to the next function in the layer stack.
-		 */
-		UMaterialExpressionSetMaterialAttributes* SetMatAttributesNode = NewObject<UMaterialExpressionSetMaterialAttributes>(this);
-		SetMatAttributesNode->Function = this;
-		SetMatAttributesNode->ConnectInputAttribute(MP_MaterialAttributes, OutputNode->A.Expression);
-		SetMatAttributesNode->ConnectInputAttribute(MP_FrontMaterial, ConvertAttributesNode);
-		ReplaceNodeAndMoveToTheRight(OutputNode, SetMatAttributesNode);
-
-		Expressions.EmplaceAt(1, SetMatAttributesNode);
-		Expressions.EmplaceAt(2, ConvertAttributesNode);
-
-		OutputNode->A.Connect(0, SetMatAttributesNode);
-	}
-	MoveNodeInHorizonalAxis(OutputNode, 50);
-	OutputNode->bCollapsed = true;
-}
-#endif //WITH_EDITORONLY_DATA
-
-void UMaterialFunction::PostLoad()
+#endif //WITH_EDITORONLY_DATA
+
+void UMaterialFunction::PostLoad()
 {
 	LLM_SCOPE(ELLMTag::Materials);
 
@@ -14454,8 +13602,6 @@ void UMaterialFunction::PostLoad()
 		EditorOnly->ExpressionCollection.EditorComments = MoveTemp(FunctionEditorComments_DEPRECATED);
 	}
 	
-	ConvertExpressionsBetweenLegacyAndSubstrate();
-
 	if (EditorOnly)
 	{
 		for (UMaterialExpression* Expression : EditorOnly->ExpressionCollection.Expressions)
@@ -14768,104 +13914,11 @@ bool UMaterialFunction::ValidateFunctionUsage(FMaterialCompiler* Compiler, const
 #if WITH_EDITOR
 	if (GetMaterialFunctionUsage() == EMaterialFunctionUsage::MaterialLayer)
 	{
-		if (Substrate::IsMaterialLayeringSupportEnabled())
-		{
-			// Material layers must have a single MA input and output only
-			for (UMaterialExpression* Expression : GetExpressions())
-			{
-				if (UMaterialExpressionFunctionInput* InputExpression = Cast<UMaterialExpressionFunctionInput>(Expression))
-				{
-					if (InputExpression->IsResultMaterialAttributes(0) || InputExpression->IsResultSubstrateMaterial(0))
-					{
-						++NumInputs;
-					}
-
-					if (NumInputs > 1)
-					{
-						Compiler->Errorf(TEXT("Layer graphs only support a single material attributes or Substrate inputs currently."));
-						bHasValidOutput = false;
-					}
-				}
-				else if (UMaterialExpressionFunctionOutput* OutputExpression = Cast<UMaterialExpressionFunctionOutput>(Expression))
-				{
-					if (OutputExpression->IsResultMaterialAttributes(0) || OutputExpression->IsResultSubstrateMaterial(0))
-					{
-						++NumOutputs;
-					}
-
-					if (NumOutputs > 1)
-					{
-						Compiler->Errorf(TEXT("Layer graphs only support a single material attributes or Substrate outputs currently."));
-						bHasValidOutput = false;
-					}
-				}
-				else if (UMaterialExpressionMaterialAttributeLayers* RecursiveLayer = Cast<UMaterialExpressionMaterialAttributeLayers>(Expression))
-				{
-					Compiler->Errorf(TEXT("Layer graphs do not support layers within layers."));
-					bHasValidOutput = false;
-				}
-			}
-
-			if (NumInputs > 1 || NumOutputs < 1)
-			{
-				Compiler->Errorf(TEXT("Layer graphs require a single material attributes or Substrate output and optionally, a single material attributes or Substrate input."));
-				bHasValidOutput = false;
-			}
-		}
-		else
-		{
-			bHasValidOutput = UMaterialExpressionLayerStack::ValidateFunctionForLayerUsage(Compiler, this);
-		}
+		bHasValidOutput = UMaterialExpressionLayerStack::ValidateFunctionForLayerUsage(Compiler, this);
 	}
 	else if (GetMaterialFunctionUsage() == EMaterialFunctionUsage::MaterialLayerBlend)
 	{
-		if(Substrate::IsMaterialLayeringSupportEnabled())
-		{
-			// Material layer blends can have up to two MA inputs and single MA output only
-			for (UMaterialExpression* Expression : GetExpressions())
-			{
-				if (UMaterialExpressionFunctionInput* InputExpression = Cast<UMaterialExpressionFunctionInput>(Expression))
-				{
-					if(InputExpression->IsResultMaterialAttributes(0) || InputExpression->IsResultSubstrateMaterial(0))
-					{
-						++NumInputs;
-					}
-
-					if (NumInputs > 2)
-					{
-						Compiler->Errorf(TEXT("Layer blend graphs only support two material attributes or Substrate inputs currently."));
-						bHasValidOutput = false;
-					}
-				}
-				else if (UMaterialExpressionFunctionOutput* OutputExpression = Cast<UMaterialExpressionFunctionOutput>(Expression))
-				{
-					if (OutputExpression->IsResultMaterialAttributes(0) || OutputExpression->IsResultSubstrateMaterial(0))
-					{
-						++NumOutputs;
-					}
-					if (NumOutputs > 1)
-					{
-						Compiler->Errorf(TEXT("Layer blend graphs only support a single MA or Substrate output currently."));
-						bHasValidOutput = false;
-					}
-				}
-				else if (UMaterialExpressionMaterialAttributeLayers* RecursiveLayer = Cast<UMaterialExpressionMaterialAttributeLayers>(Expression))
-				{
-					Compiler->Errorf(TEXT("Layer blend graphs do not support layers within layers."));
-					bHasValidOutput = false;
-				}
-			}
-
-			if (NumOutputs < 1)
-			{
-				Compiler->Errorf(TEXT("Layer blend graphs must have a only a single MA or Substrate output currently."));
-				bHasValidOutput = false;
-			}
-		}
-		else
-		{
-			bHasValidOutput = UMaterialExpressionLayerStack::ValidateFunctionForBlendUsage(Compiler, this);
-		}
+		bHasValidOutput = UMaterialExpressionLayerStack::ValidateFunctionForBlendUsage(Compiler, this);
 	}
 #endif
 
@@ -15689,615 +14742,13 @@ void FMaterialLayersFunctionsID::Append(FShaderKeyGenerator& KeyGen) const
 
 const FGuid FMaterialLayersFunctions::BackgroundGuid(2u, 0u, 0u, 0u);
 
-
-bool FMaterialLayersFunctionsTree::Serialize(FArchive& Ar)
+FMaterialLayersFunctionsRuntimeData::~FMaterialLayersFunctionsRuntimeData()
 {
-	Ar << Nodes;
-	Ar << Payloads;
-	Ar << Root;
-	return true;
-}
-
-using FLayerNodeId = FMaterialLayersFunctionsTree::FNodeId;
-using FLayerNode = FMaterialLayersFunctionsTree::FNode;
-using FLayerNodePayload = FMaterialLayersFunctionsTree::FPayload;
-using FLayerNodeIdArray = TArray<FLayerNodeId>;
-
-FLayerNodeId FMaterialLayersFunctionsTree::AllocNode(const FPayload& InPayload)
-{
-	Nodes.AddDefaulted();
-	Payloads.Add(InPayload);
-	return (FNodeId)Nodes.Num() - 1;
-}
-
-
-// Find and return the node at the specified Id or null if invalid 
-FLayerNode* FMaterialLayersFunctionsTree::GetNode(FNodeId InNodeId) const
-{
-	if (IsValidId(InNodeId))
-	{
-		return const_cast<FNode*>(Nodes.GetData() + InNodeId);
-	}
-
-	return nullptr; // invalid index => return null
-}
-
-// Find and return the node payload at the specified Id or null if invalid 
-FLayerNodePayload* FMaterialLayersFunctionsTree::GetPayload(FNodeId InNodeId) const
-{
-	if (IsValidId(InNodeId))
-	{
-		return const_cast<FPayload*>(Payloads.GetData() + InNodeId);
-	}
-
-	return nullptr; // invalid index => return null
-}
-
-// Find and return the parent node at the specified Id or null if invalid 
-FLayerNode* FMaterialLayersFunctionsTree::GetParent(FNodeId InNodeId) const
-{
-	FNode* Node = GetNode(InNodeId);
-	if (Node && IsValidId(Node->Parent))
-	{
-		return GetNode(Node->Parent);
-	}
-
-	return nullptr; // invalid parent => return null
-}
-
-// Return the depth of the node in the tree, the root is depth 0
-// the first valid node in the tree under the root is depth 1
-FLayerNodeIdArray FMaterialLayersFunctionsTree::GetChildrenIds(FNodeId InNodeId) const
-{
-	FLayerNodeIdArray ChildrendIds;
-
-	FNodeId NextChildId = Root;
-	FNode* Node = GetNode(InNodeId);
-	if (Node)
-		NextChildId = Node->ChildrenHead;
-
-	while (IsValidId(NextChildId))
-	{
-		ChildrendIds.Add(NextChildId);
-		Node = GetNode(NextChildId);
-		if (Node)
-			NextChildId = Node->NextSibling;
-		else
-			NextChildId = -1;
-	}
-
-	return ChildrendIds;
-}
-
-FLayerNodeId FMaterialLayersFunctionsTree::GetChildrenTailId(FNodeId InNodeId) const
-{
-	FNodeId NextChildId = Root;
-	FNode* Node = GetNode(InNodeId);
-	if (Node)
-		NextChildId = Node->ChildrenHead;
-
-	while (IsValidId(NextChildId))
-	{
-		Node = GetNode(NextChildId);
-		if (Node && IsValidId(Node->NextSibling))
-			NextChildId = Node->NextSibling;
-		else
-			break;
-	}
-
-	return NextChildId;
-}
-
-int32 FMaterialLayersFunctionsTree::GetDepth(FNodeId InNodeId) const
-{
-	int32 Depth = 0;
-	FNode* Node = GetNode(InNodeId);
-	while (Node)
-	{
-		Node = GetNode(Node->Parent);
-		Depth++;
-	}
-
-	return Depth;
-}
-
-FLayerNodeIdArray FMaterialLayersFunctionsTree::GetParentIds(FNodeId InNodeId) const
-{
-	FLayerNodeIdArray ParentIds;
-
-	FNode* Node = GetNode(InNodeId);
-	while (Node)
-	{
-		ParentIds.Add(Node->Parent);
-		Node = GetNode(Node->Parent);
-	}
-
-	return ParentIds;
-}
-
-// Find the node which is the head of the list of siblings where InNodeID belongs
-FLayerNodeId FMaterialLayersFunctionsTree::GetSiblingHeadId(FNodeId InNodeId) const
-{
-	if (!IsValidId(InNodeId))
-		return InvalidId;
-
-	FNode* TargetNode = const_cast<FNode*>(Nodes.GetData() + InNodeId);
-
-	FNodeId HeadNodeId = Root;
-	if (TargetNode->Parent != InvalidId)
-	{
-		HeadNodeId =Nodes[TargetNode->Parent].ChildrenHead;
-	}
-
-	return HeadNodeId;
-}
-
-FLayerNodeId FMaterialLayersFunctionsTree::GetPreviousSiblingId(FNodeId InNodeId) const
-{
-	// Go to the sibling head
-	// and walk next until next is the specified node
-	FLayerNodeId PrevSibling = GetSiblingHeadId(InNodeId);
-
-	// if the queried node IS the first of the siblings then no more prev siblings
-	if (!IsValidId(PrevSibling) || (PrevSibling == InNodeId))
-		return InvalidId;
-
-	// LEt's look for the Node whose next sibling IS the queried nodeID
-	while (IsValidId(PrevSibling))
-	{
-		FNode* Node = GetNode(PrevSibling); 
-		if (Node->NextSibling == InNodeId)
-		{
-			break;
-		}
-		PrevSibling = Node->NextSibling;
-	}
-
-	return PrevSibling;
-}
-
-FLayerNodeId FMaterialLayersFunctionsTree::AddNode(const FPayload& InPayload, FNodeId InParent, int32 InAtChildIndex)
-{
-	// Check Parent id
-	if (!IsValidId(InParent))
-	{
-		// Bad parent id, allocate at root at the end
-		InParent = -1;
-	}
-
-	// Allocate the new node setting the payload value
-	FNodeId NodeId = AllocNode(InPayload);
-
-	// Let's connect the new node in the hierarchy
-	FNode* NewNode = Nodes.GetData() + NodeId;
-	NewNode->Parent = InParent;
-
-	// First find the parent node and from that the head node for the list in which we will add the new node
-	// Initialize the case if parent is the tree root
-	FNodeId HeadSiblingId = Root;
-	FNodeId* ParentChildrenHeaddPtr = &Root;
-	// Parent is actually a node and not the root:
-	if (InParent >= 0)
-	{
-		FNode* ParentNode = Nodes.GetData() + InParent;
-		HeadSiblingId = ParentNode->ChildrenHead;
-		ParentChildrenHeaddPtr = &ParentNode->ChildrenHead;
-	}
-
-	// When Insert at head or no siblings yet then update parent sub
-	if ((InAtChildIndex == 0) || (HeadSiblingId < 0))
-	{
-		(*ParentChildrenHeaddPtr) = NodeId;
-		NewNode->NextSibling = HeadSiblingId;
-
-		// Done
-		return NodeId;
-	}
-
-	// Second insert the new node among the siblings AFTER head
-	FNode* SiblingNode = Nodes.GetData() + HeadSiblingId;
-
-	int32 SiblingNextIdx = 1;
-	while (SiblingNode->NextSibling >= 0)
-	{
-		// if the next sibling is where we want to insert then insert
-		if (InAtChildIndex == SiblingNextIdx)
-		{
-			NewNode->NextSibling = SiblingNode->NextSibling;
-			SiblingNode->NextSibling = NodeId;
-
-			// Done
-			return NodeId;
-		}
-
-		// next sibling
-		SiblingNode = Nodes.GetData() + SiblingNode->NextSibling;
-		SiblingNextIdx++;
-	}
-
-	// we went through the siblings and now hold the tail node, let's add the new node as the new tail
-	SiblingNode->NextSibling = NodeId;
-
-	// Done
-	return NodeId;
-}	
-
-int FMaterialLayersFunctionsTree::RemoveNodeBranch(FNode* BranchParentNode, TArray<FNodeId>& RemovedIds, TArray<FPayload>& RemovedPayloads)
-{
-	int NumRemoved = 0;
-	if (BranchParentNode)
-	{
-		FNodeId RemovedNodeId = BranchParentNode->ChildrenHead;
-		while (IsValidId(RemovedNodeId))
-		{
-			FNode* RemovedNode = GetNode(RemovedNodeId);
-
-			// Record the removed node and payload
-			RemovedIds.Add(RemovedNodeId);
-			RemovedPayloads.Add(Payloads[RemovedNodeId]);
-			NumRemoved++;
-
-			// Remove the potential branch of nodes UNDER this removed node
-			NumRemoved += RemoveNodeBranch(RemovedNode, RemovedIds, RemovedPayloads);
-			
-			// Clear the deleted node
-			FNodeId NextSyblingId = RemovedNode->NextSibling;
-			Nodes[RemovedNodeId] = FNode();
-			Payloads[RemovedNodeId] = FPayload();
-
-			// Move on to the next sybling
-			RemovedNodeId = NextSyblingId;
-		}
-
-		BranchParentNode->ChildrenHead = InvalidId; // Disconnect the branch removed from the parent
-	}
-
-	return NumRemoved;
-}
-#if WITH_EDITOR
-int32 FMaterialLayersFunctionsRuntimeGraphCache::FindExpressionIndex(FGuid InExpressionGuid, const FMaterialLayersFunctions* Layersfunctions) const
-{
-	TArray<int32> NodeIndicesLayer; 
-	TArray<int32> NodeIndicesBlends;
-
-	GetNodeIndices(InExpressionGuid, NodeMaterialGraphExpressions, NodeIndicesLayer);
-	GetNodeIndices(InExpressionGuid, NodeMaterialGraphExpressionsBlends, NodeIndicesBlends);
-
-	int32 NodeIndex = INDEX_NONE;
-
-	const int32 NumLayers = NodeIndicesLayer.Num();
-	for(int32 LayerIndex = NumLayers - 1; LayerIndex >= 0; LayerIndex--)
-	{
-		int32 CurrentNodeIndex = NodeIndicesLayer[LayerIndex];
-		if (Layersfunctions->GetNodeChildren(CurrentNodeIndex).IsEmpty())
-		{
-			NodeIndex = CurrentNodeIndex;
-			break;
-		}
-	}
-	if (NodeIndex == INDEX_NONE)
-	{
-		const int32 NumBlends = NodeIndicesBlends.Num();
-		for(int32 BlendIndex = NumBlends - 1; BlendIndex >= 0; BlendIndex--)
-		{
-			int32 CurrentNodeIndex = NodeIndicesBlends[BlendIndex];
-			if (Layersfunctions->GetNodeChildren(CurrentNodeIndex).IsEmpty())
-			{
-				NodeIndex = CurrentNodeIndex;
-				break;
-			}
-		}
-	}
-	return NodeIndex;
-}
-
-void FMaterialLayersFunctionsRuntimeGraphCache::GetNodeIndices(FGuid InExpressionGuid,
-	const TArray<TObjectPtr<UMaterialExpressionMaterialFunctionCall>>& ExpressionArrayToUse, TArray<int32>& OutIndices) const
-{
-	const int32 NumGraphExpressions = ExpressionArrayToUse.Num();
-	for(int32 MaterialExpressionIndex = NumGraphExpressions - 1; MaterialExpressionIndex >= 0; MaterialExpressionIndex--)
-	{
-		TObjectPtr<UMaterialExpressionMaterialFunctionCall> FunctionCall = ExpressionArrayToUse[MaterialExpressionIndex];
-		if (FunctionCall)
-		{
-			const uint32 OutputCount = FunctionCall->FunctionOutputs.Num();
-			for (uint32 OutputIndex = 0; OutputIndex < OutputCount; ++OutputIndex)
-			{
-				if (FunctionCall->IsResultSubstrateMaterial(OutputIndex))
-				{
-					FSubstrateMaterialInfo SubstrateMaterialInfo(true/*bGatherGuids*/);
-					FunctionCall->GatherSubstrateMaterialInfo(SubstrateMaterialInfo, OutputIndex);
-					TArray<FGuid> Guids = SubstrateMaterialInfo.GetGuids();
-    
-					if (Guids.Contains(InExpressionGuid))
-						OutIndices.Add(MaterialExpressionIndex);
-				}
-			}
-		}
-	}
-}
-#endif
-	
-std::pair<TArray<FMaterialLayersFunctionsTree::FNodeId>, TArray<FMaterialLayersFunctionsTree::FPayload>> FMaterialLayersFunctionsTree::RemoveNode(FNodeId RemovedNodeId)
-{
-	TArray<FNodeId> RemovedIds;
-	TArray<FPayload> RemovedPayloads;
-
-	FNode* RemovedNode = GetNode(RemovedNodeId);
-	if (!RemovedNode)
-		return { RemovedIds, RemovedPayloads };
-
-	// Record the removed node and payload
-	RemovedIds.Add(RemovedNodeId);
-	RemovedPayloads.Add(Payloads[RemovedNodeId]);
-
-	// remove the branch of nodes starting at the removed node
-	RemoveNodeBranch(RemovedNode, RemovedIds, RemovedPayloads);
-
-	// Then remove the node itself
-
-	FNodeId NextId = RemovedNode->NextSibling;
-
-	FNodeId HeadId = GetSiblingHeadId(RemovedNodeId);
-	FNode* HeadSibling = GetNode(HeadId);
-
-	// walk from head sibling to removed node
-	// then detach from tree
-	FNodeId NextPrevNodeId = HeadId;
-	FNode* PrevNode = HeadSibling; // prev node is the head sibling at first
-	FNode* NextPrevNode = HeadSibling; // next prev node is the head sibling at first
-	while (NextPrevNode)
-	{
-		if (NextPrevNodeId == RemovedNodeId)
-		{
-			break;
-		}
-		PrevNode = NextPrevNode;
-		NextPrevNodeId = PrevNode->NextSibling;
-		NextPrevNode = GetNode(PrevNode->NextSibling);
-	}
-
-	// Removed node is the sibling head
-	// adjust parent link to children
-	if (NextPrevNode == HeadSibling)
-	{
-		FNodeId OldChildHeadId = Root;
-
-		if (RemovedNode->Parent != InvalidId)
-		{
-			FNode* ParentNode = Nodes.GetData() + RemovedNode->Parent;
-			OldChildHeadId = ParentNode->ChildrenHead;
-			ParentNode->ChildrenHead = NextId;
-		}
-		else // Parent is the root so relink root
-		{
-			Root = NextId;
-		}
-	}
-	// Removed node in the middle of the siblings
-	else
-	{
-		PrevNode->NextSibling = NextId;
-	}
-
-	// Clear the node and payload that just have been removed
-	Nodes[RemovedNodeId] = FNode();
-	Payloads[RemovedNodeId] = FPayload();
-
-	return { RemovedIds, RemovedPayloads };
-}
-
-void FMaterialLayersFunctionsTree::MoveNode(FNodeId InNodeId, FNodeId DstParentId, int32 InSiblingIndex)
-{
-	if (InNodeId == InvalidId || InNodeId == DstParentId)
-    {
-        return;
-    }
-
-    FNode* NodeToMove = GetNode(InNodeId);
-    FNodeId OldParentId = NodeToMove->Parent;
-
-    FNode* OldParentNode = GetNode(OldParentId);
-	FNodeId PrevSiblingId = GetSiblingHeadId(InNodeId);
-	
-	// If this node is the head child, update the parent's ChildrenHead
-    if (PrevSiblingId == InNodeId)
-    {
-        // 1. Detach the node from its old parent
-    	if (OldParentNode)
-        {
-            OldParentNode->ChildrenHead = NodeToMove->NextSibling;
-        }
-        else
-        {
-            Root = NodeToMove->NextSibling;
-        }
-    }
-    else
-    // Otherwise, find the previous sibling and update its NextSibling
-	 {
-        while (PrevSiblingId != InvalidId && GetNode(PrevSiblingId)->NextSibling != InNodeId)
-        {
-            PrevSiblingId = GetNode(PrevSiblingId)->NextSibling;
-        }
-
-        if (PrevSiblingId != InvalidId)
-        {
-            // Previous sibling found, update its NextSibling to skip the node being moved
-            GetNode(PrevSiblingId)->NextSibling = NodeToMove->NextSibling;
-        }
-    }
-
-    // 2. Update the node's parent to the new parent
-    NodeToMove->Parent = DstParentId;
-	NodeToMove->NextSibling = InvalidId; // Reset the NextSibling for now
-	
-	FNode* NewParentNode = GetNode(DstParentId);
-	FNodeId SiblingHead = NewParentNode ? NewParentNode->ChildrenHead : Root;
-	
-	// 3. Insert the node at the given sibling index
-	if (InSiblingIndex <= 0 || SiblingHead == InvalidId)
-	{
-		// Insert as the first child if the index is 0 or if there are no children
-		NodeToMove->NextSibling = SiblingHead;
-		if (NewParentNode)
-		{
-			NewParentNode->ChildrenHead = InNodeId;
-		}
-		else
-		{
-			// this is the first depth level with no parent, we set the root in this case
-			Root = InNodeId;
-		}
-	}
-	else
-	{
-		// Insert after the specified sibling index
-		PrevSiblingId = SiblingHead;
-		int32 CurrentIndex = 0;
-
-		// Traverse to find the sibling at InSiblingIndex - 1 (previous sibling)
-		while (PrevSiblingId != InvalidId && CurrentIndex < InSiblingIndex - 1)
-		{
-			PrevSiblingId = GetNode(PrevSiblingId)->NextSibling;
-			++CurrentIndex;
-		}
-
-		// Insert the node after the previous sibling
-		if (PrevSiblingId != InvalidId)
-		{
-			NodeToMove->NextSibling = GetNode(PrevSiblingId)->NextSibling;
-			GetNode(PrevSiblingId)->NextSibling = InNodeId;
-		}
-	}
-}
-
-FString FMaterialLayersFunctionsTree::Log(FString InTab) const
-{
-	FString NewLine = TEXT("\r\n");
-	FString LogMessage;
-
-	auto Visitor = [&](FNodeId InNodeId, int32 InDepth, int32 InSiblingNum, FPayload InPayload)
-		{
-			FString Tab = InTab;
-			for (int d = 0; d < InDepth; d++)
-				Tab += "  ";
-			LogMessage += Tab + FString::Printf(TEXT("%i %i - L%i B%i"), InNodeId, InSiblingNum, InPayload.Layer, InPayload.Blend) + NewLine;
-		};
-
-	Traverse(Visitor);
-
-	return LogMessage;
-}
-
-namespace MLFT {
-
-	FString Print(const FMaterialLayersFunctionsRuntimeData* LFRD, FString InTab)
-	{
-		FString NewLine = TEXT("\r\n");
-		FString LogMessage; // = InTab + TEXT("FMaterialLayersFunctionsRuntimeData") + NewLine;
-
-		InTab += "  ";
-
-		// Layers
-		LogMessage += InTab + TEXT("Layers:") + NewLine;
-		for (int32 i = 0; i < LFRD->Layers.Num(); ++i)
-		{
-			FString LayerHead = InTab + FString::Printf(TEXT("%i - "), i);
-			auto Layer = LFRD->Layers[i];
-			if (Layer)
-			{
-				LayerHead += FString::Printf(TEXT("%s"), *Layer->GetName());
-			}
-			LogMessage += LayerHead + NewLine;
-		}
-		// Blends
-		LogMessage += InTab + TEXT("Blends:") + NewLine;
-		for (int32 i = 0; i < LFRD->Blends.Num(); ++i)
-		{
-			FString BlendHead = InTab + FString::Printf(TEXT("%i - "), i);
-			auto Blend = LFRD->Blends[i];
-			if (Blend)
-			{
-				BlendHead += FString::Printf(TEXT("%s"), *Blend->GetName());
-			}
-			LogMessage += BlendHead + NewLine;
-		}
-		// Tree
-		LogMessage += InTab + TEXT("Tree:") + NewLine;
-		//LogMessage += LFRD->Tree.Log(InTab);
-		struct TraverseNode { int32 NodeID; int32 Depth; FMaterialLayersFunctionsTree::FPayload Payload; };
-		TArray<TraverseNode> TraverseNodes;
-
-		auto Visitor = [&](FMaterialLayersFunctionsTree::FNodeId InNodeId, int32 InDepth, int32 InSiblingNum, FMaterialLayersFunctionsTree::FPayload InPayload)
-			{
-				TraverseNodes.Add({ InNodeId,  InDepth, InPayload });
-			};
-		LFRD->Tree.TraverseBottomUp(Visitor);
-
-		for (auto ItNode = TraverseNodes.begin(); ItNode != TraverseNodes.end(); ++ItNode)
-		{
-			TraverseNode* N = &(*ItNode);
-			FString Tab = InTab;
-			for (int d = 0; d < N->Depth; d++)
-				Tab += "    ";
-
-			FString OffsetTab = "";
-			for (int d = N->Depth; d < 2; d++)
-				OffsetTab += "----";
-
-			FString LayerName = TEXT("               ");
-			bool HasLayerFunc = false;
-			if (N->Payload.Layer != FMaterialLayersFunctionsTree::InvalidId && LFRD->Layers[N->Payload.Layer])
-			{
-				LayerName = FString::Printf(TEXT("L%-2i %-12s"), N->Payload.Layer, *LFRD->Layers[N->Payload.Layer]->GetName());
-				HasLayerFunc = true;
-			}
-			FString BlendName = TEXT("               ");
-			bool HasBlendFunc = false;
-			if (N->Payload.Blend != FMaterialLayersFunctionsTree::InvalidId && LFRD->Blends[N->Payload.Blend])
-			{
-				BlendName = FString::Printf(TEXT("B%-2i %-12s"), N->Payload.Blend, *LFRD->Blends[N->Payload.Blend]->GetName());
-				HasBlendFunc = true;
-			}
-
-			LogMessage += Tab + FString::Printf(TEXT("%i -%s> %s %s"), N->NodeID, *OffsetTab, *BlendName, *LayerName) + NewLine;
-		}
-
-
-		return LogMessage;
-	}
-
-
-	void Log(const FMaterialLayersFunctionsRuntimeData* LFRD, FString CallingSite = FString())
-	{
-		FString NewLine = TEXT("\r\n");
-		FString LogMessage;
-		FString Tab = " ";
-
-		LogMessage += FString::Printf(TEXT("**** FMaterialLayersFunctionsRuntimeData from <%s> ****"), *CallingSite) + NewLine;
-		LogMessage += Print(LFRD, Tab);
-		LogMessage += TEXT("**** *********** ****") + NewLine;
-
-		TArray<FString> Lines;
-		LogMessage.ParseIntoArray(Lines, TEXT("\n"));
-		for (const FString& Line : Lines)
-		{
-			UE_LOG(LogMaterial, VeryVerbose, TEXT("%s"), *Line);
-		}
-	}
-
-
-}
-
-
-FMaterialLayersFunctionsRuntimeData::~FMaterialLayersFunctionsRuntimeData()
-{
-#if WITH_EDITORONLY_DATA
-	// If this is destroyed while still holding 'LegacySerializedEditorOnlyData', that means it was serialized from an 'FMaterialLayersFunctions' in some unexpected context
-	// Need to find where this is happening, and update that code to properly handle LegacySerializedEditorOnlyData
-	checkf(!LegacySerializedEditorOnlyData, TEXT("LegacySerializedEditorOnlyData should have been acquired by FStaticParameterSet"));
-#endif
+#if WITH_EDITORONLY_DATA
+	// If this is destroyed while still holding 'LegacySerializedEditorOnlyData', that means it was serialized from an 'FMaterialLayersFunctions' in some unexpected context
+	// Need to find where this is happening, and update that code to properly handle LegacySerializedEditorOnlyData
+	checkf(!LegacySerializedEditorOnlyData, TEXT("LegacySerializedEditorOnlyData should have been acquired by FStaticParameterSet"));
+#endif
 }
 
 bool FMaterialLayersFunctionsRuntimeData::SerializeFromMismatchedTag(const FPropertyTag& Tag, FStructuredArchive::FSlot Slot)
@@ -16318,47 +14769,6 @@ bool FMaterialLayersFunctionsRuntimeData::SerializeFromMismatchedTag(const FProp
 	return false;
 }
 
-
-void FMaterialLayersFunctionsRuntimeData::PostSerialize(const FArchive& Ar)
-{
-	if (Substrate::IsMaterialLayeringSupportEnabled() && Ar.IsLoading())
-	{
-		// When loading, if a legacy PRE substrate tree layer version:
-		// The tree need to be recreated from the list
-		// of layers and blends assuming this is a flat hierarchy of layers:
-		// For PRE substrate tree layer version:
-		// TODO: We could use:
-		// 	if (Ar.CustomVer(FRenderingObjectVersion::GUID) < FRenderingObjectVersion::AddedMaterialLayersSubstrateSupport)
-		if ((Tree.Nodes.Num() != Layers.Num()) || (Layers.Num() > Blends.Num()))
-		{
-			int NumLegacyLayers = Layers.Num();
-			int NumLegacyBlends = Blends.Num();
-
-			Tree.Empty();
-
-			// First add an extra empty Blend to match the number of layers
-			while (Blends.Num() < Layers.Num())
-			{
-				Blends.AddDefaulted(); // normally we should go through this only once
-			}
-
-			Layers.AddDefaulted();
-			Blends.AddDefaulted();
-
-			// Add a layer node at root 
-			FLayerNodeId LayerNodeId = Tree.AddNode({ NumLegacyLayers,  NumLegacyLayers }, -1);
-
-			// Third rebuiild the tree of nodes referencing the layers and blends
-			for (int32 l = 0; l < NumLegacyLayers; ++l)
-			{
-				// And then ONE MORE node to be the attributes first child
-				Tree.AddNode({ l,  l == 0 ? NumLegacyLayers - 1 : l-1}, LayerNodeId);
-			}
-		}
-	}
-}
-
-
 #if WITH_EDITOR
 const FMaterialLayersFunctionsID FMaterialLayersFunctionsRuntimeData::GetID(const FMaterialLayersFunctionsEditorOnlyData& EditorOnly) const
 {
@@ -16425,11 +14835,6 @@ void FMaterialLayersFunctions::SerializeLegacy(FArchive& Ar)
 
 void FMaterialLayersFunctions::PostSerialize(const FArchive& Ar)
 {
-	if(Substrate::IsMaterialLayeringSupportEnabled())
-	{
-		FMaterialLayersFunctionsRuntimeData::PostSerialize(Ar);
-	}
-
 #if WITH_EDITORONLY_DATA
 	if (Ar.IsLoading())
 	{
@@ -16462,10 +14867,7 @@ void FMaterialLayersFunctions::PostSerialize(const FArchive& Ar)
 			EditorOnly.DeletedParentLayerGuids = MoveTemp(DeletedParentLayerGuids_DEPRECATED);
 		}
 
-		const int32 NumLayers = Substrate::IsMaterialLayeringSupportEnabled() ? EditorOnly.LayerStates.Num() : Layers.Num();
-		// In legacy version we use Layers.Num, but in the newer version above we compare against EdtorOnly.LayerStates.Num because
-		// Layers.Num is modified in the RuntimeData legacy conversion and incorrect for the upcoming test.
-		// then further in the call CheckAndRepairPostSerializeEditorOnlyDataForRuntimeData() the editor only fileds are converted to the new tree of layers datastructure.
+		const int32 NumLayers = Layers.Num();
 		if (EditorOnly.LayerGuids.Num() != NumLayers ||
 			EditorOnly.LayerLinkStates.Num() != NumLayers)
 		{
@@ -16489,8 +14891,6 @@ void FMaterialLayersFunctions::PostSerialize(const FArchive& Ar)
 				}
 			}
 		}
-
-		CheckAndRepairPostSerializeEditorOnlyDataForRuntimeData(GetRuntime(), EditorOnly);
 	} // Ar.IsLoading()
 
 #endif // WITH_EDITORONLY_DATA	
@@ -16499,58 +14899,18 @@ void FMaterialLayersFunctions::PostSerialize(const FArchive& Ar)
 
 #if WITH_EDITOR
 
-void FMaterialLayersFunctions::CheckAndRepairPostSerializeEditorOnlyDataForRuntimeData(FMaterialLayersFunctionsRuntimeData& Runtime, FMaterialLayersFunctionsEditorOnlyData& EditorOnly)
-{
-	if (Substrate::IsMaterialLayeringSupportEnabled())
-	{
-		// This method is called right after unserialization
-		// Use this fact to fix editor only data in case of a legacy version of the LayersFunctions
-
-		// Catch case when a legacy layer has been loaded
-		// EditorOnly data need to be updated
-		if ((Runtime.Layers.Num() == Runtime.Blends.Num()) && ((EditorOnly.LayerNames.Num() + 1) == Runtime.Layers.Num()))
-		{
-			int32 NumLegacyLayers = EditorOnly.LayerNames.Num();
-			// Second add an extra Layer and Blend for each legacy layer.
-			{
-				EditorOnly.LayerStates.Add(true);
-				FText LayerName = FText::FromString(TEXT("Legacy Layers"));
-				EditorOnly.LayerNames.Add(LayerName);
-				EditorOnly.RestrictToLayerRelatives.Add(false);
-				EditorOnly.RestrictToBlendRelatives.Add(false);
-				EditorOnly.LayerGuids.Add(FGuid::NewGuid());
-				EditorOnly.LayerLinkStates.Add(EMaterialLayerLinkState::Uninitialized);
-			}
-		}
-	}
-}
-
 void FMaterialLayersFunctions::AddDefaultBackgroundLayer()
 {
-	if(Substrate::IsMaterialLayeringSupportEnabled())
-	{
-		// This call is only valid for the very first layer created from the constructor.
-		check(Layers.Num() == 0 && Blends.Num() == 0 && Tree.IsEmpty());
-		AppendLayerNode(-1, -1);
-		FText LayerName = FText(LOCTEXT("Background", "Background"));
-		EditorOnly.LayerNames[0] = LayerName;
-		// Use a consistent Guid for the background layer
-		// Default constructor assigning different guids will break FStructUtils::AttemptToFindUninitializedScriptStructMembers
-		EditorOnly.LayerGuids[0] = (BackgroundGuid);
-	}
-	else
-	{
-		// Default to a non-blended "background" layer
-		Layers.AddDefaulted();	
-		EditorOnly.LayerStates.Add(true);
-		FText LayerName = FText(LOCTEXT("Background", "Background"));
-		EditorOnly.LayerNames.Add(LayerName);
-		EditorOnly.RestrictToLayerRelatives.Add(false);
-		// Use a consistent Guid for the background layer
-		// Default constructor assigning different guids will break FStructUtils::AttemptToFindUninitializedScriptStructMembers
-		EditorOnly.LayerGuids.Add(BackgroundGuid);
-		EditorOnly.LayerLinkStates.Add(EMaterialLayerLinkState::NotFromParent);
-	}
+	// Default to a non-blended "background" layer
+	Layers.AddDefaulted();	
+	EditorOnly.LayerStates.Add(true);
+	FText LayerName = FText(LOCTEXT("Background", "Background"));
+	EditorOnly.LayerNames.Add(LayerName);
+	EditorOnly.RestrictToLayerRelatives.Add(false);
+	// Use a consistent Guid for the background layer
+	// Default constructor assigning different guids will break FStructUtils::AttemptToFindUninitializedScriptStructMembers
+	EditorOnly.LayerGuids.Add(BackgroundGuid);
+	EditorOnly.LayerLinkStates.Add(EMaterialLayerLinkState::NotFromParent);
 }
 
 int32 FMaterialLayersFunctions::AppendBlendedLayer()
@@ -16580,27 +14940,7 @@ int32 FMaterialLayersFunctions::AddLayerCopy(const FMaterialLayersFunctionsRunti
 
 	Layers.Add(Source.Layers[SourceLayerIndex]);
 
-	if(Substrate::IsMaterialLayeringSupportEnabled())
-	{
-		// Legacy parents pre substrate support
-		if (Source.Blends.Num() < Source.Layers.Num())
-		{
-			if (LayerIndex > 0)
-			{
-				Blends.Add(Source.Blends[SourceLayerIndex - 1]);
-			}
-			else
-			{
-				Blends.AddDefaulted();
-			}
-		}
-		// Source is substrate tree, same number of layersand blends
-		else
-		{
-			Blends.Add(Source.Blends[SourceLayerIndex]);
-		}
-	}
-	else if (LayerIndex > 0)
+	if (LayerIndex > 0)
 	{
 		Blends.Add(Source.Blends[SourceLayerIndex - 1]);
 	}
@@ -16609,27 +14949,7 @@ int32 FMaterialLayersFunctions::AddLayerCopy(const FMaterialLayersFunctionsRunti
 	EditorOnly.LayerNames.Add(SourceEditorOnly.LayerNames[SourceLayerIndex]);
 	EditorOnly.RestrictToLayerRelatives.Add(SourceEditorOnly.RestrictToLayerRelatives[SourceLayerIndex]);
 
-	if(Substrate::IsMaterialLayeringSupportEnabled())
-	{
-		// Legacy parents pre substrate support
-		if (SourceEditorOnly.RestrictToBlendRelatives.Num() < SourceEditorOnly.RestrictToLayerRelatives.Num())
-		{
-			if (LayerIndex > 0)
-			{
-				EditorOnly.RestrictToBlendRelatives.Add(SourceEditorOnly.RestrictToBlendRelatives[SourceLayerIndex - 1]);
-	}
-			else
-			{
-				EditorOnly.RestrictToBlendRelatives.AddDefaulted();
-			}
-		}
-		// Source is substrate tree, same number of layersand blends
-		else
-		{
-			EditorOnly.RestrictToBlendRelatives.Add(SourceEditorOnly.RestrictToBlendRelatives[SourceLayerIndex]);
-		}
-	}
-	else if (LayerIndex > 0)
+	if (LayerIndex > 0)
 	{
 		EditorOnly.RestrictToBlendRelatives.Add(SourceEditorOnly.RestrictToBlendRelatives[SourceLayerIndex - 1]);
 	}
@@ -16649,16 +14969,8 @@ void FMaterialLayersFunctions::InsertLayerCopy(const FMaterialLayersFunctionsRun
 	check(LinkState != EMaterialLayerLinkState::Uninitialized);
 	check(LayerIndex > 0);
 	Layers.Insert(Source.Layers[SourceLayerIndex], LayerIndex);
-	if(Substrate::IsMaterialLayeringSupportEnabled())
-	{
-		Blends.Insert(Source.Blends[SourceLayerIndex], LayerIndex);
-		EditorOnly.RestrictToBlendRelatives.Insert(SourceEditorOnly.RestrictToBlendRelatives[SourceLayerIndex], LayerIndex);
-	}
-	else
-	{
-		Blends.Insert(Source.Blends[SourceLayerIndex - 1], LayerIndex - 1);
-		EditorOnly.RestrictToBlendRelatives.Insert(SourceEditorOnly.RestrictToBlendRelatives[SourceLayerIndex - 1], LayerIndex - 1);
-	}
+	Blends.Insert(Source.Blends[SourceLayerIndex - 1], LayerIndex - 1);
+	EditorOnly.RestrictToBlendRelatives.Insert(SourceEditorOnly.RestrictToBlendRelatives[SourceLayerIndex - 1], LayerIndex - 1);
 	EditorOnly.LayerStates.Insert(SourceEditorOnly.LayerStates[SourceLayerIndex], LayerIndex);
 	EditorOnly.LayerNames.Insert(SourceEditorOnly.LayerNames[SourceLayerIndex], LayerIndex);
 	EditorOnly.RestrictToLayerRelatives.Insert(SourceEditorOnly.RestrictToLayerRelatives[SourceLayerIndex], LayerIndex);
@@ -16670,32 +14982,16 @@ void FMaterialLayersFunctions::RemoveBlendedLayerAt(int32 Index)
 {
 	if (Layers.IsValidIndex(Index))
 	{
-		if(Substrate::IsMaterialLayeringSupportEnabled())
-		{
-			check(Layers.IsValidIndex(Index) && Blends.IsValidIndex(Index));
-			Layers.RemoveAt(Index);
-			Blends.RemoveAt(Index);
-
-			check(EditorOnly.LayerStates.IsValidIndex(Index) &&
-				EditorOnly.LayerNames.IsValidIndex(Index) &&
-				EditorOnly.RestrictToLayerRelatives.IsValidIndex(Index) &&
-				EditorOnly.RestrictToBlendRelatives.IsValidIndex(Index));
-
-			EditorOnly.RestrictToBlendRelatives.RemoveAt(Index);
-		}
-		else
-		{
-			check(Layers.IsValidIndex(Index) && Blends.IsValidIndex(Index - 1));
-			Layers.RemoveAt(Index);
-			Blends.RemoveAt(Index - 1);
+		check(Layers.IsValidIndex(Index) && Blends.IsValidIndex(Index - 1));
+		Layers.RemoveAt(Index);
+		Blends.RemoveAt(Index - 1);
 		
-			check(EditorOnly.LayerStates.IsValidIndex(Index) &&
-				EditorOnly.LayerNames.IsValidIndex(Index) &&
-				EditorOnly.RestrictToLayerRelatives.IsValidIndex(Index) &&
-				EditorOnly.RestrictToBlendRelatives.IsValidIndex(Index - 1));
+		check(EditorOnly.LayerStates.IsValidIndex(Index) &&
+			EditorOnly.LayerNames.IsValidIndex(Index) &&
+			EditorOnly.RestrictToLayerRelatives.IsValidIndex(Index) &&
+			EditorOnly.RestrictToBlendRelatives.IsValidIndex(Index - 1));
 
-			EditorOnly.RestrictToBlendRelatives.RemoveAt(Index - 1);
-		}
+		EditorOnly.RestrictToBlendRelatives.RemoveAt(Index - 1);
 
 		if (EditorOnly.LayerLinkStates[Index] != EMaterialLayerLinkState::NotFromParent)
 		{
@@ -16720,16 +15016,8 @@ void FMaterialLayersFunctions::MoveBlendedLayer(int32 SrcLayerIndex, int32 DstLa
 	if (SrcLayerIndex != DstLayerIndex)
 	{
 		Layers.Swap(SrcLayerIndex, DstLayerIndex);
-		if (Substrate::IsMaterialLayeringSupportEnabled())
-		{
-			Blends.Swap(SrcLayerIndex, DstLayerIndex);
-			EditorOnly.RestrictToBlendRelatives.Swap(SrcLayerIndex, DstLayerIndex);
-		}
-		else
-		{
-			Blends.Swap(SrcLayerIndex - 1, DstLayerIndex - 1);
-			EditorOnly.RestrictToBlendRelatives.Swap(SrcLayerIndex - 1, DstLayerIndex - 1);
-		}
+		Blends.Swap(SrcLayerIndex - 1, DstLayerIndex - 1);
+		EditorOnly.RestrictToBlendRelatives.Swap(SrcLayerIndex - 1, DstLayerIndex - 1);
 		EditorOnly.LayerStates.Swap(SrcLayerIndex, DstLayerIndex);
 		EditorOnly.LayerNames.Swap(SrcLayerIndex, DstLayerIndex);
 		EditorOnly.RestrictToLayerRelatives.Swap(SrcLayerIndex, DstLayerIndex);
@@ -16829,7 +15117,7 @@ bool FMaterialLayersFunctions::MatchesParent(const FMaterialLayersFunctionsRunti
 			return false;
 		}
 
-		if (LayerIndex > 0 && (Substrate::IsMaterialLayeringSupportEnabled() ? Runtime.Blends[LayerIndex] == ParentRuntime.Blends[ParentLayerIndex] : Runtime.Blends[LayerIndex - 1] == ParentRuntime.Blends[ParentLayerIndex - 1]))
+		if (LayerIndex > 0 && Runtime.Blends[LayerIndex - 1] == ParentRuntime.Blends[ParentLayerIndex - 1])
 		{
 			return false;
 		}
@@ -16844,10 +15132,6 @@ bool FMaterialLayersFunctions::ResolveParent(const FMaterialLayersFunctionsRunti
 	FMaterialLayersFunctionsEditorOnlyData& EditorOnly,
 	TArray<int32>& OutRemapLayerIndices)
 {
-	// For some legacy materials using FMaterialLayersFunctions, this function is called right after unserialization of FMaterialLayersFunctionsRuntimeData
-	// and the PostLoad of FMaterialLayersFunctions is not called, so we need to make sure this is called and checked once
-	CheckAndRepairPostSerializeEditorOnlyDataForRuntimeData(Runtime, EditorOnly);
-
  	check(EditorOnly.LayerGuids.Num() == Runtime.Layers.Num());
 	check(EditorOnly.LayerLinkStates.Num() == Runtime.Layers.Num());
 
@@ -16904,7 +15188,7 @@ bool FMaterialLayersFunctions::ResolveParent(const FMaterialLayersFunctionsRunti
 			{
 				// See if we match layer in parent
 				if (Runtime.Layers[LayerIndex] == ParentRuntime.Layers[ParentLayerIndex]
-					&& (LayerIndex == 0 || (Substrate::IsMaterialLayeringSupportEnabled() ? Runtime.Blends[LayerIndex] == ParentRuntime.Blends[ParentLayerIndex] : Runtime.Blends[LayerIndex - 1] == ParentRuntime.Blends[ParentLayerIndex - 1])))
+					&& (LayerIndex == 0 || Runtime.Blends[LayerIndex - 1] == ParentRuntime.Blends[ParentLayerIndex - 1]))
 				{
 					// Parent layer matches, so link to parent
 					ResolvedLayerIndex = ResolvedLayers.AddLayerCopy(ParentRuntime, ParentEditorOnly, ParentLayerIndex, bLayerVisible, EMaterialLayerLinkState::LinkedToParent);
@@ -17024,36 +15308,6 @@ bool FMaterialLayersFunctions::ResolveParent(const FMaterialLayersFunctionsRunti
 		}
 	}
 
-	if(Substrate::IsMaterialLayeringSupportEnabled())
-	{
-		struct Traverser
-		{
-			void Recursive(const FMaterialLayersFunctionsTree& Source, FLayerNodeId InSourceId, FMaterialLayersFunctionsTree& Dest, FLayerNodeId InDestParentId)
-			{
-
-				FLayerNode* SourceNode = Source.GetNode(InSourceId);
-				FLayerNodeId NewNodeId = FMaterialLayersFunctionsTree::InvalidId;
-				FLayerNodePayload Payload;
-				if (SourceNode)
-				{
-					Payload = Source.Payloads[InSourceId];
-					NewNodeId = Dest.AddNode(Payload, InDestParentId);
-				}
-
-				auto ChildrenIds = Source.GetChildrenIds(InSourceId);
-				for (auto SourceSubId : ChildrenIds)
-				{
-					Recursive(Source, SourceSubId, Dest, NewNodeId);
-				}
-
-			}
-		} ReplicateTree;
-	
-		ReplicateTree.Recursive(Runtime.Tree, -1, ResolvedLayers.Tree, -1);
-
-		// NOTE: For debug purpose at the moment, we want to monitor the end result tree data state.
-		MLFT::Log(&Runtime, __FUNCTION__);
-	}
 
 	Runtime = MoveTemp(static_cast<FMaterialLayersFunctionsRuntimeData&>(ResolvedLayers));
 	EditorOnly = MoveTemp(ResolvedLayers.EditorOnly);
@@ -17066,7 +15320,7 @@ void FMaterialLayersFunctions::Validate(const FMaterialLayersFunctionsRuntimeDat
 {
 	if (Runtime.Layers.Num() > 0)
 	{
-		check(Substrate::IsMaterialLayeringSupportEnabled() || Runtime.Blends.Num() == Runtime.Layers.Num() - 1);
+		check(Runtime.Blends.Num() == Runtime.Layers.Num() - 1);
 		check(Runtime.Layers.Num() == EditorOnly.LayerStates.Num());
 		check(Runtime.Layers.Num() == EditorOnly.LayerNames.Num());
 		check(Runtime.Layers.Num() == EditorOnly.LayerGuids.Num());
@@ -17074,191 +15328,6 @@ void FMaterialLayersFunctions::Validate(const FMaterialLayersFunctionsRuntimeDat
 	}
 }
 
-FLayerNodeId FMaterialLayersFunctions::GetNodeParent(FLayerNodeId InNodeId) const
-{
-	FLayerNode* Node = Tree.GetNode(InNodeId);
-	if (Node)
-		return Node->Parent;
-
-	return FMaterialLayersFunctionsTree::InvalidId;
-}
-
-FLayerNodeIdArray FMaterialLayersFunctions::GetNodeParents(FLayerNodeId InNode) const
-{
-	return Tree.GetParentIds(InNode);
-}
-
-FMaterialLayersFunctions::FLayerNodeIdArray FMaterialLayersFunctions::GetNodeChildren(FLayerNodeId InNodeId) const
-{
-	return Tree.GetChildrenIds(InNodeId);
-}
-
-FMaterialLayersFunctions::FLayerNodePayload FMaterialLayersFunctions::GetNodePayload(FLayerNodeId InNodeId) const
-{
-	FLayerNodePayload* Payload = Tree.GetPayload(InNodeId);
-	if (Payload)
-		return *Payload;
-
-	return {};
-}
-
-int32 FMaterialLayersFunctions::GetLayerFuncIndex(FLayerNodeId InNodeId) const
-{
-	return GetNodePayload(InNodeId).Layer;
-}
-
-int32 FMaterialLayersFunctions::GetBlendFuncIndex(FLayerNodeId InNodeId) const
-{
-	return GetNodePayload(InNodeId).Blend;
-}
-
-int32 FMaterialLayersFunctions::GetNodeDepth(FLayerNodeId InNodeId) const
-{
-	return Tree.GetDepth(InNodeId);
-}
-
-bool FMaterialLayersFunctions::CanAppendLayerNode(FLayerNodeId InParent) const
-{
-	// NOTE: Current design only allows for 2 levels depth in the tree
-	return  !(GetNodeDepth(InParent) > 1);
-}
-
-FLayerNodeId FMaterialLayersFunctions::AppendLayerNode(FLayerNodeId InParent, int32 InSiblingIndex)
-{
-	if(Substrate::IsMaterialLayeringSupportEnabled())
-	{
-		if (!CanAppendLayerNode(InParent))
-			return FMaterialLayersFunctionsTree::InvalidId;
-
-		const int32 PayloadIndex = AppendBlendedLayer();
-
-		// InPayload
-		FLayerNodeId NewNodeId = Tree.AddNode({ PayloadIndex, PayloadIndex }, InParent, InSiblingIndex);
-
-		MLFT::Log(this, __FUNCTION__);
-
-		return NewNodeId;
-	}
-	return FMaterialLayersFunctionsTree::InvalidId;
-}
-
-bool FMaterialLayersFunctions::CanRemoveLayerNode(FLayerNodeId InNodeId) const
-{
-	// Only if really valid
-	if (Tree.IsValidId(InNodeId))
-	{
-		// Must keep the very first Slab / node
-		return Tree.Root != InNodeId;
-	}
-	return false;
-}
-
-void FMaterialLayersFunctions::RemoveLayerNodeAt(FLayerNodeId InNodeId)
-{
-	if (!Substrate::IsMaterialLayeringSupportEnabled() || !CanRemoveLayerNode(InNodeId))
-		return;
-
-	if (Tree.IsValidId(InNodeId))
-	{
-		auto [RemovedNodes, RemovedPayloads] = Tree.RemoveNode(InNodeId);
-
-		// The node and its branch have been removed from the tree
-		// Now remove the actual asset from Layers and Blends functions at the corresponding indices (stored in the payload)
-		
-		// Extract the actual asset indices from the returned removed payload 
-		TArray<int32> RemovedAssetIndices;
-		for (auto& Payload : RemovedPayloads)
-		{
-			if (Payload.Layer >= 0 && Payload.Layer < Layers.Num())
-			{
-				RemovedAssetIndices.Add(Payload.Layer);
-			}
-		}
-
-		// THen sort the indices in growing order
-		RemovedAssetIndices.Sort();
-
-		// And now from back to front
-		// REMOVE the actual asset in the Layers/Blends arrays
-		// AND adjust the stored indices in the tree node payload
-		for (int i = 0; i < RemovedAssetIndices.Num(); ++i)
-		{
-			int AssetIndex = RemovedAssetIndices[RemovedAssetIndices.Num() - 1 - i];
-
-			RemoveBlendedLayerAt(AssetIndex);
-
-			for (auto& Payload : Tree.Payloads)
-			{
-				if (Payload.Layer > AssetIndex)
-				{
-					Payload.Layer--;
-				}
-				if (Payload.Blend > AssetIndex)
-				{
-					Payload.Blend--;
-				}
-			}
-		}	
-	}
-
-	MLFT::Log(this, __FUNCTION__);
-}
-	
-
-
-void FMaterialLayersFunctions::MoveLayerNode(FLayerNodeId InNodeId, FLayerNodeId DstParentId, int32 InSiblingIndex, bool bShouldDuplicate)
-{
-	if (bShouldDuplicate)
-	{
-		// create new layer function and blend function copies
-		FLayerNodeId NewNode = AppendLayerNode(DstParentId, InSiblingIndex);
-		FLayerNodePayload SourcePayload = GetNodePayload(InNodeId);
-		FLayerNodePayload DestPayload = GetNodePayload(NewNode);
-		Layers[DestPayload.Layer] = DuplicateObject<UMaterialFunctionInterface>(Layers[SourcePayload.Layer], nullptr);
-		Blends[DestPayload.Blend] = DuplicateObject<UMaterialFunctionInterface>(Blends[SourcePayload.Blend], nullptr);
-		EditorOnly.LayerNames[DestPayload.Layer] = EditorOnly.LayerNames[SourcePayload.Layer];
-	}
-	else
-	{
-		Tree.MoveNode(InNodeId, DstParentId, InSiblingIndex);
-	}
-}
-
-UMaterial* FMaterialLayersFunctions::GetRuntimeNodePreviewMaterial(FLayerNodeId InNodeId) const
-{
-	if (Substrate::IsMaterialLayeringSupportEnabled() && RuntimeGraphCache)
-	{
-		if (Tree.IsValidId(InNodeId))
-		{
-			if (!RuntimeGraphCache->NodePreviewMaterials[InNodeId])
-			{
-				if (RuntimeGraphCache->NodeMaterialGraphExpressions[InNodeId])
-				{
-					UMaterialExpressionMaterialFunctionCall* NodeCall = RuntimeGraphCache->NodeMaterialGraphExpressions[InNodeId];
-
-					UMaterial* PreviewMaterial = NewObject<UMaterial>(NodeCall, NAME_None, RF_Transient | RF_Public);
-					PreviewMaterial->bIsPreviewMaterial = true;
-
-					//Update cached expression data to ensure function calls are populated for resolving the preview
-					PreviewMaterial->AssignExpressionCollection(RuntimeGraphCache->ExpressionCollection);
-					PreviewMaterial->UpdateCachedExpressionData();
-
-					NodeCall->ConnectToPreviewMaterial(PreviewMaterial, 0);
-
-					//Compile the material.
-					PreviewMaterial->PreEditChange(nullptr);
-					PreviewMaterial->PostEditChange();
-
-					RuntimeGraphCache->NodePreviewMaterials[InNodeId] = PreviewMaterial;
-				}
-			}
-
-			return RuntimeGraphCache->NodePreviewMaterials[InNodeId];	
-		}
-	}
-	return nullptr;
-}
-
 #endif // WITH_EDITOR
 
 ///////////////////////////////////////////////////////////////////////////////
diff --git a/Engine/Source/Runtime/Engine/Private/Materials/MaterialInstance.cpp b/Engine/Source/Runtime/Engine/Private/Materials/MaterialInstance.cpp
--- a/Engine/Source/Runtime/Engine/Private/Materials/MaterialInstance.cpp
+++ b/Engine/Source/Runtime/Engine/Private/Materials/MaterialInstance.cpp
@@ -298,7 +298,7 @@ bool FMaterialInstanceResource::GetParameterValue(EMaterialParameterType Type, c
 	// Check for hard-coded parameters
 	if (Type == EMaterialParameterType::Scalar && ParameterInfo.Name == SubsurfaceProfile::GetSubsurfaceProfileParameterName())
 	{
-		check(Substrate::IsMaterialLayeringSupportEnabled() || ParameterInfo.Association == EMaterialParameterAssociation::GlobalParameter);
+		check(ParameterInfo.Association == EMaterialParameterAssociation::GlobalParameter);
 
 		const USubsurfaceProfile* MySubsurfaceProfileRT = GetSubsurfaceProfileRT();
 		OutValue = SubsurfaceProfile::GetSubsurfaceProfileId(MySubsurfaceProfileRT);
@@ -313,7 +313,7 @@ bool FMaterialInstanceResource::GetParameterValue(EMaterialParameterType Type, c
 			const USubsurfaceProfile* SSProfileRT = GetSubsurfaceProfileRT(It);
 			if (ParameterInfo.Name == SubsurfaceProfile::CreateSubsurfaceProfileParameterName(SSProfileRT))
 			{
-				check(Substrate::IsMaterialLayeringSupportEnabled() || ParameterInfo.Association == EMaterialParameterAssociation::GlobalParameter);
+				check(ParameterInfo.Association == EMaterialParameterAssociation::GlobalParameter);
 
 				// Set the root material Profile, or the profile overriden by any instances.
 				OutValue = SubsurfaceProfile::GetSubsurfaceProfileId(SSProfileOverrideRT ? SSProfileOverrideRT : SSProfileRT);
@@ -329,7 +329,7 @@ bool FMaterialInstanceResource::GetParameterValue(EMaterialParameterType Type, c
 		{
 			if (ParameterInfo.Name == SpecularProfile::GetSpecularProfileParameterName(GetSpecularProfileRT(It)))
 			{
-				check(Substrate::IsMaterialLayeringSupportEnabled() || ParameterInfo.Association == EMaterialParameterAssociation::GlobalParameter);
+				check(ParameterInfo.Association == EMaterialParameterAssociation::GlobalParameter);
 				OutValue = SpecularProfile::GetSpecularProfileId(SPOverrideRT ? SPOverrideRT : GetSpecularProfileRT(It));
 				bResult = true;
 				break;
diff --git a/Engine/Source/Runtime/Engine/Public/Materials/MaterialExpressionBlendMaterialAttributes.h b/Engine/Source/Runtime/Engine/Public/Materials/MaterialExpressionBlendMaterialAttributes.h
--- a/Engine/Source/Runtime/Engine/Public/Materials/MaterialExpressionBlendMaterialAttributes.h
+++ b/Engine/Source/Runtime/Engine/Public/Materials/MaterialExpressionBlendMaterialAttributes.h
@@ -70,7 +70,8 @@ class UMaterialExpressionBlendMaterialAttributes : public UMaterialExpression
 };
 
 
-UCLASS(collapsecategories, hidecategories = Object, MinimalAPI)
+/** Deprecated blend expression used by substrate material layers prototype. */
+UCLASS(collapsecategories, hidecategories = Object, MinimalAPI, meta = (Private))
 class UMaterialExpressionLegacyBlendMaterialAttributes : public UMaterialExpressionBlendMaterialAttributes
 {
 	GENERATED_UCLASS_BODY()
@@ -91,14 +92,4 @@ public:
 
 	UPROPERTY(EditAnywhere, Category = MaterialAttributes)
 	TEnumAsByte<EMaterialAttributeBlendFunction::Type> BlendFunctionType;
-
-#if WITH_EDITOR
-	virtual void GetCaption(TArray<FString>& OutCaptions) const override;
-	virtual bool CanEditChange(const FProperty* InProperty) const override;
-	virtual FExpressionInput* GetInput(int32 InputIndex) override;
-	virtual FName GetInputName(int32 InputIndex) const override;
-	virtual bool IsInputConnectionRequired(int32 InputIndex) const override;
-	virtual EMaterialValueType GetInputValueType(int32 InputIndex) override;
-	virtual int32 Compile(class FMaterialCompiler* Compiler, int32 OutputIndex) override;
-#endif // WITH_EDITOR
-};
\ No newline at end of file
+};
diff --git a/Engine/Source/Runtime/Engine/Public/Materials/MaterialExpressionMaterialAttributeLayers.h b/Engine/Source/Runtime/Engine/Public/Materials/MaterialExpressionMaterialAttributeLayers.h
--- a/Engine/Source/Runtime/Engine/Public/Materials/MaterialExpressionMaterialAttributeLayers.h
+++ b/Engine/Source/Runtime/Engine/Public/Materials/MaterialExpressionMaterialAttributeLayers.h
@@ -121,8 +121,7 @@ class UMaterialExpressionMaterialAttributeLayers : public UMaterialExpression
 	virtual bool IsInputConnectionRequired(int32 InputIndex) const override {return false;}
 	virtual EMaterialValueType GetInputValueType(int32 InputIndex) override;
 	virtual bool IsResultMaterialAttributes(int32 OutputIndex) override {return true;}
-	virtual bool IsResultSubstrateMaterial(int32 OutputIndex) override;
-	virtual FSubstrateOperator* SubstrateGenerateMaterialTopologyTree(class FMaterialCompiler* Compiler, class UMaterialExpression* Parent, int32 OutputIndex) override;
+	virtual bool IsResultSubstrateMaterial(int32 OutputIndex) override {return false;}
 #endif
 	//~ End UMaterialExpression Interface
 
diff --git a/Engine/Source/Runtime/Engine/Public/Materials/MaterialFunction.h b/Engine/Source/Runtime/Engine/Public/Materials/MaterialFunction.h
--- a/Engine/Source/Runtime/Engine/Public/Materials/MaterialFunction.h
+++ b/Engine/Source/Runtime/Engine/Public/Materials/MaterialFunction.h
@@ -181,9 +181,6 @@ public:
 	/** Returns If returns an empty string, use the default class name for the material function. Otherwise, the string will be the name shown when the function is exposed to users in the material graph as a node, or from the contextual menu when searching for nodes. */
 	virtual FString GetUserExposedCaption() const override { return UserExposedCaption; }
 
-	/** Converts legacy Material Layer Functions to utilise Substrate */
-	void ConvertExpressionsBetweenLegacyAndSubstrate();
-
 	/** Retrieves an array of all custom output expressions in a Material Function */
 	void GetAllCustomOutputExpressions(TArray<class UMaterialExpressionCustomOutput*>& OutCustomOutputs) const;
 
diff --git a/Engine/Source/Runtime/Engine/Public/Materials/MaterialLayersFunctions.h b/Engine/Source/Runtime/Engine/Public/Materials/MaterialLayersFunctions.h
--- a/Engine/Source/Runtime/Engine/Public/Materials/MaterialLayersFunctions.h
+++ b/Engine/Source/Runtime/Engine/Public/Materials/MaterialLayersFunctions.h
@@ -120,264 +120,6 @@ struct FMaterialLayersFunctionsEditorOnlyData
 };
 
 
-
-USTRUCT()
-struct FMaterialLayersFunctionsTree
-{
-	GENERATED_BODY()
-
-public:
-	using FNodeId = int32;
-	using FPayloadId = int32;
-
-	static const FNodeId InvalidId = -1; // Invalid id is always <0
-
-	struct FNode
-	{
-		FNodeId 		Parent = InvalidId;
-		FNodeId 		NextSibling = InvalidId;
-		FNodeId 		ChildrenHead = InvalidId;
-		FNodeId			Spare = InvalidId;
-
-		inline bool operator==(const FNode& Other) const
-		{
-			if (Parent != Other.Parent || NextSibling != Other.NextSibling || ChildrenHead != Other.ChildrenHead)
-			{
-				return false;
-			}
-			return true;
-		}
-
-		inline bool operator!=(const FNode& Other) const
-		{
-			return !operator==(Other);
-		}
-
-		void Clear()
-		{
-			Parent = NextSibling = Spare = InvalidId;
-		}
-	};
-
-	struct FPayload
-	{
-		FPayloadId Layer = InvalidId;
-		FPayloadId Blend = InvalidId;
-
-		inline bool operator==(const FPayload& Other) const
-		{
-			if (Layer != Other.Layer || Blend != Other.Blend)
-			{
-				return false;
-			}
-			return true;
-		}
-
-		inline bool operator!=(const FPayload& Other) const
-		{
-			return !operator==(Other);
-		}
-	};
-
-	TArray<FNode>		Nodes;
-	TArray<FPayload>	Payloads;
-	FNodeId				Root = -1; // the index of the head of the top level siblings
-							   	
-	FMaterialLayersFunctionsTree() = default;
-	FMaterialLayersFunctionsTree(const FMaterialLayersFunctionsTree& Rhs)
-		:	Nodes(Rhs.Nodes),
-			Payloads(Rhs.Payloads),
-			Root(Rhs.Root)
-			 
-	{}
-	FMaterialLayersFunctionsTree(const FMaterialLayersFunctions& Rhs) = delete;
-
-	FMaterialLayersFunctionsTree& operator=(const FMaterialLayersFunctionsTree& Rhs)
-	{
-		Nodes = Rhs.Nodes;
-		Payloads = Rhs.Payloads;
-		Root = Rhs.Root;
-
-		return *this;
-	}
-
-	inline bool operator==(const FMaterialLayersFunctionsTree& Other) const
-	{
-		if (Nodes != Other.Nodes || Payloads != Other.Payloads || Root != Other.Root)
-		{
-			return false;
-		}
-		return true;
-	}
-
-	inline bool operator!=(const FMaterialLayersFunctionsTree& Other) const
-	{
-		return !operator==(Other);
-	}
-
-	void Empty()
-	{
-		Nodes.Empty();
-		Payloads.Empty();
-		Root = -1;
-	}
-
-	inline bool IsEmpty() const
-	{
-		return Nodes.IsEmpty();
-	}
-
-	bool Serialize( FArchive& Ar );
-
-	friend FArchive& operator<<(FArchive& Ar, FNode& T)
-	{
-		return Ar << T.Parent << T.NextSibling << T.ChildrenHead << T.Spare;
-	}
-	friend FArchive& operator<<(FArchive& Ar, FPayload& T)
-	{
-		return Ar << T.Layer << T.Blend;
-	}
-
-	/// Authoring the Tree
-
-	inline bool IsValidId(FNodeId Id) const
-	{
-		return (Id >= 0 && Id < Nodes.Num());
-	}
-
-	// Find and return the node at the specified Id or null if invalid 
-	FNode* GetNode(FNodeId InNodeId) const;
-
-	// Find and return the node payload at the specified Id or null if invalid 
-	FPayload* GetPayload(FNodeId InNodeId) const;
-
-	// Find and return the parent node at the specified Id or null if invalid 
-	FNode* GetParent(FNodeId InNodeId) const;
-
-	// Find the node which is the head of the list of siblings where InNodeID belongs
-	FNodeId GetSiblingHeadId(FNodeId InNodeId) const;
-
-	// Find the node which is the previous sibling to the specified node
-	// The Next sibling is directly available in the FNode struct, not the previous one
-	FNodeId GetPreviousSiblingId(FNodeId InNodeId) const;
-
-	// Gather the list of children for this node
-	TArray<FNodeId> GetChildrenIds(FNodeId InNodeId) const;
-
-	// Find the last of the children in the specified node
-	FNodeId GetChildrenTailId(FNodeId InNodeId) const;
-
-	// Return the depth of the node in the tree, the root is depth 0
-	// the first valid node in the tree under the root is depth 1
-	int32 GetDepth(FNodeId InNodeId) const;
-
-	// Gather the path of parent node up to the root from the specified node.
-	// First element is the direct parent, second is parent of parent, etc up to the root of the tree, so the last element is always -1.
-	// THe size of the returned array is equivalent to the depth.
-	// Thus if the specified node's parent is root the return array is 1 element containing -1.
-	// If node is invalid, the returned array is empty
-	TArray<FNodeId> GetParentIds(FNodeId InNodeId) const;
-
-	// Add a new node with the specified Payload under the InPArent at the specified child index
-	// the new node id is returned
-	FNodeId AddNode(const FPayload& InPayload, FNodeId InParent, int32 InAtChildIndex = -1);
-
-	// Remove this node
-	// detach the node from the tree hierarchy as well as all the children recursively
-	// the nodes are all cleared in the array of nodes
-	// the payloads are NOT cleared and still hold on their value
-	// return the array of Id of the removed nodes AND the array of Payloads removed
-	std::pair<TArray<FNodeId>, TArray<FPayload>> RemoveNode(FNodeId RemovedNodeId);
-	
-	// Move node from one position to another
-	void MoveNode(FNodeId InNodeId, FNodeId DstParentId, int32 InSiblingIndex);
-
-	// Traverser Utility
-	using FNodeVisitor = void (*)(FNodeId InNodeId, int32 InDepth, int32 InSiblingNum, const FPayload& InPayload);
-	template <typename V>
-	FNodeId TraverseNode(V InVisitor, FNodeId InNodeId, int32 InDepth, int32 InSiblingNum, bool ReverseOrder = false) const
-	{
-		// First plan for the next child node and sibling we will be visiting
-		// default with the InNodeId is the root case
-		FNodeId NextChildId = Root;	
-		FNodeId NextSiblingId = InvalidId;
-
-		// Traverse this node if not the root tree
-		if (InNodeId != InvalidId)
-		{
-			if (!IsValidId(InNodeId))
-				return InvalidId;
-
-			const FNode& Node = Nodes[InNodeId];
-			const FPayload& Payload = Payloads[InNodeId];
-			InVisitor(InNodeId, InDepth, InSiblingNum, Payload);
-
-			// And since we have visited a valid node, lets visit the children next and the next sibling after
-			NextChildId = Node.ChildrenHead;
-			NextSiblingId = Node.NextSibling;
-			if (ReverseOrder)
-			{
-				NextChildId = GetChildrenTailId(InNodeId);
-				NextSiblingId = GetPreviousSiblingId(InNodeId);
-			}
-		}
-		else
-		{
-			if (ReverseOrder)
-			{
-				NextChildId = GetChildrenTailId(InvalidId);
-			}
-		}
-
-		// Traverse children
-		{
-			InDepth++;
-			int32 SiblingNum = 0;
-			while (IsValidId(NextChildId))
-			{
-				NextChildId = TraverseNode(InVisitor, NextChildId, InDepth, SiblingNum, ReverseOrder);
-				SiblingNum++;
-			}
-			InDepth--;
-		}
-
-		return NextSiblingId;
-	}
-
-	template <typename V>
-	void Traverse(V InVisitor, FNodeId InRootNodeId = InvalidId) const
-	{
-		TraverseNode(InVisitor, InRootNodeId, 0, 0);
-	}
-	template <typename V>
-	void TraverseBottomUp(V InVisitor, FNodeId InRootNodeId = InvalidId) const
-	{
-		TraverseNode(InVisitor, InRootNodeId, 0, 0, true);
-	}
-
-	// Logging the tree data structure, useful for debugging
-	// generate a multiline string representing the tree nodes hierarchy with the payload info per node
-	FString Log(FString InTab) const;
-
-private:
-	FNodeId AllocNode(const FPayload& InPayload);
-
-	// Remove iteratively a full branch of nodes
-	int RemoveNodeBranch(FNode* BranchParentNode, TArray<FNodeId>& RemovedIds, TArray<FPayload>& RemovedPayloads);
-};
-
-template<>
-struct TStructOpsTypeTraits<FMaterialLayersFunctionsTree>
-: public TStructOpsTypeTraitsBase2<FMaterialLayersFunctionsTree>
-{
-	enum
-	{
-		WithSerializer = true,
-		WithCopy = true,
-		WithIdenticalViaEquality = true
-	};
-};
-
 USTRUCT() 
 struct FMaterialLayersFunctionsRuntimeGraphCache
 {
@@ -393,9 +135,6 @@ struct FMaterialLayersFunctionsRuntimeGraphCache
 	TArray<TObjectPtr<UMaterialExpressionMaterialFunctionCall>> NodeMaterialGraphExpressions;
 	TArray<TObjectPtr<UMaterialExpressionMaterialFunctionCall>> NodeMaterialGraphExpressionsBlends;
 	TArray<TObjectPtr<UMaterial>>								NodePreviewMaterials;
-
-	ENGINE_API void GetNodeIndices(FGuid InExpressionGuid, const TArray<TObjectPtr<UMaterialExpressionMaterialFunctionCall>>& ExpressionArrayToUse, TArray<int32>& OutIndices) const;
-	ENGINE_API int32 FindExpressionIndex(FGuid InExpressionGuid, const FMaterialLayersFunctions* Layersfunctions) const;
 #endif
 };
 
@@ -417,10 +156,6 @@ struct FMaterialLayersFunctionsRuntimeData
 	UPROPERTY(EditAnywhere, Category = MaterialLayers)
 	TArray<TObjectPtr<class UMaterialFunctionInterface>> Blends;
 
-	// UProperty describing the tree data structure
-	UPROPERTY()
-	FMaterialLayersFunctionsTree Tree;
-
 	mutable TSharedPtr<FMaterialLayersFunctionsRuntimeGraphCache> RuntimeGraphCache;
 
 	// In multiple locations there are const references to this runtime data/it's children, 
@@ -431,7 +166,6 @@ struct FMaterialLayersFunctionsRuntimeData
 	FMaterialLayersFunctionsRuntimeData(const FMaterialLayersFunctionsRuntimeData& Rhs)
 		: Layers(Rhs.Layers)
 		, Blends(Rhs.Blends)
-		, Tree(Rhs.Tree)
 		, RuntimeGraphCache(Rhs.RuntimeGraphCache)
 		, LayerStackCache(Rhs.LayerStackCache)
 	{}
@@ -442,7 +176,6 @@ struct FMaterialLayersFunctionsRuntimeData
 	{
 		Layers = Rhs.Layers;
 		Blends = Rhs.Blends;
-		Tree = Rhs.Tree;
 		RuntimeGraphCache = Rhs.RuntimeGraphCache;
 		LayerStackCache = Rhs.LayerStackCache;
 		return *this;
@@ -456,16 +189,13 @@ struct FMaterialLayersFunctionsRuntimeData
 	{
 		Layers.Empty();
 		Blends.Empty();
-		Tree.Empty();
 		RuntimeGraphCache.Reset();
 		LayerStackCache.Reset();
 	}
 
 	inline bool operator==(const FMaterialLayersFunctionsRuntimeData& Other) const
 	{
-		if (Layers != Other.Layers || Blends != Other.Blends
-			|| Tree != Other.Tree
-		)
+		if (Layers != Other.Layers || Blends != Other.Blends)
 		{
 			return false;
 		}
@@ -489,8 +219,6 @@ struct FMaterialLayersFunctionsRuntimeData
 
 	bool SerializeFromMismatchedTag(const FPropertyTag& Tag, FStructuredArchive::FSlot Slot);
 	
-	ENGINE_API void PostSerialize(const FArchive& Ar);
-
 #if WITH_EDITOR
 	const FMaterialLayersFunctionsID GetID(const FMaterialLayersFunctionsEditorOnlyData& EditorOnly) const;
 #endif // WITH_EDITOR
@@ -510,7 +238,6 @@ template<>
 struct TStructOpsTypeTraits<FMaterialLayersFunctionsRuntimeData> : TStructOpsTypeTraitsBase2<FMaterialLayersFunctionsRuntimeData>
 {
 	enum { WithStructuredSerializeFromMismatchedTag = true };
-	enum { WithPostSerialize = true };
 };
 
 USTRUCT()
@@ -651,7 +378,6 @@ struct FMaterialLayersFunctions : public FMaterialLayersFunctionsRuntimeData
 		return FMaterialLayersFunctions::ResolveParent(Parent.GetRuntime(), Parent.EditorOnly, GetRuntime(), EditorOnly, OutRemapLayerIndices);
 	}
 
-
 	static ENGINE_API void Validate(const FMaterialLayersFunctionsRuntimeData& Runtime, const FMaterialLayersFunctionsEditorOnlyData& EditorOnly);
 
 	void Validate()
@@ -661,8 +387,6 @@ struct FMaterialLayersFunctions : public FMaterialLayersFunctionsRuntimeData
 
 	ENGINE_API void SerializeLegacy(FArchive& Ar);
 
-	static ENGINE_API void CheckAndRepairPostSerializeEditorOnlyDataForRuntimeData(FMaterialLayersFunctionsRuntimeData& Runtime, FMaterialLayersFunctionsEditorOnlyData& EditorOnly);
-
 #endif // WITH_EDITOR
 
 	ENGINE_API void PostSerialize(const FArchive& Ar);
@@ -687,37 +411,6 @@ struct FMaterialLayersFunctions : public FMaterialLayersFunctionsRuntimeData
 		return !operator==(Other);
 	}
 
-	// Define the Tree authoring of the layers
-#if WITH_EDITOR
-	using FLayerNodeId = FMaterialLayersFunctionsTree::FNodeId;
-	using FLayerNodeIdArray = TArray<FLayerNodeId>;
-	using FLayerNodePayload = FMaterialLayersFunctionsTree::FPayload; 
-
-	ENGINE_API FLayerNodeId GetNodeParent(FLayerNodeId InNode) const;
-	ENGINE_API int32 GetNodeDepth(FLayerNodeId InNodeId) const;
-	// generate the path from the ROOT to the specified node
-	ENGINE_API FLayerNodeIdArray GetNodeParents(FLayerNodeId InNode) const;
-
-	ENGINE_API FLayerNodeIdArray GetNodeChildren(FLayerNodeId InNode) const;
-
-	ENGINE_API FLayerNodePayload GetNodePayload(FLayerNodeId InNodeId) const;
-	ENGINE_API int32 GetLayerFuncIndex(FLayerNodeId InNodeId) const;
-	ENGINE_API int32 GetBlendFuncIndex(FLayerNodeId InNodeId) const;
-
-	ENGINE_API bool	 CanAppendLayerNode(FLayerNodeId InParent) const;
-	ENGINE_API FLayerNodeId	AppendLayerNode(FLayerNodeId InParent, int32 InSiblingIndex = -1);
-
-	ENGINE_API bool	 CanRemoveLayerNode(FLayerNodeId InNodeId) const;
-	ENGINE_API void	 RemoveLayerNodeAt(FLayerNodeId InNodeId);
-
-	ENGINE_API void MoveLayerNode(FLayerNodeId InNodeId, FLayerNodeId DstParentId, int32 InSiblingIndex, bool bShouldDuplicate);
-
-	// Cached preview material can be generated directly from the embedded runtime graph cache struct
-	// Return such a material for the specific node if available, nullptr otherwise
-	ENGINE_API UMaterial* GetRuntimeNodePreviewMaterial(FLayerNodeId InNodeId) const;
-
-#endif // WITH_EDITOR
-
 private:
 	UPROPERTY()
 	TArray<bool> LayerStates_DEPRECATED;
diff --git a/Engine/Source/Runtime/RenderCore/Private/RenderUtils.cpp b/Engine/Source/Runtime/RenderCore/Private/RenderUtils.cpp
--- a/Engine/Source/Runtime/RenderCore/Private/RenderUtils.cpp
+++ b/Engine/Source/Runtime/RenderCore/Private/RenderUtils.cpp
@@ -2273,7 +2273,7 @@ namespace Substrate
 
 	bool IsMaterialLayeringSupportEnabled()
 	{
-		return IsSubstrateEnabled() && CVarSubstrateEnableLayerSupport.GetValueOnAnyThread() > 0;
+		return false;
 	}
 
 	bool IsHiddenMaterialAssetConversionEnabled()
