FSpline: Only reparameterize during update, not while setting point location

diff --git a/Engine/Source/Runtime/Engine/Private/Curves/Spline.cpp b/Engine/Source/Runtime/Engine/Private/Curves/Spline.cpp
--- a/Engine/Source/Runtime/Engine/Private/Curves/Spline.cpp
+++ b/Engine/Source/Runtime/Engine/Private/Curves/Spline.cpp
@@ -2208,8 +2208,6 @@ void FNewSpline::SetLocation(int32 Index, const FVector& InLocation)
 	
 	GetSpline().ModifyPoint(Index, {InLocation, GetSpline().GetTangentIn(Index), GetSpline().GetTangentOut(Index), GetSpline().GetTangentMode(Index)});
 
-	// We do this so that parameterization remains proportional to the square root of segment chord lengths.
-	GetSpline().Reparameterize(EParameterizationPolicy::Centripetal);
 	
 	MarkReparamTableDirty();
     MarkLegacyCurvesDirty();
@@ -2436,9 +2434,17 @@ FVector FNewSpline::EvaluateScale(float Parameter) const
 
 void FNewSpline::UpdateSpline()
 {
-    GetSpline().UpdateTangents();			// updates the tangents on our version of the points
+	// Ensure all tangents are up to date w.r.t. point types
+    GetSpline().UpdateTangents();			
+	
+	// We do this so that parameterization remains proportional to the square root of segment chord lengths.
+	GetSpline().Reparameterize(EParameterizationPolicy::Centripetal);
+
+	// Mark reparam table dirty, the table is not rebuild until necessary.
     MarkReparamTableDirty();
-    RebuildLegacyCurves();					// updates legacy curves based on points, never evaluates internal spline.
+
+	// Update the legacy interface. This cannot be deferred because users expect const references to be up-to-date after calling UpdateSpline.
+    RebuildLegacyCurves();
 }
 
 void FNewSpline::UpdateSpline(const FSpline::FUpdateSplineParams& Params)
