Iris
* Improved NetActorFactory CreationHeader classes so they can be derived from more easily.
* Made an header type for pre-registered actors

#rnx

diff --git a/Engine/Source/Runtime/Engine/Private/Net/Iris/ReplicationSystem/NetActorFactory.cpp b/Engine/Source/Runtime/Engine/Private/Net/Iris/ReplicationSystem/NetActorFactory.cpp
--- a/Engine/Source/Runtime/Engine/Private/Net/Iris/ReplicationSystem/NetActorFactory.cpp
+++ b/Engine/Source/Runtime/Engine/Private/Net/Iris/ReplicationSystem/NetActorFactory.cpp
@@ -44,16 +44,6 @@ namespace UE::Net::Private
 
 	static FName NetActorFactoryName(TEXT("NetActorFactory"));
 
-	enum class EActorNetSpawnInfoFlags : uint32
-	{
-		None = 0U,
-		QuantizeScale = 1U,
-		QuantizeLocation = QuantizeScale << 1U,
-		QuantizeVelocity = QuantizeLocation << 1U,
-		BitCount = 3U,
-	};
-	ENUM_CLASS_FLAGS(EActorNetSpawnInfoFlags)
-	
 	EActorNetSpawnInfoFlags GetSpawnInfoFlags()
 	{
 		// Disable performance warnings on FindConsoleVariable. We we call it once per NetActorFactory that is created, but in tests this can be thousands of times.
@@ -127,16 +117,51 @@ TUniquePtr<UE::Net::FNetObjectCreationHeader> UNetActorFactory::CreateAndFillHea
 		return nullptr;
 	}
 
-	UEngineReplicationBridge* ActorBridge = CastChecked<UEngineReplicationBridge>(Bridge);
+	const FNetObjectReference ActorReference = Bridge->GetOrCreateObjectReference(Actor);
 
 	TUniquePtr<FBaseActorNetCreationHeader> BaseHeader;
 
-	const FNetObjectReference ActorReference = ActorBridge->GetOrCreateObjectReference(Actor);	
+	// Allocate the right header type for this actor
+	if (!ActorReference.GetRefHandle().IsDynamic())
+	{
+		BaseHeader.Reset(new FStaticActorNetCreationHeader);
+	}
+	else
+	{
+		if (Bridge->IsNetRefHandlePreRegistered(Handle))
+		{
+			BaseHeader.Reset(new FPreRegisteredActorNetCreationHeader);
+		}
+		else
+		{
+			BaseHeader.Reset(new FDynamicActorNetCreationHeader);
+		}
+	}
+
+
+	const bool bSuccess = FillHeader(BaseHeader.Get(), Handle, Actor);
+	
+	return bSuccess ? MoveTemp(BaseHeader) : nullptr;
+}
+
+bool UNetActorFactory::FillHeader(UE::Net::FBaseActorNetCreationHeader* BaseHeader, UE::Net::FNetRefHandle Handle, AActor* Actor)
+{
+	using namespace UE::Net;
+
+	if (!BaseHeader->IsDynamic())
+	{
+		FStaticActorNetCreationHeader* Header = static_cast<FStaticActorNetCreationHeader*>(BaseHeader);
 		
-	if (ActorReference.GetRefHandle().IsDynamic())
+		// Refer by path for static actors
+		Header->ObjectReference = Bridge->GetOrCreateObjectReference(Actor);;
+	}
+	else if (BaseHeader->IsPreregistered())
 	{
-		FDynamicActorNetCreationHeader* Header = new FDynamicActorNetCreationHeader;
-		BaseHeader.Reset(Header);
+		// Nothing to fill!
+	}
+	else
+	{
+		FDynamicActorNetCreationHeader* Header = static_cast<FDynamicActorNetCreationHeader*>(BaseHeader);
 
 		// This is more or less a straight copy from ClientPackageMap and needs to be updated accordingly
 		UObject* Archetype = nullptr;
@@ -160,14 +185,15 @@ TUniquePtr<UE::Net::FNetObjectCreationHeader> UNetActorFactory::CreateAndFillHea
 		check(Actor->NeedsLoadForClient());			// We have no business sending this unless the client can load
 		check(Archetype->NeedsLoadForClient());		// We have no business sending this unless the client can load
 
+		UEngineReplicationBridge* EngineBridge = CastChecked<UEngineReplicationBridge>(Bridge);
+
 		// Fill in Header
-		Header->ArchetypeReference = ActorBridge->GetOrCreateObjectReference(Archetype);
-		Header->bUsePersistentLevel = (UE::Net::Private::SerializeNewActorOverrideLevel == 0) || (ActorBridge->GetNetDriver()->GetWorld()->PersistentLevel == ActorLevel);
-		Header->bIsPreRegistered = ActorBridge->IsNetRefHandlePreRegistered(Handle);
+		Header->ArchetypeReference = Bridge->GetOrCreateObjectReference(Archetype);
+		Header->bUsePersistentLevel = (UE::Net::Private::SerializeNewActorOverrideLevel == 0) || (EngineBridge->GetNetDriver()->GetWorld()->PersistentLevel == ActorLevel);
 
 		if (!Header->bUsePersistentLevel)
 		{
-			Header->LevelReference = ActorBridge->GetOrCreateObjectReference(ActorLevel);
+			Header->LevelReference = Bridge->GetOrCreateObjectReference(ActorLevel);
 		}
 
 		if (const USceneComponent* RootComponent = Actor->GetRootComponent())
@@ -195,15 +221,8 @@ TUniquePtr<UE::Net::FNetObjectCreationHeader> UNetActorFactory::CreateAndFillHea
 			check(!Header->SpawnInfo.Location.ContainsNaN());
 		}
 	}
-	else 
-	{
-		FStaticActorNetCreationHeader* Header = new FStaticActorNetCreationHeader;
-		BaseHeader.Reset(Header);
-		// Refer by path for static actors
-		Header->ObjectReference = ActorReference;
-	}
 
-	// Custom actor creation data
+	// Finally fill the custom actor creation data for all.
 	{
 		FOutBunch Bunch(UE::Net::Private::SerializeNewActorMaxBunchSize);
 		Actor->OnSerializeNewActor(Bunch);
@@ -215,7 +234,7 @@ TUniquePtr<UE::Net::FNetObjectCreationHeader> UNetActorFactory::CreateAndFillHea
 		}
 	}
 
-	return BaseHeader;
+	return true;
 }
 
 bool UNetActorFactory::SerializeHeader(const UE::Net::FCreationHeaderContext& Context, const UE::Net::FNetObjectCreationHeader* Header) 
@@ -226,19 +245,16 @@ bool UNetActorFactory::SerializeHeader(const UE::Net::FCreationHeaderContext& Co
 
 	const FBaseActorNetCreationHeader* BaseActorHeader = static_cast<const FBaseActorNetCreationHeader*>(Header);
 
-	Writer->WriteBool(BaseActorHeader->IsDynamic());
+	// Serialize the type of header it is (static -> preregistered -> dynamic)
+	const bool bIsDynamic = Writer->WriteBool(BaseActorHeader->IsDynamic());
 	
-	if (BaseActorHeader->IsDynamic())
-	{
-		const FDynamicActorNetCreationHeader* DynamicHeader = static_cast<const FDynamicActorNetCreationHeader*>(Header);
-		return DynamicHeader->Serialize(Context, SpawnInfoFlags, DefaultSpawnInfo);
-	}
-	else
+	if (bIsDynamic)
 	{
-		const FStaticActorNetCreationHeader* StaticHeader = static_cast<const FStaticActorNetCreationHeader*>(Header);
-		return StaticHeader->Serialize(Context);
+		const bool bIsPreRegistered = Writer->WriteBool(BaseActorHeader->IsPreregistered());
 	}
-	
+
+	// Serialize the content of the header
+	return BaseActorHeader->Serialize(Context);
 	
 }
 
@@ -250,21 +266,28 @@ TUniquePtr<UE::Net::FNetObjectCreationHeader> UNetActorFactory::CreateAndDeseria
 
 	FNetBitStreamReader* Writer = Context.Serialization.GetBitStreamReader();
 
+	// Deserialize the type of header and create it
 	const bool bIsDynamic = Writer->ReadBool();
 	if (bIsDynamic)
 	{
-		FDynamicActorNetCreationHeader* DynamicHeader = new FDynamicActorNetCreationHeader;
-		Header.Reset(DynamicHeader);
+		const bool bIsPreRegistered = Writer->ReadBool();
 
-		DynamicHeader->Deserialize(Context, DefaultSpawnInfo);
+		if (!bIsPreRegistered)
+		{
+			Header.Reset(new FDynamicActorNetCreationHeader);
+		}
+		else
+		{
+			Header.Reset(new FPreRegisteredActorNetCreationHeader);
+		}
 	}
 	else
 	{
-		FStaticActorNetCreationHeader* StaticHeader = new FStaticActorNetCreationHeader;
-		Header.Reset(StaticHeader);
-
-		StaticHeader->Deserialize(Context);
+		Header.Reset(new FStaticActorNetCreationHeader);
 	}
+
+	// Deserialize the content of the header
+	Header->Deserialize(Context);
 	
 	return Header;
 }
@@ -309,12 +332,11 @@ UNetObjectFactory::FInstantiateResult UNetActorFactory::InstantiateReplicatedObj
 		return Result;
 	}
 
-	// For dynamic actors, spawn a new actor using the provided information, or find a pre-registered actor
-	const FDynamicActorNetCreationHeader* DynamicHeader = static_cast<const FDynamicActorNetCreationHeader*>(Header);
-
-	// If pre-registered, look for existing instance.
-	if (DynamicHeader->bIsPreRegistered)
+	// If pre-registered, look for the existing instance.
+	if (BaseHeader->IsPreregistered())
 	{
+		// Nothing to read in the header for now
+
 		AActor* FoundActor = Cast<AActor>(Bridge->GetPreRegisteredObject(Context.Handle));
 		if (!FoundActor)
 		{
@@ -335,6 +357,9 @@ UNetObjectFactory::FInstantiateResult UNetActorFactory::InstantiateReplicatedObj
 		return Result;
 	}
 
+	// For dynamic actors, spawn a new actor using the provided information, or find a pre-registered actor
+	const FDynamicActorNetCreationHeader* DynamicHeader = static_cast<const FDynamicActorNetCreationHeader*>(Header);
+
 	// Find archetype
 	AActor* Archetype = Cast<AActor>(Bridge->ResolveObjectReference(DynamicHeader->ArchetypeReference, Context.ResolveContext));
 	if (!Archetype)
@@ -541,48 +566,65 @@ float UNetActorFactory::GetPollFrequency(UE::Net::FNetRefHandle RootObjectHandle
  {
 
 //------------------------------------------------------------------------
-// FStaticActorNetCreationHeader
+// FBaseActorNetCreationHeader
 //------------------------------------------------------------------------
 
-bool FStaticActorNetCreationHeader::Serialize(const FCreationHeaderContext& Context) const
+EActorNetSpawnInfoFlags FBaseActorNetCreationHeader::GetFactorySpawnFlags(const UNetActorFactory* ActorFactory) const
 {
-	using namespace UE::Net::Private;
+	return ActorFactory->SpawnInfoFlags;
+}
 
+bool FBaseActorNetCreationHeader::Serialize(const FCreationHeaderContext& Context) const
+{
 	FNetBitStreamWriter* Writer = Context.Serialization.GetBitStreamWriter();
 
-	WriteFullNetObjectReference(Context.Serialization, ObjectReference);
-
-	if (Writer->WriteBool(CustomCreationDataBitCount > 0))
+	const bool bHasCustomCreationData = Writer->WriteBool(CustomCreationDataBitCount > 0);
+	if (bHasCustomCreationData)
 	{
 		Writer->WriteBits(CustomCreationDataBitCount - 1U, 16U);
 		Writer->WriteBitStream(reinterpret_cast<const uint32*>(CustomCreationData.GetData()), 0U, CustomCreationDataBitCount);
 	}
 
-	return true;
+	return Writer->IsOverflown() == false;
 }
 
-bool FStaticActorNetCreationHeader::Deserialize(const FCreationHeaderContext& Context)
+bool FBaseActorNetCreationHeader::Deserialize(const FCreationHeaderContext& Context)
 {
-	using namespace UE::Net::Private;
-
 	FNetBitStreamReader* Reader = Context.Serialization.GetBitStreamReader();
 
-	ReadFullNetObjectReference(Context.Serialization, ObjectReference);
-
-	if (Reader->ReadBool())
+	const bool bHasCustomCreationData = Reader->ReadBool();
+	if (bHasCustomCreationData)
 	{
 		CustomCreationDataBitCount = 1U + Reader->ReadBits(16U);
 		CustomCreationData.SetNumZeroed(((CustomCreationDataBitCount + 31U) & ~31U) >> 3U);
 		Reader->ReadBitStream(reinterpret_cast<uint32*>(CustomCreationData.GetData()), CustomCreationDataBitCount);
 	}
+
+	return Reader->IsOverflown() == false;
+}
+
+//------------------------------------------------------------------------
+// FStaticActorNetCreationHeader
+//------------------------------------------------------------------------
+
+bool FStaticActorNetCreationHeader::Serialize(const FCreationHeaderContext& Context) const
+{
+	WriteFullNetObjectReference(Context.Serialization, ObjectReference);
+
+	return Super::Serialize(Context);
+}
+
+bool FStaticActorNetCreationHeader::Deserialize(const FCreationHeaderContext& Context)
+{
+	ReadFullNetObjectReference(Context.Serialization, ObjectReference);
 	
-	return true;
+	return Super::Deserialize(Context);
 }
 
 
 FString FStaticActorNetCreationHeader::ToString() const
 {
-	return FString::Printf(TEXT("FStaticActorNetCreationHeader (ProtocolId:0x%x):\n\t"
+	return FString::Printf(TEXT("\n\tFStaticActorNetCreationHeader (ProtocolId:0x%x):\n\t"
 								"ObjectReference=%s\n\t"
 								"CustomCreationData=%u bits"),
 						   		GetProtocolId(),
@@ -594,11 +636,13 @@ FString FStaticActorNetCreationHeader::ToString() const
 // FDynamicActorNetCreationHeader
 //------------------------------------------------------------------------
 
-bool FDynamicActorNetCreationHeader::Serialize(const FCreationHeaderContext& Context, UE::Net::Private::EActorNetSpawnInfoFlags SpawnFlags, const FActorNetSpawnInfo& DefaultSpawnInfo) const
+bool FDynamicActorNetCreationHeader::Serialize(const FCreationHeaderContext& Context) const
 {
 	using namespace UE::Net;
 	using namespace UE::Net::Private;
 
+	const UNetActorFactory* ActorFactory = CastChecked<UNetActorFactory>(Context.Factory);
+
 	FNetBitStreamWriter* Writer = Context.Serialization.GetBitStreamWriter();
 
 	// Write Archetype and LevelPath
@@ -610,35 +654,30 @@ bool FDynamicActorNetCreationHeader::Serialize(const FCreationHeaderContext& Con
 		WriteFullNetObjectReference(Context.Serialization, LevelReference);
 	}
 
-	const bool bQuantizeLocation = (SpawnFlags&EActorNetSpawnInfoFlags::QuantizeLocation)!=EActorNetSpawnInfoFlags::None;
-	const bool bQuantizeScale = (SpawnFlags&EActorNetSpawnInfoFlags::QuantizeScale)!=EActorNetSpawnInfoFlags::None;
-	const bool bQuantizeVelocity = (SpawnFlags&EActorNetSpawnInfoFlags::QuantizeVelocity)!=EActorNetSpawnInfoFlags::None;
+	const EActorNetSpawnInfoFlags SpawnFlags = GetFactorySpawnFlags(ActorFactory);
+	const bool bQuantizeLocation = EnumHasAnyFlags(SpawnFlags, EActorNetSpawnInfoFlags::QuantizeLocation);
+	const bool bQuantizeScale = EnumHasAnyFlags(SpawnFlags, EActorNetSpawnInfoFlags::QuantizeScale);
+	const bool bQuantizeVelocity = EnumHasAnyFlags(SpawnFlags, EActorNetSpawnInfoFlags::QuantizeVelocity);
 
 	// Write actor spawn info
-	WriteConditionallyQuantizedVector(Writer, SpawnInfo.Location, DefaultSpawnInfo.Location, bQuantizeLocation);
-	WriteConditionallyQuantizedVector(Writer, SpawnInfo.Scale, DefaultSpawnInfo.Scale, bQuantizeScale);
-	WriteConditionallyQuantizedVector(Writer, SpawnInfo.Velocity, DefaultSpawnInfo.Velocity, bQuantizeVelocity);
+	WriteConditionallyQuantizedVector(Writer, SpawnInfo.Location, ActorFactory->DefaultSpawnInfo.Location, bQuantizeLocation);
+	WriteConditionallyQuantizedVector(Writer, SpawnInfo.Scale, ActorFactory->DefaultSpawnInfo.Scale, bQuantizeScale);
+	WriteConditionallyQuantizedVector(Writer, SpawnInfo.Velocity, ActorFactory->DefaultSpawnInfo.Velocity, bQuantizeVelocity);
 
 	// For rotation we use 0.001f for Rotation comparison to keep consistency with old behavior.
 	static constexpr float RotationEpsilon = 0.001f;
-	WriteRotator(Writer, SpawnInfo.Rotation, DefaultSpawnInfo.Rotation, RotationEpsilon);
+	WriteRotator(Writer, SpawnInfo.Rotation, ActorFactory->DefaultSpawnInfo.Rotation, RotationEpsilon);
 
-	Writer->WriteBool(bIsPreRegistered);
-
-	if (Writer->WriteBool(CustomCreationDataBitCount > 0))
-	{
-		Writer->WriteBits(CustomCreationDataBitCount - 1U, 16U);
-		Writer->WriteBitStream(reinterpret_cast<const uint32*>(CustomCreationData.GetData()), 0U, CustomCreationDataBitCount);
-	}
-
-	return true;
+	return Super::Serialize(Context);
 }
 
-bool FDynamicActorNetCreationHeader::Deserialize(const FCreationHeaderContext& Context, const FActorNetSpawnInfo& DefaultSpawnInfo)
+bool FDynamicActorNetCreationHeader::Deserialize(const FCreationHeaderContext& Context)
 {
 	using namespace UE::Net;
 	using namespace UE::Net::Private;
 
+	const UNetActorFactory* ActorFactory = CastChecked<UNetActorFactory>(Context.Factory);
+
 	FNetBitStreamReader* Reader = Context.Serialization.GetBitStreamReader();
 
 	// Read Archetype
@@ -651,28 +690,19 @@ bool FDynamicActorNetCreationHeader::Deserialize(const FCreationHeaderContext& C
 	}
 
 	// Read actor spawn info
-	ReadConditionallyQuantizedVector(Reader, SpawnInfo.Location, DefaultSpawnInfo.Location);
-	ReadConditionallyQuantizedVector(Reader, SpawnInfo.Scale, DefaultSpawnInfo.Scale);
-	ReadConditionallyQuantizedVector(Reader, SpawnInfo.Velocity, DefaultSpawnInfo.Velocity);
-
-	ReadRotator(Reader, SpawnInfo.Rotation, DefaultSpawnInfo.Rotation);
-	
-	bIsPreRegistered = Reader->ReadBool();
+	ReadConditionallyQuantizedVector(Reader, SpawnInfo.Location, ActorFactory->DefaultSpawnInfo.Location);
+	ReadConditionallyQuantizedVector(Reader, SpawnInfo.Scale, ActorFactory->DefaultSpawnInfo.Scale);
+	ReadConditionallyQuantizedVector(Reader, SpawnInfo.Velocity, ActorFactory->DefaultSpawnInfo.Velocity);
 
-	if (Reader->ReadBool())
-	{
-		CustomCreationDataBitCount = 1U + Reader->ReadBits(16U);
-		CustomCreationData.SetNumZeroed(((CustomCreationDataBitCount + 31U) & ~31U) >> 3U);
-		Reader->ReadBitStream(reinterpret_cast<uint32*>(CustomCreationData.GetData()), CustomCreationDataBitCount);
-	}
+	ReadRotator(Reader, SpawnInfo.Rotation, ActorFactory->DefaultSpawnInfo.Rotation);
 	
-	return true;
+	return Super::Deserialize(Context);
 }
 
 
 FString FDynamicActorNetCreationHeader::ToString() const
 {
-	return FString::Printf(TEXT("FDynamicActorNetCreationHeader (ProtocolId:0x%x):\n\t"
+	return FString::Printf(TEXT("\n\tFDynamicActorNetCreationHeader (ProtocolId:0x%x):\n\t"
 								"ArchetypeReference=%s\n\t"
 								"SpawnInfo.Location=%s\n\t"
 								"SpawnInfo.Rotation=%s\n\t"
@@ -692,6 +722,28 @@ FString FDynamicActorNetCreationHeader::ToString() const
 								CustomCreationDataBitCount);
 }
 
+//------------------------------------------------------------------------
+// FPreRegisteredActorNetCreationHeader
+//------------------------------------------------------------------------
+bool FPreRegisteredActorNetCreationHeader::Serialize(const FCreationHeaderContext& Context) const
+{
+	// Nothing to serialize
+	return Super::Serialize(Context);
+}
+
+bool FPreRegisteredActorNetCreationHeader::Deserialize(const FCreationHeaderContext& Context)
+{
+	// Nothing to deserialize
+	return Super::Deserialize(Context);
+}
+
+FString FPreRegisteredActorNetCreationHeader::ToString() const
+{
+	return FString::Printf(TEXT("\n\tFPreRegisteredActorNetCreationHeader (ProtocolId:0x%x):\n\t"
+								"CustomCreationData=%u bits"),
+						   		GetProtocolId(),
+						   		CustomCreationDataBitCount);
+}
 
 } // end namespace UE::Net
 
diff --git a/Engine/Source/Runtime/Engine/Public/Net/Iris/ReplicationSystem/NetActorFactory.h b/Engine/Source/Runtime/Engine/Public/Net/Iris/ReplicationSystem/NetActorFactory.h
--- a/Engine/Source/Runtime/Engine/Public/Net/Iris/ReplicationSystem/NetActorFactory.h
+++ b/Engine/Source/Runtime/Engine/Public/Net/Iris/ReplicationSystem/NetActorFactory.h
@@ -10,11 +10,20 @@
 
 #include "NetActorFactory.generated.h"
 
-namespace UE::Net::Private
+namespace UE::Net
 {
-	enum class EActorNetSpawnInfoFlags : uint32;
+	enum class EActorNetSpawnInfoFlags : uint32
+	{
+		None = 0U,
+		QuantizeScale = 1U,
+		QuantizeLocation = QuantizeScale << 1U,
+		QuantizeVelocity = QuantizeLocation << 1U,
+	};
+	ENUM_CLASS_FLAGS(EActorNetSpawnInfoFlags)
 }
 
+class UNetActorFactory;
+
 
 namespace UE::Net
 {
@@ -27,7 +36,17 @@ namespace UE::Net
 class FBaseActorNetCreationHeader : public FNetObjectCreationHeader
 {
 public:
+	/** Is the header representing a dynamic or stable actor */
 	virtual bool IsDynamic() const = 0;
+	/** Is the header representing a dynamic pre-registered actor */
+	virtual bool IsPreregistered() const = 0;
+
+	ENGINE_API EActorNetSpawnInfoFlags GetFactorySpawnFlags(const UNetActorFactory* ActorFactory) const;
+
+	ENGINE_API virtual bool Serialize(const FCreationHeaderContext& Context) const;
+	ENGINE_API virtual bool Deserialize(const FCreationHeaderContext& Context);
+
+public:
 
 	TArray<uint8> CustomCreationData;
 	uint16 CustomCreationDataBitCount = 0;
@@ -38,6 +57,10 @@ public:
  */
 class FStaticActorNetCreationHeader : public FBaseActorNetCreationHeader
 {
+private:
+
+	typedef FBaseActorNetCreationHeader Super;
+
 public:
 
 	virtual bool IsDynamic() const
@@ -45,12 +68,19 @@ public:
 		return false;
 	}
 
-	FNetObjectReference ObjectReference;
-	
-	bool Serialize(const FCreationHeaderContext& Context) const;
-	bool Deserialize(const FCreationHeaderContext& Context);
+	virtual bool IsPreregistered() const
+	{
+		return false;
+	}
+
+	ENGINE_API virtual bool Serialize(const FCreationHeaderContext& Context) const override;
+	ENGINE_API virtual bool Deserialize(const FCreationHeaderContext& Context) override;
 
 	ENGINE_API virtual FString ToString() const override;
+
+public:
+
+	FNetObjectReference ObjectReference;
 };
 
 
@@ -59,14 +89,11 @@ public:
  */
 class FDynamicActorNetCreationHeader : public FBaseActorNetCreationHeader
 {
-public:
+private:
 
-	virtual bool IsDynamic() const override
-	{ 
-		return true;
-	}
+	typedef FBaseActorNetCreationHeader Super;
 
-	ENGINE_API virtual FString ToString() const override;
+public:
 
 	struct FActorNetSpawnInfo
 	{
@@ -83,6 +110,25 @@ public:
 		FVector Velocity;
 	};
 
+public:
+
+	virtual bool IsDynamic() const override
+	{ 
+		return true;
+	}
+
+	virtual bool IsPreregistered() const
+	{
+		return false;
+	}
+
+	ENGINE_API virtual FString ToString() const override;
+
+	ENGINE_API virtual bool Serialize(const FCreationHeaderContext& Context) const override;
+	ENGINE_API virtual bool Deserialize(const FCreationHeaderContext& Context) override;
+
+public:
+
 	FActorNetSpawnInfo SpawnInfo;
 
 	FNetObjectReference ArchetypeReference;
@@ -90,9 +136,34 @@ public:
 	
 	bool bUsePersistentLevel = false;
 	bool bIsPreRegistered = false;
+};
 
-	bool Serialize(const FCreationHeaderContext& Context, UE::Net::Private::EActorNetSpawnInfoFlags SpawnFlags, const FActorNetSpawnInfo& DefaultSpawnInfo) const;
-	bool Deserialize(const FCreationHeaderContext& Context, const FActorNetSpawnInfo& DefaultSpawnInfo);
+/**
+ * Header information representing dynamic pre-registered actors
+ * @See UObjectReplicationBridge::PreRegisterNewObjectReferenceHandle
+ */
+class FPreRegisteredActorNetCreationHeader : public FBaseActorNetCreationHeader
+{
+private:
+
+	typedef FBaseActorNetCreationHeader Super;
+
+public:
+
+	virtual bool IsDynamic() const override
+	{ 
+		return true;
+	}
+
+	virtual bool IsPreregistered() const
+	{
+		return true;
+	}
+
+	ENGINE_API virtual bool Serialize(const FCreationHeaderContext& Context) const override;
+	ENGINE_API virtual bool Deserialize(const FCreationHeaderContext& Context) override;
+
+	ENGINE_API virtual FString ToString() const override;
 };
 
 #endif // UE_WITH_IRIS
@@ -138,13 +209,23 @@ protected:
 
 	ENGINE_API virtual bool SerializeHeader(const UE::Net::FCreationHeaderContext& Context, const UE::Net::FNetObjectCreationHeader* Header)  override;
 
-private:
+protected:
 
-	/** Cache of the spawn info flags created from the net.QuantizeActor cvars */
-	UE::Net::Private::EActorNetSpawnInfoFlags SpawnInfoFlags;
+	/** Fill the required header information */
+	ENGINE_API bool FillHeader(UE::Net::FBaseActorNetCreationHeader* BaseHeader, UE::Net::FNetRefHandle Handle, AActor* Actor);
 
+public:
+
+	/** Cache of the default values of the spawn info to delta serialize against */
 	const UE::Net::FDynamicActorNetCreationHeader::FActorNetSpawnInfo DefaultSpawnInfo;
 
+protected:
+
+	friend UE::Net::FBaseActorNetCreationHeader;
+
+	/** Cache of the spawn info flags created from the net.QuantizeActor cvars */
+	UE::Net::EActorNetSpawnInfoFlags SpawnInfoFlags;
+
 #endif // UE_WITH_IRIS
 };
 
