RemoteObject: Reschedule aborted ServerReplicateActor calls rather than retrying them.

#rnx

diff --git a/Engine/Source/Runtime/Engine/Private/NetDriver.cpp b/Engine/Source/Runtime/Engine/Private/NetDriver.cpp
--- a/Engine/Source/Runtime/Engine/Private/NetDriver.cpp
+++ b/Engine/Source/Runtime/Engine/Private/NetDriver.cpp
@@ -6065,10 +6065,15 @@ int32 UNetDriver::ServerReplicateActors(float DeltaSeconds)
 
 	int32 Updated = 0;
 
-	const int32 NumClientsToTick = [this, DeltaSeconds]()
-		{
 #if UE_WITH_REMOTE_OBJECT_HANDLE
-			if (bReplicateTransactionally && UE::Net::Private::bServerReplicateActorsGranularTransactions)
+	const bool bUseGranularTransactions = (bReplicateTransactionally && UE::Net::Private::bServerReplicateActorsGranularTransactions);
+#else
+	constexpr bool bUseGranularTransactions = false;
+#endif
+
+	const int32 NumClientsToTick = [this, DeltaSeconds, bUseGranularTransactions]()
+		{
+			if (bUseGranularTransactions)
 			{
 				int32 LocalNumClientsToTick = 0;
 				static FName TransactionalWorkName("ServerReplicateActors_PrepConnections");
@@ -6079,7 +6084,6 @@ int32 UNetDriver::ServerReplicateActors(float DeltaSeconds)
 				return LocalNumClientsToTick;
 			}
 			else
-#endif // UE_WITH_REMOTE_OBJECT_HANDLE
 			{
 				return ServerReplicateActors_PrepConnections(DeltaSeconds);
 			}
@@ -6110,17 +6114,19 @@ int32 UNetDriver::ServerReplicateActors(float DeltaSeconds)
 	ConsiderList.Reserve( GetNetworkObjectList().GetActiveObjects().Num() );
 
 	// Build the consider list (actors that are ready to replicate)
-#if UE_WITH_REMOTE_OBJECT_HANDLE
-	if (bReplicateTransactionally && UE::Net::Private::bServerReplicateActorsGranularTransactions)
+	if (bUseGranularTransactions)
 	{
-		static FName TransactionalWorkName("ServerReplicateActors_BuildConsiderList");
+		// Start the region here and end it when we're done ExecutePendingWork
+		static const FName FuncName { __func__ };
+		UE::RemoteExecutor::OnRegionBeginDelegate.ExecuteIfBound(FuncName);
+
+		static const FName TransactionalWorkName("ServerReplicateActors_BuildConsiderList");
 		UE::RemoteExecutor::ExecuteTransactional(TransactionalWorkName, [this, &ConsiderList, ServerTickTime]()
 			{
 				ServerReplicateActors_BuildConsiderList(ConsiderList, ServerTickTime);
 			});
 	}
 	else
-#endif // UE_WITH_REMOTE_OBJECT_HANDLE
 	{
 		ServerReplicateActors_BuildConsiderList(ConsiderList, ServerTickTime);
 	}
@@ -6194,17 +6200,15 @@ int32 UNetDriver::ServerReplicateActors(float DeltaSeconds)
 				.bCPUSaturated = bCPUSaturated
 			};
 
-#if UE_WITH_REMOTE_OBJECT_HANDLE
-			if (bReplicateTransactionally && UE::Net::Private::bServerReplicateActorsGranularTransactions)
+			if (bUseGranularTransactions)
 			{
-				static FName TransactionalWorkName("ServerReplicateActors_ForConnection");
-				UE::RemoteExecutor::ExecuteTransactional(TransactionalWorkName, [this, &Params]()
+				const FName TransactionalWorkName = Connection->ViewTarget.GetFName();
+				UE::RemoteExecutor::EnqueueWork(TransactionalWorkName, true, [this, &Params]()
 					{
 						ServerReplicateActors_ForConnection(Params);
 					});
 			}
 			else
-#endif // UE_WITH_REMOTE_OBJECT_HANDLE
 			{
 				ServerReplicateActors_ForConnection(Params);
 			}
@@ -6216,6 +6220,12 @@ int32 UNetDriver::ServerReplicateActors(float DeltaSeconds)
 		}
 	}
 
+	if (bUseGranularTransactions)
+	{
+		UE::RemoteExecutor::ExecutePendingWork();
+		UE::RemoteExecutor::OnRegionEndDelegate.ExecuteIfBound(TEXT(""));
+	}
+
 	if (OnPostConsiderListUpdateOverride.IsBound())
 	{
 		OnPostConsiderListUpdateOverride.ExecuteIfBound( { DeltaSeconds, nullptr, bCPUSaturated }, Updated, ConsiderList );
