SplineComponent: Refactor rendering

diff --git a/Engine/Source/Runtime/Engine/Classes/Components/SplineComponent.h b/Engine/Source/Runtime/Engine/Classes/Components/SplineComponent.h
--- a/Engine/Source/Runtime/Engine/Classes/Components/SplineComponent.h
+++ b/Engine/Source/Runtime/Engine/Classes/Components/SplineComponent.h
@@ -293,6 +293,14 @@ private:
 	UPROPERTY(EditAnywhere, Replicated, Category = Spline, meta=(EditCondition="bLoopPositionOverride"))
 	float LoopPosition;
 
+	/** Material used for rendering spline segments. */
+	UPROPERTY(Transient)
+	TObjectPtr<UMaterialInterface> LineMaterial = nullptr;
+
+	/** Material used for rendering spline points. */
+	UPROPERTY(Transient)
+	TObjectPtr<UMaterialInterface> PointMaterial = nullptr;
+	
 public:
 	/** Default up vector in local space to be used when calculating transforms along the spline */
 	UPROPERTY(EditAnywhere, Replicated, BlueprintReadWrite, Category = Spline)
@@ -348,6 +356,7 @@ public:
 
 #if UE_ENABLE_DEBUG_DRAWING
 	//~ Begin UPrimitiveComponent Interface.
+	ENGINE_API virtual void GetUsedMaterials(TArray<UMaterialInterface*>& OutMaterials, bool bGetDebugMaterials = false) const override;
 	ENGINE_API virtual FPrimitiveSceneProxy* CreateSceneProxy() override;
 #endif
 	
diff --git a/Engine/Source/Runtime/Engine/Private/Components/SplineComponent.cpp b/Engine/Source/Runtime/Engine/Private/Components/SplineComponent.cpp
--- a/Engine/Source/Runtime/Engine/Private/Components/SplineComponent.cpp
+++ b/Engine/Source/Runtime/Engine/Private/Components/SplineComponent.cpp
@@ -19,6 +19,13 @@
 #include "Styling/SlateColor.h"
 #include "Styling/StyleColors.h"
 
+#include "LocalVertexFactory.h"
+#include "SceneInterface.h"
+#include "Materials/Material.h"
+#include "DynamicMeshBuilder.h"
+#include "StaticMeshResources.h"
+
+
 #include UE_INLINE_GENERATED_CPP_BY_NAME(SplineComponent)
 
 #if WITH_EDITOR
@@ -2366,6 +2373,12 @@ FVector USplineComponent::GetVectorPropertyAtSplinePoint(int32 Index, FName Prop
 }
 
 #if UE_ENABLE_DEBUG_DRAWING
+void USplineComponent::GetUsedMaterials(TArray<UMaterialInterface*>& OutMaterials, bool bGetDebugMaterials) const
+{
+	if (LineMaterial) { OutMaterials.Add(LineMaterial.Get()); }
+	if (PointMaterial) { OutMaterials.Add(PointMaterial.Get()); }
+}
+
 FPrimitiveSceneProxy* USplineComponent::CreateSceneProxy()
 {
 	if (!bDrawDebug)
@@ -2373,8 +2386,23 @@ FPrimitiveSceneProxy* USplineComponent::CreateSceneProxy()
 		return Super::CreateSceneProxy();
 	}
 
+#if UE_ENABLE_DEBUG_DRAWING
+	LineMaterial = LoadObject<UMaterial>(this, TEXT("/Engine/EngineMaterials/LineSetComponentMaterial"));
+	PointMaterial = LoadObject<UMaterial>(this, TEXT("/Engine/EngineMaterials/LineSetComponentMaterial"));
+#endif
+	
 	class FSplineSceneProxy final : public FPrimitiveSceneProxy
 	{
+		/** Inspired by FLineMeshBatchData and FPointMeshBatchData. */
+		struct FMeshBatchData
+		{
+			FMaterialRenderProxy* MaterialProxy = nullptr;
+			int32 StartIndex = 0;
+			int32 NumPrimitives = 0;
+			int32 MinVertexIndex = 0;
+			int32 MaxVertexIndex = 0;
+		};
+		
 	public:
 		SIZE_T GetTypeHash() const override
 		{
@@ -2391,7 +2419,267 @@ FPrimitiveSceneProxy* USplineComponent::CreateSceneProxy()
 #else
 			, LineColor(FLinearColor::White)
 #endif
-		{}
+			, LineVertexFactory(GetScene().GetFeatureLevel(), "FSplineProxyLineVertexFactory")
+			, PointVertexFactory(GetScene().GetFeatureLevel(), "FSplineProxyPointVertexFactory")
+		{
+			CSV_SCOPED_TIMING_STAT_GLOBAL(FSplineSceneProxy_FSplineSceneProxy)
+			
+			constexpr int32 NumStepsPerSegment = 40;
+			
+			const int32 NumPoints = SplineInfo.Points.Num();
+			const int32 NumSegments = FMath::Max(SplineInfo.bIsLooped ? NumPoints : NumPoints - 1, 0);
+			const int32 NumLines = NumSegments * NumStepsPerSegment;
+			
+			const int32 NumLineVertices = NumLines * 4;
+			const int32 NumLineIndices = NumLines * 6;
+			const int32 NumPointVertices = NumPoints * 4;
+			const int32 NumPointIndices = NumPoints * 6;
+			constexpr int32 NumTextureCoordinates = 1;
+
+			LineVertexBuffers.PositionVertexBuffer.Init(NumLineVertices);
+			LineVertexBuffers.StaticMeshVertexBuffer.Init(NumLineVertices, NumTextureCoordinates);
+			LineVertexBuffers.ColorVertexBuffer.Init(NumLineVertices);
+			LineIndexBuffer.Indices.SetNumUninitialized(NumLineIndices);
+			int32 LineVertexBufferIndex = 0;
+			int32 LineIndexBufferIndex = 0;
+
+			LineBatchData.MinVertexIndex = LineVertexBufferIndex;
+			LineBatchData.MaxVertexIndex = LineVertexBufferIndex + NumLineVertices - 1;
+			LineBatchData.StartIndex = LineIndexBufferIndex;
+			LineBatchData.NumPrimitives = NumLines * 2;
+			LineBatchData.MaterialProxy = InComponent->LineMaterial->GetRenderProxy();
+
+			PointVertexBuffers.PositionVertexBuffer.Init(NumPointVertices);
+			PointVertexBuffers.StaticMeshVertexBuffer.Init(NumPointVertices, NumTextureCoordinates);
+			PointVertexBuffers.ColorVertexBuffer.Init(NumPointVertices);
+			PointIndexBuffer.Indices.SetNumUninitialized(NumPointIndices);
+			int32 PointVertexBufferIndex = 0;
+			int32 PointIndexBufferIndex = 0;
+
+			PointBatchData.MinVertexIndex = PointVertexBufferIndex;
+			PointBatchData.MaxVertexIndex = PointVertexBufferIndex + NumPointVertices - 1;
+			PointBatchData.StartIndex = PointIndexBufferIndex;
+			PointBatchData.NumPrimitives = NumPoints * 2;
+			PointBatchData.MaterialProxy = InComponent->PointMaterial->GetRenderProxy();
+			
+			FVector OldKeyPos(0);
+			
+			auto AppendLineToBuffers = [this, &LineVertexBufferIndex, &LineIndexBufferIndex](FVector Start, FVector End, FColor Color, float Thickness, float DepthBias = 0.f) -> void
+			{
+				const FVector LineDirection = (End - Start).GetSafeNormal();
+				const FVector2f UV(Thickness, DepthBias);
+
+				LineVertexBuffers.PositionVertexBuffer.VertexPosition(LineVertexBufferIndex + 0) = static_cast<FVector3f>(Start);
+				LineVertexBuffers.PositionVertexBuffer.VertexPosition(LineVertexBufferIndex + 1) = static_cast<FVector3f>(End);
+				LineVertexBuffers.PositionVertexBuffer.VertexPosition(LineVertexBufferIndex + 2) = static_cast<FVector3f>(End);
+				LineVertexBuffers.PositionVertexBuffer.VertexPosition(LineVertexBufferIndex + 3) = static_cast<FVector3f>(Start);
+
+				// On each end of the line, we store the line direction pointing in opposite directions so that the verts
+				// are moved in opposite directions to make the rectangle
+				LineVertexBuffers.StaticMeshVertexBuffer.SetVertexTangents(LineVertexBufferIndex + 0, FVector3f::ZeroVector, FVector3f::ZeroVector, static_cast<FVector3f>(-LineDirection));
+				LineVertexBuffers.StaticMeshVertexBuffer.SetVertexTangents(LineVertexBufferIndex + 1, FVector3f::ZeroVector, FVector3f::ZeroVector, static_cast<FVector3f>(-LineDirection));
+				LineVertexBuffers.StaticMeshVertexBuffer.SetVertexTangents(LineVertexBufferIndex + 2, FVector3f::ZeroVector, FVector3f::ZeroVector, static_cast<FVector3f>(LineDirection));
+				LineVertexBuffers.StaticMeshVertexBuffer.SetVertexTangents(LineVertexBufferIndex + 3, FVector3f::ZeroVector, FVector3f::ZeroVector, static_cast<FVector3f>(LineDirection));
+
+				LineVertexBuffers.StaticMeshVertexBuffer.SetVertexUV(LineVertexBufferIndex + 0, 0, UV);
+				LineVertexBuffers.StaticMeshVertexBuffer.SetVertexUV(LineVertexBufferIndex + 1, 0, UV);
+				LineVertexBuffers.StaticMeshVertexBuffer.SetVertexUV(LineVertexBufferIndex + 2, 0, UV);
+				LineVertexBuffers.StaticMeshVertexBuffer.SetVertexUV(LineVertexBufferIndex + 3, 0, UV);
+
+				// The color stored in the vertices actually gets interpreted as a linear color by the material,
+				// whereas it is more convenient for the user of the LineSet to specify colors as sRGB. So we actually
+				// have to convert it back to linear. The ToFColor(false) call just scales back into 0-255 space.
+				const FColor ConvertedColor = FLinearColor::FromSRGBColor(Color).ToFColor(false);
+				LineVertexBuffers.ColorVertexBuffer.VertexColor(LineVertexBufferIndex + 0) = ConvertedColor;
+				LineVertexBuffers.ColorVertexBuffer.VertexColor(LineVertexBufferIndex + 1) = ConvertedColor;
+				LineVertexBuffers.ColorVertexBuffer.VertexColor(LineVertexBufferIndex + 2) = ConvertedColor;
+				LineVertexBuffers.ColorVertexBuffer.VertexColor(LineVertexBufferIndex + 3) = ConvertedColor;
+
+				LineIndexBuffer.Indices[LineIndexBufferIndex + 0] = LineVertexBufferIndex + 0;
+				LineIndexBuffer.Indices[LineIndexBufferIndex + 1] = LineVertexBufferIndex + 1;
+				LineIndexBuffer.Indices[LineIndexBufferIndex + 2] = LineVertexBufferIndex + 2;
+				LineIndexBuffer.Indices[LineIndexBufferIndex + 3] = LineVertexBufferIndex + 2;
+				LineIndexBuffer.Indices[LineIndexBufferIndex + 4] = LineVertexBufferIndex + 3;
+				LineIndexBuffer.Indices[LineIndexBufferIndex + 5] = LineVertexBufferIndex + 0;
+
+				LineVertexBufferIndex += 4;
+				LineIndexBufferIndex += 6;
+			};
+			
+			auto AppendPointToBuffers = [this, &PointVertexBufferIndex, &PointIndexBufferIndex](FVector Position, FColor Color, float Size, float DepthBias = 0.f) -> void
+			{
+				const FVector2f UV(Size, DepthBias);
+
+				PointVertexBuffers.PositionVertexBuffer.VertexPosition(PointVertexBufferIndex + 0) = static_cast<FVector3f>(Position);
+				PointVertexBuffers.PositionVertexBuffer.VertexPosition(PointVertexBufferIndex + 1) = static_cast<FVector3f>(Position);
+				PointVertexBuffers.PositionVertexBuffer.VertexPosition(PointVertexBufferIndex + 2) = static_cast<FVector3f>(Position);
+				PointVertexBuffers.PositionVertexBuffer.VertexPosition(PointVertexBufferIndex + 3) = static_cast<FVector3f>(Position);
+				
+				PointVertexBuffers.StaticMeshVertexBuffer.SetVertexUV(PointVertexBufferIndex + 0, 0, UV);
+				PointVertexBuffers.StaticMeshVertexBuffer.SetVertexUV(PointVertexBufferIndex + 1, 0, UV);
+				PointVertexBuffers.StaticMeshVertexBuffer.SetVertexUV(PointVertexBufferIndex + 2, 0, UV);
+				PointVertexBuffers.StaticMeshVertexBuffer.SetVertexUV(PointVertexBufferIndex + 3, 0, UV);
+				
+				// The color stored in the vertices actually gets interpreted as a linear color by the material,
+				// whereas it is more convenient for the user of the LineSet to specify colors as sRGB. So we actually
+				// have to convert it back to linear. The ToFColor(false) call just scales back into 0-255 space.
+				const FColor ConvertedColor = FLinearColor::FromSRGBColor(Color).ToFColor(false);
+				PointVertexBuffers.ColorVertexBuffer.VertexColor(PointVertexBufferIndex + 0) = ConvertedColor;
+				PointVertexBuffers.ColorVertexBuffer.VertexColor(PointVertexBufferIndex + 1) = ConvertedColor;
+				PointVertexBuffers.ColorVertexBuffer.VertexColor(PointVertexBufferIndex + 2) = ConvertedColor;
+				PointVertexBuffers.ColorVertexBuffer.VertexColor(PointVertexBufferIndex + 3) = ConvertedColor;
+				
+				PointVertexBuffers.StaticMeshVertexBuffer.SetVertexTangents(PointVertexBufferIndex + 0, FVector3f::ZeroVector, FVector3f::ZeroVector, FVector3f(1.0f, -1.0f, 0.0f));
+				PointVertexBuffers.StaticMeshVertexBuffer.SetVertexTangents(PointVertexBufferIndex + 1, FVector3f::ZeroVector, FVector3f::ZeroVector, FVector3f(1.0f, 1.0f, 0.0f));
+				PointVertexBuffers.StaticMeshVertexBuffer.SetVertexTangents(PointVertexBufferIndex + 2, FVector3f::ZeroVector, FVector3f::ZeroVector, FVector3f(-1.0f, 1.0f, 0.0f));
+				PointVertexBuffers.StaticMeshVertexBuffer.SetVertexTangents(PointVertexBufferIndex + 3, FVector3f::ZeroVector, FVector3f::ZeroVector, FVector3f(-1.0f, -1.0f, 0.0f));
+					
+				PointIndexBuffer.Indices[PointIndexBufferIndex + 0] = PointVertexBufferIndex + 0;
+				PointIndexBuffer.Indices[PointIndexBufferIndex + 1] = PointVertexBufferIndex + 1;
+				PointIndexBuffer.Indices[PointIndexBufferIndex + 2] = PointVertexBufferIndex + 2;
+				PointIndexBuffer.Indices[PointIndexBufferIndex + 3] = PointVertexBufferIndex + 2;
+				PointIndexBuffer.Indices[PointIndexBufferIndex + 4] = PointVertexBufferIndex + 3;
+				PointIndexBuffer.Indices[PointIndexBufferIndex + 5] = PointVertexBufferIndex + 0;
+
+				PointVertexBufferIndex += 4;
+				PointIndexBufferIndex += 6;
+			};
+			
+			for (int32 KeyIdx = 0; KeyIdx < NumSegments + 1; KeyIdx++)
+			{
+				const FVector NewKeyPos = SplineInfo.Eval(static_cast<float>(KeyIdx), FVector(0));
+
+				// Draw the keypoint
+				if (KeyIdx < NumPoints)
+				{
+					constexpr int32 GrabHandleSize = 6;
+					AppendPointToBuffers(NewKeyPos, LineColor.ToFColorSRGB(), GrabHandleSize);
+				}
+
+				// If not the first keypoint, draw a line to the previous keypoint.
+				if (KeyIdx > 0)
+				{
+					// For constant interpolation - don't draw ticks - just draw dotted line.
+					if (SplineInfo.Points[KeyIdx - 1].InterpMode == CIM_Constant)
+					{
+						FVector OldPos = OldKeyPos;
+						
+						float SegmentLineThickness = 1.4f;
+#if WITH_EDITOR
+						SegmentLineThickness += GetDefault<ULevelEditorViewportSettings>()->SplineLineThicknessAdjustment * 6.f;
+#endif
+						const FVector SegmentVector = SplineInfo.Points[KeyIdx].OutVal - OldKeyPos;
+						const float SegmentLength = SegmentVector.Length();
+						const FVector SegmentDirection = SegmentVector.GetSafeNormal();
+						const float StepLength = SegmentLength / NumStepsPerSegment;
+						const FVector SegmentStep = SegmentDirection * StepLength;
+						
+						for (int32 StepIdx = 1; StepIdx <= NumStepsPerSegment; StepIdx++)
+						{
+							const FVector NewPos = OldPos + SegmentStep;
+
+							// In order to get a dashed line (which was the old behavior), we subdivide the line in the same way that we draw normal segments
+							// but we hide every other subdivision by giving it a thickness of 0. While the fact that we render more lines than before may seem
+							// wasteful, we do it for 2 reasons:
+							// 1) Consistent logic for subdivision makes preallocating the vertex buffers trivial
+							// 2) We are rendering lines in a much more efficient way than before
+							const bool bHideLine = SplineInfo.Points[KeyIdx - 1].InterpMode == CIM_Constant && StepIdx % 2;
+							AppendLineToBuffers(OldPos, NewPos, LineColor.ToFColorSRGB(), bHideLine ? 0.f : SegmentLineThickness);
+
+							OldPos = NewPos;
+						}
+					}
+					else
+					{
+						FVector OldPos = OldKeyPos;
+
+						// This initial value of 1.4f is an experimentally determined value which gives the same effective segment thickness as the
+						// previous PDI rendering when the level editor thickness adjustment setting was at its default value (0). PDI special cases
+						// lines of 0 thickness while our approach does not, so if we are not artificially increasing it here we will get invisible
+						// lines by default.
+						float SegmentLineThickness = 1.4f;
+#if WITH_EDITOR
+						SegmentLineThickness += GetDefault<ULevelEditorViewportSettings>()->SplineLineThicknessAdjustment;
+#endif
+
+						for (int32 StepIdx = 1; StepIdx <= NumStepsPerSegment; StepIdx++)
+						{
+							const float Key = (KeyIdx - 1) + (StepIdx / static_cast<float>(NumStepsPerSegment));
+							const FVector NewPos = SplineInfo.Eval(Key, FVector(0));
+							
+							AppendLineToBuffers(OldPos, NewPos, LineColor.ToFColorSRGB(), SegmentLineThickness);
+
+							OldPos = NewPos;
+						}
+					}
+				}
+
+				OldKeyPos = NewKeyPos;
+			}
+
+			if (NumLines > 0)
+			{
+				ENQUEUE_RENDER_COMMAND(LineSetVertexBuffersInit)(
+					[this](FRHICommandListImmediate& RHICmdList)
+				{
+					LineVertexBuffers.PositionVertexBuffer.InitResource(RHICmdList);
+					LineVertexBuffers.StaticMeshVertexBuffer.InitResource(RHICmdList);
+					LineVertexBuffers.ColorVertexBuffer.InitResource(RHICmdList);
+
+					FLocalVertexFactory::FDataType Data;
+					LineVertexBuffers.PositionVertexBuffer.BindPositionVertexBuffer(&LineVertexFactory, Data);
+					LineVertexBuffers.StaticMeshVertexBuffer.BindTangentVertexBuffer(&LineVertexFactory, Data);
+					LineVertexBuffers.StaticMeshVertexBuffer.BindTexCoordVertexBuffer(&LineVertexFactory, Data);
+					LineVertexBuffers.ColorVertexBuffer.BindColorVertexBuffer(&LineVertexFactory, Data);
+					LineVertexFactory.SetData(RHICmdList, Data);
+
+					LineVertexFactory.InitResource(RHICmdList);
+					LineIndexBuffer.InitResource(RHICmdList);
+				});
+			}
+
+			if (NumPoints > 0)
+			{
+				ENQUEUE_RENDER_COMMAND(OverlayVertexBuffersInit)(
+					[this](FRHICommandListImmediate& RHICmdList)
+				{
+					PointVertexBuffers.PositionVertexBuffer.InitResource(RHICmdList);
+					PointVertexBuffers.StaticMeshVertexBuffer.InitResource(RHICmdList);
+					PointVertexBuffers.ColorVertexBuffer.InitResource(RHICmdList);
+
+					FLocalVertexFactory::FDataType Data;
+					PointVertexBuffers.PositionVertexBuffer.BindPositionVertexBuffer(&PointVertexFactory, Data);
+					PointVertexBuffers.StaticMeshVertexBuffer.BindTangentVertexBuffer(&PointVertexFactory, Data);
+					PointVertexBuffers.StaticMeshVertexBuffer.BindTexCoordVertexBuffer(&PointVertexFactory, Data);
+					PointVertexBuffers.ColorVertexBuffer.BindColorVertexBuffer(&PointVertexFactory, Data);
+					PointVertexFactory.SetData(RHICmdList, Data);
+
+					PointVertexFactory.InitResource(RHICmdList);
+					PointIndexBuffer.InitResource(RHICmdList);
+				});
+			}
+		}
+
+		virtual ~FSplineSceneProxy() override
+		{
+			if (LineVertexBuffers.PositionVertexBuffer.IsInitialized())
+			{
+				LineVertexBuffers.PositionVertexBuffer.ReleaseResource();
+				LineVertexBuffers.StaticMeshVertexBuffer.ReleaseResource();
+				LineVertexBuffers.ColorVertexBuffer.ReleaseResource();
+				LineIndexBuffer.ReleaseResource();
+				LineVertexFactory.ReleaseResource();
+			}
+
+			if (PointVertexBuffers.PositionVertexBuffer.IsInitialized())
+			{
+				PointVertexBuffers.PositionVertexBuffer.ReleaseResource();
+				PointVertexBuffers.StaticMeshVertexBuffer.ReleaseResource();
+				PointVertexBuffers.ColorVertexBuffer.ReleaseResource();
+				PointIndexBuffer.ReleaseResource();
+				PointVertexFactory.ReleaseResource();
+			}
+		}
 
 		virtual void GetDynamicMeshElements(const TArray<const FSceneView*>& Views, const FSceneViewFamily& ViewFamily, uint32 VisibilityMap, FMeshElementCollector& Collector) const override
 		{
@@ -2401,19 +2689,53 @@ FPrimitiveSceneProxy* USplineComponent::CreateSceneProxy()
 			{
 				if (VisibilityMap & (1 << ViewIndex))
 				{
-					const FSceneView* View = Views[ViewIndex];
-					FPrimitiveDrawInterface* PDI = Collector.GetPDI(ViewIndex);
-
-					const FMatrix& LocalToWorld = GetLocalToWorld();
-
-					// Taking into account the min and maximum drawing distance
-					const float DistanceSqr = (View->ViewMatrices.GetViewOrigin() - LocalToWorld.GetOrigin()).SizeSquared();
-					if (DistanceSqr < FMath::Square(GetMinDrawDistance()) || DistanceSqr > FMath::Square(GetMaxDrawDistance()))
+					if (LineVertexBuffers.PositionVertexBuffer.IsInitialized())
 					{
-						continue;
+						FMeshBatch& Mesh = Collector.AllocateMesh();
+						FMeshBatchElement& BatchElement = Mesh.Elements[0];
+						BatchElement.IndexBuffer = &LineIndexBuffer;
+						Mesh.bWireframe = false;
+						Mesh.VertexFactory = &LineVertexFactory;
+						Mesh.MaterialRenderProxy = LineBatchData.MaterialProxy;
+
+						FDynamicPrimitiveUniformBuffer& DynamicPrimitiveUniformBuffer = Collector.AllocateOneFrameResource<FDynamicPrimitiveUniformBuffer>();
+						DynamicPrimitiveUniformBuffer.Set(Collector.GetRHICommandList(), GetLocalToWorld(), GetLocalToWorld(), GetBounds(), GetLocalBounds(), false, false, AlwaysHasVelocity());
+						BatchElement.PrimitiveUniformBufferResource = &DynamicPrimitiveUniformBuffer.UniformBuffer;
+
+						BatchElement.FirstIndex = LineBatchData.StartIndex;
+						BatchElement.NumPrimitives = LineBatchData.NumPrimitives;
+						BatchElement.MinVertexIndex = LineBatchData.MinVertexIndex;
+						BatchElement.MaxVertexIndex = LineBatchData.MaxVertexIndex;
+						Mesh.ReverseCulling = IsLocalToWorldDeterminantNegative();
+						Mesh.Type = PT_TriangleList;
+						Mesh.DepthPriorityGroup = SDPG_World;
+						Mesh.bCanApplyViewModeOverrides = false;
+						Collector.AddMesh(ViewIndex, Mesh);
 					}
 
-					USplineComponent::Draw(PDI, View, SplineInfo, LocalToWorld, LineColor, SDPG_World);
+					if (PointVertexBuffers.PositionVertexBuffer.IsInitialized())
+					{
+						FMeshBatch& Mesh = Collector.AllocateMesh();
+						FMeshBatchElement& BatchElement = Mesh.Elements[0];
+						BatchElement.IndexBuffer = &PointIndexBuffer;
+						Mesh.bWireframe = false;
+						Mesh.VertexFactory = &PointVertexFactory;
+						Mesh.MaterialRenderProxy = PointBatchData.MaterialProxy;
+
+						FDynamicPrimitiveUniformBuffer& DynamicPrimitiveUniformBuffer = Collector.AllocateOneFrameResource<FDynamicPrimitiveUniformBuffer>();
+						DynamicPrimitiveUniformBuffer.Set(Collector.GetRHICommandList(), GetLocalToWorld(), GetLocalToWorld(), GetBounds(), GetLocalBounds(), true, false, AlwaysHasVelocity());
+						BatchElement.PrimitiveUniformBufferResource = &DynamicPrimitiveUniformBuffer.UniformBuffer;
+
+						BatchElement.FirstIndex = PointBatchData.StartIndex;
+						BatchElement.NumPrimitives = PointBatchData.NumPrimitives;
+						BatchElement.MinVertexIndex = PointBatchData.MinVertexIndex;
+						BatchElement.MaxVertexIndex = PointBatchData.MaxVertexIndex;
+						Mesh.ReverseCulling = IsLocalToWorldDeterminantNegative();
+						Mesh.Type = PT_TriangleList;
+						Mesh.DepthPriorityGroup = SDPG_World;
+						Mesh.bCanApplyViewModeOverrides = false;
+						Collector.AddMesh(ViewIndex, Mesh);
+					}
 				}
 			}
 		}
@@ -2421,10 +2743,15 @@ FPrimitiveSceneProxy* USplineComponent::CreateSceneProxy()
 		virtual FPrimitiveViewRelevance GetViewRelevance(const FSceneView* View) const override
 		{
 			FPrimitiveViewRelevance Result;
-			Result.bDrawRelevance = bDrawDebug && IsShown(View) && View->Family->EngineShowFlags.Splines;
+			Result.bDrawRelevance = bDrawDebug && !IsSelected() && IsShown(View) && View->Family->EngineShowFlags.Splines;
 			Result.bDynamicRelevance = true;
 			Result.bShadowRelevance = IsShadowCast(View);
 			Result.bEditorPrimitiveRelevance = UseEditorCompositing(View);
+			Result.bRenderInMainPass = ShouldRenderInMainPass();
+			Result.bUsesLightingChannels = GetLightingChannelMask() != GetDefaultLightingChannelMask();
+			Result.bRenderCustomDepth = ShouldRenderCustomDepth();
+			Result.bTranslucentSelfShadow = bCastVolumetricTranslucentShadow;
+			Result.bVelocityRelevance = DrawsVelocity() && Result.bOpaque && Result.bRenderInMainPass;
 			return Result;
 		}
 
@@ -2435,6 +2762,16 @@ FPrimitiveSceneProxy* USplineComponent::CreateSceneProxy()
 		bool bDrawDebug;
 		FInterpCurveVector SplineInfo;
 		FLinearColor LineColor;
+		
+		FMeshBatchData LineBatchData;
+		FLocalVertexFactory LineVertexFactory;
+		FStaticMeshVertexBuffers LineVertexBuffers;
+		FDynamicMeshIndexBuffer32 LineIndexBuffer;
+		
+		FMeshBatchData PointBatchData;
+		FLocalVertexFactory PointVertexFactory;
+		FStaticMeshVertexBuffers PointVertexBuffers;
+		FDynamicMeshIndexBuffer32 PointIndexBuffer;
 	};
 
 	return new FSplineSceneProxy(this);
