SplineComponent: Stop spamming the strictly increasing input key ensure when an invalid spline is detected

We should still catch all bad splines without spamming this ensure more than once

diff --git a/Engine/Source/Runtime/Engine/Private/Components/SplineComponent.cpp b/Engine/Source/Runtime/Engine/Private/Components/SplineComponent.cpp
--- a/Engine/Source/Runtime/Engine/Private/Components/SplineComponent.cpp
+++ b/Engine/Source/Runtime/Engine/Private/Components/SplineComponent.cpp
@@ -301,7 +301,11 @@ void FSplineCurves::UpdateSpline(bool bClosedLoop, bool bStationaryEndpoints, in
 	// Ensure input keys are strictly ascending
 	for (int32 Index = 1; Index < NumPoints; Index++)
 	{
-		ensureAlways(Position.Points[Index - 1].InVal < Position.Points[Index].InVal);
+		if (!ensureAlways(Position.Points[Index - 1].InVal < Position.Points[Index].InVal))
+		if (!ensureAlways(Position.Points[Index - 1].InVal < Position.Points[Index].InVal))
+		{
+			break;
+		}
 	}
 #endif
 
