SplineComponent: Removing rogue line of code

diff --git a/Engine/Source/Runtime/Engine/Private/Components/SplineComponent.cpp b/Engine/Source/Runtime/Engine/Private/Components/SplineComponent.cpp
--- a/Engine/Source/Runtime/Engine/Private/Components/SplineComponent.cpp
+++ b/Engine/Source/Runtime/Engine/Private/Components/SplineComponent.cpp
@@ -301,7 +301,6 @@ void FSplineCurves::UpdateSpline(bool bClosedLoop, bool bStationaryEndpoints, in
 	// Ensure input keys are strictly ascending
 	for (int32 Index = 1; Index < NumPoints; Index++)
 	{
-		if (!ensureAlways(Position.Points[Index - 1].InVal < Position.Points[Index].InVal))
 		if (!ensureAlways(Position.Points[Index - 1].InVal < Position.Points[Index].InVal))
 		{
 			break;
