World Partition - HLOD: Improve instance batching by ignoring some of the ISMComponentDescriptor fields that aren't relevant for HLOD

diff --git a/Engine/Plugins/Editor/WorldPartitionHLODUtilities/Source/Private/WorldPartition/HLOD/Builders/HLODBuilderInstancing.cpp b/Engine/Plugins/Editor/WorldPartitionHLODUtilities/Source/Private/WorldPartition/HLOD/Builders/HLODBuilderInstancing.cpp
--- a/Engine/Plugins/Editor/WorldPartitionHLODUtilities/Source/Private/WorldPartition/HLOD/Builders/HLODBuilderInstancing.cpp
+++ b/Engine/Plugins/Editor/WorldPartitionHLODUtilities/Source/Private/WorldPartition/HLOD/Builders/HLODBuilderInstancing.cpp
@@ -25,7 +25,7 @@ uint32 UHLODBuilderInstancingSettings::GetCRC() const
 	FArchiveCrc32 Ar;
 
 	// Base key, changing this will force a rebuild of all HLODs from this builder
-	FString HLODBaseKey = "53809597CD9C4FB7AC75827A628513D6";
+	FString HLODBaseKey = "5197F75BF850402EABA5D29E4D8BB1DA";
 	Ar << HLODBaseKey;
 
 	Ar << This.bDisallowNanite;
diff --git a/Engine/Source/Runtime/Engine/Private/WorldPartition/HLOD/HLODInstancedStaticMeshComponent.cpp b/Engine/Source/Runtime/Engine/Private/WorldPartition/HLOD/HLODInstancedStaticMeshComponent.cpp
--- a/Engine/Source/Runtime/Engine/Private/WorldPartition/HLOD/HLODInstancedStaticMeshComponent.cpp
+++ b/Engine/Source/Runtime/Engine/Private/WorldPartition/HLOD/HLODInstancedStaticMeshComponent.cpp
@@ -29,11 +29,22 @@ void FHLODISMComponentDescriptor::InitFrom(const UStaticMeshComponent* Component
 {
 	Super::InitFrom(Component, bInitBodyInstance);
 
-	// Stationnary can be considered as static for the purpose of HLODs
-	if (Mobility == EComponentMobility::Stationary)
-	{
-		Mobility = EComponentMobility::Static;
-	}
+	// Improve instance batching by ignoring some of the fields that aren't relevant for HLOD
+	Mobility = EComponentMobility::Static;
+	bHasCustomNavigableGeometry = EHasCustomNavigableGeometry::No;
+	HLODBatchingPolicy = EHLODBatchingPolicy::Instancing;
+	bSelectable = true;
+	bHasPerInstanceHitProxies = false;
+	bConsiderForActorPlacementWhenHidden = false;
+	bUseDefaultCollision = true;
+	bGenerateOverlapEvents = false;
+	bOverrideNavigationExport = false;
+	bForceNavigationObstacle = false;
+	bCanEverAffectNavigation = false;
+	bFillCollisionUnderneathForNavmesh = false;	
+
+	// Force the correct ISM component class as it can be changed by the parent InitFrom() implementation
+	ComponentClass = UHLODBuilder::GetInstancedStaticMeshComponentClass();
 }
 
 
