Fixed UWorldPartitionStreamingPolicy::UpdateStreamingStateInternal accessing world members outside the game thread with wp.Runtime.UpdateStreaming.EnableAsyncUpdate enabled, potentially causing concurrent access.

diff --git a/Engine/Source/Runtime/Engine/Private/WorldPartition/WorldPartitionRuntimeHash.cpp b/Engine/Source/Runtime/Engine/Private/WorldPartition/WorldPartitionRuntimeHash.cpp
--- a/Engine/Source/Runtime/Engine/Private/WorldPartition/WorldPartitionRuntimeHash.cpp
+++ b/Engine/Source/Runtime/Engine/Private/WorldPartition/WorldPartitionRuntimeHash.cpp
@@ -624,10 +624,8 @@ EWorldPartitionStreamingPerformance UWorldPartitionRuntimeHash::GetStreamingPerf
 	EWorldPartitionStreamingPerformance StreamingPerformance = EWorldPartitionStreamingPerformance::Good;
 	bOutShouldBlock = false;
 
-	if (!CellsToActivate.IsEmpty() && GetWorld()->bMatchStarted)
+	if (!CellsToActivate.IsEmpty())
 	{
-		UWorld* World = GetWorld();
-
 		for (const UWorldPartitionRuntimeCell* Cell : CellsToActivate)
 		{
 			if (!Cell->IsAlwaysLoaded() && Cell->GetStreamingStatus() != LEVEL_Visible)
diff --git a/Engine/Source/Runtime/Engine/Private/WorldPartition/WorldPartitionStreamingPolicy.cpp b/Engine/Source/Runtime/Engine/Private/WorldPartition/WorldPartitionStreamingPolicy.cpp
--- a/Engine/Source/Runtime/Engine/Private/WorldPartition/WorldPartitionStreamingPolicy.cpp
+++ b/Engine/Source/Runtime/Engine/Private/WorldPartition/WorldPartitionStreamingPolicy.cpp
@@ -395,8 +395,6 @@ void UWorldPartitionStreamingPolicy::UpdateStreamingStateInternal(const FUpdateS
 
 	OutTargetState.bUpdateServerEpoch = false;
 
-	check(InParams.World);
-	check(InParams.World->IsGameWorld());
 	check(InParams.FrameActivateCells.IsEmpty());
 	check(InParams.FrameLoadCells.IsEmpty());
 	check(OutTargetState.IsEmpty());
@@ -412,7 +410,7 @@ void UWorldPartitionStreamingPolicy::UpdateStreamingStateInternal(const FUpdateS
 
 	if (InParams.bCanStream)
 	{
-		if (!InParams.bIsServer || InParams.bIsServerStreamingEnabled || AWorldPartitionReplay::IsPlaybackEnabled(InParams.World))
+		if (!InParams.bIsServer || InParams.bIsServerStreamingEnabled || InParams.bIsPlaybackEnabled)
 		{
 			// When world partition can't stream, all cells must be unloaded
 			if (InParams.RuntimeHash)
@@ -659,7 +657,15 @@ void UWorldPartitionStreamingPolicy::UpdateStreamingStateInternal(const FUpdateS
 
 	if (InParams.RuntimeHash)
 	{
-		OutTargetState.StreamingPerformance = InParams.RuntimeHash->GetStreamingPerformance(InParams.FrameActivateCells, OutTargetState.bBlockOnSlowStreaming);
+		if (!InParams.bMatchStarted)
+		{
+			OutTargetState.StreamingPerformance = EWorldPartitionStreamingPerformance::Good;
+			OutTargetState.bBlockOnSlowStreaming = false;
+		}
+		else
+		{
+			OutTargetState.StreamingPerformance = InParams.RuntimeHash->GetStreamingPerformance(InParams.FrameActivateCells, OutTargetState.bBlockOnSlowStreaming);
+		}
 	}
 }
 
@@ -1260,6 +1266,8 @@ UWorldPartitionStreamingPolicy::FUpdateStreamingStateParams::FUpdateStreamingSta
 	, bIsServerStreamingEnabled(InPolicy->GetOuterUWorldPartition()->IsServerStreamingEnabled())
 	, bIsServerStreamingOutEnabled(InPolicy->GetOuterUWorldPartition()->IsServerStreamingOutEnabled())
 	, bIsBlockingCausedByBadStreamingPerformance(InPolicy->IsInBlockTillLevelStreamingCompleted(true))
+	, bIsPlaybackEnabled(AWorldPartitionReplay::IsPlaybackEnabled(InPolicy->GetWorld()))
+	, bMatchStarted(InPolicy->GetWorld()->bMatchStarted)
 	, bShouldMergeStreamingSourceInfo(InPolicy->bShouldMergeStreamingSourceInfo)
 	, PolicyUpdateStreamingStateEpoch(InPolicy->UpdateStreamingStateCounter)
 	, DataLayersLogicOperator(InPolicy->GetOuterUWorldPartition()->GetDataLayersLogicOperator())
@@ -1271,6 +1279,7 @@ UWorldPartitionStreamingPolicy::FUpdateStreamingStateParams::FUpdateStreamingSta
 	, WorldDataLayersEffectiveStatesRef(FWorldDataLayersEffectiveStatesAccessor::Get(InPolicy->GetOuterUWorldPartition()->GetTypedOuter<UWorld>()->GetWorldDataLayers()))
 {
 	check(IsInGameThread());
+	check(World->IsGameWorld());
 }
 
 UWorldPartitionStreamingPolicy::FUpdateStreamingStateParams& UWorldPartitionStreamingPolicy::FUpdateStreamingStateParams::SetRequiredWorldDataLayersEffectiveStatesCopy(bool bInRequiredEffectiveStatesCopy)
diff --git a/Engine/Source/Runtime/Engine/Public/WorldPartition/WorldPartitionStreamingPolicy.h b/Engine/Source/Runtime/Engine/Public/WorldPartition/WorldPartitionStreamingPolicy.h
--- a/Engine/Source/Runtime/Engine/Public/WorldPartition/WorldPartitionStreamingPolicy.h
+++ b/Engine/Source/Runtime/Engine/Public/WorldPartition/WorldPartitionStreamingPolicy.h
@@ -181,6 +181,8 @@ private:
 		const bool bIsServerStreamingEnabled;
 		const bool bIsServerStreamingOutEnabled;
 		const bool bIsBlockingCausedByBadStreamingPerformance;
+		const bool bIsPlaybackEnabled;
+		const bool bMatchStarted;
 		const bool bShouldMergeStreamingSourceInfo;
 		const int32 PolicyUpdateStreamingStateEpoch;
 		const EWorldPartitionDataLayersLogicOperator DataLayersLogicOperator;
