World Partition Actor Component Descriptors: fix bad class lookup when retrieving base compojnent class.

diff --git a/Engine/Source/Runtime/Engine/Private/WorldPartition/WorldPartitionActorDesc.cpp b/Engine/Source/Runtime/Engine/Private/WorldPartition/WorldPartitionActorDesc.cpp
--- a/Engine/Source/Runtime/Engine/Private/WorldPartition/WorldPartitionActorDesc.cpp
+++ b/Engine/Source/Runtime/Engine/Private/WorldPartition/WorldPartitionActorDesc.cpp
@@ -60,7 +60,7 @@ void FWorldPartitionComponentDesc::Init(const UActorComponent* InComponent)
 	NativeClass = *ComponentNativeClass->GetPathName();
 	
 	// For native class, don't set this
-	if (!ComponentNativeClass->IsNative())
+	if (!ComponentClass->IsNative())
 	{
 		BaseClass = *ComponentClass->GetPathName();
 	}
