MaterialInterface - Silence message logged on explicit material rename, related to the associated EditorOnlyData object
* MaterialInterface X has a incorrectly name EditorOnlyData Y. This may cause issues when running the editor on cooked data. Trying to rename it to the correct name XY
* There are call sites for which SafeRenameEditorOnlyData() should output a message, but for UMaterialInterface::Rename(), it shouldn't

diff --git a/Engine/Source/Runtime/Engine/Private/Materials/MaterialInterface.cpp b/Engine/Source/Runtime/Engine/Private/Materials/MaterialInterface.cpp
--- a/Engine/Source/Runtime/Engine/Private/Materials/MaterialInterface.cpp
+++ b/Engine/Source/Runtime/Engine/Private/Materials/MaterialInterface.cpp
@@ -2129,7 +2129,7 @@ bool UMaterialInterface::Rename(const TCHAR* NewName, UObject* NewOuter, ERename
 	if (bRenamed && NewName && EditorOnlyData)
 	{
 		FString EditorOnlyDataName = UE::MaterialInterface::GetEditorOnlyDataName(NewName);
-		bRenamed = SafeRenameEditorOnlyData(this, EditorOnlyDataName, true);
+		bRenamed = SafeRenameEditorOnlyData(this, EditorOnlyDataName, false);
 	}
 #endif
 	return bRenamed;
