Fixed ensure in UMaterial::GetUsedTextures when s.StreamableAssets.UseSimpleStreamableAssetManager = true.

diff --git a/Engine/Source/Runtime/Engine/Private/Components/PrimitiveComponent.cpp b/Engine/Source/Runtime/Engine/Private/Components/PrimitiveComponent.cpp
--- a/Engine/Source/Runtime/Engine/Private/Components/PrimitiveComponent.cpp
+++ b/Engine/Source/Runtime/Engine/Private/Components/PrimitiveComponent.cpp
@@ -59,6 +59,7 @@
 #include "Physics/Experimental/PhysScene_Chaos.h"
 #include "PhysicsProxy/SingleParticlePhysicsProxy.h"
 #include "PrimitiveComponentHelper.h"
+#include "Materials/MaterialInterface.h"
 
 #include UE_INLINE_GENERATED_CPP_BY_NAME(PrimitiveComponent)
 
@@ -5425,7 +5426,7 @@ void FActorPrimitiveComponentInterface::GetUsedMaterials(TArray<UMaterialInterfa
 
 void FActorPrimitiveComponentInterface::GetStreamableRenderAssetInfo(TArray<FStreamingRenderAssetPrimitiveInfo>& StreamableRenderAssets) const
 {
-	FStreamingTextureLevelContext LevelContext{EMaterialQualityLevel::Num,UPrimitiveComponent::GetPrimitiveComponent(this)}; 
+	FStreamingTextureLevelContext LevelContext{GetCurrentMaterialQualityLevelChecked(),UPrimitiveComponent::GetPrimitiveComponent(this)};
 	UPrimitiveComponent::GetPrimitiveComponent(this)->GetStreamingRenderAssetInfo(LevelContext, StreamableRenderAssets);
 }
 
