Further removal of unmaintained substrate material layers code.
This code was only ref'd by the already removed SMaterialSubstrateTree class.
#rnx

diff --git a/Engine/Source/Runtime/Engine/Private/Materials/Material.cpp b/Engine/Source/Runtime/Engine/Private/Materials/Material.cpp
--- a/Engine/Source/Runtime/Engine/Private/Materials/Material.cpp
+++ b/Engine/Source/Runtime/Engine/Private/Materials/Material.cpp
@@ -2379,23 +2379,6 @@ bool UMaterial::GetMaterialLayers(FMaterialLayersFunctions& OutLayers, TMicRecur
 	return false;
 }
 
-#if WITH_EDITORONLY_DATA
-void UMaterial::SyncLayersRuntimeGraphCache(FMaterialLayersFunctions* OverrideLayers)
-{
-	if(UMaterialEditorOnlyData* EditorOnly = GetEditorOnlyData())
-	{
-		for (UMaterialExpression* Expression : EditorOnly->ExpressionCollection.Expressions)
-		{
-			if (UMaterialExpressionMaterialAttributeLayers* LayersExpression = Cast<UMaterialExpressionMaterialAttributeLayers>(Expression))
-			{
-				LayersExpression->OverrideLayerGraph(OverrideLayers);
-				return;
-			}
-		}
-	}
-}
-#endif
-
 bool UMaterial::GetRefractionSettings(float& OutBiasValue) const
 {
 	OutBiasValue = RefractionDepthBias;
diff --git a/Engine/Source/Runtime/Engine/Private/Materials/MaterialInstance.cpp b/Engine/Source/Runtime/Engine/Private/Materials/MaterialInstance.cpp
--- a/Engine/Source/Runtime/Engine/Private/Materials/MaterialInstance.cpp
+++ b/Engine/Source/Runtime/Engine/Private/Materials/MaterialInstance.cpp
@@ -3745,14 +3745,6 @@ void UMaterialInstance::DeclareConstructClasses(TArray<FTopLevelAssetPath>& OutC
 	Super::DeclareConstructClasses(OutConstructClasses, SpecificSubclass);
 	OutConstructClasses.Add(FTopLevelAssetPath(UMaterialInstanceEditorOnlyData::StaticClass()));
 }
-
-void UMaterialInstance::SyncLayersRuntimeGraphCache(FMaterialLayersFunctions* OverrideLayers)
-{
-	if (Parent)
-	{
-		Parent->SyncLayersRuntimeGraphCache(OverrideLayers);
-	}
-}
 #endif
 
 void UMaterialInstance::BeginDestroy()
diff --git a/Engine/Source/Runtime/Engine/Public/Materials/Material.h b/Engine/Source/Runtime/Engine/Public/Materials/Material.h
--- a/Engine/Source/Runtime/Engine/Public/Materials/Material.h
+++ b/Engine/Source/Runtime/Engine/Public/Materials/Material.h
@@ -1309,8 +1309,6 @@ public:
 #if WITH_EDITORONLY_DATA
 	/** Return true if the material support shading model override */
 	ENGINE_API bool SupportsShadingModelOverride() const;
-
-	ENGINE_API virtual void SyncLayersRuntimeGraphCache(FMaterialLayersFunctions* OverrideLayers) override;
 #endif
 
 #if WITH_EDITOR
diff --git a/Engine/Source/Runtime/Engine/Public/Materials/MaterialInstance.h b/Engine/Source/Runtime/Engine/Public/Materials/MaterialInstance.h
--- a/Engine/Source/Runtime/Engine/Public/Materials/MaterialInstance.h
+++ b/Engine/Source/Runtime/Engine/Public/Materials/MaterialInstance.h
@@ -1005,7 +1005,6 @@ public:
 	virtual ENGINE_API void PostLoad() override;
 #if WITH_EDITORONLY_DATA
 	ENGINE_API static void DeclareConstructClasses(TArray<FTopLevelAssetPath>& OutConstructClasses, const UClass* SpecificSubclass);
-	ENGINE_API virtual void SyncLayersRuntimeGraphCache(FMaterialLayersFunctions* OverrideLayers) override;
 #endif
 	virtual ENGINE_API void BeginDestroy() override;
 	virtual ENGINE_API bool IsReadyForFinishDestroy() override;
diff --git a/Engine/Source/Runtime/Engine/Public/Materials/MaterialInterface.h b/Engine/Source/Runtime/Engine/Public/Materials/MaterialInterface.h
--- a/Engine/Source/Runtime/Engine/Public/Materials/MaterialInterface.h
+++ b/Engine/Source/Runtime/Engine/Public/Materials/MaterialInterface.h
@@ -694,11 +694,6 @@ public:
 	* @return					True if successful
 	*/
 	ENGINE_API bool GetStaticComponentMaskParameterValue(const FHashedMaterialParameterInfo& ParameterInfo, bool& R, bool& G, bool& B, bool& A, FGuid& OutExpressionGuid, bool bOveriddenOnly = false) const;
-
-	/**
-	 * Update layers functions RuntimeGraphCache for compiling (primarily for Preview Materials)
-	 */
-	virtual void SyncLayersRuntimeGraphCache(FMaterialLayersFunctions* OverrideLayers) PURE_VIRTUAL(UMaterialInterface::SyncLayersRuntimeGraphCache);
 #endif // WITH_EDITORONLY_DATA
 
 #if WITH_EDITOR
diff --git a/Engine/Source/Runtime/Engine/Public/Materials/MaterialLayersFunctions.h b/Engine/Source/Runtime/Engine/Public/Materials/MaterialLayersFunctions.h
--- a/Engine/Source/Runtime/Engine/Public/Materials/MaterialLayersFunctions.h
+++ b/Engine/Source/Runtime/Engine/Public/Materials/MaterialLayersFunctions.h
@@ -119,25 +119,6 @@ struct FMaterialLayersFunctionsEditorOnlyData
 #endif // WITH_EDITOR
 };
 
-
-USTRUCT() 
-struct FMaterialLayersFunctionsRuntimeGraphCache
-{
-	GENERATED_BODY()
-
-
-	TArray<TObjectPtr<UMaterialExpressionMaterialFunctionCall>> LayerCallers;
-	TArray<TObjectPtr<UMaterialExpressionMaterialFunctionCall>> BlendCallers;
-
-#if WITH_EDITOR
-	FMaterialExpressionCollection								ExpressionCollection;
-
-	TArray<TObjectPtr<UMaterialExpressionMaterialFunctionCall>> NodeMaterialGraphExpressions;
-	TArray<TObjectPtr<UMaterialExpressionMaterialFunctionCall>> NodeMaterialGraphExpressionsBlends;
-	TArray<TObjectPtr<UMaterial>>								NodePreviewMaterials;
-#endif
-};
-
 struct FMaterialLayerStackFunctionsCache
 {
 	//Helper struct for managing available layers and blends across parent/instance materials utilising the LayerStack node.
@@ -156,8 +137,6 @@ struct FMaterialLayersFunctionsRuntimeData
 	UPROPERTY(EditAnywhere, Category = MaterialLayers)
 	TArray<TObjectPtr<class UMaterialFunctionInterface>> Blends;
 
-	mutable TSharedPtr<FMaterialLayersFunctionsRuntimeGraphCache> RuntimeGraphCache;
-
 	// In multiple locations there are const references to this runtime data/it's children, 
 	// so this mutable entry is required to ensure we can keep available layers/blend up to date across all instances.
 	mutable TSharedPtr<FMaterialLayerStackFunctionsCache> LayerStackCache;
@@ -166,7 +145,6 @@ struct FMaterialLayersFunctionsRuntimeData
 	FMaterialLayersFunctionsRuntimeData(const FMaterialLayersFunctionsRuntimeData& Rhs)
 		: Layers(Rhs.Layers)
 		, Blends(Rhs.Blends)
-		, RuntimeGraphCache(Rhs.RuntimeGraphCache)
 		, LayerStackCache(Rhs.LayerStackCache)
 	{}
 
@@ -176,7 +154,6 @@ struct FMaterialLayersFunctionsRuntimeData
 	{
 		Layers = Rhs.Layers;
 		Blends = Rhs.Blends;
-		RuntimeGraphCache = Rhs.RuntimeGraphCache;
 		LayerStackCache = Rhs.LayerStackCache;
 		return *this;
 	}
@@ -189,7 +166,6 @@ struct FMaterialLayersFunctionsRuntimeData
 	{
 		Layers.Empty();
 		Blends.Empty();
-		RuntimeGraphCache.Reset();
 		LayerStackCache.Reset();
 	}
 
