Deprecated legacy "browse to" override API, in favor of UBrowseToAssetOverrideSubsystem

diff --git a/Engine/Source/Editor/UnrealEd/Private/Subsystems/BrowseToAssetOverrideSubsystem.cpp b/Engine/Source/Editor/UnrealEd/Private/Subsystems/BrowseToAssetOverrideSubsystem.cpp
--- a/Engine/Source/Editor/UnrealEd/Private/Subsystems/BrowseToAssetOverrideSubsystem.cpp
+++ b/Engine/Source/Editor/UnrealEd/Private/Subsystems/BrowseToAssetOverrideSubsystem.cpp
@@ -22,7 +22,9 @@ FName UBrowseToAssetOverrideSubsystem::GetBrowseToAssetOverride(const UObject* O
 	// If set, that takes priority over any per-class overrides
 	if (const AActor* Actor = Cast<AActor>(Object))
 	{
+		PRAGMA_DISABLE_DEPRECATION_WARNINGS
 		const FString& ActorBrowseToAssetOverride = Actor->GetBrowseToAssetOverride();
+		PRAGMA_ENABLE_DEPRECATION_WARNINGS
 		if (!ActorBrowseToAssetOverride.IsEmpty())
 		{
 			return *ActorBrowseToAssetOverride;
diff --git a/Engine/Source/Runtime/Engine/Classes/GameFramework/Actor.h b/Engine/Source/Runtime/Engine/Classes/GameFramework/Actor.h
--- a/Engine/Source/Runtime/Engine/Classes/GameFramework/Actor.h
+++ b/Engine/Source/Runtime/Engine/Classes/GameFramework/Actor.h
@@ -2810,8 +2810,10 @@ public:
 	}
 
 	/** Sets metadata on the Actor that allows it to point to a different asset than its source when browsing in the Content Browser */
+	UE_DEPRECATED(5.7, "SetBrowseToAssetOverride is deprecated as it was unreliable. Register a callback with UBrowseToAssetOverrideSubsystem instead, and store the override as a property on your actor.")
 	ENGINE_API void SetBrowseToAssetOverride(const FString& PackageName);
 	/** Gets metadata on the Actor that says which package to show in the Content Browser when browsing to it, or an empty string if there is no override */
+	UE_DEPRECATED(5.7, "GetBrowseToAssetOverride is deprecated. Use UBrowseToAssetOverrideSubsystem::GetBrowseToAssetOverride instead.")
 	ENGINE_API const FString& GetBrowseToAssetOverride() const;
 #endif		// WITH_EDITOR
 
