Tentative fix for crash occuring when navigation is being rebuild because nav overrride is being removed when clearing world components.

diff --git a/Engine/Source/Runtime/Engine/Private/World.cpp b/Engine/Source/Runtime/Engine/Private/World.cpp
--- a/Engine/Source/Runtime/Engine/Private/World.cpp
+++ b/Engine/Source/Runtime/Engine/Private/World.cpp
@@ -2821,6 +2821,8 @@ void UWorld::ShrinkLevel()
 
 void UWorld::ClearWorldComponents()
 {
+	TGuardValue<bool> IsBeingCleanedUp(bIsBeingCleanedUp, true);
+	
 	for( int32 LevelIndex=0; LevelIndex<Levels.Num(); LevelIndex++ )
 	{
 		ULevel* Level = Levels[LevelIndex];
