World Partition - HLOD: Include the HLOD Modifier class in the HLOD layer hash computation

diff --git a/Engine/Source/Runtime/Engine/Private/WorldPartition/HLOD/HLODSourceActors.cpp b/Engine/Source/Runtime/Engine/Private/WorldPartition/HLOD/HLODSourceActors.cpp
--- a/Engine/Source/Runtime/Engine/Private/WorldPartition/HLOD/HLODSourceActors.cpp
+++ b/Engine/Source/Runtime/Engine/Private/WorldPartition/HLOD/HLODSourceActors.cpp
@@ -5,6 +5,7 @@
 #if WITH_EDITOR
 #include "WorldPartition/HLOD/HLODActor.h"
 #include "WorldPartition/HLOD/HLODLayer.h"
+#include "WorldPartition/HLOD/HLODModifier.h"
 #endif
 
 #include UE_INLINE_GENERATED_CPP_BY_NAME(HLODSourceActors)
@@ -24,6 +25,12 @@ static uint32 GetHLODHash(const UHLODLayer* InHLODLayer)
 	HLODHash = HashCombineFast(HLODLayer.GetHLODBuilderSettings()->GetCRC(), HLODHash);
 	UE_LOG(LogHLODHash, VeryVerbose, TEXT(" - HLODBuilderSettings = %d"), HLODHash);
 
+	if (TSubclassOf<UWorldPartitionHLODModifier> HLODModifier = HLODLayer.GetHLODModifierClass(); HLODModifier.Get())
+	{
+		HLODHash = FCrc::StrCrc32(*HLODModifier->GetPathName(), HLODHash);
+		UE_LOG(LogHLODHash, VeryVerbose, TEXT(" - HLODModifier = %d"), HLODHash);
+	}
+
 	return HLODHash;
 }
 
