Reworked refraction.
Now it is either specified on the root node or another mode allows to use IOR from F0 explicitely. It is no longer based on plugged in or not since the default value on the root node kind of breaks that workflow and enforce override.
No conversion is needed since with the default value on the root node, we would always only use the override from the root node instead of F0 from the materials.

#tests pc editor dx12
#rnx
#DisabledValidations Preflight

diff --git a/Engine/Shaders/Private/BasePassPixelShader.usf b/Engine/Shaders/Private/BasePassPixelShader.usf
--- a/Engine/Shaders/Private/BasePassPixelShader.usf
+++ b/Engine/Shaders/Private/BasePassPixelShader.usf
@@ -2498,11 +2498,11 @@ void FPixelShaderInOut_MainPS(
 #if SUPPORT_MATERIAL_PRIMITIVE_ALPHA_HOLDOUT
 		
 	//Calculate IoR for translucent material with ior and update holdout property
-	#if (MATERIALBLENDING_TRANSLUCENT || SUBSTRATE_TRANSLUCENT_MATERIAL) && REFRACTION_USE_INDEX_OF_REFRACTION
-		#if REFRACTION_ROOT_NODE_OVERRIDES_DEFAULT
+	#if (MATERIALBLENDING_TRANSLUCENT || SUBSTRATE_TRANSLUCENT_MATERIAL) && (REFRACTION_USE_INDEX_OF_REFRACTION || REFRACTION_USE_INDEX_OF_REFRACTION_FROM_F0)
+		#if REFRACTION_USE_INDEX_OF_REFRACTION
 			FMaterialRefractionData RefractionData = GetMaterialRefraction(PixelMaterialInputs);
 			float MaterialIOR = GetMaterialRefractionIOR(RefractionData);
-		#else
+		#else // REFRACTION_USE_INDEX_OF_REFRACTION_FROM_F0
 			float MaterialIOR = DielectricF0RGBToIor(ComputeF0(Specular, BaseColor, Metallic));
 		#endif
 		
diff --git a/Engine/Shaders/Private/DistortionCommon.ush b/Engine/Shaders/Private/DistortionCommon.ush
--- a/Engine/Shaders/Private/DistortionCommon.ush
+++ b/Engine/Shaders/Private/DistortionCommon.ush
@@ -21,8 +21,18 @@ float2 ComputeBufferUVDistortion(
 #elif REFRACTION_USE_INDEX_OF_REFRACTION
 	// we assume the camera is in air
 	float AirIOR = 1.0f;
-#if REFRACTION_ROOT_NODE_OVERRIDES_DEFAULT || SUBSTRATE_OPTIMIZED_UNLIT
+#if SUBSTRATE_OPTIMIZED_UNLIT
+	float MaterialIOR = GetUnlitMaterialRefractionIOR(RefractionData);
+#else
 	float MaterialIOR = GetMaterialRefractionIOR(RefractionData);
+#endif
+	float2 ViewportUVDistortion = ViewNormal.xy * (MaterialIOR - AirIOR);
+
+#elif REFRACTION_USE_INDEX_OF_REFRACTION_FROM_F0
+	// we assume the camera is in air
+	float AirIOR = 1.0f;
+#if SUBSTRATE_OPTIMIZED_UNLIT
+	float MaterialIOR = GetUnlitMaterialRefractionIOR(RefractionData);
 #else
 	float MaterialIOR = InMaterialIOR;
 #endif
diff --git a/Engine/Shaders/Private/MaterialTemplate.ush b/Engine/Shaders/Private/MaterialTemplate.ush
--- a/Engine/Shaders/Private/MaterialTemplate.ush
+++ b/Engine/Shaders/Private/MaterialTemplate.ush
@@ -4059,6 +4059,11 @@ struct FMaterialRefractionData
 	float  RefractionDepthBias;
 };
 
+float GetUnlitMaterialRefractionIOR(in FMaterialRefractionData RefractionData)
+{
+	return RefractionData.Data.x;
+}
+
 float GetMaterialRefractionIOR(in FMaterialRefractionData RefractionData)
 {
 #if REFRACTION_USE_INDEX_OF_REFRACTION
diff --git a/Engine/Shaders/Private/PathTracing/PathTracingMaterialHitShader.usf b/Engine/Shaders/Private/PathTracing/PathTracingMaterialHitShader.usf
--- a/Engine/Shaders/Private/PathTracing/PathTracingMaterialHitShader.usf
+++ b/Engine/Shaders/Private/PathTracing/PathTracingMaterialHitShader.usf
@@ -55,15 +55,15 @@ float AdjustMaterialRoughness(float Roughness, float PathRoughness)
 	return max(Roughness, PathRoughness);
 }
 
-#if REFRACTION_USE_INDEX_OF_REFRACTION
+#if REFRACTION_USE_INDEX_OF_REFRACTION || REFRACTION_USE_INDEX_OF_REFRACTION_FROM_F0
 
 #if SUBSTRATE_ENABLED
 
 float GetRefractionIor(FPixelMaterialInputs PixelMaterialInputs, float3 F0)
 {
-#if REFRACTION_ROOT_NODE_OVERRIDES_DEFAULT
+#if REFRACTION_USE_INDEX_OF_REFRACTION
 	float Ior = GetMaterialRefractionIOR(GetMaterialRefraction(PixelMaterialInputs));
-#else
+#else // REFRACTION_USE_INDEX_OF_REFRACTION_FROM_F0
 	float Ior = DielectricF0RGBToIor(F0);
 #endif
 	return Ior;
@@ -73,9 +73,9 @@ float GetRefractionIor(FPixelMaterialInputs PixelMaterialInputs, float3 F0)
 
 float GetRefractionIor(FPixelMaterialInputs PixelMaterialInputs)
 {
-#if REFRACTION_ROOT_NODE_OVERRIDES_DEFAULT
+#if REFRACTION_USE_INDEX_OF_REFRACTION
 	float Ior = GetMaterialRefractionIOR(GetMaterialRefraction(PixelMaterialInputs));
-#else
+#else // REFRACTION_USE_INDEX_OF_REFRACTION_FROM_F0
 	float3 BaseColor = GetMaterialBaseColor(PixelMaterialInputs);
 	float  Metallic = GetMaterialMetallic(PixelMaterialInputs);
 	float  Specular = GetMaterialSpecular(PixelMaterialInputs);
diff --git a/Engine/Shaders/Private/RayTracing/RayTracingMaterialHitShaders.usf b/Engine/Shaders/Private/RayTracing/RayTracingMaterialHitShaders.usf
--- a/Engine/Shaders/Private/RayTracing/RayTracingMaterialHitShaders.usf
+++ b/Engine/Shaders/Private/RayTracing/RayTracingMaterialHitShaders.usf
@@ -147,9 +147,9 @@ void GetSkyLighting(
 // TODO: Consolidate this definition with similar one in PathTracing Material
 float GetRefractionIor(FPixelMaterialInputs PixelMaterialInputs)
 {
-#if MATERIALBLENDING_TRANSLUCENT && REFRACTION_USE_INDEX_OF_REFRACTION
+#if MATERIALBLENDING_TRANSLUCENT && (REFRACTION_USE_INDEX_OF_REFRACTION || REFRACTION_USE_INDEX_OF_REFRACTION_FROM_F0)
 	// Is the material using refraction at all?
-#if REFRACTION_ROOT_NODE_OVERRIDES_DEFAULT
+#if REFRACTION_USE_INDEX_OF_REFRACTION
 	// Did the user plug something into the Refraction port?
 	// Return the user's IOR value, or 0 (ensure we don't get anything negative)
 	return max(GetMaterialRefractionIOR(GetMaterialRefraction(PixelMaterialInputs)), 0.0);
@@ -386,8 +386,8 @@ void GetMaterialPayload(
 	Payload.Radiance += EmissiveLuminance;
 	Payload.TranslatedWorldPos = DFFastToTranslatedWorld(MaterialParameters.AbsoluteWorldPosition, ResolvedView.PreViewTranslation);
 
-#if MATERIALBLENDING_ANY_TRANSLUCENT && REFRACTION_USE_INDEX_OF_REFRACTION
-	#if REFRACTION_ROOT_NODE_OVERRIDES_DEFAULT
+#if MATERIALBLENDING_ANY_TRANSLUCENT && (REFRACTION_USE_INDEX_OF_REFRACTION || REFRACTION_USE_INDEX_OF_REFRACTION_FROM_F0)
+	#if REFRACTION_USE_INDEX_OF_REFRACTION
 	Payload.IorOverride = max(GetMaterialRefractionIOR(GetMaterialRefraction(PixelMaterialInputs)), 0);
 	#else
 	// IOR will be derived from specular color
diff --git a/Engine/Source/Developer/MaterialBaking/Private/ExportMaterialProxy.h b/Engine/Source/Developer/MaterialBaking/Private/ExportMaterialProxy.h
--- a/Engine/Source/Developer/MaterialBaking/Private/ExportMaterialProxy.h
+++ b/Engine/Source/Developer/MaterialBaking/Private/ExportMaterialProxy.h
@@ -594,7 +594,7 @@ public:
 	virtual bool IsMasked() const override { return ProxyBlendMode == BLEND_Masked; }
 	virtual enum EBlendMode GetBlendMode() const override { return ProxyBlendMode; }
 	virtual enum ERefractionMode GetRefractionMode() const override { return Material ? (ERefractionMode)Material->RefractionMethod : RM_None; }
-	virtual bool GetRootNodeOverridesDefaultRefraction()const override { return Material ? Material->bRootNodeOverridesDefaultDistortion : false; }
+	virtual bool GetRootNodeOverridesDefaultRefraction()const override { return false; }
 
 	virtual FMaterialShadingModelField GetShadingModels() const override
 	{
diff --git a/Engine/Source/Developer/MaterialUtilities/Private/MaterialUtilities.cpp b/Engine/Source/Developer/MaterialUtilities/Private/MaterialUtilities.cpp
--- a/Engine/Source/Developer/MaterialUtilities/Private/MaterialUtilities.cpp
+++ b/Engine/Source/Developer/MaterialUtilities/Private/MaterialUtilities.cpp
@@ -728,7 +728,7 @@ public:
 	virtual bool IsMasked() const override									{ return false; }
 	virtual enum EBlendMode GetBlendMode() const override					{ return BLEND_Opaque; }
 	virtual enum ERefractionMode GetRefractionMode() const override			{ return Material ? (ERefractionMode)Material->RefractionMethod : RM_None; }
-	virtual bool GetRootNodeOverridesDefaultRefraction()const override		{ return Material ? Material->bRootNodeOverridesDefaultDistortion : false; }
+	virtual bool GetRootNodeOverridesDefaultRefraction()const override		{ return false; }
 	virtual FMaterialShadingModelField GetShadingModels() const override	{ return MSM_Unlit; }
 	virtual bool IsShadingModelFromMaterialExpression() const override		{ return false; }
 	virtual float GetOpacityMaskClipValue() const override					{ return 0.5f; }
diff --git a/Engine/Source/Editor/MaterialEditor/Private/MaterialEditor.h b/Engine/Source/Editor/MaterialEditor/Private/MaterialEditor.h
--- a/Engine/Source/Editor/MaterialEditor/Private/MaterialEditor.h
+++ b/Engine/Source/Editor/MaterialEditor/Private/MaterialEditor.h
@@ -134,7 +134,7 @@ public:
 	virtual bool IsWireframe() const override { return false; }
 	virtual bool IsMasked() const override { return false; }
 	virtual enum EBlendMode GetBlendMode() const override { return BLEND_Translucent; }
-	virtual bool GetRootNodeOverridesDefaultRefraction()const override { return false; } // refraction unused for material preview
+	virtual bool GetRootNodeOverridesDefaultRefraction()const override { return false; } 
 	virtual FMaterialShadingModelField GetShadingModels() const override { return MSM_Unlit; }
 	virtual bool IsShadingModelFromMaterialExpression() const override { return false; }
 	virtual float GetOpacityMaskClipValue() const override { return 0.5f; }
diff --git a/Engine/Source/Editor/UnrealEd/Private/Commandlets/DumpMaterialInfo.cpp b/Engine/Source/Editor/UnrealEd/Private/Commandlets/DumpMaterialInfo.cpp
--- a/Engine/Source/Editor/UnrealEd/Private/Commandlets/DumpMaterialInfo.cpp
+++ b/Engine/Source/Editor/UnrealEd/Private/Commandlets/DumpMaterialInfo.cpp
@@ -276,7 +276,6 @@ TArray<MaterialInfo::FPropertySet> GetMaterialInfoProperties()
 		FPropertySet::FromAccessor<bool, &FMaterialResource::UseLmDirectionality>(TEXT("UseLmDirectionality")),
 		FPropertySet::FromEnumAccessor<EBlendMode, &FMaterialResource::GetBlendMode>(TEXT("GetBlendMode")),
 		FPropertySet::FromEnumAccessor<ERefractionMode, &FMaterialResource::GetRefractionMode>(TEXT("GetRefractionMode")),
-		FPropertySet::FromAccessor<bool, &FMaterialResource::GetRootNodeOverridesDefaultRefraction>(TEXT("GetRootNodeOverridesDefaultRefraction")),
 		FPropertySet::FromAccessor<uint32, &FMaterialResource::GetMaterialDecalResponse>(TEXT("GetMaterialDecalResponse")),
 		FPropertySet::FromAccessor<bool, &FMaterialResource::HasBaseColorConnected>(TEXT("HasBaseColorConnected")),
 		FPropertySet::FromAccessor<bool, &FMaterialResource::HasNormalConnected>(TEXT("HasNormalConnected")),
diff --git a/Engine/Source/Editor/UnrealEd/Private/Lightmass/LightmassRender.cpp b/Engine/Source/Editor/UnrealEd/Private/Lightmass/LightmassRender.cpp
--- a/Engine/Source/Editor/UnrealEd/Private/Lightmass/LightmassRender.cpp
+++ b/Engine/Source/Editor/UnrealEd/Private/Lightmass/LightmassRender.cpp
@@ -621,7 +621,7 @@ public:
 	virtual bool IsMasked() const override									{ return false; }
 	virtual enum EBlendMode GetBlendMode() const override					{ return BLEND_Opaque; }
 	virtual enum ERefractionMode GetRefractionMode() const override			{ return Material ? (ERefractionMode)Material->RefractionMethod : RM_None; }
-	virtual bool GetRootNodeOverridesDefaultRefraction()const override		{ return Material ? Material->bRootNodeOverridesDefaultDistortion : false; }
+	virtual bool GetRootNodeOverridesDefaultRefraction()const override		{ return false; }
 	virtual float GetOpacityMaskClipValue() const override					{ return 0.5f; }
 	virtual bool GetCastDynamicShadowAsMasked() const override				{ return false; }
 	virtual FString GetFriendlyName() const override { return FString::Printf(TEXT("FLightmassMaterialRenderer %s"), MaterialInterface ? *MaterialInterface->GetName() : TEXT("NULL")); }
diff --git a/Engine/Source/Editor/UnrealEd/Private/MaterialGraphNode_Root.cpp b/Engine/Source/Editor/UnrealEd/Private/MaterialGraphNode_Root.cpp
--- a/Engine/Source/Editor/UnrealEd/Private/MaterialGraphNode_Root.cpp
+++ b/Engine/Source/Editor/UnrealEd/Private/MaterialGraphNode_Root.cpp
@@ -238,7 +238,7 @@ void UMaterialGraphNode_Root::CreateInputPins()
 					TempStr = TEXT("Disabled");
 					break;
 				case ERefractionMode::RM_IndexOfRefraction:
-					TempStr = TEXT("IOR, from Material F0 if Unplugged");
+					TempStr = TEXT("IOR");
 					break;
 				case ERefractionMode::RM_PixelNormalOffset:
 					TempStr = TEXT("Pixel Normal Offset");
@@ -246,6 +246,9 @@ void UMaterialGraphNode_Root::CreateInputPins()
 				case ERefractionMode::RM_2DOffset:
 					TempStr = TEXT("2D Offset");
 					break;
+				case ERefractionMode::RM_IndexOfRefractionFromF0:
+					TempStr = TEXT("IOR from F0");
+					break;
 				default:
 					TempStr = TEXT("UNKNOWN ERefractionMode");
 				}
diff --git a/Engine/Source/Runtime/Core/Private/UObject/FortniteMainBranchObjectVersion.cpp b/Engine/Source/Runtime/Core/Private/UObject/FortniteMainBranchObjectVersion.cpp
--- a/Engine/Source/Runtime/Core/Private/UObject/FortniteMainBranchObjectVersion.cpp
+++ b/Engine/Source/Runtime/Core/Private/UObject/FortniteMainBranchObjectVersion.cpp
@@ -14,6 +14,6 @@ TMap<FGuid, FGuid> FFortniteMainBranchObjectVersion::GetSystemGuids()
 	SystemGuids.Add(DevGuids.Niagara_LatestScriptCompileVersion, FGuid("783FF3700BB641C799B2D43552FBF91E"));
 	SystemGuids.Add(DevGuids.SkeletalMeshDerivedDataVersion, FGuid("C998F2F3B2884A02B7E290B8C25835F0"));
 	SystemGuids.Add(DevGuids.STATICMESH_DERIVEDDATA_VER, FGuid("3ABF32149F9D4D838750368AB95573BE"));
-	SystemGuids.Add(DevGuids.MaterialTranslationDDCVersion, FGuid("9D0538021BDB47DAA63FFC09DDAAC00E"));
+	SystemGuids.Add(DevGuids.MaterialTranslationDDCVersion, FGuid("0D98D99759864679974E9D06A7CFA7ED"));
 	return SystemGuids;
 }
diff --git a/Engine/Source/Runtime/Engine/Classes/Engine/EngineTypes.h b/Engine/Source/Runtime/Engine/Classes/Engine/EngineTypes.h
--- a/Engine/Source/Runtime/Engine/Classes/Engine/EngineTypes.h
+++ b/Engine/Source/Runtime/Engine/Classes/Engine/EngineTypes.h
@@ -358,7 +358,6 @@ UENUM()
 enum ERefractionMode : int
 {
 	/** 
-	 * By default, when the root node refraction pin is unplugged, relies on the material IOR evaluated from F0.
 	 * Refraction is computed based on the camera vector entering a medium whose index of refraction is defined by the Refraction material input.  
 	 * The new medium's surface is defined by the material's normal.  With this mode, a flat plane seen from the side will have a constant refraction offset.
 	 * This is a physical model of refraction but causes reading outside the scene color texture so is a poor fit for large refractive surfaces like water.
@@ -384,6 +383,13 @@ enum ERefractionMode : int
 	 * Refraction is disabled.
 	 */
 	RM_None UMETA(DisplayName = "None"),
+
+	/**
+	 * Refraction is computed based on the camera vector entering a medium whose index of refraction is defined by the material IOR evaluated from F0.
+	 * The new medium's surface is defined by the material's normal.  With this mode, a flat plane seen from the side will have a constant refraction offset.
+	 * This is a physical model of refraction but causes reading outside the scene color texture so is a poor fit for large refractive surfaces like water.
+	 */
+	RM_IndexOfRefractionFromF0 UMETA(DisplayName = "Index Of Refraction From F0"),
 };
 
 /** Determines how the refraction account for the coverage with Substrate. It can only be used when Substrate is enabled. */
diff --git a/Engine/Source/Runtime/Engine/Private/Materials/Material.cpp b/Engine/Source/Runtime/Engine/Private/Materials/Material.cpp
--- a/Engine/Source/Runtime/Engine/Private/Materials/Material.cpp
+++ b/Engine/Source/Runtime/Engine/Private/Materials/Material.cpp
@@ -4690,17 +4690,12 @@ void UMaterial::PostLoad()
 			// The root node refraction pin was not plugged in so set refraction mode as none.
 			// We do this for all domains and blending modes (translucent, opaque, etc).
 			RefractionMethod = RM_None;
-			bRootNodeOverridesDefaultDistortion = false;
 
 			// Need to mutate since UPROPERTY like bUsesDistortion is changed and shaders need to recompile.
 			static FGuid MaterialRefractionModeOFFConversionGuid(TEXT("094C0316-EA95-4B39-A3FA-CA126A92989B"));
 			ReleaseResourcesAndMutateDDCKey(MaterialRefractionModeOFFConversionGuid);
 		}
-		else
-		{
-			// Keep the current refraction mode and notify that it is overriden on the root node.
-			bRootNodeOverridesDefaultDistortion = true;
-		}
+		// else keep the current refraction mode.
 	}
 #endif // WITH_EDITOR
 
@@ -5311,7 +5306,6 @@ void UMaterial::PostEditChangePropertyInternal(FPropertyChangedEvent& PropertyCh
 
 	// Check for distortion in material 
 	bUsesDistortion = RefractionMethod != RM_None;
-	bRootNodeOverridesDefaultDistortion = bUsesDistortion ? IsRefractionPinPluggedIn(EditorOnly) : false;
 
 	//If we can be sure this material would be the same opaque as it is masked then allow it to be assumed opaque.
 	bCanMaskedBeAssumedOpaque = !EditorOnly->OpacityMask.Expression && !(EditorOnly->OpacityMask.UseConstant && EditorOnly->OpacityMask.Constant < 0.999f) && !bUseMaterialAttributes;
diff --git a/Engine/Source/Runtime/Engine/Private/Materials/MaterialShared.cpp b/Engine/Source/Runtime/Engine/Private/Materials/MaterialShared.cpp
--- a/Engine/Source/Runtime/Engine/Private/Materials/MaterialShared.cpp
+++ b/Engine/Source/Runtime/Engine/Private/Materials/MaterialShared.cpp
@@ -1940,7 +1940,7 @@ ERefractionMode FMaterialResource::GetRefractionMode() const
 
 bool FMaterialResource::GetRootNodeOverridesDefaultRefraction() const
 {
-	return Material->bRootNodeOverridesDefaultDistortion;
+	return false;
 }
 
 FMaterialShadingModelField FMaterialResource::GetShadingModels() const 
@@ -2655,15 +2655,15 @@ void FMaterial::SetupMaterialEnvironment(
 
 	switch(GetRefractionMode())
 	{
-	case RM_IndexOfRefraction:	SET_SHADER_DEFINE(OutEnvironment, REFRACTION_USE_INDEX_OF_REFRACTION,	1); break;
-	case RM_PixelNormalOffset:	SET_SHADER_DEFINE(OutEnvironment, REFRACTION_USE_PIXEL_NORMAL_OFFSET,	1); break;
-	case RM_2DOffset:			SET_SHADER_DEFINE(OutEnvironment, REFRACTION_USE_2D_OFFSET,				1); break;
-	case RM_None:				SET_SHADER_DEFINE(OutEnvironment, REFRACTION_USE_NONE,					1); break;
+	case RM_IndexOfRefraction:			SET_SHADER_DEFINE(OutEnvironment, REFRACTION_USE_INDEX_OF_REFRACTION,			1); break;
+	case RM_PixelNormalOffset:			SET_SHADER_DEFINE(OutEnvironment, REFRACTION_USE_PIXEL_NORMAL_OFFSET,			1); break;
+	case RM_2DOffset:					SET_SHADER_DEFINE(OutEnvironment, REFRACTION_USE_2D_OFFSET,						1); break;
+	case RM_None:						SET_SHADER_DEFINE(OutEnvironment, REFRACTION_USE_NONE,							1); break;
+	case RM_IndexOfRefractionFromF0:	SET_SHADER_DEFINE(OutEnvironment, REFRACTION_USE_INDEX_OF_REFRACTION_FROM_F0,	1); break;
 	default: 
 		UE_LOG(LogMaterial, Warning, TEXT("Unknown material refraction mode: %u  Setting to RM_IndexOfRefraction"),(int32)GetRefractionMode());
 		SET_SHADER_DEFINE(OutEnvironment, REFRACTION_USE_INDEX_OF_REFRACTION, 1);
 	}
-	SET_SHADER_DEFINE(OutEnvironment, REFRACTION_ROOT_NODE_OVERRIDES_DEFAULT, GetRootNodeOverridesDefaultRefraction() ? 1 : 0);
 
 	SET_SHADER_DEFINE(OutEnvironment, USE_DITHERED_LOD_TRANSITION_FROM_MATERIAL, IsDitheredLODTransition());
 	SET_SHADER_DEFINE(OutEnvironment, MATERIAL_TWOSIDED, IsTwoSided());
diff --git a/Engine/Source/Runtime/Engine/Public/MaterialShared.h b/Engine/Source/Runtime/Engine/Public/MaterialShared.h
--- a/Engine/Source/Runtime/Engine/Public/MaterialShared.h
+++ b/Engine/Source/Runtime/Engine/Public/MaterialShared.h
@@ -2344,6 +2344,7 @@ public:
 	virtual bool AllowNegativeEmissiveColor() const { return false; }
 	virtual enum EBlendMode GetBlendMode() const = 0;
 	ENGINE_API virtual enum ERefractionMode GetRefractionMode() const;
+	UE_DEPRECATED(5.7, "This is not longer valid since even when not plugged in, the default value on the root node can be used")
 	virtual bool GetRootNodeOverridesDefaultRefraction() const = 0;
 	virtual FMaterialShadingModelField GetShadingModels() const = 0;
 	virtual bool IsShadingModelFromMaterialExpression() const = 0;
@@ -3144,6 +3145,7 @@ public:
 	ENGINE_API virtual bool UseLmDirectionality() const override;
 	ENGINE_API virtual enum EBlendMode GetBlendMode() const override;
 	ENGINE_API virtual enum ERefractionMode GetRefractionMode() const override;
+	UE_DEPRECATED(5.7, "This is not longer valid since even when not plugged in, the default value on the root node can be used")
 	ENGINE_API virtual bool GetRootNodeOverridesDefaultRefraction() const override;
 	ENGINE_API virtual uint32 GetMaterialDecalResponse() const override;
 	ENGINE_API virtual bool HasBaseColorConnected() const override;
diff --git a/Engine/Source/Runtime/Engine/Public/Materials/Material.h b/Engine/Source/Runtime/Engine/Public/Materials/Material.h
--- a/Engine/Source/Runtime/Engine/Public/Materials/Material.h
+++ b/Engine/Source/Runtime/Engine/Public/Materials/Material.h
@@ -771,12 +771,6 @@ public:
 	UPROPERTY()
 	uint8 bUsesDistortion : 1;
 
-	/** 
-	 * Indicates that the material do not use the material physically based refraction (e.g. IOR from reflectivity F0), but overrides it for artistic purposes.
-	 */
-	 UPROPERTY()
-	uint8 bRootNodeOverridesDefaultDistortion : 1;
-
 	/** 
 	 * Indicates that the material and its instances can be used with clothing
 	 * This will result in the shaders required to support clothing being compiled which will increase shader compile time and memory usage.
