[DataTableJSON] added new implementation for WriteStruct using a PropertyVisitor to be able to output StructUtils types (e.g., FInstancedStruct)
Console variable 'DataTableJSON.ExportUsingPropertyVisitor' can be used to use the previous implementation
Added optional callbacks to FPropertyVisitorPath when new info is added/removed to the path.
Updated uses of FPropertyVisitorScope to make sure it gets pushed/popped for each elements in containers

diff --git a/Engine/Source/Runtime/CoreUObject/Private/UObject/PropertyArray.cpp b/Engine/Source/Runtime/CoreUObject/Private/UObject/PropertyArray.cpp
--- a/Engine/Source/Runtime/CoreUObject/Private/UObject/PropertyArray.cpp
+++ b/Engine/Source/Runtime/CoreUObject/Private/UObject/PropertyArray.cpp
@@ -1456,12 +1456,9 @@ EPropertyVisitorControlFlow FArrayProperty::Visit(FPropertyVisitorContext& Conte
 	{
 		checkf(Inner, TEXT("Expecting a valid inner property type"));
 
-		FPropertyVisitorScope Scope(Context.Path, FPropertyVisitorInfo(Inner));
-
 		for (int32 ContainerIndex = 0; ContainerIndex < ArrayNum; ContainerIndex++)
 		{
-			Context.Path.Top().SetIndex(ContainerIndex, EPropertyVisitorInfoType::ContainerIndex);
-
+			FPropertyVisitorScope Scope(Context.Path, FPropertyVisitorInfo(Inner, ContainerIndex, EPropertyVisitorInfoType::ContainerIndex));
 			FPropertyVisitorContext SubContext = Context.VisitPropertyData(ArrayHelper.GetRawPtr(ContainerIndex));
 
 			RetVal = Inner->Visit(SubContext, InFunc);
diff --git a/Engine/Source/Runtime/CoreUObject/Private/UObject/PropertySet.cpp b/Engine/Source/Runtime/CoreUObject/Private/UObject/PropertySet.cpp
--- a/Engine/Source/Runtime/CoreUObject/Private/UObject/PropertySet.cpp
+++ b/Engine/Source/Runtime/CoreUObject/Private/UObject/PropertySet.cpp
@@ -1135,12 +1135,10 @@ EPropertyVisitorControlFlow FSetProperty::Visit(FPropertyVisitorContext& Context
 	{
 		checkf(ElementProp, TEXT("Expecting a valid inner property type"));
 
-		FPropertyVisitorScope Scope(Context.Path, FPropertyVisitorInfo(ElementProp));
 		for (FScriptSetHelper::FIterator It(SetHelper); It; ++It)
 		{
 			// visit element
-			Context.Path.Top().SetIndex(It.GetLogicalIndex(), EPropertyVisitorInfoType::ContainerIndex);
-
+			FPropertyVisitorScope Scope(Context.Path, FPropertyVisitorInfo(ElementProp, It.GetLogicalIndex(), EPropertyVisitorInfoType::ContainerIndex));
 			FPropertyVisitorContext SubContext = Context.VisitPropertyData(SetHelper.GetElementPtr(It));
 
 			RetVal = ElementProp->Visit(SubContext, InFunc);
diff --git a/Engine/Source/Runtime/CoreUObject/Private/UObject/PropertyVisitor.cpp b/Engine/Source/Runtime/CoreUObject/Private/UObject/PropertyVisitor.cpp
--- a/Engine/Source/Runtime/CoreUObject/Private/UObject/PropertyVisitor.cpp
+++ b/Engine/Source/Runtime/CoreUObject/Private/UObject/PropertyVisitor.cpp
@@ -113,12 +113,10 @@ EPropertyVisitorControlFlow PropertyVisitorHelpers::VisitProperty(const UStruct*
 {
 	EPropertyVisitorControlFlow RetVal = EPropertyVisitorControlFlow::StepOver;
 
-	FPropertyVisitorScope Scope(Context.Path, FPropertyVisitorInfo(Property));
-	Context.Path.Top().ParentStructType = PropertyOwner;
-
 	// ArrayDim of one means it is just a single property, not a static array
 	if (Property->ArrayDim == 1)
 	{
+		FPropertyVisitorScope Scope(Context.Path, FPropertyVisitorInfo(Property, PropertyOwner));
 		FPropertyVisitorData Data(Property->ContainerPtrToValuePtr<void>(Context.Data.PropertyData), /*ParentStructData*/Context.Data.PropertyData);
 		FPropertyVisitorContext SubContext(Context.Path, Data, Context.Scope);
 
@@ -129,7 +127,7 @@ EPropertyVisitorControlFlow PropertyVisitorHelpers::VisitProperty(const UStruct*
 		// In the case of a static array, we need to setup the EPropertyVisitorInfoType
 		for (int32 StaticArrayIndex = 0; StaticArrayIndex < Property->ArrayDim; ++StaticArrayIndex)
 		{
-			Context.Path.Top().SetIndex(StaticArrayIndex, EPropertyVisitorInfoType::StaticArrayIndex);
+			FPropertyVisitorScope Scope(Context.Path, FPropertyVisitorInfo(Property, StaticArrayIndex, EPropertyVisitorInfoType::StaticArrayIndex, PropertyOwner));
 
 			FPropertyVisitorData Data(Property->ContainerPtrToValuePtr<void>(Context.Data.PropertyData, StaticArrayIndex), /*ParentStructData*/Context.Data.PropertyData);
 			FPropertyVisitorContext SubContext(Context.Path, Data, Context.Scope);
diff --git a/Engine/Source/Runtime/CoreUObject/Public/UObject/PropertyVisitor.h b/Engine/Source/Runtime/CoreUObject/Public/UObject/PropertyVisitor.h
--- a/Engine/Source/Runtime/CoreUObject/Public/UObject/PropertyVisitor.h
+++ b/Engine/Source/Runtime/CoreUObject/Public/UObject/PropertyVisitor.h
@@ -4,7 +4,7 @@
 
 #include "Containers/Array.h"
 #include "Containers/ArrayView.h"
-#include "Templates/FunctionFwd.h"
+#include "Templates/Function.h"
 
 class FProperty;
 class FEditPropertyChain;
@@ -30,12 +30,23 @@ enum class EPropertyVisitorInfoType : uint8
 
 struct FPropertyVisitorInfo
 {
-	explicit FPropertyVisitorInfo(const FProperty* InProperty, int32 InIndex = INDEX_NONE, EPropertyVisitorInfoType InPropertyInfo = EPropertyVisitorInfoType::None)
+	explicit FPropertyVisitorInfo(const FProperty* InProperty
+		, int32 InIndex = INDEX_NONE
+		, EPropertyVisitorInfoType InPropertyInfo = EPropertyVisitorInfoType::None
+		, const UStruct* InParentStructType = nullptr
+		)
 		: Property(InProperty)
+		, ParentStructType(InParentStructType)
 		, Index(InIndex)
 		, PropertyInfo(InPropertyInfo)
-		, bContainsInnerProperties(false)
-	{}
+	{
+	}
+
+	explicit FPropertyVisitorInfo(const FProperty* InProperty, const UStruct* InParentStructType)
+	: Property(InProperty)
+	, ParentStructType(InParentStructType)
+	{
+	}
 
 	FPropertyVisitorInfo(const FPropertyVisitorInfo&) = default;
 	FPropertyVisitorInfo(FPropertyVisitorInfo&&) = default;
@@ -43,12 +54,6 @@ struct FPropertyVisitorInfo
 	FPropertyVisitorInfo& operator=(const FPropertyVisitorInfo&) = default;
 	FPropertyVisitorInfo& operator=(FPropertyVisitorInfo&&) = default;
 
-	void SetIndex(int32 InIndex, EPropertyVisitorInfoType InPropertyInfo)
-	{
-		Index = InIndex;
-		PropertyInfo = InPropertyInfo;
-	}
-
 	COREUOBJECT_API bool Identical(const FPropertyVisitorInfo& Other) const;
 
 	/** @note The default comparison only compares the key data to match GetTypeHash and is for use with hashed containers and ResolveVisitedPathInfo_Generic; use Identical for an exact comparison */
@@ -67,11 +72,11 @@ struct FPropertyVisitorInfo
 	}
 
 	/** The property currently being visited */
-	const FProperty* Property;
+	const FProperty* Property = nullptr;
 
 	/**
 	 * The parent struct that provided the property being iterated, if iterating a sub-property within a struct.
-	 * @note This is slighty different than Property->GetOwnerStruct() as you might be iterating a FDerived instance 
+	 * @note This is slightly different from Property->GetOwnerStruct() as you might be iterating a FDerived instance 
 	 *       but processing a FBase struct property. In this case this will be set to FDerived rather than FBase.
 	 */
 	const UStruct* ParentStructType = nullptr;
@@ -79,20 +84,25 @@ struct FPropertyVisitorInfo
 	/**
      * Index of the element being visited in the container, otherwise INDEX_NONE.
      * For maps and sets it indicates the logical index. */
-	int32 Index;
+	int32 Index = INDEX_NONE;
 
 	/** Whether this property is inside a container and if it is key or a value of a map */
-	EPropertyVisitorInfoType PropertyInfo;
+	EPropertyVisitorInfoType PropertyInfo = EPropertyVisitorInfoType::None;
 
 	/* Indicate that this property contains inner properties */
-	bool bContainsInnerProperties;
+	bool bContainsInnerProperties = false;
 };
 
 struct FPropertyVisitorPath
 {
 public:
 	FPropertyVisitorPath() = default;
-	
+	FPropertyVisitorPath(TFunction<void(const FPropertyVisitorInfo&)> InOnPushFunc, TFunction<void(const FPropertyVisitorInfo&)> InOnPopFunc)
+		: OnPushFunc(InOnPushFunc)
+		, OnPopFunc(InOnPopFunc)
+	{
+	}
+
 	FPropertyVisitorPath(const FPropertyVisitorPath&) = default;
 	FPropertyVisitorPath(FPropertyVisitorPath&&) = default;
 
@@ -115,11 +125,20 @@ public:
 	void Push(const FPropertyVisitorInfo& Info)
 	{
 		Path.Push(Info);
+		if (OnPushFunc)
+		{
+			OnPushFunc(Info);
+		}
 	}
 
 	FPropertyVisitorInfo Pop()
 	{
-		return Path.Pop(EAllowShrinking::No);
+		FPropertyVisitorInfo Info = Path.Pop(EAllowShrinking::No);
+		if (OnPopFunc)
+		{
+			OnPopFunc(Info);
+		}
+		return MoveTemp(Info);
 	}
 
 	int32 Num() const
@@ -204,6 +223,8 @@ public:
 
 protected:
 	TArray<FPropertyVisitorInfo> Path;
+	TFunction<void(const FPropertyVisitorInfo&)> OnPushFunc;
+	TFunction<void(const FPropertyVisitorInfo&)> OnPopFunc;
 };
 
 struct FPropertyVisitorScope
diff --git a/Engine/Source/Runtime/Engine/Private/DataTableJSON.cpp b/Engine/Source/Runtime/Engine/Private/DataTableJSON.cpp
--- a/Engine/Source/Runtime/Engine/Private/DataTableJSON.cpp
+++ b/Engine/Source/Runtime/Engine/Private/DataTableJSON.cpp
@@ -1,12 +1,13 @@
 // Copyright Epic Games, Inc. All Rights Reserved.
 
 #include "DataTableJSON.h"
-#include "UObject/EnumProperty.h"
 #include "Engine/DataTable.h"
+#include "HAL/IConsoleManager.h"
 #include "Serialization/JsonSerializer.h"
+#include "UObject/EnumProperty.h"
 
 
-namespace
+namespace UE::DataTableJSON
 {
 	const TCHAR* JSONTypeToString(const EJson InType)
 	{
@@ -30,6 +31,7 @@ namespace
 			return TEXT("Unknown");
 		}
 	}
+
 #if WITH_EDITOR
 
 	template <typename CharType>
@@ -58,6 +60,19 @@ namespace
 		}
 	}
 
+	/**
+	 * Console variable to be able to run the previous implementation of the exporter
+	 * in case issues get reported and validate potential regression.
+	 * We could eventually get rid of the FieldIterator implementation.
+	 */
+	bool bExportUsingPropertyVisitor = true;
+	static FAutoConsoleVariableRef CVarRuntimeValidationContext(
+		TEXT("DataTableJSON.ExportUsingPropertyVisitor"),
+		bExportUsingPropertyVisitor,
+		TEXT("True indicates that the PropertyVisitor will be used to export the data table row including support for new types like instanced struct."
+				" Otherwise the legacy implementation based on the FieldIterator is used.")
+	);
+
 #endif	// WITH_EDITOR
 
 }
@@ -174,6 +189,20 @@ bool TDataTableExporterJSON<CharType>::WriteRow(const UScriptStruct* InRowStruct
 
 template<typename CharType>
 bool TDataTableExporterJSON<CharType>::WriteStruct(const UScriptStruct* InStruct, const void* InStructData, const FString* FieldToSkip)
+{
+	if (UE::DataTableJSON::bExportUsingPropertyVisitor)
+	{
+		return InStruct ? WriteStructUsingVisitor(InStruct, InStructData, FieldToSkip) : false;
+	}
+
+	return WriteStructUsingFieldIterator(InStruct, InStructData, FieldToSkip);
+}
+
+template <typename CharType>
+bool TDataTableExporterJSON<CharType>::WriteStructUsingFieldIterator(
+	const UScriptStruct* InStruct
+	, const void* InStructData
+	, const FString* FieldToSkip)
 {
 	for (TFieldIterator<const FProperty> It(InStruct); It; ++It)
 	{
@@ -316,6 +345,8 @@ bool TDataTableExporterJSON<CharType>::WriteStructEntry(const void* InRowData, c
 template<typename CharType>
 bool TDataTableExporterJSON<CharType>::WriteContainerEntry(const FProperty* InProperty, const void* InPropertyData, const FString* InIdentifier)
 {
+	using namespace UE::DataTableJSON;
+
 	if (const FEnumProperty* EnumProp = CastField<const FEnumProperty>(InProperty))
 	{
 		const FString PropertyValue = DataTableUtils::GetPropertyValueAsStringDirect(InProperty, (uint8*)InPropertyData, DTExportFlags);
@@ -382,6 +413,223 @@ bool TDataTableExporterJSON<CharType>::WriteContainerEntry(const FProperty* InPr
 	return true;
 }
 
+template <typename CharType>
+bool TDataTableExporterJSON<CharType>::WriteStructUsingVisitor(TNotNull<const UScriptStruct*> InStruct, const void* InStructData
+	, const FString* FieldToSkip)
+{
+	using namespace UE::DataTableJSON;
+
+	// Keep a stack of identifier overrides since when entering arrays we don't want struct to output their identifier
+	// We always push override for structs and arrays so an empty override list or an unset value on top of the stack
+	// indicates that no override is effective and the property export name should be used.
+	TArray<TOptional<FString>> IdentifierOverrides;
+
+	auto PushFunc = [this, &IdentifierOverrides](const FPropertyVisitorInfo& InVisitorInfo)
+		{
+			const FProperty* Property = InVisitorInfo.Property;
+			const FString Identifier = DataTableUtils::GetPropertyExportName(Property, DTExportFlags);
+
+			if (InVisitorInfo.PropertyInfo == EPropertyVisitorInfoType::MapKey
+				|| InVisitorInfo.PropertyInfo == EPropertyVisitorInfoType::MapValue)
+			{
+				// Special case for map keys since we won't write them but use their values as string when writing the map values.
+				// In order to do so we don't push a new identifier, and we will set it when visiting the key value.
+				// Then we can reset when popping the values.
+				return;
+			}
+
+			// Process static array first to properly handle arrays of structs
+			if (Property->ArrayDim != 1)
+			{
+				ensure(InVisitorInfo.PropertyInfo == EPropertyVisitorInfoType::StaticArrayIndex);
+				if (InVisitorInfo.Index == 0)
+				{
+					JsonWriter->WriteArrayStart(Identifier);
+
+					// Push override (empty identifier) for array elements
+					IdentifierOverrides.AddDefaulted_GetRef().Emplace(FString{});
+				}
+			}
+
+			if (Property->IsA<FMapProperty>())
+			{
+				JsonWriter->WriteObjectStart(Identifier);
+			}
+			else if (Property->IsA<FArrayProperty>()
+				|| Property->IsA<FSetProperty>())
+			{
+				JsonWriter->WriteArrayStart(Identifier);
+
+				// Push override (empty identifier) for array elements
+				IdentifierOverrides.AddDefaulted_GetRef().Emplace(FString{});
+			}
+			else if (Property->IsA<FStructProperty>())
+			{
+				if (!!(DTExportFlags & EDataTableExportFlags::UseJsonObjectsForStructs))
+				{
+					if (IdentifierOverrides.Num() && IdentifierOverrides.Top().IsSet())
+					{
+						const FString* Override = IdentifierOverrides.Top().GetValue().IsEmpty() ? nullptr : &IdentifierOverrides.Top().GetValue();
+						WriteJSONObjectStartWithOptionalIdentifier<CharType>(*JsonWriter, Override);
+					}
+					else
+					{
+						JsonWriter->WriteObjectStart(Identifier);
+					}
+				}
+
+				// Push a "no override" entry when entering a struct
+				IdentifierOverrides.AddDefaulted();
+			}
+		};
+
+	auto PopFunc = [this, &IdentifierOverrides](const FPropertyVisitorInfo& InVisitorInfo)
+		{
+			// Map keys don't do anything when pushed in the visitor scope, so nothing to do here
+			if (InVisitorInfo.PropertyInfo == EPropertyVisitorInfoType::MapKey)
+			{
+				return;
+			}
+
+			// Map element is done after writing its value; we can reset the override set by the key and close its WriteObjectStart.
+			if (InVisitorInfo.PropertyInfo == EPropertyVisitorInfoType::MapValue)
+			{
+				JsonWriter->WriteObjectEnd();
+				return;
+			}
+
+			const FProperty* Property = InVisitorInfo.Property;
+
+			if (Property->IsA<FMapProperty>())
+			{
+				JsonWriter->WriteObjectEnd();
+			}
+			else if (Property->IsA<FArrayProperty>()
+				|| Property->IsA<FSetProperty>())
+			{
+				JsonWriter->WriteArrayEnd();
+				IdentifierOverrides.Pop();
+			}
+			else if (Property->IsA<FStructProperty>())
+			{
+				if (!!(DTExportFlags & EDataTableExportFlags::UseJsonObjectsForStructs))
+				{
+					JsonWriter->WriteObjectEnd();
+				}
+				IdentifierOverrides.Pop();
+			}
+
+			// Process static array last to properly close an array of structs
+			if (Property->ArrayDim != 1)
+			{
+				ensure(InVisitorInfo.PropertyInfo == EPropertyVisitorInfoType::StaticArrayIndex);
+				if (InVisitorInfo.Index + 1 == Property->ArrayDim)
+				{
+					JsonWriter->WriteArrayEnd();
+					IdentifierOverrides.Pop();
+				}
+			}
+		};
+
+	// Visitor
+	FPropertyVisitorPath Path(PushFunc, PopFunc);
+	FPropertyVisitorData Data(const_cast<void*>(InStructData), /*ParentStructData*/nullptr);
+	FPropertyVisitorContext Context(Path, Data, FPropertyVisitorContext::EScope::All);
+
+	InStruct->Visit(Context,
+		[FieldToSkip, this, &IdentifierOverrides](const FPropertyVisitorContext& InContext)-> EPropertyVisitorControlFlow
+		{
+			const FProperty* Property = InContext.Path.Top().Property;
+			if (FieldToSkip && *FieldToSkip == DataTableUtils::GetPropertyExportName(Property, DTExportFlags))
+			{
+				// Skip this field
+				return EPropertyVisitorControlFlow::StepOver;
+			}
+
+			// Special case for map when we use the key to set the identifier for the values
+			if (InContext.Path.Top().PropertyInfo == EPropertyVisitorInfoType::MapKey)
+			{
+				const FString Identifier = (DataTableUtils::GetPropertyValueAsString(Property, static_cast<uint8*>(InContext.Data.PropertyData), DTExportFlags));
+				JsonWriter->WriteObjectStart(Identifier);
+				return EPropertyVisitorControlFlow::StepOver;
+			}
+
+			// Use override identifier if one is set on the stack
+			if (IdentifierOverrides.Num() && IdentifierOverrides.Top().IsSet())
+			{
+				const FString* Override = IdentifierOverrides.Top().GetValue().IsEmpty() ? nullptr : &IdentifierOverrides.Top().GetValue();
+				return WriteVisitedEntry(Property, InContext.Data.PropertyData, Override);
+			}
+
+			const FString ExportName = DataTableUtils::GetPropertyExportName(Property, DTExportFlags);
+			return WriteVisitedEntry(Property, InContext.Data.PropertyData, &ExportName);
+		});
+
+	return true;
+}
+
+template<typename CharType>
+EPropertyVisitorControlFlow TDataTableExporterJSON<CharType>::WriteVisitedEntry(TNotNull<const FProperty*> InProperty, const void* InValueAddress, const FString* InIdentifier)
+{
+	using namespace UE::DataTableJSON;
+
+	if (const FEnumProperty* EnumProp = CastField<const FEnumProperty>(InProperty))
+	{
+		const FString PropertyValue = DataTableUtils::GetPropertyValueAsString(EnumProp, (uint8*)InValueAddress, DTExportFlags);
+		WriteJSONValueWithOptionalIdentifier<CharType>(*JsonWriter, InIdentifier, PropertyValue);
+	}
+	else if (const FNumericProperty* NumProp = CastField<const FNumericProperty>(InProperty))
+	{
+		if (NumProp->IsEnum())
+		{
+			const FString PropertyValue = DataTableUtils::GetPropertyValueAsString(InProperty, (uint8*)InValueAddress, DTExportFlags);
+			WriteJSONValueWithOptionalIdentifier<CharType>(*JsonWriter, InIdentifier, PropertyValue);
+		}
+		else if (NumProp->IsInteger())
+		{
+			const int64 PropertyValue = NumProp->GetSignedIntPropertyValue(InValueAddress);
+			WriteJSONValueWithOptionalIdentifier<CharType>(*JsonWriter, InIdentifier, PropertyValue);
+		}
+		else if (NumProp->IsA<FFloatProperty>())
+		{
+			const float PropertyValue = (float)NumProp->GetFloatingPointPropertyValue(InValueAddress);
+			WriteJSONValueWithOptionalIdentifier<CharType>(*JsonWriter, InIdentifier, PropertyValue);
+		}
+		else
+		{
+			const double PropertyValue = NumProp->GetFloatingPointPropertyValue(InValueAddress);
+			WriteJSONValueWithOptionalIdentifier<CharType>(*JsonWriter, InIdentifier, PropertyValue);
+		}
+	}
+	else if (const FBoolProperty* BoolProp = CastField<const FBoolProperty>(InProperty))
+	{
+		const bool PropertyValue = BoolProp->GetPropertyValue(InValueAddress);
+		WriteJSONValueWithOptionalIdentifier<CharType>(*JsonWriter, InIdentifier, PropertyValue);
+	}
+	else if (InProperty->IsA<const FStructProperty>())
+	{
+		if (!(DTExportFlags & EDataTableExportFlags::UseJsonObjectsForStructs))
+		{
+			const FString PropertyValue = DataTableUtils::GetPropertyValueAsString(InProperty, (uint8*)InValueAddress, DTExportFlags);
+			WriteJSONValueWithOptionalIdentifier<CharType>(*JsonWriter, InIdentifier, PropertyValue);
+			return EPropertyVisitorControlFlow::StepOver;
+		}
+	}
+	else if (InProperty->IsA<FArrayProperty>()
+		|| InProperty->IsA<FSetProperty>()
+		|| InProperty->IsA<FMapProperty>())
+	{
+		// Cannot nest arrays, sets and maps
+	}
+	else
+	{
+		const FString PropertyValue = DataTableUtils::GetPropertyValueAsStringDirect(InProperty, (uint8*)InValueAddress, DTExportFlags);
+		WriteJSONValueWithOptionalIdentifier<CharType>(*JsonWriter, InIdentifier, PropertyValue);
+	}
+
+	return EPropertyVisitorControlFlow::StepInto;
+};
+
 template class TDataTableExporterJSON<TCHAR>;
 template class TDataTableExporterJSON<ANSICHAR>;
 
@@ -599,7 +847,7 @@ bool FDataTableImporterJSON::ReadStruct(const TSharedRef<FJsonObject>& InParsedO
 		}
 		else
 		{
-			const TCHAR* const ParsedPropertyType = JSONTypeToString(ParsedPropertyValue->Type);
+			const TCHAR* const ParsedPropertyType = UE::DataTableJSON::JSONTypeToString(ParsedPropertyValue->Type);
 
 			const TArray< TSharedPtr<FJsonValue> >* PropertyValuesPtr;
 			if (!ParsedPropertyValue->TryGetArray(PropertyValuesPtr))
@@ -630,7 +878,7 @@ bool FDataTableImporterJSON::ReadStruct(const TSharedRef<FJsonObject>& InParsedO
 
 bool FDataTableImporterJSON::ReadStructEntry(const TSharedRef<FJsonValue>& InParsedPropertyValue, const FName InRowName, const FString& InColumnName, const void* InRowData, FProperty* InProperty, void* InPropertyData)
 {
-	const TCHAR* const ParsedPropertyType = JSONTypeToString(InParsedPropertyValue->Type);
+	const TCHAR* const ParsedPropertyType = UE::DataTableJSON::JSONTypeToString(InParsedPropertyValue->Type);
 
 	if (FEnumProperty* EnumProp = CastField<FEnumProperty>(InProperty))
 	{
@@ -822,7 +1070,7 @@ bool FDataTableImporterJSON::ReadStructEntry(const TSharedRef<FJsonValue>& InPar
 
 bool FDataTableImporterJSON::ReadContainerEntry(const TSharedRef<FJsonValue>& InParsedPropertyValue, const FName InRowName, const FString& InColumnName, const int32 InArrayEntryIndex, FProperty* InProperty, void* InPropertyData)
 {
-	const TCHAR* const ParsedPropertyType = JSONTypeToString(InParsedPropertyValue->Type);
+	const TCHAR* const ParsedPropertyType = UE::DataTableJSON::JSONTypeToString(InParsedPropertyValue->Type);
 
 	if (FEnumProperty* EnumProp = CastField<FEnumProperty>(InProperty))
 	{
diff --git a/Engine/Source/Runtime/Engine/Private/DataTableJSON.h b/Engine/Source/Runtime/Engine/Private/DataTableJSON.h
--- a/Engine/Source/Runtime/Engine/Private/DataTableJSON.h
+++ b/Engine/Source/Runtime/Engine/Private/DataTableJSON.h
@@ -9,6 +9,7 @@
 
 class UDataTable;
 enum class EDataTableExportFlags : uint8;
+enum class EPropertyVisitorControlFlow : uint8;
 
 template <class CharType> struct TPrettyJsonPrintPolicy;
 
@@ -52,9 +53,16 @@ public:
 	bool WriteStruct(const UScriptStruct* InStruct, const void* InStructData, const FString* FieldToSkip = nullptr);
 
 protected:
+	//~ Field iterator based methods BEGIN
+	bool WriteStructUsingFieldIterator(const UScriptStruct* InStruct, const void* InStructData, const FString* FieldToSkip = nullptr);
 	bool WriteStructEntry(const void* InRowData, const FProperty* InProperty, const void* InPropertyData);
-
 	bool WriteContainerEntry(const FProperty* InProperty, const void* InPropertyData, const FString* InIdentifier = nullptr);
+	//~ Field iterator based methods END
+
+	//~ PropertyVisitor based methods BEGIN
+	bool WriteStructUsingVisitor(TNotNull<const UScriptStruct*> InStruct, const void* InStructData, const FString* FieldToSkip = nullptr);
+	EPropertyVisitorControlFlow WriteVisitedEntry(TNotNull<const FProperty*> InProperty, const void* InPropertyData, const FString* InIdentifier = nullptr);
+	//~ PropertyVisitor based methods END
 
 	EDataTableExportFlags DTExportFlags;
 	TSharedRef<FDataTableJsonWriter> JsonWriter;
