World Partition Actor Component Descriptors: make sure to serialize the component name as part of the initialization data to allow matching with the base class component for delta serialization.

diff --git a/Engine/Source/Runtime/Engine/Private/WorldPartition/WorldPartitionActorDesc.cpp b/Engine/Source/Runtime/Engine/Private/WorldPartition/WorldPartitionActorDesc.cpp
--- a/Engine/Source/Runtime/Engine/Private/WorldPartition/WorldPartitionActorDesc.cpp
+++ b/Engine/Source/Runtime/Engine/Private/WorldPartition/WorldPartitionActorDesc.cpp
@@ -72,7 +72,7 @@ void FWorldPartitionComponentDesc::Init(const UActorComponent* InComponent)
 
 void FWorldPartitionComponentDesc::Serialize(FArchive& Ar)
 {
-	Ar << BaseClass << ComponentName;		
+	Ar << BaseClass;
 }
 
 FWorldPartitionActorDesc::FWorldPartitionActorDesc()
@@ -962,29 +962,30 @@ void FWorldPartitionActorDesc::Serialize(FArchive& Ar)
 
 	if (Ar.CustomVer(FFortniteMainBranchObjectVersion::GUID) >= FFortniteMainBranchObjectVersion::WorldPartitionActorComponentDesc)
 	{
-		TArray<FTopLevelAssetPath> ComponentDescsType;
+		TArray<FWorldPartitionComponentDescInitData> ComponentDescsInitData;
 
 		if (Ar.IsSaving())
 		{
-			ComponentDescsType.Reserve(ComponentDescs.Num());
+			ComponentDescsInitData.Reserve(ComponentDescs.Num());
 			for (const TUniquePtr<FWorldPartitionComponentDesc>& ComponentDesc : ComponentDescs)
 			{
-				ComponentDescsType.Add(ComponentDesc->NativeClass);
+				ComponentDescsInitData.Add({ ComponentDesc->NativeClass, ComponentDesc->ComponentName });
 			}
 		}
 
-		Ar << ComponentDescsType;
+		Ar << ComponentDescsInitData;
 
 		if (Ar.IsLoading())
 		{
-			for (const FTopLevelAssetPath& ComponentDescType : ComponentDescsType)
+			for (const FWorldPartitionComponentDescInitData& ComponentDescInitData : ComponentDescsInitData)
 			{
-				UClass* ComponentNativeClass = FWorldPartitionActorDescUtils::GetNativeClassFromClassName(*ComponentDescType.ToString());
+				UClass* ComponentNativeClass = FWorldPartitionActorDescUtils::GetNativeClassFromClassName(*ComponentDescInitData.Type.ToString());
 				check(ComponentNativeClass);
 
 				TUniquePtr<class FWorldPartitionComponentDesc> ComponentDesc = UActorComponent::StaticCreateClassComponentDesc(ComponentNativeClass);
 				ComponentDesc->NativeClass = *ComponentNativeClass->GetPathName();
 				ComponentDesc->ComponentNativeClass = ComponentNativeClass;
+				ComponentDesc->ComponentName = ComponentDescInitData.Name;
 
 				ComponentDescs.Add(MoveTemp(ComponentDesc));
 			}
diff --git a/Engine/Source/Runtime/Engine/Public/WorldPartition/WorldPartitionActorDesc.h b/Engine/Source/Runtime/Engine/Public/WorldPartition/WorldPartitionActorDesc.h
--- a/Engine/Source/Runtime/Engine/Public/WorldPartition/WorldPartitionActorDesc.h
+++ b/Engine/Source/Runtime/Engine/Public/WorldPartition/WorldPartitionActorDesc.h
@@ -233,6 +233,19 @@ inline bool CompareUnsortedArrays(const TArray<FName>& Array1, const TArray<FNam
 /**
  * Represents a Actor Metadata (editor-only)
  */
+struct FWorldPartitionComponentDescInitData
+{
+#if WITH_EDITOR
+	FTopLevelAssetPath Type;
+	FName Name;
+
+	friend FArchive& operator<<(FArchive& Ar, FWorldPartitionComponentDescInitData& ComponentDescInitData)
+	{
+		return Ar << ComponentDescInitData.Type << ComponentDescInitData.Name;
+	}
+#endif // WITH_EDITOR
+};
+
 class FWorldPartitionComponentDesc
 {
 #if WITH_EDITOR
@@ -246,9 +259,11 @@ public:
 	ENGINE_API virtual void Init(const UActorComponent* InComponent);
 	ENGINE_API virtual void Serialize(FArchive& Ar);
 	
-	UClass* GetComponentNativeClass() const { return ComponentNativeClass; }
+	inline FTopLevelAssetPath GetBaseClass() const { return BaseClass; }
+	inline FTopLevelAssetPath GetNativeClass() const { return NativeClass; }
+	inline UClass* GetComponentNativeClass() const { return ComponentNativeClass; }
 
-	FString ToString() const
+	inline FString ToString() const
 	{
 		return FString::Printf(TEXT("BaseClass=%s, NativeClass=%s, ComponentName=%s"), *BaseClass.ToString(), *NativeClass.ToString(), *ComponentName.ToString());
 	}
