Disable HTTP tests in AutoRTFMEngineTests on CI.

These occasionally flake and cause spurious one-off errors on Horde.

diff --git a/Engine/Source/Runtime/Engine/Private/Tests/AutoRTFM/AutoRTFMHttpModuleTests.cpp b/Engine/Source/Runtime/Engine/Private/Tests/AutoRTFM/AutoRTFMHttpModuleTests.cpp
--- a/Engine/Source/Runtime/Engine/Private/Tests/AutoRTFM/AutoRTFMHttpModuleTests.cpp
+++ b/Engine/Source/Runtime/Engine/Private/Tests/AutoRTFM/AutoRTFMHttpModuleTests.cpp
@@ -72,7 +72,8 @@ private:
 
 IMPLEMENT_SIMPLE_AUTOMATION_TEST(FAutoRTFMHttpModuleNetworking, "AutoRTFM + HttpModuleNetworking", \
 								 EAutomationTestFlags::EngineFilter | EAutomationTestFlags::ClientContext | \
-								 EAutomationTestFlags::ServerContext | EAutomationTestFlags::CommandletContext)
+								 EAutomationTestFlags::ServerContext | EAutomationTestFlags::CommandletContext | \
+								 EAutomationTestFlags::RequiresUser /* to disable testing on CI */)
 
 bool FAutoRTFMHttpModuleNetworking::RunTest(const FString& Parameters)
 {
