[SimpleStreamableAssetManager] Removed FRenderAssetStreamingManager::AddReferencedObjects because it could retain references to trashed packages.
Refactored SimpleStreamableAssetManager to ensure UpdatePerfectWantedMips_Async no longer dereferences UStreamableRenderAsset.

diff --git a/Engine/Source/Runtime/Engine/Classes/Engine/StreamableRenderAsset.h b/Engine/Source/Runtime/Engine/Classes/Engine/StreamableRenderAsset.h
--- a/Engine/Source/Runtime/Engine/Classes/Engine/StreamableRenderAsset.h
+++ b/Engine/Source/Runtime/Engine/Classes/Engine/StreamableRenderAsset.h
@@ -293,7 +293,7 @@ protected:
 	
 	friend class FAsyncRenderAssetStreamingData;
 	friend class FSimpleStreamableAssetManager;
-	int32 SimpleStreamableAssetManagerIndex = INDEX_NONE;
+	TSharedPtr<int32, ESPMode::ThreadSafe> SimpleStreamableAssetManagerIndex = MakeShared<int32, ESPMode::ThreadSafe>(INDEX_NONE);
 public:
 	UPROPERTY(EditAnywhere, BlueprintReadWrite, Category = LevelOfDetail, AssetRegistrySearchable, AdvancedDisplay)
 	uint8 NeverStream : 1;
diff --git a/Engine/Source/Runtime/Engine/Private/Streaming/AsyncTextureStreaming.cpp b/Engine/Source/Runtime/Engine/Private/Streaming/AsyncTextureStreaming.cpp
--- a/Engine/Source/Runtime/Engine/Private/Streaming/AsyncTextureStreaming.cpp
+++ b/Engine/Source/Runtime/Engine/Private/Streaming/AsyncTextureStreaming.cpp
@@ -183,7 +183,7 @@ void FAsyncRenderAssetStreamingData::UpdatePerfectWantedMips_Async(FStreamingRen
 		const EStreamableRenderAssetType AssetType = StreamingRenderAsset.RenderAssetType;
 		DynamicInstancesView.GetRenderAssetScreenSize(AssetType, RenderAsset, MaxSize, MaxSize_VisibleOnly, MaxNumForcedLODs, MaxAssetSize, bOutputToLog ? TEXT("Dynamic") : nullptr);
 		
-		if (FSimpleStreamableAssetManager::IsEnabled() && RenderAsset->SimpleStreamableAssetManagerIndex != INDEX_NONE)
+		if (FSimpleStreamableAssetManager::IsEnabled() && *StreamingRenderAsset.SimpleStreamableAssetManagerIndex != INDEX_NONE)
 		{
 			// No need to iterate more if render asset is already at maximum resolution.
 			const bool bAlreadyAtMaximumResolution = MaxSize_VisibleOnly >= MaxAssetSize || MaxNumForcedLODs >= StreamingRenderAsset.MaxAllowedMips;
@@ -193,7 +193,7 @@ void FAsyncRenderAssetStreamingData::UpdatePerfectWantedMips_Async(FStreamingRen
 				float TmpMaxVisibleOnly = MaxSize_VisibleOnly;
 				int32 TmpMaxNumForcedLODs = MaxNumForcedLODs;
 				
-				FSimpleStreamableAssetManager::GetRenderAssetScreenSize(AssetType, RenderAsset, TmpMaxSize, TmpMaxVisibleOnly, TmpMaxNumForcedLODs, MaxAssetSize, StreamingRenderAsset.MaxAllowedMips, bOutputToLog ? TEXT("Static") : nullptr);
+				FSimpleStreamableAssetManager::GetRenderAssetScreenSize(AssetType, *StreamingRenderAsset.SimpleStreamableAssetManagerIndex, TmpMaxSize, TmpMaxVisibleOnly, TmpMaxNumForcedLODs, MaxAssetSize, StreamingRenderAsset.MaxAllowedMips, bOutputToLog ? TEXT("Static") : nullptr);
 
 				MaxSize = FMath::Max(TmpMaxSize, MaxSize);
 				MaxSize_VisibleOnly = FMath::Max(TmpMaxVisibleOnly, MaxSize_VisibleOnly);
diff --git a/Engine/Source/Runtime/Engine/Private/Streaming/SimpleStreamableAssetManager.cpp b/Engine/Source/Runtime/Engine/Private/Streaming/SimpleStreamableAssetManager.cpp
--- a/Engine/Source/Runtime/Engine/Private/Streaming/SimpleStreamableAssetManager.cpp
+++ b/Engine/Source/Runtime/Engine/Private/Streaming/SimpleStreamableAssetManager.cpp
@@ -80,13 +80,15 @@ void FSimpleStreamableAssetManager::UnregisterAsset(UStreamableRenderAsset* InAs
 	if (IsEnabled())
 	{
 		check(Instance != nullptr);
-		Instance->RemovedAssetsRecords.Push(FAssetRecord{
-			.AssetRegistrationIndex = InAsset->SimpleStreamableAssetManagerIndex
+		// Lock here since TSharedPtr SimpleStreamableAssetManagerIndex is modified while it could be accessed by other threads
+		FScopedLock ScopedLock(GetCriticalSection(), IsEnabled());
+		Instance->RemovedAssetsRecords.Push(FRemovedAssetRecord{
+			.SimpleStreamableAssetManagerIndex = MoveTemp(InAsset->SimpleStreamableAssetManagerIndex)
 #if SIMPLE_STREAMABLE_ASSET_MANAGER_ALLOW_DEBUG_POINTER
 			, .StreamableRenderAsset_ForDebug = InAsset
 #endif
 		});
-		InAsset->SimpleStreamableAssetManagerIndex = INDEX_NONE;
+		InAsset->SimpleStreamableAssetManagerIndex = MakeShared<int32, ESPMode::ThreadSafe>(INDEX_NONE);
 	}
 }
 
@@ -120,7 +122,7 @@ void FSimpleStreamableAssetManager::UpdateBoundSizes(
 
 void FSimpleStreamableAssetManager::GetRenderAssetScreenSize(
 		EStreamableRenderAssetType AssetType,
-		const UStreamableRenderAsset* InAsset,
+		const int32 InAssetIndex,
 		float& MaxSize,
 		float& MaxSize_VisibleOnly,
 		int32& MaxNumForcedLODs,
@@ -129,7 +131,7 @@ void FSimpleStreamableAssetManager::GetRenderAssetScreenSize(
 		const TCHAR* LogPrefix)
 {
 	check(Instance != nullptr);
-	Instance->GetRenderAssetScreenSize_Impl(AssetType, InAsset, MaxSize, MaxSize_VisibleOnly, MaxNumForcedLODs, MaxAssetSize, MaxAllowedMip, LogPrefix);
+	Instance->GetRenderAssetScreenSize_Impl(AssetType, InAssetIndex, MaxSize, MaxSize_VisibleOnly, MaxNumForcedLODs, MaxAssetSize, MaxAllowedMip, LogPrefix);
 }
 
 void FSimpleStreamableAssetManager::GetDistanceAndRange(const FUpdate& Record, float& MinDistanceSq, float& MinRangeSq, float& MaxRangeSq)
@@ -163,7 +165,8 @@ void FSimpleStreamableAssetManager::Update(FUpdate&& Record)
 
 void FSimpleStreamableAssetManager::GetAssetReferenceBounds_Impl(const UStreamableRenderAsset* Asset, TArray<FBox>& AssetBoxes) const
 {
-	const int32 AssetIndex = Asset->SimpleStreamableAssetManagerIndex;
+	check(Asset->SimpleStreamableAssetManagerIndex.IsValid());
+	const int32 AssetIndex = *Asset->SimpleStreamableAssetManagerIndex;
 
 	if (AssetIndex != INDEX_NONE)
 	{
@@ -218,7 +221,7 @@ void FSimpleStreamableAssetManager::UpdateTask_Async()
 	TArray<TLocklessGrowingStorage<FUpdate>::FStorageShard*> Pending_UpdateRecords{};
 	TArray<TLocklessGrowingStorage<FRegister>::FStorageShard*> Pending_RegisterRecords{};
 	TArray<TLocklessGrowingStorage<FUnregister>::FStorageShard*> Pending_UnregisterRecords{};
-	TArray<TLocklessGrowingStorage<FAssetRecord>::FStorageShard*> Pending_RemovedAssetRecords{};
+	TArray<TLocklessGrowingStorage<FRemovedAssetRecord>::FStorageShard*> Pending_RemovedAssetRecords{};
 
 	{
 		/** It is important to extract Shards in this order
@@ -315,19 +318,19 @@ void FSimpleStreamableAssetManager::UpdateTask_Async()
 
 		{
 			SCOPED_NAMED_EVENT(FSimpleStreamableAssetManager_RemoveAssets, FColor::Silver);
-			for (TLocklessGrowingStorage<FAssetRecord>::FStorageShard* Shard : Pending_RemovedAssetRecords)
+			for (TLocklessGrowingStorage<FRemovedAssetRecord>::FStorageShard* Shard : Pending_RemovedAssetRecords)
 			{
-				TArrayView<FAssetRecord> View = Shard->GetData();
-				for (const FAssetRecord& Record : View)
+				TArrayView<FRemovedAssetRecord> View = Shard->GetData();
+				for (const FRemovedAssetRecord& Record : View)
 				{
-					const int32 AssetIndex = Record.AssetRegistrationIndex;
+					const int32 AssetIndex = *Record.SimpleStreamableAssetManagerIndex;
 					if (AssetIndex != INDEX_NONE)
 					{
 						AssetUsedIndices[AssetIndex] = false;
 						--UsedAssetIndices;
 						FreeAssetIndexHint = FMath::Min(FreeAssetIndexHint, AssetIndex);
-						
 						AssetIndexToBounds4Index[AssetIndex].Empty();
+						*Record.SimpleStreamableAssetManagerIndex = INDEX_NONE;
 					}
 				}
 			}
@@ -381,7 +384,7 @@ void FSimpleStreamableAssetManager::UpdateTask_Async()
 		}
 		TRACE_COUNTER_SET(RemovedObjects, Removed);
 
-		for (TLocklessGrowingStorage<FAssetRecord>::FStorageShard* Shard : Pending_RemovedAssetRecords)
+		for (TLocklessGrowingStorage<FRemovedAssetRecord>::FStorageShard* Shard : Pending_RemovedAssetRecords)
 		{
 			delete Shard;
 		}
@@ -549,9 +552,10 @@ void FSimpleStreamableAssetManager::AddRenderAssetElements(const TArrayView<FStr
 		
 		const int32 AssetIndex = [&]
 		{
-			if (RenderAsset->SimpleStreamableAssetManagerIndex != INDEX_NONE)
+			int32& SimpleStreamableAssetManagerIndex = *RenderAsset->SimpleStreamableAssetManagerIndex;
+			if (SimpleStreamableAssetManagerIndex != INDEX_NONE)
 			{
-				return RenderAsset->SimpleStreamableAssetManagerIndex;
+				return SimpleStreamableAssetManagerIndex;
 			}
 			else
 			{
@@ -564,7 +568,7 @@ void FSimpleStreamableAssetManager::AddRenderAssetElements(const TArrayView<FStr
 				check(NewAssetIndex != INDEX_NONE);
 				FreeAssetIndexHint = NewAssetIndex + 1;
 				++UsedAssetIndices;
-				RenderAsset->SimpleStreamableAssetManagerIndex = NewAssetIndex;
+				SimpleStreamableAssetManagerIndex = NewAssetIndex;
 				
 				return NewAssetIndex;
 			}
@@ -618,7 +622,7 @@ void FSimpleStreamableAssetManager::RemoveRenderAssetElements(int32 ObjectRegist
 
 void FSimpleStreamableAssetManager::GetRenderAssetScreenSize_Impl(
 		EStreamableRenderAssetType AssetType,
-		const UStreamableRenderAsset* InAsset,
+		const int32 InAssetIndex,
 		float& MaxSize,
 		float& MaxSize_VisibleOnly,
 		int32& MaxNumForcedLODs,
@@ -628,11 +632,10 @@ void FSimpleStreamableAssetManager::GetRenderAssetScreenSize_Impl(
 {
 	if (AssetType != EStreamableRenderAssetType::Texture || MaxSize_VisibleOnly < MaxAssetSize || LogPrefix)
 	{
-		const int32 AssetIndex = InAsset->SimpleStreamableAssetManagerIndex;
 		// Asset might be registered, but it might not be yet used by any proxy
-		if (AssetIndex != INDEX_NONE)
+		if (InAssetIndex != INDEX_NONE)
 		{
-			const TSimpleSparseArray<FAssetBoundElement>& AssetBoundElements = AssetIndexToBounds4Index[AssetIndex];
+			const TSimpleSparseArray<FAssetBoundElement>& AssetBoundElements = AssetIndexToBounds4Index[InAssetIndex];
 						
 			for (const FAssetBoundElement& AssetBoundElement : AssetBoundElements.GetSparseView())	
 			{
diff --git a/Engine/Source/Runtime/Engine/Private/Streaming/StreamingManagerTexture.cpp b/Engine/Source/Runtime/Engine/Private/Streaming/StreamingManagerTexture.cpp
--- a/Engine/Source/Runtime/Engine/Private/Streaming/StreamingManagerTexture.cpp
+++ b/Engine/Source/Runtime/Engine/Private/Streaming/StreamingManagerTexture.cpp
@@ -608,7 +608,7 @@ void FRenderAssetStreamingManager::ProcessAddedRenderAssets()
 			Asset->StreamingIndex = StreamingRenderAssets.Num();
 			const int32* NumStreamedMips = nullptr;
 			const int32 NumLODGroups = GetNumStreamedMipsArray(Asset->GetRenderAssetType(), NumStreamedMips);
-			new (StreamingRenderAssets) FStreamingRenderAsset(Asset, NumStreamedMips, NumLODGroups, Settings);
+			new (StreamingRenderAssets) FStreamingRenderAsset(Asset, Asset->SimpleStreamableAssetManagerIndex, NumStreamedMips, NumLODGroups, Settings);
 		}
 	}
 	PendingStreamingRenderAssets.Reset();
@@ -3195,25 +3195,3 @@ bool FRenderAssetStreamingManager::Exec( UWorld* InWorld, const TCHAR* Cmd, FOut
 
 	return false;
 }
-
-void FRenderAssetStreamingManager::AddReferencedObjects(FReferenceCollector& Collector)
-{
-	/* [Note] UStreamableRenderAsset is registering/unregistering from manager
-	 * during the creation/destruction - providing reference to GC will prevent
-	 * Asset destruction -> unregistration but in rare case of background task
-	 * FRenderAssetStreamingMipCalcTask::DoWork when in spans for more than one frame
-	 * it might try touch Asset that just has been removed so to prevent this
-	 * we provide GC with referenced assets only if there is a background task running
-	 */
-	if(!AsyncWork->IsDone())
-	{
-		for (TObjectPtr<UStreamableRenderAsset>& ReferencedRenderAsset : ReferencedRenderAssets)
-		{
-			Collector.AddReferencedObject(ReferencedRenderAsset);
-		}
-	}
-}
-FString FRenderAssetStreamingManager::GetReferencerName() const
-{
-	return FString(TEXT("FRenderAssetStreamingManager"));
-}
diff --git a/Engine/Source/Runtime/Engine/Private/Streaming/StreamingManagerTexture.h b/Engine/Source/Runtime/Engine/Private/Streaming/StreamingManagerTexture.h
--- a/Engine/Source/Runtime/Engine/Private/Streaming/StreamingManagerTexture.h
+++ b/Engine/Source/Runtime/Engine/Private/Streaming/StreamingManagerTexture.h
@@ -27,7 +27,7 @@ template<typename TTask> class FAsyncTask;
 /**
  * Streaming manager dealing with textures/meshes.
  */
-struct FRenderAssetStreamingManager final : public IRenderAssetStreamingManager, FGCObject
+struct FRenderAssetStreamingManager final : public IRenderAssetStreamingManager
 {
 	/** Constructor, initializing all members */
 	FRenderAssetStreamingManager();
@@ -456,9 +456,4 @@ private:
 	TArray<FStreamingRenderAsset>& GetStreamingRenderAssetsAsyncSafe();
 
 	friend bool TrackRenderAssetEvent( FStreamingRenderAsset* StreamingRenderAsset, UStreamableRenderAsset* RenderAsset, bool bForceMipLevelsToBeResident, const FRenderAssetStreamingManager* Manager);
-
-	/** FGCObject interface begin **/
-	virtual void AddReferencedObjects(FReferenceCollector& Collector) override;
-	virtual FString GetReferencerName() const override;
-	/** FGCObject interface end **/
 };
diff --git a/Engine/Source/Runtime/Engine/Private/Streaming/StreamingTexture.cpp b/Engine/Source/Runtime/Engine/Private/Streaming/StreamingTexture.cpp
--- a/Engine/Source/Runtime/Engine/Private/Streaming/StreamingTexture.cpp
+++ b/Engine/Source/Runtime/Engine/Private/Streaming/StreamingTexture.cpp
@@ -29,10 +29,12 @@ static FAutoConsoleVariableRef CVarDefaultNoRefBias(
 
 FStreamingRenderAsset::FStreamingRenderAsset(
 	UStreamableRenderAsset* InRenderAsset,
+	TSharedPtr<int32, ESPMode::ThreadSafe>& InSimpleStreamableAssetManagerIndex,
 	const int32* NumStreamedMips,
 	int32 NumLODGroups,
 	const FRenderAssetStreamingSettings& Settings)
 	: RenderAsset(InRenderAsset)
+	, SimpleStreamableAssetManagerIndex(InSimpleStreamableAssetManagerIndex)
 	, RenderAssetType(InRenderAsset->GetRenderAssetType())
 {
 	UpdateStaticData(Settings);
diff --git a/Engine/Source/Runtime/Engine/Private/Streaming/StreamingTexture.h b/Engine/Source/Runtime/Engine/Private/Streaming/StreamingTexture.h
--- a/Engine/Source/Runtime/Engine/Private/Streaming/StreamingTexture.h
+++ b/Engine/Source/Runtime/Engine/Private/Streaming/StreamingTexture.h
@@ -30,6 +30,7 @@ struct FStreamingRenderAsset
 
 	FStreamingRenderAsset(
 		UStreamableRenderAsset* InRenderAsset,
+		TSharedPtr<int32, ESPMode::ThreadSafe>& InSimpleStreamableAssetManagerIndex,
 		const int32* NumStreamedMips,
 		int32 NumLODGroups,
 		const FRenderAssetStreamingSettings& Settings);
@@ -196,6 +197,10 @@ struct FStreamingRenderAsset
 
 	/** (1) Texture/mesh to manage. Note that this becomes null when the texture/mesh is removed. */
 	UStreamableRenderAsset*		RenderAsset;
+
+	/** (1) Index in SSAM. Note that this becomes INDEX_NONE when the texture/mesh is removed. */
+	TSharedPtr<int32, ESPMode::ThreadSafe> SimpleStreamableAssetManagerIndex;
+	
 	/** (1) The optional mip filename hash, see FRenderAssetStreamingManager::OnPakFileChanged() */
 	FIoFilenameHash OptionalFileHash = INVALID_IO_FILENAME_HASH;
 	/** (1) Cached texture/mesh LOD group. */
diff --git a/Engine/Source/Runtime/Engine/Public/Streaming/SimpleStreamableAssetManager.h b/Engine/Source/Runtime/Engine/Public/Streaming/SimpleStreamableAssetManager.h
--- a/Engine/Source/Runtime/Engine/Public/Streaming/SimpleStreamableAssetManager.h
+++ b/Engine/Source/Runtime/Engine/Public/Streaming/SimpleStreamableAssetManager.h
@@ -251,6 +251,16 @@ public:
 		friend bool operator==(const FAssetRecord& A, const FAssetRecord& B) { return A.AssetRegistrationIndex == B.AssetRegistrationIndex; }
 	};
 
+	struct FRemovedAssetRecord
+	{
+		TSharedPtr<int32, ESPMode::ThreadSafe> SimpleStreamableAssetManagerIndex  = nullptr;
+		int32 AssetElementIndex = INDEX_NONE; 
+#if SIMPLE_STREAMABLE_ASSET_MANAGER_ALLOW_DEBUG_POINTER
+		const UStreamableRenderAsset* StreamableRenderAsset_ForDebug = nullptr;
+#endif
+		friend bool operator==(const FRemovedAssetRecord& A, const FRemovedAssetRecord& B) { return A.SimpleStreamableAssetManagerIndex == B.SimpleStreamableAssetManagerIndex; }
+	};
+
 	struct FObjectBoundsRecord
 	{
 		int32 BoundsIndex = INDEX_NONE;
@@ -326,7 +336,7 @@ private:
 	static FSimpleStreamableAssetManager* Instance;
 	FCriticalSection CriticalSection;
 
-	TLocklessGrowingStorage<FAssetRecord> RemovedAssetsRecords;
+	TLocklessGrowingStorage<FRemovedAssetRecord> RemovedAssetsRecords;
 	TLocklessGrowingStorage<FRegister> RegisterRecords;
 	TLocklessGrowingStorage<FUnregister> UnregisterRecords;
 	TLocklessGrowingStorage<FUpdate> UpdateRecords;
@@ -361,7 +371,7 @@ private:
 	
 	void GetRenderAssetScreenSize_Impl(
 		EStreamableRenderAssetType AssetType,
-		const UStreamableRenderAsset* InAsset,
+		const int32 InAssetIndex,
 		float& MaxSize,
 		float& MaxSize_VisibleOnly,
 		int32& MaxNumForcedLODs,
@@ -407,7 +417,7 @@ public:
 
 	static void GetRenderAssetScreenSize(
 		EStreamableRenderAssetType AssetType,
-		const UStreamableRenderAsset* InAsset,
+		const int32 InAssetIndex,
 		float& MaxSize,
 		float& MaxSize_VisibleOnly,
 		int32& MaxNumForcedLODs,
