Minor UStaticMesh::GetUsedMaterials optimization: don't use TSet::Sort to avoid having to rehash all values after sorting and only sort if required, as most of the time the indices are already sorted.

diff --git a/Engine/Source/Runtime/Engine/Private/StaticMesh.cpp b/Engine/Source/Runtime/Engine/Private/StaticMesh.cpp
--- a/Engine/Source/Runtime/Engine/Private/StaticMesh.cpp
+++ b/Engine/Source/Runtime/Engine/Private/StaticMesh.cpp
@@ -9903,7 +9903,8 @@ void UStaticMesh::GetUsedMaterials(TArray<UMaterialInterface*>& OutMaterials, TF
 	TRACE_CPUPROFILER_EVENT_SCOPE(UStaticMesh::GetUsedMaterials);
 
 	if (const FStaticMeshRenderData* ThisRenderData = GetRenderData())
-	{		
+	{
+		int32 LastIndex = 0;
 		TSet<int32> UniqueIndex;
 		for (int32 LODIndex = 0, Num = GetRenderData()->LODResources.Num(); LODIndex < Num; LODIndex++)
 		{
@@ -9911,19 +9912,28 @@ void UStaticMesh::GetUsedMaterials(TArray<UMaterialInterface*>& OutMaterials, TF
 			for (int32 SectionIndex = 0; SectionIndex < LODResources.Sections.Num(); SectionIndex++)
 			{
 				// Get the material for each element at the current lod index
-				UniqueIndex.Add(LODResources.Sections[SectionIndex].MaterialIndex);
+				const int32 MaterialIndex = LODResources.Sections[SectionIndex].MaterialIndex;
+				UniqueIndex.Add(MaterialIndex);
+				LastIndex = (MaterialIndex < LastIndex) ? -1 : MaterialIndex;
 			}
 		}
 
-		if (UniqueIndex.Num() > 0)
-		{
-			//We need to output the material in the correct order (follow the material index)
-			//So we sort the map with the material index
-			UniqueIndex.Sort([](int32 A, int32 B) {
-				return A < B; // sort keys in order
-			});
+		OutMaterials.Reserve(UniqueIndex.Num());
 
-			OutMaterials.Reserve(UniqueIndex.Num());
+		if (LastIndex == -1)
+		{
+			// We need to output the material in the correct order (follow the material index),
+			// So we sort the map with the material index
+			TArray<int32> SortedUniqueIndex = UniqueIndex.Array();
+			SortedUniqueIndex.Sort(TLess<int32>());
+			
+			for (int32 MaterialIndex : SortedUniqueIndex)
+			{
+				OutMaterials.Add(OverrideMaterial(MaterialIndex));
+			}
+		}
+		else
+		{
 			for (int32 MaterialIndex : UniqueIndex)
 			{
 				OutMaterials.Add(OverrideMaterial(MaterialIndex));
