Spline: Remove Reparameterize call in SetLocation which has bad perf implications

diff --git a/Engine/Source/Runtime/Engine/Private/Curves/Spline.cpp b/Engine/Source/Runtime/Engine/Private/Curves/Spline.cpp
--- a/Engine/Source/Runtime/Engine/Private/Curves/Spline.cpp
+++ b/Engine/Source/Runtime/Engine/Private/Curves/Spline.cpp
@@ -2207,9 +2207,6 @@ void FNewSpline::SetLocation(int32 Index, const FVector& InLocation)
     }
 	
 	GetSpline().ModifyPoint(Index, {InLocation, GetSpline().GetTangentIn(Index), GetSpline().GetTangentOut(Index), GetSpline().GetTangentMode(Index)});
-
-	// We do this so that parameterization remains proportional to the square root of segment chord lengths.
-	GetSpline().Reparameterize(EParameterizationPolicy::Centripetal);
 	
 	MarkReparamTableDirty();
     MarkLegacyCurvesDirty();
@@ -2437,6 +2434,11 @@ FVector FNewSpline::EvaluateScale(float Parameter) const
 void FNewSpline::UpdateSpline()
 {
     GetSpline().UpdateTangents();			// updates the tangents on our version of the points
+	
+	// We do this so that parameterization remains proportional to the square root of segment chord lengths.
+	// todo: reenable this code when underlying knot vector issues are resolved.
+	// GetSpline().Reparameterize(EParameterizationPolicy::Centripetal);
+	
     MarkReparamTableDirty();
     RebuildLegacyCurves();					// updates legacy curves based on points, never evaluates internal spline.
 }
