Follow up to change 45022486:
- Recursive parent MI validation for enabling static switches turned into a warning for now.
- UMaterialInstance::PostDuplicate() now performs parent validation.
- Added further defensive code that clears static switch overrides when a material instance is deemed not able to create new shader permutations.

#rnx

diff --git a/Engine/Source/Runtime/Engine/Private/Materials/MaterialInstance.cpp b/Engine/Source/Runtime/Engine/Private/Materials/MaterialInstance.cpp
--- a/Engine/Source/Runtime/Engine/Private/Materials/MaterialInstance.cpp
+++ b/Engine/Source/Runtime/Engine/Private/Materials/MaterialInstance.cpp
@@ -3413,6 +3413,15 @@ void UMaterialInstance::Serialize(FArchive& Ar)
 	}
 }
 
+void UMaterialInstance::PostDuplicate(bool bDuplicateForPIE)
+{
+	Super::PostDuplicate(bDuplicateForPIE);
+
+#if WITH_EDITOR
+	ValidateStaticPermutationAllowed();
+#endif
+}
+
 void UMaterialInstance::PostLoad()
 {
 	LLM_SCOPE(ELLMTag::MaterialInstance);
@@ -3687,7 +3696,10 @@ bool UMaterialInstance::IsStaticPermutationAllowedForCandidateParent(UMaterialIn
 	UMaterialInstance* ParentAsInstance = Cast<UMaterialInstance>(CandidateParent);
 	if (ParentAsInstance && !IsStaticPermutationAllowedForCandidateParent(ParentAsInstance->Parent))
 	{
-		return false;
+		UE_LOG(LogMaterial, Warning,
+			   TEXT("MaterialInstance '%s' cannot allow static parameter permutations because its parent '%s' ultimately derives from a restricted material."),
+			   *this->GetName(),
+			   *CandidateParent->GetName());
 	}
 
 	// Allow candidate material if it is included in base game.
@@ -3733,6 +3745,13 @@ void UMaterialInstance::ValidateStaticPermutationAllowed()
 		SetParentInternal(nullptr, true);
 		bDisallowStaticParameterPermutations = false;
 	}
+
+	// If static switches are disabled, clear any latent existing switche
+	if (bDisallowStaticParameterPermutations)
+	{
+		StaticParametersRuntime.Empty();
+		bHasStaticPermutationResource = false;
+	}
 }
 
 #endif // WITH_EDITOR
diff --git a/Engine/Source/Runtime/Engine/Public/Materials/MaterialInstance.h b/Engine/Source/Runtime/Engine/Public/Materials/MaterialInstance.h
--- a/Engine/Source/Runtime/Engine/Public/Materials/MaterialInstance.h
+++ b/Engine/Source/Runtime/Engine/Public/Materials/MaterialInstance.h
@@ -1002,6 +1002,7 @@ public:
 	virtual ENGINE_API void AppendCompileStateDebugInfo(FStringBuilderBase& OutDebugInfo) const override;
 #endif
 	virtual ENGINE_API void Serialize(FArchive& Ar) override;
+	virtual ENGINE_API void PostDuplicate(bool bDuplicateForPIE) override;
 	virtual ENGINE_API void PostLoad() override;
 #if WITH_EDITORONLY_DATA
 	ENGINE_API static void DeclareConstructClasses(TArray<FTopLevelAssetPath>& OutConstructClasses, const UClass* SpecificSubclass);
