Fixed an issue where an instanced partitioned world failed to reload if any of its cells had previously failed to load due to missing package(s).

diff --git a/Engine/Source/Runtime/Engine/Private/WorldPartition/WorldPartitionSubsystem.cpp b/Engine/Source/Runtime/Engine/Private/WorldPartition/WorldPartitionSubsystem.cpp
--- a/Engine/Source/Runtime/Engine/Private/WorldPartition/WorldPartitionSubsystem.cpp
+++ b/Engine/Source/Runtime/Engine/Private/WorldPartition/WorldPartitionSubsystem.cpp
@@ -610,7 +610,9 @@ void UWorldPartitionSubsystem::OnWorldPartitionUninitialized(UWorldPartition* In
 			for (ULevelStreaming* StreamingLevel : OwningWorld->GetStreamingLevels())
 			{
 				UWorldPartitionLevelStreamingDynamic* WorldPartitionStreamingLevel = Cast<UWorldPartitionLevelStreamingDynamic>(StreamingLevel);
-				if (WorldPartitionStreamingLevel && (WorldPartitionStreamingLevel->GetOuterWorldPartition() == WorldPartition))
+				if (WorldPartitionStreamingLevel && 
+					(WorldPartitionStreamingLevel->GetOuterWorldPartition() == WorldPartition) && 
+					(WorldPartitionStreamingLevel->GetLevelStreamingState() != ELevelStreamingState::FailedToLoad))
 				{
 					GetPendingStreamingLevels().Add(WorldPartitionStreamingLevel);
 				}
@@ -680,7 +682,7 @@ void UWorldPartitionSubsystem::OnLevelStreamingStateChanged(UWorld* InWorld, con
 
 	UpdateLoadingAndPendingLoadStreamingLevels(InStreamingLevel);
 
-	if (NewState == ELevelStreamingState::Removed)
+	if ((NewState == ELevelStreamingState::Removed) || (NewState == ELevelStreamingState::FailedToLoad))
 	{
 		if (const UWorldPartitionLevelStreamingDynamic* WorldPartitionStreamingLevel = Cast<const UWorldPartitionLevelStreamingDynamic>(InStreamingLevel))
 		{
