[SimpleStreamableAssetManager] Clarified a comment about SimpleStreamableAssetManagerIndex in FStreamingRenderAsset.

diff --git a/Engine/Source/Runtime/Engine/Private/Streaming/StreamingTexture.h b/Engine/Source/Runtime/Engine/Private/Streaming/StreamingTexture.h
--- a/Engine/Source/Runtime/Engine/Private/Streaming/StreamingTexture.h
+++ b/Engine/Source/Runtime/Engine/Private/Streaming/StreamingTexture.h
@@ -198,7 +198,7 @@ struct FStreamingRenderAsset
 	/** (1) Texture/mesh to manage. Note that this becomes null when the texture/mesh is removed. */
 	UStreamableRenderAsset*		RenderAsset;
 
-	/** (1) Index in SSAM. Note that this becomes INDEX_NONE when the texture/mesh is removed. */
+	/** (1) Index in SSAM. Note that the underlying index becomes INDEX_NONE when the texture/mesh is removed. */
 	TSharedPtr<int32, ESPMode::ThreadSafe> SimpleStreamableAssetManagerIndex;
 	
 	/** (1) The optional mip filename hash, see FRenderAssetStreamingManager::OnPakFileChanged() */
