Cloth - Fixed the SkeletalMesh cloth LOD rendering in editor by re-adding a missing world readiness conditional accidentaly removed in CL 45029949.

#rnx

diff --git a/Engine/Source/Runtime/Engine/Private/SkeletalMeshComponentPhysics.cpp b/Engine/Source/Runtime/Engine/Private/SkeletalMeshComponentPhysics.cpp
--- a/Engine/Source/Runtime/Engine/Private/SkeletalMeshComponentPhysics.cpp
+++ b/Engine/Source/Runtime/Engine/Private/SkeletalMeshComponentPhysics.cpp
@@ -2927,7 +2927,16 @@ void USkeletalMeshComponent::RecreateClothingActors()
 		return;
 	}
 
-	if (bAllowClothActors && CVarEnableClothPhysics.GetValueOnGameThread())
+	// Only create cloth sim actors when the world is ready for it
+	const UWorld* const World = GetWorld();
+	const bool bIsWorldReadyForSimulation =
+#if WITH_EDITORONLY_DATA
+		World && (World->bShouldSimulatePhysics || bUpdateClothInEditor) && World->GetPhysicsScene();
+#else
+		World && World->bShouldSimulatePhysics && World->GetPhysicsScene();
+#endif
+
+	if (bIsWorldReadyForSimulation && bAllowClothActors && CVarEnableClothPhysics.GetValueOnGameThread())
 	{
 		const TArray<UClothingAssetBase*>& MeshClothingAssets = SkelMesh->GetMeshClothingAssets();
 		if (MeshClothingAssets.Num())
