Adding console command 'obj remote id=id' that dumps object information related to the provided remote id

diff --git a/Engine/Source/Runtime/CoreUObject/Private/UObject/RemoteObject.cpp b/Engine/Source/Runtime/CoreUObject/Private/UObject/RemoteObject.cpp
--- a/Engine/Source/Runtime/CoreUObject/Private/UObject/RemoteObject.cpp
+++ b/Engine/Source/Runtime/CoreUObject/Private/UObject/RemoteObject.cpp
@@ -951,6 +951,34 @@ FString FRemoteObjectId::ToString() const
 	return FString::Printf(TEXT("%s-%llu"), *GetServerId().ToString(), SerialNumber);
 }
 
+FRemoteObjectId FRemoteObjectId::FromString(const FStringView& InText)
+{
+	int32 ServerDelimiterIndex = -1;
+	if (InText.FindChar('-', ServerDelimiterIndex))
+	{
+		// Parse formatted id (as returned from FRemoteObjectId::ToString()): 'ServerId-SerialNumber', e.g. 'Asset-12345'
+		FStringView ServerIdText(InText.GetData(), ServerDelimiterIndex);
+		const int32 ObjectSerialStartIndex = ServerDelimiterIndex + 1;
+		FStringView ObjectSerialText(InText.GetData() + ObjectSerialStartIndex, InText.Len() - ObjectSerialStartIndex);
+		FRemoteServerId ObjectServerId = FRemoteServerId::FromString(ServerIdText);
+		uint64 ObjectSerial = 0;
+		LexFromString(ObjectSerial, ObjectSerialText);
+
+		return FRemoteObjectId(ObjectServerId, ObjectSerial);
+	}
+	else
+	{
+		// Parse the string as a number as if it was the memory image of an id
+		FRemoteObjectId RemoteId;
+		LexFromString(RemoteId.Id, InText);
+		if (RemoteId.ServerId != (uint32)ERemoteServerIdConstants::Invalid)
+		{
+			return RemoteId;
+		}
+	}		
+	return FRemoteObjectId();
+}
+
 FRemoteObjectId::FRemoteObjectId(const UObjectBase* Object)
 	: FRemoteObjectId(UE::CoreUObject::Private::FObjectHandleUtils::GetRemoteId(Object))
 {
@@ -1009,7 +1037,7 @@ FRemoteServerId::FRemoteServerId(const FString& InText)
 	*this = FromString(InText);
 }
 
-FRemoteServerId FRemoteServerId::FromString(const FString& InText)
+FRemoteServerId FRemoteServerId::FromString(const FStringView& InText)
 {
 	FRemoteServerId Result;
 
diff --git a/Engine/Source/Runtime/CoreUObject/Public/UObject/RemoteObjectTypes.h b/Engine/Source/Runtime/CoreUObject/Public/UObject/RemoteObjectTypes.h
--- a/Engine/Source/Runtime/CoreUObject/Public/UObject/RemoteObjectTypes.h
+++ b/Engine/Source/Runtime/CoreUObject/Public/UObject/RemoteObjectTypes.h
@@ -72,7 +72,7 @@ struct FRemoteServerId
 	UE_DEPRECATED(5.6, "Use FRemoteServerId::FromString(const FString&) instead.")
 	explicit COREUOBJECT_API FRemoteServerId(const FString& InText);
 
-	static COREUOBJECT_API FRemoteServerId FromString(const FString& InText);
+	static COREUOBJECT_API FRemoteServerId FromString(const FStringView& InText);
 	COREUOBJECT_API FString ToString() const;
 
 	uint32 GetIdNumber() const
@@ -315,6 +315,7 @@ public:
 	}
 
 	COREUOBJECT_API FString ToString() const;
+	static COREUOBJECT_API FRemoteObjectId FromString(const FStringView& InText);
 
 	COREUOBJECT_API static FRemoteObjectId Generate(UObjectBase* InObject, const TCHAR* InName, EInternalObjectFlags InInitialFlags = EInternalObjectFlags::None);
 
diff --git a/Engine/Source/Runtime/Engine/Private/UnrealEngine.cpp b/Engine/Source/Runtime/Engine/Private/UnrealEngine.cpp
--- a/Engine/Source/Runtime/Engine/Private/UnrealEngine.cpp
+++ b/Engine/Source/Runtime/Engine/Private/UnrealEngine.cpp
@@ -10152,7 +10152,68 @@ bool UEngine::HandleObjCommand( const TCHAR* Cmd, FOutputDevice& Ar )
 		LogHashMemoryOverheadStatistics(Ar, Options);
 		return true;
 	}
-#endif
+#if UE_WITH_REMOTE_OBJECT_HANDLE
+	else if (FParse::Command(&Cmd, TEXT("REMOTE")))
+	{
+		using namespace UE::RemoteObject::Handle;
+
+		FRemoteObjectId RemoteId;
+		FString IdText;
+
+		if (FParse::Value(Cmd, TEXT("ID="), IdText))
+		{
+			RemoteId = FRemoteObjectId::FromString(IdText);
+		}
+
+		if (RemoteId.IsValid())
+		{
+			FString PathName;
+			FString ClassName;
+			bool bStubOnly = false;
+
+			if (UObject* ExistingObject = StaticFindObjectFastInternal(RemoteId))
+			{
+				PathName = ExistingObject->GetPathName();
+				ClassName = ExistingObject->GetClass()->GetName();
+			}
+			else if (UClass* Class = UE::RemoteObject::Handle::GetClass(RemoteId, ERemoteObjectGetClassBehavior::ReturnNullIfNeverLocal))
+			{
+				FRemoteObjectPathName RemotePathName(RemoteId);
+				PathName = RemotePathName.ToString();
+				ClassName = Class->GetName();
+				bStubOnly = true;
+			}
+
+			if (PathName.Len())
+			{
+				FString Attributes;
+				if (IsRemote(RemoteId))
+				{
+					Attributes += TEXT(" (remote)");
+				}
+				if (IsOwned(RemoteId))
+				{
+					Attributes += TEXT(" (owned)");
+				}
+				if (bStubOnly)
+				{
+					Attributes += TEXT(" (stub)");
+				}
+				Ar.Logf(TEXT("%s%s %s %s"), *RemoteId.ToString(), *Attributes, *ClassName, *PathName);
+			}
+			else
+			{
+				Ar.Logf(TEXT("No object found matching remote id '%s'"), *RemoteId.ToString());
+			}
+		}
+		else
+		{
+			Ar.Logf(TEXT("Unable to parse remote object id '%s'"), *IdText);
+		}
+
+	}
+#endif // UE_WITH_REMOTE_OBJECT_HANDLE
+#endif // !UE_BUILD_SHIPPING
 	else
 	{
 		// OBJ command but not supported here
