Repair cook editor-only-data nondeterminancy in materials

diff --git a/Engine/Source/Runtime/Engine/Private/Materials/MaterialExpressions.cpp b/Engine/Source/Runtime/Engine/Private/Materials/MaterialExpressions.cpp
--- a/Engine/Source/Runtime/Engine/Private/Materials/MaterialExpressions.cpp
+++ b/Engine/Source/Runtime/Engine/Private/Materials/MaterialExpressions.cpp
@@ -47,6 +47,8 @@
 #include "SparseVolumeTexture/SparseVolumeTexture.h"
 #include "ProfilingDebugging/LoadTimeTracker.h"
 #include "ColorManagement/ColorSpace.h"
+#include "Containers/Array.h"
+#include "Logging/StructuredLog.h"
 
 #include "Materials/MaterialExpressionAbs.h"
 #include "Materials/MaterialExpressionAbsorptionMediumMaterialOutput.h"
@@ -349,6 +351,54 @@ FUObjectAnnotationSparseBool GMaterialFunctionsThatNeedFeatureLevelSM6Fix;
 static const TCHAR* CPD_UI_ErrorMessage = TEXT("Custom Primitive Data can't be used with the UI material domain.");
 #endif // #if WITH_EDITOR
 
+// Helper functions to handle determinism during cooking
+namespace CookDeterminism
+{
+	// Used to ensure that same object can get different values for different fields
+	enum class ESeed : uint32
+	{
+		Expression = 0,
+		Parameter,
+		InterfaceInit,
+		InterfaceDuplicate,
+		InterfacePostLoad,
+		Compile,
+		FunctionInput,
+		FunctionOutput,
+		RerouteVariable,
+		RerouteColor,
+	};
+
+	FLinearColor MakeRandomColor(FStringView Str, ESeed Seed)
+	{
+		if (IsRunningCookCommandlet())
+		{
+			uint32 Hue = 157 * (GetTypeHash(Str) + static_cast<uint32>(Seed));
+			return FLinearColor::MakeFromHSV8((uint8)Hue, 255, 255);
+		}
+
+		return FLinearColor::MakeRandomColor();
+	}
+
+	FGuid NewGuid(FStringView Str, ESeed Seed)
+	{
+		if (IsRunningCookCommandlet())
+		{
+			return FGuid::NewDeterministicGuid(Str, static_cast<uint64>(Seed));
+		}
+
+		return FGuid::NewGuid();
+	}
+
+	void LogWarningIfCooking(FStringView Message)
+	{
+		if (IsRunningCookCommandlet())
+		{
+			UE_LOGFMT(LogMaterial, Warning, "Non-determinism in material cooking: {Str}", Message);
+		}
+	}
+};
+
 /** Returns whether the given expression class is allowed. */
 bool IsAllowedExpressionType(const UClass* Class, const bool bMaterialFunction)
 {
@@ -2052,16 +2102,7 @@ void UMaterialExpression::UpdateMaterialExpressionGuid(bool bForceGeneration, bo
 
 		if (bForceGeneration || !Guid.IsValid())
 		{
-			if (IsRunningCookCommandlet())
-			{
-				// Prevent indeterminism in the cook by constructing the guid deterministically
-				constexpr int64 MaterialExpressionGuidDeterminismSeed = 0;
-				Guid = FGuid::NewDeterministicGuid(GetPathName(), MaterialExpressionGuidDeterminismSeed);
-			}
-			else
-			{
-				Guid = FGuid::NewGuid();
-			}
+			Guid = CookDeterminism::NewGuid(GetPathName(), CookDeterminism::ESeed::Expression);
 
 			if (bAllowMarkingPackageDirty)
 			{
@@ -2088,16 +2129,7 @@ void UMaterialExpression::UpdateParameterGuid(bool bForceGeneration, bool bAllow
 
 			if (bForceGeneration || !Guid.IsValid())
 			{
-				if (IsRunningCookCommandlet())
-				{
-					// Prevent indeterminism in the cook by constructing the guid deterministically
-					constexpr int64 ParameterGuidDeterminismSeed = 1;
-					Guid = FGuid::NewDeterministicGuid(GetPathName(), ParameterGuidDeterminismSeed);
-				}
-				else
-				{
-					Guid = FGuid::NewGuid();
-				}
+				Guid = CookDeterminism::NewGuid(GetPathName(), CookDeterminism::ESeed::Parameter);
 
 				if (bAllowMarkingPackageDirty)
 				{
@@ -13028,7 +13060,7 @@ void UMaterialFunctionInterface::PostInitProperties()
 	// Initialize StateId to something unique, in case this is a new function
 	if (!IsTemplate())
 	{
-		StateId = FGuid::NewGuid();
+		StateId = CookDeterminism::NewGuid(GetPathName(), CookDeterminism::ESeed::InterfaceInit);
 	}
 }
 
@@ -13041,7 +13073,7 @@ void UMaterialFunctionInterface::PostDuplicate(bool bDuplicateForPIE)
 	if (!bKeepStateId)
 	{
 		// Initialize StateId to something unique, in case this is a new function
-		StateId = FGuid::NewGuid();
+		StateId = CookDeterminism::NewGuid(GetPathName(), CookDeterminism::ESeed::InterfaceDuplicate);
 	}
 }
 
@@ -13074,7 +13106,7 @@ void UMaterialFunctionInterface::PostLoad()
 	
 	if (!StateId.IsValid())
 	{
-		StateId = FGuid::NewGuid();
+		StateId = CookDeterminism::NewGuid(GetPathName(), CookDeterminism::ESeed::InterfacePostLoad);
 	}
 }
 
@@ -13188,7 +13220,7 @@ void UMaterialFunctionInterface::ForceRecompileForRendering(FMaterialUpdateConte
 	TRACE_CPUPROFILER_EVENT_SCOPE(UMaterialFunctionInterface::ForceRecompileForRendering)
 
 	//@todo - recreate guid only when needed, not when a comment changes
-	StateId = FGuid::NewGuid();
+	StateId = CookDeterminism::NewGuid(GetPathName(), CookDeterminism::ESeed::Compile);
 
 	// Go through all materials in memory and recompile them if they use this function
 	for (TObjectIterator<UMaterialInterface> It(/*AdditionalExclusionFlags = */RF_ClassDefaultObject, /*bIncludeDerivedClasses = */true, /*InInternalExclusionFlags = */EInternalObjectFlags::Garbage); It; ++It)
@@ -14923,6 +14955,7 @@ int32 FMaterialLayersFunctions::AppendBlendedLayer()
 	EditorOnly.LayerNames.Add(LayerName);
 	EditorOnly.RestrictToLayerRelatives.Add(false);
 	EditorOnly.RestrictToBlendRelatives.Add(false);
+	CookDeterminism::LogWarningIfCooking(TEXT("AppendBlendedLayer creating NewGuid"));
 	EditorOnly.LayerGuids.Add(FGuid::NewGuid());
 	EditorOnly.LayerLinkStates.Add(EMaterialLayerLinkState::NotFromParent);
 
@@ -16322,7 +16355,7 @@ void UMaterialExpressionFunctionInput::ConditionallyGenerateId(bool bForce)
 {
 	if (bForce || !Id.IsValid())
 	{
-		Id = FGuid::NewGuid();
+		Id = CookDeterminism::NewGuid(GetPathName(), CookDeterminism::ESeed::FunctionInput);
 	}
 }
 
@@ -16705,7 +16738,7 @@ void UMaterialExpressionFunctionOutput::ConditionallyGenerateId(bool bForce)
 {
 	if (bForce || !Id.IsValid())
 	{
-		Id = FGuid::NewGuid();
+		Id = CookDeterminism::NewGuid(GetPathName(), CookDeterminism::ESeed::FunctionOutput);
 	}
 }
 
@@ -17844,7 +17877,7 @@ UMaterialExpressionNamedRerouteDeclaration::UMaterialExpressionNamedRerouteDecla
 	: Super(ObjectInitializer)
 {
 #if WITH_EDITORONLY_DATA
-	NodeColor = FLinearColor::MakeRandomColor();
+	NodeColor = CookDeterminism::MakeRandomColor(GetPathName(), CookDeterminism::ESeed::RerouteColor);
 #endif
 	Name = TEXT("Name");
 }
@@ -18000,7 +18033,7 @@ void UMaterialExpressionNamedRerouteDeclaration::UpdateVariableGuid(bool bForceG
 	{
 		if (bForceGeneration || !VariableGuid.IsValid())
 		{
-			VariableGuid = FGuid::NewGuid();
+			VariableGuid = CookDeterminism::NewGuid(GetPathName(), CookDeterminism::ESeed::RerouteVariable);
 
 			if (bAllowMarkingPackageDirty)
 			{
