Allowing transactional RPCs to be rescheduled per Actor.  This will allow the transaction to abort and try other transactions while we wait for needed dependencies to service the RPC.

diff --git a/Engine/Source/Runtime/Engine/Classes/Engine/NetDriver.h b/Engine/Source/Runtime/Engine/Classes/Engine/NetDriver.h
--- a/Engine/Source/Runtime/Engine/Classes/Engine/NetDriver.h
+++ b/Engine/Source/Runtime/Engine/Classes/Engine/NetDriver.h
@@ -1502,6 +1502,9 @@ private:
 	bool bCallPreReplication = true;
 	bool bReplicatingAutonomousAsSimulated = false;
 
+	/** A set of replicators that have queued RPCs for transactions (server only) */
+	TSet< FObjectReplicator* > ReplicatorsWithPendingTransactions;
+
 public:
 
 	ENGINE_API ENetMode	GetNetMode() const;
@@ -2362,11 +2365,19 @@ public:
 		bReplicateTransactionally = bInReplicateTransactionally;
 	}
 
+#if UE_WITH_REMOTE_OBJECT_HANDLE
 	/** Returns true if this driver should run certain operations via UE::RemoteExecutor::ExecuteTransactional */
 	bool ShouldReplicateTransactionally() const
 	{
 		return bReplicateTransactionally;
 	}
+#else
+	/** This function is not valid unless UE_WITH_REMOTE_OBJECT_HANDLE is set */ 
+	constexpr bool ShouldReplicateTransactionally() const
+	{
+		return false;
+	}
+#endif
 
 	void SetDeferTransactionalSend(bool bInDeferTransactionalSend)
 	{
diff --git a/Engine/Source/Runtime/Engine/Private/DataReplication.cpp b/Engine/Source/Runtime/Engine/Private/DataReplication.cpp
--- a/Engine/Source/Runtime/Engine/Private/DataReplication.cpp
+++ b/Engine/Source/Runtime/Engine/Private/DataReplication.cpp
@@ -84,6 +84,12 @@ FAutoConsoleVariableRef CVarPushModelSkipUndirtiedFastArrays(
 	GbPushModelSkipUndirtiedFastArrays,
 	TEXT("When true, include fast arrays when skipping objects that we can safely see aren't dirty."));
 
+static TAutoConsoleVariable<int32> CVarQueueRpcsForTransactions(
+	TEXT("net.QueueRpcsForTransactions"),
+	1,
+	TEXT("When set to 1, any transactional connection will queue RPCs and execute them in transactions allowing for rescheduling"));
+
+
 extern int32 GNumSkippedObjectEmptyUpdates;
 
 class FNetSerializeCB : public INetSerializeCB
@@ -603,6 +609,7 @@ void FObjectReplicator::CleanUp()
 		}
 
 		Connection->Driver->UnmappedReplicators.Remove( this );
+		Connection->Driver->ReplicatorsWithPendingTransactions.Remove( this );
 
 		Connection->Driver->TotalTrackedGuidMemoryBytes -= TrackedGuidMemoryBytes;
 
@@ -983,6 +990,12 @@ bool FObjectReplicator::ReceivedBunch(FNetBitReader& Bunch, const FReplicationFl
 	const bool bCanDelayRPCs = (CVarDelayUnmappedRPCs.GetValueOnGameThread() > 0) && !bIsServer;
     const uint32 DriverReplicationFrame = ConnectionNetDriver->ReplicationFrame;
 
+#if UE_WITH_REMOTE_OBJECT_HANDLE
+	const bool bQueueRpcForTransactions = ConnectionNetDriver->ShouldReplicateTransactionally() && (CVarQueueRpcsForTransactions.GetValueOnGameThread() > 0);
+#else
+	constexpr bool bQueueRpcForTransactions = false;
+#endif
+
 	const FClassNetCache* const ClassCache = ConnectionNetDriver->NetCache->GetClassNetCache(ObjectClass);
 
 	if (!ClassCache)
@@ -1173,21 +1186,43 @@ bool FObjectReplicator::ReceivedBunch(FNetBitReader& Bunch, const FReplicationFl
 		// Handle function call
 		else if ( Cast< UFunction >( FieldCache->Field.ToUObject() ) )
 		{
-			bool bDelayFunction = false;
+			ESkipRpcBehavior SkipRpcBehavior = bCanDelayRPCs ? ESkipRpcBehavior::SkipIfNotReady : ESkipRpcBehavior::NeverSkip;
+			if (bQueueRpcForTransactions)
+			{
+				SkipRpcBehavior = ESkipRpcBehavior::AlwaysSkip;
+
+				// To debug why RPC property deserialization is aborting:
+				//UE_AUTORTFM_ONABORT()
+				//{
+				//	ensureMsgf(false, TEXT("Expecting that we never abort when we are using queued transactional RPCs"));
+				//};
+			}
+
+			bool bSkippedRpcExec = false;
 			TSet<FNetworkGUID> UnmappedGuids;
-			bool bSuccess = ReceivedRPC(Reader, RepFlags, FieldCache, bCanDelayRPCs, bDelayFunction, UnmappedGuids);
+			bool bSuccess = ReceivedRPC(Reader, UnmappedGuids, bSkippedRpcExec, RepFlags, FieldCache, SkipRpcBehavior);
 
 			if (!bSuccess)
 			{
 				return false;
 			}
-			else if (bDelayFunction)
+			else if (bSkippedRpcExec)
 			{
-				// This invalidates Reader's buffer
-				PendingLocalRPCs.Emplace(FieldCache, RepFlags, Reader, DriverReplicationFrame, UnmappedGuids);
-				bOutHasUnmapped = true;
-				bGuidsChanged = true;
-				bForceUpdateUnmapped = true;
+				if (bQueueRpcForTransactions && UnmappedGuids.IsEmpty())
+				{
+					// We are skipping now and executing later (PostTickDispatch)
+					PendingTransactionalRPCs.Emplace(FieldCache, RepFlags, Reader, DriverReplicationFrame, UnmappedGuids);
+					ConnectionNetDriver->ReplicatorsWithPendingTransactions.Add(this);
+				}
+				else
+				{
+					// Unmapped GUIDs: Either due to not receiving a GUID yet from replication (client only),
+					// or this can also happen on the client & server due to async loading.
+					PendingLocalRPCs.Emplace(FieldCache, RepFlags, Reader, DriverReplicationFrame, UnmappedGuids);
+					bOutHasUnmapped = true;
+					bGuidsChanged = true;
+					bForceUpdateUnmapped = true;
+				}
 			}
 			else if (!IsValid(Object))
 			{
@@ -1275,8 +1310,11 @@ ENGINE_API TArray<UFunction*> FindScopedRPCTrackers(UNetConnection* Connection =
 	return FuncList;
 }
 
-bool FObjectReplicator::ReceivedRPC(FNetBitReader& Reader, const FReplicationFlags& RepFlags, const FFieldNetCache* FieldCache, const bool bCanDelayRPC, bool& bOutDelayRPC, TSet<FNetworkGUID>& UnmappedGuids)
+bool FObjectReplicator::ReceivedRPC(FNetBitReader& Reader, TSet<FNetworkGUID>& UnmappedGuids, bool& bOutSkippedRpcExec, const FReplicationFlags& RepFlags, const FFieldNetCache* FieldCache, ESkipRpcBehavior SkipRpcBehavior)
 {
+	check(FieldCache);
+	bOutSkippedRpcExec = false;
+
 	CSV_SCOPED_TIMING_STAT_EXCLUSIVE(HandleRPC);
 	const bool bIsServer = Connection->Driver->IsServer();
 	UObject* Object = GetObject();
@@ -1346,9 +1384,6 @@ bool FObjectReplicator::ReceivedRPC(FNetBitReader& Reader, const FReplicationFla
 	{
 		TRACE_CPUPROFILER_EVENT_SCOPE(EvaluateCanExecute);
 
-		// Only delay if reliable and CVar is enabled
-		bool bCanDelayUnmapped = bCanDelayRPC && Function->FunctionFlags & FUNC_NetReliable;
-
 		// Get the parameters.
 		FMemMark Mark(FMemStack::Get());
 		uint8* Parms = new(FMemStack::Get(), MEM_Zeroed, Function->ParmsSize)uint8;
@@ -1383,12 +1418,12 @@ bool FObjectReplicator::ReceivedRPC(FNetBitReader& Reader, const FReplicationFla
 
 		RPC_ResetLastFailedReason();
 
-		if (bCanDelayUnmapped && (UnmappedGuids.Num() > 0 || PendingLocalRPCs.Num() > 0))
-		{
-			// If this has unmapped guids or there are already some queued, add to queue
-			bOutDelayRPC = true;	
-		}
-		else
+		// See if the caller wants us to potentially skip the execution of this RPC
+		const bool bCanDelayUnmapped = (SkipRpcBehavior == ESkipRpcBehavior::SkipIfNotReady) && (Function->FunctionFlags & FUNC_NetReliable);
+		const bool bDelayUnmappedRPCs = bCanDelayUnmapped && (UnmappedGuids.Num() > 0 || PendingLocalRPCs.Num() > 0);
+		bOutSkippedRpcExec = (bDelayUnmappedRPCs || SkipRpcBehavior == ESkipRpcBehavior::AlwaysSkip);
+
+		if (!bOutSkippedRpcExec)
 		{
 			AActor* OwningActor = OwningChannel->Actor;
 			UObject* const SubObject = Object != OwningChannel->Actor ? Object : nullptr;
@@ -2500,7 +2535,6 @@ void FObjectReplicator::UpdateUnmappedObjects(bool & bOutHasMoreUnmapped)
 			}
 
 			TSet<FNetworkGUID> UnmappedGuids;
-			bool bCanDelayRPCs = bIsGuidPending; // Force execute if none of our RPC guids are pending, even if other guids are. This is more consistent behavior as it is less dependent on unrelated actors
 			bool bFunctionWasUnmapped = false;
 			bool bSuccess = true;
 			FString FunctionName = TEXT("(Unknown)");
@@ -2512,8 +2546,11 @@ void FObjectReplicator::UpdateUnmappedObjects(bool & bOutHasMoreUnmapped)
 			}
 			else
 			{
+				// Force execute if none of our RPC guids are pending, even if other guids are. This is more consistent behavior as it is less dependent on unrelated actors
+				ESkipRpcBehavior AllowDelay = bIsGuidPending ? ESkipRpcBehavior::SkipIfNotReady : ESkipRpcBehavior::NeverSkip;
+
 				FunctionName = FieldCache->Field.GetName();
-				bSuccess = ReceivedRPC(Reader, Pending.RepFlags, FieldCache, bCanDelayRPCs, bFunctionWasUnmapped, UnmappedGuids);
+				bSuccess = ReceivedRPC(Reader, UnmappedGuids, bFunctionWasUnmapped, Pending.RepFlags, FieldCache, AllowDelay);
 			}
 
 			if (!bSuccess)
@@ -2566,6 +2603,89 @@ void FObjectReplicator::UpdateUnmappedObjects(bool & bOutHasMoreUnmapped)
 	}
 }
 
+void FObjectReplicator::ExecutePendingTransactionalRPCs()
+{
+	ensureMsgf(AutoRTFM::IsClosed(), TEXT("%hs is meant to be executed inside of a transaction"), __func__);
+	if (!ensure(Connection) || !ensure(Connection->Driver) || !ensure(Connection->Driver->NetCache))
+	{
+		return;
+	}
+
+	const bool bIsServer = Connection->Driver->IsServer();
+	if (!ensureMsgf(bIsServer, TEXT("%hs: Expected to only be executed on the server"), __func__))
+	{
+		return;
+	}
+
+	UPackageMapClient* PackageMapClient = Cast<UPackageMapClient>(Connection->PackageMap);
+	if (!PackageMapClient || !OwningChannel)
+	{
+		return;
+	}
+
+	const UObject* Object = GetObject();
+	const FClassNetCache* ClassCache = Connection->Driver->NetCache->GetClassNetCache(ObjectClass);
+
+	// Handle pending RPCs, in order
+	TArray<FRPCPendingLocalCall> ServiceTheseRPCs;
+	Swap(PendingTransactionalRPCs, ServiceTheseRPCs);
+
+	for (const FRPCPendingLocalCall& Pending : ServiceTheseRPCs)
+	{
+		ensureMsgf(Pending.UnmappedGuids.IsEmpty(), TEXT("PendingTransactionalRPCs are never expected to have UnmappedGuids (UnmappedGuids are in PendingLocalRPCs)"));
+
+		bool bSuccess;
+		const FFieldNetCache* FieldCache = ClassCache->GetFromIndex(Pending.RPCFieldIndex);			
+		if (FieldCache == nullptr)
+		{
+			UE_LOG(LogNet, Warning, TEXT("%hs: FieldCache not found. Object: %s"), __func__, *Object->GetFullName());
+			bSuccess = false;
+		}
+		else
+		{
+			// Help us keep track of this function's execution time
+			FName FunctionName = FieldCache->Field.GetFName();
+			UE::RemoteExecutor::OnRegionBeginDelegate.ExecuteIfBound(FunctionName);
+			AutoRTFM::PushOnAbortHandler(&Pending, []()
+			{
+				UE::RemoteExecutor::OnRegionEndDelegate.ExecuteIfBound(TEXT("Abort"));
+			});
+
+			FNetBitReader Reader(Connection->PackageMap, Pending.Buffer.GetData(), Pending.NumBits);
+			Connection->SetNetVersionsOnArchive(Reader);
+
+			TSet<FNetworkGUID> UnmappedGuids;
+			bool bSkipped = false;
+			bSuccess = ReceivedRPC(Reader, UnmappedGuids, bSkipped, Pending.RepFlags, FieldCache, ESkipRpcBehavior::NeverSkip);
+
+			ensureMsgf(UnmappedGuids.IsEmpty(), TEXT("%hs: Servers are expected to never have UnmappedGuids"), __func__);
+			ensureMsgf(!bSkipped, TEXT("%hs: We should never skip an RPC when requesting NeverSkip"), __func__);
+
+			UE::RemoteExecutor::OnRegionEndDelegate.ExecuteIfBound(TEXT(""));
+			AutoRTFM::PopOnAbortHandler(&Pending);
+		}
+
+		if (!bSuccess)
+		{
+			FString FunctionName = FieldCache ? FieldCache->Field.GetName() : TEXT("Unknown");
+
+			const bool bDemoMode = Connection->IsInternalAck();
+			if (!bDemoMode)
+			{
+				// Close our connection and abort rpcs as things are invalid
+				UE_LOG(LogNet, Error, TEXT("%hs: Failed executing delayed RPC %s on Object %s, closing connection!"), __func__, *FunctionName, *Object->GetFullName());
+
+				Connection->Close();
+				return;
+			}
+			else
+			{
+				UE_LOG(LogNet, Warning, TEXT("%hs: Failed executing delayed RPC %s on Object %s, skipping RPC!"), __func__, *FunctionName, *Object->GetFullName());
+			}
+		}
+	}
+}
+
 void FObjectReplicator::QueuePropertyRepNotify(
 	UObject* Object,
 	FProperty * Property,
diff --git a/Engine/Source/Runtime/Engine/Private/NetDriver.cpp b/Engine/Source/Runtime/Engine/Private/NetDriver.cpp
--- a/Engine/Source/Runtime/Engine/Private/NetDriver.cpp
+++ b/Engine/Source/Runtime/Engine/Private/NetDriver.cpp
@@ -1221,7 +1221,18 @@ void UNetDriver::TickFlush(float DeltaSeconds)
 
 	if (!IsUsingIrisReplication())
 	{
-		UpdateUnmappedObjects();
+		if (ShouldReplicateTransactionally())
+		{
+			static const FName NAME_UpdateUnmappedObjects { "UpdateUnmappedObjects" };
+			UE::RemoteExecutor::ExecuteTransactional(NAME_UpdateUnmappedObjects, [this]()
+			{
+				UpdateUnmappedObjects();
+			});
+		}
+		else
+		{
+			UpdateUnmappedObjects();
+		}
 
 		CleanupStaleDormantReplicators();
 	}
@@ -2767,6 +2778,48 @@ void UNetDriver::PostTickDispatch()
 {
 	TGuardValue<bool> GuardInNetTick(bInTick, true);
 
+	if (ShouldReplicateTransactionally())
+	{
+		ensureMsgf(IsServer(), TEXT("We are not expecting Clients to replicate transactionally"));
+		if (ReplicatorsWithPendingTransactions.Num() > 0)
+		{
+			TSet<FObjectReplicator*> TransactingReplicators; 
+			Swap(TransactingReplicators, ReplicatorsWithPendingTransactions);
+
+			for (FObjectReplicator* Replicator : TransactingReplicators)
+			{
+				static const FName WorkName { __func__ };
+				constexpr bool bIsTransactional = true;
+				UE::RemoteExecutor::EnqueueWork(WorkName, bIsTransactional, [Replicator, &TransactingReplicators]()
+				{
+					if (const UObject* ReplicatorObject = Replicator->GetObject())
+					{
+						// Help us keep track of this object's execution time
+						UE::RemoteExecutor::OnRegionBeginDelegate.ExecuteIfBound(ReplicatorObject->GetFName());
+						UE_AUTORTFM_ONABORT()
+						{
+							UE::RemoteExecutor::OnRegionEndDelegate.ExecuteIfBound(TEXT("Abort"));
+						};
+						UE_AUTORTFM_ONCOMMIT()
+						{
+							UE::RemoteExecutor::OnRegionEndDelegate.ExecuteIfBound(TEXT("Commit"));
+						};
+
+						Replicator->ExecutePendingTransactionalRPCs();
+					}
+
+					TransactingReplicators.Remove(Replicator);
+				});
+			}
+
+			// Actually do all of the work
+			UE::RemoteExecutor::ExecutePendingWork();
+
+			ensureMsgf(TransactingReplicators.IsEmpty(), TEXT("We expected all TransactingReplicators to be serviced by now!"));
+			ensureMsgf(ReplicatorsWithPendingTransactions.IsEmpty(), TEXT("Execution of RPCs or pending work caused more RPCs to be enqueued -- those will be delayed another frame"));
+		}
+	}
+
 	// Flush out of order packet caches for connections that did not receive the missing packets during TickDispatch
 	if (ServerConnection != nullptr)
 	{
@@ -6119,6 +6172,10 @@ int32 UNetDriver::ServerReplicateActors(float DeltaSeconds)
 		// Start the region here and end it when we're done ExecutePendingWork
 		static const FName FuncName { __func__ };
 		UE::RemoteExecutor::OnRegionBeginDelegate.ExecuteIfBound(FuncName);
+		AutoRTFM::PushOnAbortHandler(&ServerTickTime, []()
+		{
+			UE::RemoteExecutor::OnRegionEndDelegate.ExecuteIfBound(TEXT("Aborted"));
+		});
 
 		static const FName TransactionalWorkName("ServerReplicateActors_BuildConsiderList");
 		UE::RemoteExecutor::ExecuteTransactional(TransactionalWorkName, [this, &ConsiderList, ServerTickTime]()
@@ -6203,8 +6260,9 @@ int32 UNetDriver::ServerReplicateActors(float DeltaSeconds)
 			if (bUseGranularTransactions)
 			{
 				const FName TransactionalWorkName = Connection->ViewTarget.GetFName();
-				UE::RemoteExecutor::EnqueueWork(TransactionalWorkName, true, [this, &Params]()
+				UE::RemoteExecutor::EnqueueWork(TransactionalWorkName, true, [this, ConstParams = Params]()
 					{
+						UE::Net::FServerReplicateActors_ForConnectionParams Params = ConstParams;
 						ServerReplicateActors_ForConnection(Params);
 					});
 			}
@@ -6223,7 +6281,9 @@ int32 UNetDriver::ServerReplicateActors(float DeltaSeconds)
 	if (bUseGranularTransactions)
 	{
 		UE::RemoteExecutor::ExecutePendingWork();
+
 		UE::RemoteExecutor::OnRegionEndDelegate.ExecuteIfBound(TEXT(""));
+		AutoRTFM::PopOnAbortHandler(&ServerTickTime);
 	}
 
 	if (OnPostConsiderListUpdateOverride.IsBound())
diff --git a/Engine/Source/Runtime/Engine/Public/Net/DataReplication.h b/Engine/Source/Runtime/Engine/Public/Net/DataReplication.h
--- a/Engine/Source/Runtime/Engine/Public/Net/DataReplication.h
+++ b/Engine/Source/Runtime/Engine/Public/Net/DataReplication.h
@@ -74,6 +74,18 @@ class FObjectReplicator
 {
 public:
 
+	enum class ESkipRpcBehavior
+	{
+		// Always execute RPCs immediately
+		NeverSkip,
+
+		// Skip RPC execution only if it's not ready to execute
+		SkipIfNotReady,
+
+		// Always skip execution of RPC (do validation only, caller will handle execution)
+		AlwaysSkip
+	};
+
 	ENGINE_API FObjectReplicator();
 	ENGINE_API ~FObjectReplicator();
 
@@ -164,13 +176,35 @@ public:
 		const bool bHasRepLayout,
 		bool& bOutHasUnmapped);
 
+	/**
+	 * @param Reader the buffer that contains the RPC (this function reads RPC parameters from the current position)
+	 * @param OutUnmappedGuids any GUIDs that are not yet resolvable because we do not have the data (yet)
+	 * @param bOutSkippedRpcExec did we skip execution of this RPC? If so, the caller must handle execution (e.g. delay).
+	 * @param RepFlags the replication flags (state) of the replicated object
+	 * @param FieldCache the RPC FieldCache used to speed-up property lookups.  This cannot be null.
+	 * @param SkipRpcBehavior under what circumstances should we skip the execution of the RPC?
+	 *
+	 * @return True if the Reader is still valid. False if the Reader should be abandoned.
+	 * @see ReceivedBunch(), UpdateUnmappedObjects()
+	 */
 	ENGINE_API bool ReceivedRPC(
+		FNetBitReader& Reader,
+		TSet<FNetworkGUID>& OutUnmappedGuids,
+		bool& bOutSkippedRpcExec,
+		const FReplicationFlags& RepFlags,
+		const FFieldNetCache* FieldCache,
+		ESkipRpcBehavior SkipRpcBehavior);
+
+	inline bool ReceivedRPC(
 		FNetBitReader& Reader,
 		const FReplicationFlags& RepFlags,
 		const FFieldNetCache* FieldCache,
 		const bool bCanDelayRPC,
 		bool& bOutDelayRPC,
-		TSet<FNetworkGUID>& OutUnmappedGuids);
+		TSet<FNetworkGUID>& OutUnmappedGuids)
+	{
+		return ReceivedRPC(Reader, OutUnmappedGuids, bOutDelayRPC, RepFlags, FieldCache, bCanDelayRPC ? ESkipRpcBehavior::SkipIfNotReady : ESkipRpcBehavior::NeverSkip);
+	}
 
 	ENGINE_API void UpdateGuidToReplicatorMap();
 	ENGINE_API bool MoveMappedObjectToUnmapped(const FNetworkGUID& GUID);
@@ -188,6 +222,12 @@ public:
 
 	ENGINE_API void UpdateUnmappedObjects(bool& bOutHasMoreUnmapped);
 
+	/**
+	 * Execute any pending RPCs meant to be done within an AutoRTFM transaction.
+	 * @see UNetDriver::PostTickDispatch
+	 */
+	ENGINE_API void ExecutePendingTransactionalRPCs();
+
 	inline TWeakObjectPtr<UObject>	GetWeakObjectPtr() const
 	{
 		return WeakObjectPtr;
@@ -321,6 +361,12 @@ public:
 private:
 	TWeakObjectPtr<UObject> WeakObjectPtr;
 
+	/**
+	 * These are RPCs that are saved so that they may be executed transactionally at a later time
+	 * @see ExecutePendingTransactionalRPCs
+	 */
+	TArray<FRPCPendingLocalCall> PendingTransactionalRPCs;
+
 	static void CallProcessEventForReceivedRPC(UObject* Object, UFunction* Function, uint8* Params);
 };
 
