Exposing EMaterialProperty::MP_MaterialAttributes to Blueprints, Python and the Editor UI.

- This way users can use the Blueprint MaterialEditingLibrary to connect expressions to the MaterialAttributes output on the material block for materials that use MA.

[RN] minor | Rendering - Materials and Shaders

diff --git a/Engine/Source/Runtime/Engine/Public/SceneTypes.h b/Engine/Source/Runtime/Engine/Public/SceneTypes.h
--- a/Engine/Source/Runtime/Engine/Public/SceneTypes.h
+++ b/Engine/Source/Runtime/Engine/Public/SceneTypes.h
@@ -181,7 +181,7 @@ enum EMaterialProperty : int
 	MP_Displacement UMETA(Hidden),
 
 	//^^^ New material properties go above here ^^^^
-	MP_MaterialAttributes UMETA(Hidden),
+	MP_MaterialAttributes,
 	MP_CustomOutput UMETA(Hidden),
 	MP_MAX UMETA(DisplayName = "None"),
 };
