- Make FXSystemInterface sharable to fix world tear down reference issues
- In some situations an FX component can be destroyed after the world
- When the world is destroyed it will destroy the FX system set, which will enqueue a command to destroy the FX interfaces
- The render command can execute before the FX component's are destroyed leaving them with a dangling reference

diff --git a/Engine/Plugins/FX/Niagara/Source/Niagara/Internal/NiagaraEmitterInstanceImpl.h b/Engine/Plugins/FX/Niagara/Source/Niagara/Internal/NiagaraEmitterInstanceImpl.h
--- a/Engine/Plugins/FX/Niagara/Source/Niagara/Internal/NiagaraEmitterInstanceImpl.h
+++ b/Engine/Plugins/FX/Niagara/Source/Niagara/Internal/NiagaraEmitterInstanceImpl.h
@@ -138,7 +138,7 @@ private:
 
 	TUniquePtr<FEventInstanceData> EventInstanceData;
 
-	FNiagaraGpuComputeDispatchInterface* ComputeDispatchInterface = nullptr;
+	TSharedPtr<FNiagaraGpuComputeDispatchInterface> ComputeDispatchInterfacePtr;
 
 	/** Raw pointer to the emitter that we're instanced from. Raw ptr should be safe here as we check for the validity of the system and it's emitters higher up before any ticking. */
 	//FVersionedNiagaraEmitter CachedEmitter;
diff --git a/Engine/Plugins/FX/Niagara/Source/Niagara/Private/NiagaraEmitterInstanceImpl.cpp b/Engine/Plugins/FX/Niagara/Source/Niagara/Private/NiagaraEmitterInstanceImpl.cpp
--- a/Engine/Plugins/FX/Niagara/Source/Niagara/Private/NiagaraEmitterInstanceImpl.cpp
+++ b/Engine/Plugins/FX/Niagara/Source/Niagara/Private/NiagaraEmitterInstanceImpl.cpp
@@ -116,8 +116,10 @@ FNiagaraEmitterInstanceImpl::FNiagaraEmitterInstanceImpl(FNiagaraSystemInstance*
 {
 	ParticleDataSet = new FNiagaraDataSet();
 
-	ComputeDispatchInterface = ParentSystemInstance ? ParentSystemInstance->GetComputeDispatchInterface() : nullptr;
-	check(ComputeDispatchInterface != nullptr);
+	if (FNiagaraGpuComputeDispatchInterface* ComputeDispatchInterface = ParentSystemInstance->GetComputeDispatchInterface())
+	{
+		ComputeDispatchInterfacePtr = StaticCastSharedRef<FNiagaraGpuComputeDispatchInterface>(ComputeDispatchInterface->AsShared());
+	}
 }
 
 FNiagaraEmitterInstanceImpl::~FNiagaraEmitterInstanceImpl()
@@ -876,7 +878,7 @@ void FNiagaraEmitterInstanceImpl::CalculateFixedBounds(const FTransform& ToWorld
 			return;
 		}
 
-		ScopedGPUReadback.ReadbackData(ComputeDispatchInterface, GPUExecContext->MainDataSet);
+		ScopedGPUReadback.ReadbackData(ComputeDispatchInterfacePtr.Get(), GPUExecContext->MainDataSet);
 		NumInstances = ScopedGPUReadback.GetNumInstances();
 	}
 	else
@@ -1002,7 +1004,7 @@ bool FNiagaraEmitterInstanceImpl::HandleCompletion(bool bForce)
 	{
 		if( GPUExecContext )
 		{
-			GPUExecContext->Reset(ComputeDispatchInterface);
+			GPUExecContext->Reset(ComputeDispatchInterfacePtr.Get());
 		}
 		ParticleDataSet->ResetBuffers();
 		if (EventInstanceData.IsValid())
@@ -1175,7 +1177,15 @@ bool FNiagaraEmitterInstanceImpl::WaitForDebugInfo()
 	FNiagaraComputeExecutionContext* DebugContext = GPUExecContext;
 	if (VersionedEmitter.GetEmitterData()->SimTarget == ENiagaraSimTarget::GPUComputeSim && DebugContext)
 	{
-		ENQUEUE_RENDER_COMMAND(CaptureCommand)([this](FRHICommandListImmediate& RHICmdList) { ComputeDispatchInterface->ProcessDebugReadbacks(RHICmdList, true); });
+		FNiagaraGpuComputeDispatchInterface* ComputeDispatchInterface = ComputeDispatchInterfacePtr.Get();
+		if (ComputeDispatchInterface)
+		{
+			ENQUEUE_RENDER_COMMAND(CaptureCommand)([ComputeDispatchInterface](FRHICommandListImmediate& RHICmdList) { ComputeDispatchInterface->ProcessDebugReadbacks(RHICmdList, true); });
+		}
+		else
+		{
+			UE_LOG(LogNiagara, Fatal, TEXT("Compute interface is invalid and we have a GPU simulation, this may result in a softlock."));
+		}
 		FlushRenderingCommands(); 
 		return true;
 	}
@@ -1231,7 +1241,7 @@ void FNiagaraEmitterInstanceImpl::Tick(float DeltaSeconds)
 	{
 		if (GPUExecContext)
 		{
-			GPUExecContext->Reset(ComputeDispatchInterface);
+			GPUExecContext->Reset(ComputeDispatchInterfacePtr.Get());
 		}
 		Data.ResetBuffers();
 		ExecutionState = ENiagaraExecutionState::Inactive;
@@ -1360,7 +1370,7 @@ void FNiagaraEmitterInstanceImpl::Tick(float DeltaSeconds)
 		if (ParentSystemInstance->ShouldCaptureThisFrame())
 		{
 			TSharedPtr<FNiagaraScriptDebuggerInfo, ESPMode::ThreadSafe> DebugInfo = ParentSystemInstance->GetActiveCaptureWrite(CachedIDName, ENiagaraScriptUsage::ParticleGPUComputeScript, FGuid());
-			if (DebugInfo && ComputeDispatchInterface != nullptr)
+			if (DebugInfo && ComputeDispatchInterfacePtr.IsValid())
 			{
 				//Data.Dump(DebugInfo->Frame, true, 0, OrigNumParticles);
 				//DebugInfo->Frame.Dump(true, 0, OrigNumParticles);
@@ -1371,7 +1381,7 @@ void FNiagaraEmitterInstanceImpl::Tick(float DeltaSeconds)
 
 				// Execute a readback
 				ENQUEUE_RENDER_COMMAND(NiagaraReadbackGpuSim)(
-					[RT_ComputeDispatchInterface = ComputeDispatchInterface, RT_InstanceID=ParentSystemInstance->GetId(), RT_DebugInfo=DebugInfo, RT_Context=GPUExecContext](FRHICommandListImmediate& RHICmdList)
+					[RT_ComputeDispatchInterface = ComputeDispatchInterfacePtr.Get(), RT_InstanceID = ParentSystemInstance->GetId(), RT_DebugInfo = DebugInfo, RT_Context = GPUExecContext](FRHICommandListImmediate& RHICmdList)
 					{
 						RT_ComputeDispatchInterface->AddDebugReadback(RT_InstanceID, RT_DebugInfo, RT_Context);
 					}
diff --git a/Engine/Plugins/FX/Niagara/Source/Niagara/Private/NiagaraGpuComputeDispatch.cpp b/Engine/Plugins/FX/Niagara/Source/Niagara/Private/NiagaraGpuComputeDispatch.cpp
--- a/Engine/Plugins/FX/Niagara/Source/Niagara/Private/NiagaraGpuComputeDispatch.cpp
+++ b/Engine/Plugins/FX/Niagara/Source/Niagara/Private/NiagaraGpuComputeDispatch.cpp
@@ -2323,10 +2323,10 @@ void FNiagaraGpuComputeDispatch::PreRender(FRDGBuilder& GraphBuilder, TConstStri
 	LLM_SCOPE(ELLMTag::Niagara);
 }
 
-void FNiagaraGpuComputeDispatch::OnDestroy()
+void FNiagaraGpuComputeDispatch::OnMarkPendingKill()
 {
 	FNiagaraWorldManager::OnComputeDispatchInterfaceDestroyed(this);
-	FFXSystemInterface::OnDestroy();
+	FFXSystemInterface::OnMarkPendingKill();
 }
 
 bool FNiagaraGpuComputeDispatch::AddSortedGPUSimulation(FRHICommandListBase& RHICmdList, FNiagaraGPUSortInfo& SortInfo)
diff --git a/Engine/Plugins/FX/Niagara/Source/Niagara/Private/NiagaraGpuComputeDispatch.h b/Engine/Plugins/FX/Niagara/Source/Niagara/Private/NiagaraGpuComputeDispatch.h
--- a/Engine/Plugins/FX/Niagara/Source/Niagara/Private/NiagaraGpuComputeDispatch.h
+++ b/Engine/Plugins/FX/Niagara/Source/Niagara/Private/NiagaraGpuComputeDispatch.h
@@ -121,7 +121,7 @@ public:
 	virtual bool RequiresEarlyViewUniformBuffer() const override;
 	virtual bool RequiresRayTracingScene() const override;
 	virtual void PreRender(FRDGBuilder& GraphBuilder, TConstStridedView<FSceneView> Views, FSceneUniformBuffer &SceneUniformBuffer, bool bAllowGPUParticleUpdate) override;
-	virtual void OnDestroy() override; // Called on the gamethread to delete the batcher on the renderthread.
+	virtual void OnMarkPendingKill() override;
 
 	virtual void Tick(UWorld* World, float DeltaTime) override;
 
diff --git a/Engine/Plugins/FX/Niagara/Source/Niagara/Private/NiagaraSystemGpuComputeProxy.cpp b/Engine/Plugins/FX/Niagara/Source/Niagara/Private/NiagaraSystemGpuComputeProxy.cpp
--- a/Engine/Plugins/FX/Niagara/Source/Niagara/Private/NiagaraSystemGpuComputeProxy.cpp
+++ b/Engine/Plugins/FX/Niagara/Source/Niagara/Private/NiagaraSystemGpuComputeProxy.cpp
@@ -88,7 +88,10 @@ void FNiagaraSystemGpuComputeProxy::AddToRenderThread(FNiagaraGpuComputeDispatch
 	ENQUEUE_RENDER_COMMAND(AddProxyToComputeDispatchInterface)(
 		[this, ComputeDispatchInterface](FRHICommandListImmediate& RHICmdList)
 		{
-			ComputeDispatchInterface->AddGpuComputeProxy(this);
+			if (ComputeDispatchInterface)
+			{
+				ComputeDispatchInterface->AddGpuComputeProxy(this);
+			}
 
 			for (FNiagaraComputeExecutionContext* ComputeContext : ComputeContexts)
 			{
@@ -115,12 +118,6 @@ void FNiagaraSystemGpuComputeProxy::RemoveFromRenderThread(FNiagaraGpuComputeDis
 	check(DebugOwnerComputeDispatchInterface == ComputeDispatchInterface);
 	DebugOwnerComputeDispatchInterface = nullptr;
 
-	// If the interface is pending kill we can't use it in the render command
-	if (ComputeDispatchInterface->IsPendingKill())
-	{
-		ComputeDispatchInterface = nullptr;
-	}
-
 	ENQUEUE_RENDER_COMMAND(RemoveFromRenderThread)(
 		[this, ComputeDispatchInterface, bDeleteProxy](FRHICommandListImmediate& RHICmdList)
 		{
@@ -158,7 +155,7 @@ void FNiagaraSystemGpuComputeProxy::ClearTicksFromRenderThread(FNiagaraGpuComput
 	check(DebugOwnerComputeDispatchInterface == ComputeDispatchInterface);
 
 	ENQUEUE_RENDER_COMMAND(ClearTicksFromProxy)(
-		[this, ComputeDispatchInterface](FRHICommandListImmediate& RHICmdList)
+		[this](FRHICommandListImmediate& RHICmdList)
 		{
 			ReleaseTicks(TNumericLimits<int32>::Max(), true);
 		}
diff --git a/Engine/Plugins/FX/Niagara/Source/Niagara/Private/NiagaraSystemInstance.cpp b/Engine/Plugins/FX/Niagara/Source/Niagara/Private/NiagaraSystemInstance.cpp
--- a/Engine/Plugins/FX/Niagara/Source/Niagara/Private/NiagaraSystemInstance.cpp
+++ b/Engine/Plugins/FX/Niagara/Source/Niagara/Private/NiagaraSystemInstance.cpp
@@ -157,7 +157,10 @@ FNiagaraSystemInstance::FNiagaraSystemInstance(UWorld& InWorld, UNiagaraSystem&
 		InstanceParameters.SetOwner(InAttachComponent);
 	}
 
-	ComputeDispatchInterface = FNiagaraGpuComputeDispatchInterface::Get(World);
+	if (FNiagaraGpuComputeDispatchInterface* ComputeDispatchInterface = FNiagaraGpuComputeDispatchInterface::Get(World))
+	{
+		ComputeDispatchInterfacePtr = StaticCastSharedRef<FNiagaraGpuComputeDispatchInterface>(ComputeDispatchInterface->AsShared());
+	}
 	FeatureLevel = World->GetFeatureLevel();
 
 	// In some cases the system may have already stated that you should ignore dependencies and tick as early as possible.
@@ -508,7 +511,7 @@ void FNiagaraSystemInstance::SetSolo(bool bInSolo)
 void FNiagaraSystemInstance::SetGpuComputeDebug(bool bEnableDebug)
 {
 #if NIAGARA_COMPUTEDEBUG_ENABLED
-	if (ComputeDispatchInterface == nullptr || !::IsValid(System))
+	if (!ComputeDispatchInterfacePtr.IsValid() || !::IsValid(System))
 	{
 		return;
 	}
@@ -531,7 +534,7 @@ void FNiagaraSystemInstance::SetGpuComputeDebug(bool bEnableDebug)
 
 		ENQUEUE_RENDER_COMMAND(NiagaraAddGPUSystemDebug)
 		(
-			[RT_ComputeDispatchInterface=ComputeDispatchInterface, RT_InstanceID=GetId(), RT_SystemName=SystemName](FRHICommandListImmediate& RHICmdList)
+			[RT_ComputeDispatchInterface=ComputeDispatchInterfacePtr.Get(), RT_InstanceID = GetId(), RT_SystemName = SystemName](FRHICommandListImmediate& RHICmdList)
 			{
 				if (FNiagaraGpuComputeDebug* GpuComputeDebug = RT_ComputeDispatchInterface->GetGpuComputeDebugPrivate())
 				{
@@ -544,7 +547,7 @@ void FNiagaraSystemInstance::SetGpuComputeDebug(bool bEnableDebug)
 	{
 		ENQUEUE_RENDER_COMMAND(NiagaraRemoveGPUSystemDebug)
 		(
-			[RT_ComputeDispatchInterface=ComputeDispatchInterface, RT_InstanceID=GetId()](FRHICommandListImmediate& RHICmdList)
+			[RT_ComputeDispatchInterface=ComputeDispatchInterfacePtr.Get(), RT_InstanceID = GetId()](FRHICommandListImmediate& RHICmdList)
 			{
 				if (FNiagaraGpuComputeDebug* GpuComputeDebug = RT_ComputeDispatchInterface->GetGpuComputeDebugPrivate())
 				{
@@ -1358,7 +1361,7 @@ void FNiagaraSystemInstance::UpdateLWCTile(const FVector3f& InLWCTile, FNiagaraD
 
 	// Rebase data sets
 	NiagaraLWCHelper::RebaseDataSet(TileShift, EmitterDataSet, SystemInstanceIndex);
-	NiagaraLWCHelper::RebaseEmitters(TileShift, Emitters, ComputeDispatchInterface);
+	NiagaraLWCHelper::RebaseEmitters(TileShift, Emitters, ComputeDispatchInterfacePtr.Get());
 
 	//-TODO: Rebase data interface data
 
diff --git a/Engine/Plugins/FX/Niagara/Source/Niagara/Public/NiagaraSystemInstance.h b/Engine/Plugins/FX/Niagara/Source/Niagara/Public/NiagaraSystemInstance.h
--- a/Engine/Plugins/FX/Niagara/Source/Niagara/Public/NiagaraSystemInstance.h
+++ b/Engine/Plugins/FX/Niagara/Source/Niagara/Public/NiagaraSystemInstance.h
@@ -416,7 +416,7 @@ public:
 	/** Dumps information about the instances tick to the log */
 	NIAGARA_API void DumpTickInfo(FOutputDevice& Ar);
 
-	FNiagaraGpuComputeDispatchInterface* GetComputeDispatchInterface() const { return ComputeDispatchInterface; }
+	FNiagaraGpuComputeDispatchInterface* GetComputeDispatchInterface() const { return ComputeDispatchInterfacePtr.Get(); }
 
 	static NIAGARA_API bool AllocateSystemInstance(FNiagaraSystemInstancePtr& OutSystemInstanceAllocation, UWorld& InWorld, UNiagaraSystem& InAsset,
 		FNiagaraUserRedirectionParameterStore* InOverrideParameters = nullptr, USceneComponent* InAttachComponent = nullptr,
@@ -657,7 +657,7 @@ private:
 	/** Copy of simulations internal state so that it can be passed to emitters etc. */
 	ENiagaraExecutionState ActualExecutionState;
 
-	FNiagaraGpuComputeDispatchInterface* ComputeDispatchInterface = nullptr;
+	TSharedPtr<FNiagaraGpuComputeDispatchInterface> ComputeDispatchInterfacePtr;
 	TUniquePtr<FNiagaraSystemGpuComputeProxy> SystemGpuComputeProxy;
 
 	/** Tag we feed into crash reporter for this instance. */
diff --git a/Engine/Source/Editor/UnrealEd/Private/ThumbnailHelpers.cpp b/Engine/Source/Editor/UnrealEd/Private/ThumbnailHelpers.cpp
--- a/Engine/Source/Editor/UnrealEd/Private/ThumbnailHelpers.cpp
+++ b/Engine/Source/Editor/UnrealEd/Private/ThumbnailHelpers.cpp
@@ -214,7 +214,7 @@ FParticleSystemThumbnailScene::FParticleSystemThumbnailScene()
 
 FParticleSystemThumbnailScene::~FParticleSystemThumbnailScene()
 {
-	FFXSystemInterface::Destroy( ThumbnailFXSystem );
+	FFXSystemInterface::MarkPendingKill( ThumbnailFXSystem );
 }
 
 void FParticleSystemThumbnailScene::SetParticleSystem(UParticleSystem* ParticleSystem)
diff --git a/Engine/Source/Runtime/Engine/Private/Particles/FXSystem.cpp b/Engine/Source/Runtime/Engine/Private/Particles/FXSystem.cpp
--- a/Engine/Source/Runtime/Engine/Private/Particles/FXSystem.cpp
+++ b/Engine/Source/Runtime/Engine/Private/Particles/FXSystem.cpp
@@ -39,6 +39,17 @@ bool IsParticleCollisionModeSupported(EShaderPlatform InPlatform, EParticleColli
 	return IsFeatureLevelSupported(InPlatform, ERHIFeatureLevel::SM5);
 }
 
+void FFXSystemInterface::FFXInterfaceDeletor::operator()(FFXSystemInterface* ObjectToDelete) const
+{
+	ENQUEUE_RENDER_COMMAND(FFXInterfaceDeletor)
+	(
+		[ObjectToDelete](FRHICommandListImmediate&)
+		{
+			delete ObjectToDelete;
+		}
+	);
+}
+
 /*-----------------------------------------------------------------------------
 	External FX system interface.
 -----------------------------------------------------------------------------*/
@@ -59,7 +70,7 @@ FFXSystemInterface* FFXSystemInterface::Create(ERHIFeatureLevel::Type InFeatureL
 		Set->SetSceneInterface(Scene);
 		Set->SetScene((FScene*)Scene);
 
-		Set->FXSystems.Add(new FFXSystem(InFeatureLevel, InShaderPlatform, GPUSortManager));
+		Set->FXSystems.Emplace(MakeShareable(new FFXSystem(InFeatureLevel, InShaderPlatform, GPUSortManager), FFXInterfaceDeletor()));
 
 		for (TMap<FName, FCreateCustomFXSystemDelegate>::TConstIterator Ite(CreateCustomFXDelegates); Ite; ++Ite)
 		{
@@ -68,7 +79,7 @@ FFXSystemInterface* FFXSystemInterface::Create(ERHIFeatureLevel::Type InFeatureL
 			{
 				CustomFX->SetSceneInterface(Scene);
 				CustomFX->SetScene((FScene*)Scene);
-				Set->FXSystems.Add(CustomFX);
+				Set->FXSystems.Emplace(MakeShareable(CustomFX, FFXInterfaceDeletor()));
 			}
 		}
 		return Set;
@@ -93,16 +104,15 @@ void FFXSystemInterface::QueueDestroyGPUSimulation(FFXSystemInterface* FXSystem)
 	});
 }
 
-void FFXSystemInterface::Destroy( FFXSystemInterface* FXSystem )
+void FFXSystemInterface::MarkPendingKill(FFXSystemInterface* FXSystem)
 {
 	check(FXSystem);
-	FXSystem->OnDestroy();
+	FXSystem->OnMarkPendingKill();
+}
 
-	ENQUEUE_RENDER_COMMAND(FDestroyFXSystemCommand)(
-		[FXSystem](FRHICommandList& RHICmdList)
-		{
-			delete FXSystem;
-		});
+void FFXSystemInterface::Destroy( FFXSystemInterface* FXSystem )
+{
+	MarkPendingKill(FXSystem);
 }
 
 void FFXSystemInterface::RegisterCustomFXSystem(const FName& InterfaceName, const FCreateCustomFXSystemDelegate& InCreateDelegate)
diff --git a/Engine/Source/Runtime/Engine/Private/Particles/FXSystemSet.cpp b/Engine/Source/Runtime/Engine/Private/Particles/FXSystemSet.cpp
--- a/Engine/Source/Runtime/Engine/Private/Particles/FXSystemSet.cpp
+++ b/Engine/Source/Runtime/Engine/Private/Particles/FXSystemSet.cpp
@@ -16,13 +16,11 @@ FFXSystemSet::FFXSystemSet(FGPUSortManager* InGPUSortManager)
 
 FFXSystemInterface* FFXSystemSet::GetInterface(const FName& InName)
 {
-	for (FFXSystemInterface* FXSystem : FXSystems)
+	for (const TSharedRef<FFXSystemInterface>& FXSystem : FXSystems)
 	{
-		check(FXSystem);
-		FXSystem = FXSystem->GetInterface(InName);
-		if (FXSystem)
+		if (FFXSystemInterface* FXSystemInterface = FXSystem->GetInterface(InName))
 		{
-			return FXSystem;
+			return FXSystemInterface;
 		}
 	}
 	return nullptr;
@@ -30,9 +28,8 @@ FFXSystemInterface* FFXSystemSet::GetInterface(const FName& InName)
 
 void FFXSystemSet::Tick(UWorld* World, float DeltaSeconds)
 {
-	for (FFXSystemInterface* FXSystem : FXSystems)
+	for (const TSharedRef<FFXSystemInterface>& FXSystem : FXSystems)
 	{
-		check(FXSystem);
 		FXSystem->Tick(World, DeltaSeconds);
 	}
 }
@@ -41,18 +38,16 @@ void FFXSystemSet::Tick(UWorld* World, float DeltaSeconds)
 
 void FFXSystemSet::Suspend()
 {
-	for (FFXSystemInterface* FXSystem : FXSystems)
+	for (const TSharedRef<FFXSystemInterface>& FXSystem : FXSystems)
 	{
-		check(FXSystem);
 		FXSystem->Suspend();
 	}
 }
 
 void FFXSystemSet::Resume()
 {
-	for (FFXSystemInterface* FXSystem : FXSystems)
+	for (const TSharedRef<FFXSystemInterface>& FXSystem : FXSystems)
 	{
-		check(FXSystem);
 		FXSystem->Resume();
 	}
 }
@@ -61,18 +56,16 @@ void FFXSystemSet::Resume()
 
 void FFXSystemSet::DrawDebug(FCanvas* Canvas)
 {
-	for (FFXSystemInterface* FXSystem : FXSystems)
+	for (const TSharedRef<FFXSystemInterface>& FXSystem : FXSystems)
 	{
-		check(FXSystem);
 		FXSystem->DrawDebug(Canvas);
 	}
 }
 
 bool FFXSystemSet::ShouldDebugDraw_RenderThread() const
 {
-	for (const FFXSystemInterface* FXSystem : FXSystems)
+	for (const TSharedRef<FFXSystemInterface>& FXSystem : FXSystems)
 	{
-		check(FXSystem);
 		if (FXSystem->ShouldDebugDraw_RenderThread())
 		{
 			return true;
@@ -83,72 +76,64 @@ bool FFXSystemSet::ShouldDebugDraw_RenderThread() const
 
 void FFXSystemSet::DrawDebug_RenderThread(FRDGBuilder& GraphBuilder, const FSceneView& View, const struct FScreenPassRenderTarget& Output)
 {
-	for (FFXSystemInterface* FXSystem : FXSystems)
+	for (const TSharedRef<FFXSystemInterface>& FXSystem : FXSystems)
 	{
-		check(FXSystem);
 		FXSystem->DrawDebug_RenderThread(GraphBuilder, View, Output);
 	}
 }
 
 void FFXSystemSet::DrawSceneDebug_RenderThread(FRDGBuilder& GraphBuilder, const FSceneView& View, FRDGTextureRef SceneColor, FRDGTextureRef SceneDepth)
 {
-	for (FFXSystemInterface* FXSystem : FXSystems)
+	for (const TSharedRef<FFXSystemInterface>& FXSystem : FXSystems)
 	{
-		check(FXSystem);
 		FXSystem->DrawSceneDebug_RenderThread(GraphBuilder, View, SceneColor, SceneDepth);
 	}
 }
 
 void FFXSystemSet::AddVectorField(UVectorFieldComponent* VectorFieldComponent)
 {
-	for (FFXSystemInterface* FXSystem : FXSystems)
+	for (const TSharedRef<FFXSystemInterface>& FXSystem : FXSystems)
 	{
-		check(FXSystem);
 		FXSystem->AddVectorField(VectorFieldComponent);
 	}
 }
 
 void FFXSystemSet::RemoveVectorField(UVectorFieldComponent* VectorFieldComponent)
 {
-	for (FFXSystemInterface* FXSystem : FXSystems)
+	for (const TSharedRef<FFXSystemInterface>& FXSystem : FXSystems)
 	{
-		check(FXSystem);
 		FXSystem->RemoveVectorField(VectorFieldComponent);
 	}
 }
 
 void FFXSystemSet::UpdateVectorField(UVectorFieldComponent* VectorFieldComponent)
 {
-	for (FFXSystemInterface* FXSystem : FXSystems)
+	for (const TSharedRef<FFXSystemInterface>& FXSystem : FXSystems)
 	{
-		check(FXSystem);
 		FXSystem->UpdateVectorField(VectorFieldComponent);
 	}
 }
 
 void FFXSystemSet::PreInitViews(FRDGBuilder& GraphBuilder, bool bAllowGPUParticleUpdate, const TArrayView<const FSceneViewFamily*>& ViewFamilies, const FSceneViewFamily* CurrentFamily)
 {
-	for (FFXSystemInterface* FXSystem : FXSystems)
+	for (const TSharedRef<FFXSystemInterface>& FXSystem : FXSystems)
 	{
-		check(FXSystem);
 		FXSystem->PreInitViews(GraphBuilder, bAllowGPUParticleUpdate, ViewFamilies, CurrentFamily);
 	}
 }
 
 void FFXSystemSet::PostInitViews(FRDGBuilder& GraphBuilder, TConstStridedView<FSceneView> Views, bool bAllowGPUParticleUpdate)
 {
-	for (FFXSystemInterface* FXSystem : FXSystems)
+	for (const TSharedRef<FFXSystemInterface>& FXSystem : FXSystems)
 	{
-		check(FXSystem);
 		FXSystem->PostInitViews(GraphBuilder, Views, bAllowGPUParticleUpdate);
 	}
 }
 
 bool FFXSystemSet::UsesGlobalDistanceField() const
 {
-	for (FFXSystemInterface* FXSystem : FXSystems)
+	for (const TSharedRef<FFXSystemInterface>& FXSystem : FXSystems)
 	{
-		check(FXSystem);
 		if (FXSystem->UsesGlobalDistanceField())
 		{
 			return true;
@@ -159,9 +144,8 @@ bool FFXSystemSet::UsesGlobalDistanceField() const
 
 bool FFXSystemSet::UsesDepthBuffer() const
 {
-	for (FFXSystemInterface* FXSystem : FXSystems)
+	for (const TSharedRef<FFXSystemInterface>& FXSystem : FXSystems)
 	{
-		check(FXSystem);
 		if (FXSystem->UsesDepthBuffer())
 		{
 			return true;
@@ -172,9 +156,8 @@ bool FFXSystemSet::UsesDepthBuffer() const
 
 bool FFXSystemSet::RequiresEarlyViewUniformBuffer() const
 {
-	for (FFXSystemInterface* FXSystem : FXSystems)
+	for (const TSharedRef<FFXSystemInterface>& FXSystem : FXSystems)
 	{
-		check(FXSystem);
 		if (FXSystem->RequiresEarlyViewUniformBuffer())
 		{
 			return true;
@@ -185,9 +168,8 @@ bool FFXSystemSet::RequiresEarlyViewUniformBuffer() const
 
 bool FFXSystemSet::RequiresRayTracingScene() const
 {
-	for (FFXSystemInterface* FXSystem : FXSystems)
+	for (const TSharedRef<FFXSystemInterface>& FXSystem : FXSystems)
 	{
-		check(FXSystem);
 		if (FXSystem->RequiresRayTracingScene())
 		{
 			return true;
@@ -198,59 +180,50 @@ bool FFXSystemSet::RequiresRayTracingScene() const
 
 void FFXSystemSet::PreRender(FRDGBuilder& GraphBuilder, TConstStridedView<FSceneView> Views, FSceneUniformBuffer &SceneUniformBuffer, bool bAllowGPUParticleSceneUpdate)
 {
-	for (FFXSystemInterface* FXSystem : FXSystems)
+	for (const TSharedRef<FFXSystemInterface>& FXSystem : FXSystems)
 	{
-		check(FXSystem);
 		FXSystem->PreRender(GraphBuilder, Views, SceneUniformBuffer, bAllowGPUParticleSceneUpdate);
 	}
 }
 
 void FFXSystemSet::SetSceneTexturesUniformBuffer(const TUniformBufferRef<FSceneTextureUniformParameters>& InSceneTexturesUniformParams)
 {
-	for (FFXSystemInterface* FXSystem : FXSystems)
+	for (const TSharedRef<FFXSystemInterface>& FXSystem : FXSystems)
 	{
-		check(FXSystem);
 		FXSystem->SetSceneTexturesUniformBuffer(InSceneTexturesUniformParams);
 	}
 }
 
 void FFXSystemSet::PostRenderOpaque(FRDGBuilder& GraphBuilder, TConstStridedView<FSceneView> Views, FSceneUniformBuffer &SceneUniformBuffer, bool bAllowGPUParticleSceneUpdate)
 {
-	for (FFXSystemInterface* FXSystem : FXSystems)
+	for (const TSharedRef<FFXSystemInterface>& FXSystem : FXSystems)
 	{
-		check(FXSystem);
 		FXSystem->PostRenderOpaque(GraphBuilder, Views, SceneUniformBuffer, bAllowGPUParticleSceneUpdate);
 	}
 }
 
-void FFXSystemSet::OnDestroy()
+void FFXSystemSet::OnMarkPendingKill()
 {
-	for (FFXSystemInterface*& FXSystem : FXSystems)
+	for (const TSharedRef<FFXSystemInterface>& FXSystem : FXSystems)
 	{
-		check(FXSystem);
-		FXSystem->OnDestroy();
+		FXSystem->OnMarkPendingKill();
 	}
+	FXSystems.Empty();
 
-	FFXSystemInterface::OnDestroy();
+	FFXSystemInterface::OnMarkPendingKill();
 }
 
-
 void FFXSystemSet::DestroyGPUSimulation()
 {
-	for (FFXSystemInterface*& FXSystem : FXSystems)
+	for (const TSharedRef<FFXSystemInterface>& FXSystem : FXSystems)
 	{
-		check(FXSystem);
 		FXSystem->DestroyGPUSimulation();
 	}
 }
 
 FFXSystemSet::~FFXSystemSet()
 {
-	for (FFXSystemInterface* FXSystem : FXSystems)
-	{
-		check(FXSystem);
-		delete FXSystem;
-	}
+	FXSystems.Empty();
 }
 
 FGPUSortManager* FFXSystemSet::GetGPUSortManager() const
diff --git a/Engine/Source/Runtime/Engine/Private/Particles/FXSystemSet.h b/Engine/Source/Runtime/Engine/Private/Particles/FXSystemSet.h
--- a/Engine/Source/Runtime/Engine/Private/Particles/FXSystemSet.h
+++ b/Engine/Source/Runtime/Engine/Private/Particles/FXSystemSet.h
@@ -22,7 +22,7 @@ public:
 
 	FFXSystemSet(FGPUSortManager* InGPUSortManager);
 
-	TArray<FFXSystemInterface*> FXSystems;
+	TArray<TSharedRef<FFXSystemInterface>> FXSystems;
 
 	virtual FFXSystemInterface* GetInterface(const FName& InName) override;
 	virtual void Tick(UWorld* World, float DeltaSeconds) override;
@@ -48,7 +48,7 @@ public:
 	virtual void PreRender(FRDGBuilder& GraphBuilder, TConstStridedView<FSceneView> Views, FSceneUniformBuffer &SceneUniformBuffer, bool bAllowGPUParticleSceneUpdate) override;
 	virtual void PostRenderOpaque(FRDGBuilder& GraphBuilder, TConstStridedView<FSceneView> Views, FSceneUniformBuffer &SceneUniformBuffer, bool bAllowGPUParticleSceneUpdate) override;
 
-	virtual void OnDestroy() override;
+	virtual void OnMarkPendingKill() override;
 	virtual void DestroyGPUSimulation() override;
 
 	/** Get the shared SortManager, used in the rendering loop to call FGPUSortManager::OnPreRender() and FGPUSortManager::OnPostRenderOpaque() */
diff --git a/Engine/Source/Runtime/Engine/Private/World.cpp b/Engine/Source/Runtime/Engine/Private/World.cpp
--- a/Engine/Source/Runtime/Engine/Private/World.cpp
+++ b/Engine/Source/Runtime/Engine/Private/World.cpp
@@ -1547,7 +1547,7 @@ void UWorld::FinishDestroy()
 
 		if (FXSystem)
 		{
-			FFXSystemInterface::Destroy( FXSystem );
+			FFXSystemInterface::MarkPendingKill( FXSystem );
 			FXSystem = NULL;
 		}
 
@@ -6353,7 +6353,7 @@ void UWorld::CleanupWorldInternal(bool bSessionEnded, bool bCleanupResources, bo
 
 	if(FXSystem && bWorldChanged)
 	{
-		FFXSystemInterface::Destroy( FXSystem );
+		FFXSystemInterface::MarkPendingKill( FXSystem );
 		Scene->SetFXSystem(NULL);
 		FXSystem = NULL;
 	}
@@ -9982,7 +9982,7 @@ void UWorld::RecreateScene(ERHIFeatureLevel::Type InFeatureLevel, bool bBroadcas
 		if (FXSystem)
 		{
 			bCreateFXSystem = true;
-			FFXSystemInterface::Destroy(FXSystem);
+			FFXSystemInterface::MarkPendingKill(FXSystem);
 			FXSystem = nullptr;
 			Scene->SetFXSystem(nullptr);
 		}
diff --git a/Engine/Source/Runtime/Engine/Public/FXSystem.h b/Engine/Source/Runtime/Engine/Public/FXSystem.h
--- a/Engine/Source/Runtime/Engine/Public/FXSystem.h
+++ b/Engine/Source/Runtime/Engine/Public/FXSystem.h
@@ -127,19 +127,25 @@ DECLARE_DELEGATE_RetVal_ThreeParams(FFXSystemInterface*, FCreateCustomFXSystemDe
 /**
  * The interface to an effects system.
  */
-class FFXSystemInterface
+class FFXSystemInterface : public TSharedFromThis<FFXSystemInterface>
 {
-public:
+	struct FFXInterfaceDeletor
+	{
+		void operator()(FFXSystemInterface* ObjectToDelete) const;
+	};
 
+public:
 	/**
 	 * Create an effects system instance.
 	 */
 	ENGINE_API static FFXSystemInterface* Create(ERHIFeatureLevel::Type InFeatureLevel, FSceneInterface* Scene);
 
 	/**
-	 * Destroy an effects system instance.
+	 * Notify the effect system instance that it should release resources in preperation for destruction.
+	 * Once the interface is marked pending kill using the interface is undefined, as the scene (for example) may be invalid.
+	 * The destruction of the instance will happen later on the render thread, and may be much later due to shared references to the interface.
 	 */
-	ENGINE_API static void Destroy(FFXSystemInterface* FXSystem);
+	ENGINE_API static void MarkPendingKill(FFXSystemInterface* FXSystem);
 
 	/**
 	 * Queue Destroy the gpu simulation on the render thread
@@ -162,22 +168,30 @@ public:
 	virtual FFXSystemInterface* GetInterface(const FName& InName) { return nullptr; };
 
 	/**
-	 * Gamethread callback when destroy gets called, allows to clean up references.
+	 * Gamethread callback when MarkPendingKill is called, allows to clean up references.
+	 * Once marked pending kill the interface can release resources and is in an undefined state.
 	 */
-	virtual void OnDestroy() { bIsPendingKill = true; }
+	virtual void OnMarkPendingKill() { bIsPendingKill = true; }
 
 	/**
 	 * Gamethread callback when destroy gets called, allows to clean up references.
 	 */
 	virtual void DestroyGPUSimulation() { }
 
-
 	/**
 	 * Tick the effects system.
 	 * @param DeltaSeconds The number of seconds by which to step simulations forward.
 	 */
 	virtual void Tick(UWorld* World, float DeltaSeconds) = 0;
 
+	UE_DEPRECATED(5.7, "Destroy is no longer supported, please update your code call MarkPendingKill")
+	ENGINE_API static void Destroy(FFXSystemInterface* FXSystem);
+
+protected:
+	UE_DEPRECATED(5.7, "OnDestroy is no longer supported, please update your code call/implement OnMarkPendingKill")
+	virtual void OnDestroy() final {}
+public:
+
 #if WITH_EDITOR
 	/**
 	 * Suspend the FX system. This will cause internal state to be released.
diff --git a/Engine/Source/Runtime/Renderer/Private/RendererScene.cpp b/Engine/Source/Runtime/Renderer/Private/RendererScene.cpp
--- a/Engine/Source/Runtime/Renderer/Private/RendererScene.cpp
+++ b/Engine/Source/Runtime/Renderer/Private/RendererScene.cpp
@@ -1187,7 +1187,7 @@ FScene::FScene(UWorld* InWorld, bool bInRequiresHitProxies, bool bInIsEditorScen
 
 	if (World->FXSystem)
 	{
-		FFXSystemInterface::Destroy(World->FXSystem);
+		FFXSystemInterface::MarkPendingKill(World->FXSystem);
 	}
 
 	if (bCreateFXSystem)
