Allow specifying a screenshot view rect via FScreenshotRequest::RequestScreenshot.
The view rect will be ignored if it is invalid (empty or not within the extent of the render target), or if bShowUI=true, since that takes the slate widget screenshot path.
Default behavior remains the same (to screenshot the whole window).

diff --git a/Engine/Source/Runtime/Engine/Private/GameViewportClient.cpp b/Engine/Source/Runtime/Engine/Private/GameViewportClient.cpp
--- a/Engine/Source/Runtime/Engine/Private/GameViewportClient.cpp
+++ b/Engine/Source/Runtime/Engine/Private/GameViewportClient.cpp
@@ -2280,6 +2280,8 @@ bool UGameViewportClient::ProcessScreenShots(FViewport* InViewport)
 
 		if( bShowUI && FSlateApplication::IsInitialized() )
 		{
+			ensureMsgf(FScreenshotRequest::GetViewRect().IsEmpty(), TEXT("Specifying a screenshot view rect is not supported for slate widget screenshots (bShowUI=true). Using the full window size."));
+
 			TSharedRef<SWidget> WindowRef = WindowPtr.ToSharedRef();
 			if (bHdrEnabled)
 			{
@@ -2294,101 +2296,125 @@ bool UGameViewportClient::ProcessScreenShots(FViewport* InViewport)
 			GScreenshotResolutionY = Size.Y;
 			bIsUI = true;
 		}
-		else if (bHdrEnabled)
-		{
-			bScreenshotSuccessful = GetViewportScreenShotHDR(InViewport, BitmapHDR, FIntRect(0, 0, Size.X, Size.Y));
-			ConvertPixelDataToSCRGB(BitmapHDR, ViewportOutputFormat);
-		}
-		else if (InViewport->GetShaderResourceTexture() && InViewport->GetShaderResourceTexture()->GetDesc().Dimension == ETextureDimension::Texture2DArray)
+		else
 		{
-			// Handle Texture2DArray, we need to determine which array slices to include
-			FReadSurfaceDataFlags ReadSurfaceFlags = FReadSurfaceDataFlags();
-			const int32 ScreenshotTextureArrayMode = CVarScreenshotTextureArrayMode.GetValueOnAnyThread();
-			const int32 ArraySize = InViewport->GetShaderResourceTexture()->GetDesc().ArraySize;
+			// Set up view rect
+			FIntRect ViewRect = FScreenshotRequest::GetViewRect();
+			if (ViewRect.IsEmpty())
+			{
+				// If an empty view rect was provided (default), screenshot the entire render target
+				ViewRect = FIntRect(0, 0, Size.X, Size.Y);
+			}
+			else if (!FIntRect(0, 0, Size.X, Size.Y).Contains(ViewRect))
+			{
+				ensureMsgf(false, TEXT("Requested a screenshot with view rect %s, which does not fit within the viewport size of %dx%d. Using the full viewport size."), *ViewRect.ToString(), Size.X, Size.Y);
+				ViewRect = FIntRect(0, 0, Size.X, Size.Y);
+			}
+			else
+			{
+				// If a valid view rect was provided, update the output size to match the view rect size
+				Size.X = ViewRect.Width();
+				Size.Y = ViewRect.Height();
+			}
 
-			ensureAlwaysMsgf(ScreenshotTextureArrayMode >= 0 && ScreenshotTextureArrayMode < 3, TEXT("r.Screenshot.TextureArrayMode %d is not a valid mode, defaulting to mode 0"), ScreenshotTextureArrayMode);
+			check(FIntPoint(Size.X, Size.Y) == ViewRect.Size());
 
-			// Screenshot all slices, then collate horizontally or vertically
-			if (ScreenshotTextureArrayMode == 1 || ScreenshotTextureArrayMode == 2)
+			// Take screenshot
+			if (bHdrEnabled)
 			{
-				TArray<TArray<FColor>> BitmapArray;
-				BitmapArray.SetNum(ArraySize);
+				bScreenshotSuccessful = GetViewportScreenShotHDR(InViewport, BitmapHDR, ViewRect);
+				ConvertPixelDataToSCRGB(BitmapHDR, ViewportOutputFormat);
+			}
+			else if (InViewport->GetShaderResourceTexture() && InViewport->GetShaderResourceTexture()->GetDesc().Dimension == ETextureDimension::Texture2DArray)
+			{
+				// Handle Texture2DArray, we need to determine which array slices to include
+				FReadSurfaceDataFlags ReadSurfaceFlags = FReadSurfaceDataFlags();
+				const int32 ScreenshotTextureArrayMode = CVarScreenshotTextureArrayMode.GetValueOnAnyThread();
+				const int32 ArraySize = InViewport->GetShaderResourceTexture()->GetDesc().ArraySize;
+
+				ensureAlwaysMsgf(ScreenshotTextureArrayMode >= 0 && ScreenshotTextureArrayMode < 3, TEXT("r.Screenshot.TextureArrayMode %d is not a valid mode. Defaulting to mode 0."), ScreenshotTextureArrayMode);
 
-				for (int CurrentArrayIndex = 0; CurrentArrayIndex < ArraySize; CurrentArrayIndex++)
+				// Screenshot all slices, then collate horizontally or vertically
+				if (ScreenshotTextureArrayMode == 1 || ScreenshotTextureArrayMode == 2)
 				{
-					ReadSurfaceFlags.SetArrayIndex(CurrentArrayIndex);
-					bScreenshotSuccessful = GetViewportScreenShot(InViewport, BitmapArray[CurrentArrayIndex], FIntRect(0, 0, Size.X, Size.Y), ReadSurfaceFlags);
+					TArray<TArray<FColor>> BitmapArray;
+					BitmapArray.SetNum(ArraySize);
 
-					if (!bScreenshotSuccessful)
+					for (int CurrentArrayIndex = 0; CurrentArrayIndex < ArraySize; CurrentArrayIndex++)
 					{
-						break;
+						ReadSurfaceFlags.SetArrayIndex(CurrentArrayIndex);
+						bScreenshotSuccessful = GetViewportScreenShot(InViewport, BitmapArray[CurrentArrayIndex], ViewRect, ReadSurfaceFlags);
+
+						if (!bScreenshotSuccessful)
+						{
+							break;
+						}
 					}
-				}
 
-				if (bScreenshotSuccessful)
-				{
-					const int32 NumPixelsPerSlice = Size.X * Size.Y;
-					const int32 NumPixelsTotal = NumPixelsPerSlice * ArraySize;
+					if (bScreenshotSuccessful)
+					{
+						const int32 NumPixelsPerSlice = Size.X * Size.Y;
+						const int32 NumPixelsTotal = NumPixelsPerSlice * ArraySize;
 
-					Bitmap.Reset();
-					Bitmap.SetNumUninitialized(NumPixelsTotal);
+						Bitmap.Reset();
+						Bitmap.SetNumUninitialized(NumPixelsTotal);
 
-					// Horizontal collation
-					if (ScreenshotTextureArrayMode == 1)
-					{
-						for (int CurrentPixel = 0; CurrentPixel < NumPixelsTotal; CurrentPixel++)
+						// Horizontal collation
+						if (ScreenshotTextureArrayMode == 1)
 						{
-							// Complex case, need to alternate between array entries for each horizontal line
-							const int32 CurrentRow = CurrentPixel / (ArraySize * Size.X);
-							const int32 CurrentColumn = CurrentPixel % Size.X;
+							for (int CurrentPixel = 0; CurrentPixel < NumPixelsTotal; CurrentPixel++)
+							{
+								// Complex case, need to alternate between array entries for each horizontal line
+								const int32 CurrentRow = CurrentPixel / (ArraySize * Size.X);
+								const int32 CurrentColumn = CurrentPixel % Size.X;
 								 
-							const int32 ArrayIndex = (CurrentPixel / Size.X) % ArraySize;
-							const int32 PixelIndex = CurrentColumn + CurrentRow * Size.X;
-							Bitmap[CurrentPixel] = BitmapArray[ArrayIndex][PixelIndex];
-						}
+								const int32 ArrayIndex = (CurrentPixel / Size.X) % ArraySize;
+								const int32 PixelIndex = CurrentColumn + CurrentRow * Size.X;
+								Bitmap[CurrentPixel] = BitmapArray[ArrayIndex][PixelIndex];
+							}
 
-						Size.X *= ArraySize;
-						GScreenshotResolutionX = Size.X;
-					}
+							Size.X *= ArraySize;
+							GScreenshotResolutionX = Size.X;
+						}
 
-					// Vertical collation
-					else
-					{
-						for (int CurrentPixel = 0; CurrentPixel < NumPixelsTotal; CurrentPixel++)
+						// Vertical collation
+						else
 						{
-							// Simple case, just concatenate all bitmaps
-							const int32 ArrayIndex = CurrentPixel / NumPixelsPerSlice;
-							const int32 PixelIndex = CurrentPixel % NumPixelsPerSlice;
-							Bitmap[CurrentPixel] = BitmapArray[ArrayIndex][PixelIndex];
+							for (int CurrentPixel = 0; CurrentPixel < NumPixelsTotal; CurrentPixel++)
+							{
+								// Simple case, just concatenate all bitmaps
+								const int32 ArrayIndex = CurrentPixel / NumPixelsPerSlice;
+								const int32 PixelIndex = CurrentPixel % NumPixelsPerSlice;
+								Bitmap[CurrentPixel] = BitmapArray[ArrayIndex][PixelIndex];
+							}
+
+							Size.Y *= ArraySize;
+							GScreenshotResolutionY = Size.Y;
 						}
+					}
+				}
 
-						Size.Y *= ArraySize;
-						GScreenshotResolutionY = Size.Y;
+				// Screenshot a specific array slice (ScreenshotTextureArrayMode == 0 or other invalid values)
+				else
+				{
+					int32 ArrayIndex = CVarScreenshotTextureArrayIndex.GetValueOnAnyThread();
+					const int32 ClampedArrayIndex = FMath::Clamp(ArrayIndex, 0, ArraySize - 1);
+
+					if (ArrayIndex != ClampedArrayIndex)
+					{
+						ensureAlwaysMsgf(false, TEXT("r.Screenshot.TextureArrayIndex %d is out of range for the current viewport, which has %d texture array slices. Defaulting to slice %d."), ArrayIndex, ArraySize, ClampedArrayIndex);
+						ArrayIndex = ClampedArrayIndex;
 					}
+
+					ReadSurfaceFlags.SetArrayIndex(ArrayIndex);
+					bScreenshotSuccessful = GetViewportScreenShot(InViewport, Bitmap, ViewRect, ReadSurfaceFlags);
 				}
 			}
-
-			// Screenshot a specific array slice (ScreenshotTextureArrayMode == 0 or other invalid values)
 			else
 			{
-				int32 ArrayIndex = CVarScreenshotTextureArrayIndex.GetValueOnAnyThread();
-				const int32 ClampedArrayIndex = FMath::Clamp(ArrayIndex, 0, ArraySize - 1);
-
-				if (ArrayIndex != ClampedArrayIndex)
-				{
-					ensureAlwaysMsgf(false, TEXT("r.Screenshot.TextureArrayIndex %d is out of range for the current viewport, which has %d texture array slices, defaulting to slice %d"), ArrayIndex, ArraySize, ClampedArrayIndex);
-					ArrayIndex = ClampedArrayIndex;
-				}
-
-				ReadSurfaceFlags.SetArrayIndex(ArrayIndex);
-				bScreenshotSuccessful = GetViewportScreenShot(InViewport, Bitmap, FIntRect(0, 0, Size.X, Size.Y), ReadSurfaceFlags);
+				bScreenshotSuccessful = GetViewportScreenShot(InViewport, Bitmap, ViewRect);
 			}
 		}
-		else
-		{
-			bScreenshotSuccessful = GetViewportScreenShot(InViewport, Bitmap, FIntRect(0, 0, Size.X, Size.Y));
-		}
-		
 
 		if (bScreenshotSuccessful)
 		{
diff --git a/Engine/Source/Runtime/Engine/Private/UnrealClient.cpp b/Engine/Source/Runtime/Engine/Private/UnrealClient.cpp
--- a/Engine/Source/Runtime/Engine/Private/UnrealClient.cpp
+++ b/Engine/Source/Runtime/Engine/Private/UnrealClient.cpp
@@ -206,10 +206,10 @@ const FTextureRHIRef& FRenderTarget::GetShaderResourceTexture() const
 void FScreenshotRequest::RequestScreenshot(bool bInShowUI)
 {
 	// empty string means we'll later pick the name
-	RequestScreenshot(TEXT(""), bInShowUI, /*bAddUniqueSuffix*/ true, /*bHdrScreenshot*/ false);
+	RequestScreenshot(TEXT(""), bInShowUI, /*bAddUniqueSuffix*/ true, /*bHdrScreenshot*/ false, FIntRect());
 }
 
-void FScreenshotRequest::RequestScreenshot(const FString& InFilename, bool bInShowUI, bool bAddUniqueSuffix, bool bHdrScreenshot)
+void FScreenshotRequest::RequestScreenshot(const FString& InFilename, bool bInShowUI, bool bAddUniqueSuffix, bool bHdrScreenshot, const FIntRect& InViewRect)
 {
 	FString GeneratedFilename = InFilename;
 	CreateViewportScreenShotFilename(GeneratedFilename);
@@ -243,6 +243,7 @@ void FScreenshotRequest::RequestScreenshot(const FString& InFilename, bool bInSh
 	if (!Filename.IsEmpty())
 	{
 		bShowUI = bInShowUI;
+		ViewRect = InViewRect;
 		bIsScreenshotRequested = true;
 	}
 
@@ -261,6 +262,7 @@ void FScreenshotRequest::Reset()
 	bIsScreenshotRequested = false;
 	Filename.Empty();
 	bShowUI = false;
+	ViewRect = FIntRect();
 }
 
 void FScreenshotRequest::CreateViewportScreenShotFilename(FString& InOutFilename)
@@ -337,6 +339,7 @@ FString FScreenshotRequest::NextScreenshotName;
 bool FScreenshotRequest::bShowUI = false;
 TArray<FColor> FScreenshotRequest::HighresScreenshotMaskColorArray;
 FIntPoint FScreenshotRequest::HighresScreenshotMaskExtents;
+FIntRect FScreenshotRequest::ViewRect = FIntRect();
 
 static TAutoConsoleVariable<int32> CVarFullSizeUnitGraph(
 	TEXT("FullSizeUnitGraph"),
diff --git a/Engine/Source/Runtime/Engine/Public/UnrealClient.h b/Engine/Source/Runtime/Engine/Public/UnrealClient.h
--- a/Engine/Source/Runtime/Engine/Public/UnrealClient.h
+++ b/Engine/Source/Runtime/Engine/Public/UnrealClient.h
@@ -209,8 +209,10 @@ struct FScreenshotRequest
 	 * @param InFilename			The filename to use
 	 * @param bInShowUI				Whether or not to show Slate UI
 	 * @param bAddFilenameSuffix	Whether an auto-generated unique suffix should be added to the supplied filename
+	 * @param bHdrScreenshot		If true, the screenshot is saved as .exr rather than .png
+	 * @param InViewRect			If non-empty, within the viewport extent, and bInShowUI=false, only this view rect will be included in the screenshot
 	 */
-	static ENGINE_API void RequestScreenshot(const FString& InFilename, bool bInShowUI, bool bAddFilenameSuffix, bool bHdrScreenshot=false);
+	static ENGINE_API void RequestScreenshot(const FString& InFilename, bool bInShowUI, bool bAddFilenameSuffix, bool bHdrScreenshot = false, const FIntRect& InViewRect = FIntRect());
 
 	/**
 	 * Resets a screenshot request
@@ -232,6 +234,11 @@ struct FScreenshotRequest
 	 */
 	static bool ShouldShowUI() { return bShowUI; }
 
+	/**
+	 * @return The view rect for the next screenshot
+	 */
+	static FIntRect GetViewRect() { return ViewRect; }
+
 	/**
 	 * Creates a new screenshot filename from the passed in filename template
 	 */
@@ -266,6 +273,7 @@ private:
 	static ENGINE_API bool bShowUI;
 	static ENGINE_API TArray<FColor> HighresScreenshotMaskColorArray;
 	static ENGINE_API FIntPoint HighresScreenshotMaskExtents;
+	static ENGINE_API FIntRect ViewRect;
 };
 
 // @param bAutoType true: automatically choose GB/MB/KB/... false: always use MB for easier comparisons
