Subtitles: remove ESRB, add comment field.
UE-308936, UE-314969

#virtualized

diff --git a/Engine/Source/Runtime/Engine/Public/Subtitles/SubtitlesAndClosedCaptionsDelegates.h b/Engine/Source/Runtime/Engine/Public/Subtitles/SubtitlesAndClosedCaptionsDelegates.h
--- a/Engine/Source/Runtime/Engine/Public/Subtitles/SubtitlesAndClosedCaptionsDelegates.h
+++ b/Engine/Source/Runtime/Engine/Public/Subtitles/SubtitlesAndClosedCaptionsDelegates.h
@@ -28,19 +28,6 @@ enum class ESubtitleDurationType : uint8
 	UseDurationProperty UMETA(ToolTip = "Time this subtitle using its Duration property")
 };
 
-// ESRB rating categories.
-UENUM(BlueprintType)
-enum class EESRB : uint8
-{
-	Everyone,
-	Everyone10Plus,
-	Teen,
-	Mature,
-	AdultsOnly,
-	RatingPending,
-	RatingPending17Plus
-};
-
 // Subtitle type for type-specific rendering.
 UENUM()
 enum class ESubtitleType : uint8
@@ -97,13 +84,15 @@ public:
 	UPROPERTY(EditAnywhere, BlueprintReadWrite, Category = "Subtitles")
 	float Priority = 1.f;
 
-	// ESRB rating category.
-	UPROPERTY(EditAnywhere, BlueprintReadWrite, Category = "Subtitles", meta = (DisplayName = "ESRB"))
-	EESRB Esrb = EESRB::Everyone;
-
 	// Subtitle type for type-specific rendering.
 	UPROPERTY(EditAnywhere, BlueprintReadWrite, Category = "Subtitles")
 	ESubtitleType SubtitleType = ESubtitleType::Subtitle;
+
+#if WITH_EDITORONLY_DATA
+	// Author comments, editor-only.
+	UPROPERTY(EditAnywhere, BlueprintReadOnly, Category = "Subtitles")
+	FString Comment;
+#endif // WITH_EDITORONLY_DATA
 };
 
 
