[HWRT] Remove deprecated FRayTracingMaterialGatheringContext.

diff --git a/Engine/Plugins/FX/Niagara/Source/Niagara/Private/NiagaraSystemRenderData.cpp b/Engine/Plugins/FX/Niagara/Source/Niagara/Private/NiagaraSystemRenderData.cpp
--- a/Engine/Plugins/FX/Niagara/Source/Niagara/Private/NiagaraSystemRenderData.cpp
+++ b/Engine/Plugins/FX/Niagara/Source/Niagara/Private/NiagaraSystemRenderData.cpp
@@ -248,18 +248,6 @@ void FNiagaraSystemRenderData::GetDynamicMeshElements(const TArray<const FSceneV
 }
 
 #if RHI_RAYTRACING
-PRAGMA_DISABLE_DEPRECATION_WARNINGS
-void FNiagaraSystemRenderData::GetDynamicRayTracingInstances(FRayTracingMaterialGatheringContext& Context, TArray<FRayTracingInstance>& OutRayTracingInstances, const FNiagaraSceneProxy& SceneProxy)
-{
-	for (auto Renderer : EmitterRenderers_RT)
-	{
-		if (Renderer)
-		{
-			Renderer->GetDynamicRayTracingInstances(Context, OutRayTracingInstances, &SceneProxy);
-		}
-	}
-}
-PRAGMA_ENABLE_DEPRECATION_WARNINGS
 void FNiagaraSystemRenderData::GetDynamicRayTracingInstances(FRayTracingInstanceCollector& Collector, const FNiagaraSceneProxy& SceneProxy)
 {
 	for (auto Renderer : EmitterRenderers_RT)
diff --git a/Engine/Plugins/FX/Niagara/Source/Niagara/Public/NiagaraRenderer.h b/Engine/Plugins/FX/Niagara/Source/Niagara/Public/NiagaraRenderer.h
--- a/Engine/Plugins/FX/Niagara/Source/Niagara/Public/NiagaraRenderer.h
+++ b/Engine/Plugins/FX/Niagara/Source/Niagara/Public/NiagaraRenderer.h
@@ -26,7 +26,6 @@ class FNiagaraGPURendererCount;
 class FNiagaraSystemInstanceController;
 class FRayTracingInstanceCollector;
 struct FRayTracingInstance;
-struct FRayTracingMaterialGatheringContext;
 
 /** Struct used to pass dynamic data from game thread to render thread */
 struct FNiagaraDynamicDataBase
@@ -108,10 +107,6 @@ public:
 	inline bool IsMotionBlurEnabled() const { return bMotionBlurEnabled; }
 
 #if RHI_RAYTRACING
-	PRAGMA_DISABLE_DEPRECATION_WARNINGS
-	UE_DEPRECATED(5.5, "Use FRayTracingInstanceCollector instead.")
-	virtual void GetDynamicRayTracingInstances(FRayTracingMaterialGatheringContext& Context, TArray<FRayTracingInstance>& OutRayTracingInstances, const FNiagaraSceneProxy* Proxy) {}
-	PRAGMA_ENABLE_DEPRECATION_WARNINGS
 	virtual void GetDynamicRayTracingInstances(FRayTracingInstanceCollector& Collector, const FNiagaraSceneProxy* Proxy) {}
 #endif
 
diff --git a/Engine/Plugins/FX/Niagara/Source/Niagara/Public/NiagaraSystemRenderData.h b/Engine/Plugins/FX/Niagara/Source/Niagara/Public/NiagaraSystemRenderData.h
--- a/Engine/Plugins/FX/Niagara/Source/Niagara/Public/NiagaraSystemRenderData.h
+++ b/Engine/Plugins/FX/Niagara/Source/Niagara/Public/NiagaraSystemRenderData.h
@@ -68,10 +68,6 @@ public:
 	/** Gets the dynamic mesh elements from all renderers */
 	void GetDynamicMeshElements(const TArray<const FSceneView*>& Views, const FSceneViewFamily& ViewFamily, uint32 VisibilityMap, FMeshElementCollector& Collector, const FNiagaraSceneProxy& SceneProxy);
 #if RHI_RAYTRACING
-	PRAGMA_DISABLE_DEPRECATION_WARNINGS
-	UE_DEPRECATED(5.5, "Use FRayTracingInstanceCollector instead.")
-	void GetDynamicRayTracingInstances(FRayTracingMaterialGatheringContext& Context, TArray<FRayTracingInstance>& OutRayTracingInstances, const FNiagaraSceneProxy& SceneProxy);
-	PRAGMA_ENABLE_DEPRECATION_WARNINGS
 	void GetDynamicRayTracingInstances(FRayTracingInstanceCollector& Collector, const FNiagaraSceneProxy& SceneProxy);
 #endif
 	void GatherSimpleLights(FSimpleLightArray& OutParticleLights) const;
diff --git a/Engine/Source/Runtime/Engine/Private/SceneManagement.cpp b/Engine/Source/Runtime/Engine/Private/SceneManagement.cpp
--- a/Engine/Source/Runtime/Engine/Private/SceneManagement.cpp
+++ b/Engine/Source/Runtime/Engine/Private/SceneManagement.cpp
@@ -368,70 +368,6 @@ void FRayTracingInstanceCollector::AddRayTracingGeometryUpdate(FRayTracingDynami
 {
 	RayTracingGeometriesToUpdate.Add(MoveTemp(Params));
 }
-
-PRAGMA_DISABLE_DEPRECATION_WARNINGS
-FRayTracingMaterialGatheringContext::FRayTracingMaterialGatheringContext(
-	const FScene* InScene,
-	const FSceneView* InReferenceView,
-	const FSceneViewFamily& InReferenceViewFamily,
-	FRDGBuilder& InGraphBuilder,
-	FRayTracingMeshResourceCollector& InRayTracingMeshResourceCollector,
-	FGPUScenePrimitiveCollector& InDynamicPrimitiveCollector,
-	FGlobalDynamicReadBuffer& InDynamicReadBuffer)
-	: Scene(InScene)
-	, ReferenceView(InReferenceView)
-	, ReferenceViewFamily(InReferenceViewFamily)
-	, GraphBuilder(InGraphBuilder)
-	, RHICmdList(GraphBuilder.RHICmdList)
-	, RayTracingMeshResourceCollector(InRayTracingMeshResourceCollector)
-	, DynamicVertexBuffer(GraphBuilder.RHICmdList)
-	, DynamicIndexBuffer(GraphBuilder.RHICmdList)
-	, DynamicReadBuffer(InDynamicReadBuffer)
-	, bUsingReferenceBasedResidency(IsRayTracingUsingReferenceBasedResidency())
-{
-	RayTracingMeshResourceCollector.Start(RHICmdList, DynamicVertexBuffer, DynamicIndexBuffer, DynamicReadBuffer);
-
-	RayTracingMeshResourceCollector.AddViewMeshArrays(
-		ReferenceView,
-		nullptr,
-		nullptr,
-		&InDynamicPrimitiveCollector
-#if UE_ENABLE_DEBUG_DRAWING
-		, nullptr
-#endif
-	);
-}
-
-FRayTracingMaterialGatheringContext::~FRayTracingMaterialGatheringContext()
-{
-	RayTracingMeshResourceCollector.Finish();
-	DynamicReadBuffer.Commit(GraphBuilder.RHICmdList);
-}
-
-void FRayTracingMaterialGatheringContext::SetPrimitive(const FPrimitiveSceneProxy* InPrimitiveSceneProxy)
-{
-	RayTracingMeshResourceCollector.SetPrimitive(InPrimitiveSceneProxy, FHitProxyId::InvisibleHitProxyId);
-}
-
-void FRayTracingMaterialGatheringContext::Reset()
-{
-	DynamicRayTracingGeometriesToUpdate.Reset();
-	ReferencedGeometryGroups.Reset();
-}
-
-void FRayTracingMaterialGatheringContext::AddReferencedGeometryGroup(RayTracing::FGeometryGroupHandle GeometryGroup)
-{
-	if (bUsingReferenceBasedResidency)
-	{
-		ReferencedGeometryGroups.Add(GeometryGroup);
-	}
-}
-
-const TSet<RayTracing::FGeometryGroupHandle>& FRayTracingMaterialGatheringContext::GetReferencedGeometryGroups() const
-{
-	return ReferencedGeometryGroups;
-}
-PRAGMA_ENABLE_DEPRECATION_WARNINGS
 #endif
 
 FMeshElementCollector::FMeshElementCollector(ERHIFeatureLevel::Type InFeatureLevel, FSceneRenderingBulkObjectAllocator& InBulkAllocator, ECommitFlags InCommitFlags) :
diff --git a/Engine/Source/Runtime/Engine/Public/MeshElementCollector.h b/Engine/Source/Runtime/Engine/Public/MeshElementCollector.h
--- a/Engine/Source/Runtime/Engine/Public/MeshElementCollector.h
+++ b/Engine/Source/Runtime/Engine/Public/MeshElementCollector.h
@@ -229,7 +229,6 @@ protected:
 	friend class FViewFamilyInfo;
 	friend class FShadowMeshCollector;
 	friend class FDynamicMeshElementContext;
-	friend struct FRayTracingMaterialGatheringContext;
 	friend class RuntimeVirtualTexture::FDynamicMeshCollector;
 	friend FSceneRenderingBulkObjectAllocator;
 };
diff --git a/Engine/Source/Runtime/Engine/Public/PrimitiveSceneProxy.h b/Engine/Source/Runtime/Engine/Public/PrimitiveSceneProxy.h
--- a/Engine/Source/Runtime/Engine/Public/PrimitiveSceneProxy.h
+++ b/Engine/Source/Runtime/Engine/Public/PrimitiveSceneProxy.h
@@ -462,16 +462,13 @@ public:
 	/** Gathers dynamic ray tracing instances from this proxy. */
 	virtual void GetDynamicRayTracingInstances(class FRayTracingInstanceCollector& Collector) {}
 
-	UE_DEPRECATED(5.5, "Use FRayTracingInstanceCollector instead.")
-	virtual void GetDynamicRayTracingInstances(struct FRayTracingMaterialGatheringContext& Context, TArray<struct FRayTracingInstance>& OutRayTracingInstances) {}
-
 	virtual TArray<FRayTracingGeometry*> GetStaticRayTracingGeometries() const { return {}; }
 
 	/** 
 	 * Gathers static ray tracing primitives from this proxy. 
 	 * Fields of the instances can be partially filled, depending on what is desired to be cached, which is described in the returned flags
 	 */
-	ENGINE_API virtual ERayTracingPrimitiveFlags GetCachedRayTracingInstance(FRayTracingInstance& OutRayTracingInstance);
+	ENGINE_API virtual ERayTracingPrimitiveFlags GetCachedRayTracingInstance(struct FRayTracingInstance& OutRayTracingInstance);
 
 	/**
 	 * If the ray tracing data is streaming then get the coarse mesh streaming handle 
diff --git a/Engine/Source/Runtime/Engine/Public/SceneManagement.h b/Engine/Source/Runtime/Engine/Public/SceneManagement.h
--- a/Engine/Source/Runtime/Engine/Public/SceneManagement.h
+++ b/Engine/Source/Runtime/Engine/Public/SceneManagement.h
@@ -1521,47 +1521,6 @@ struct FRayTracingDynamicGeometryUpdateParams
 	FBufferRHIRef IndexBuffer = nullptr;
 };
 
-struct UE_DEPRECATED(5.5, "Use FRayTracingInstanceCollector instead.") FRayTracingMaterialGatheringContext
-{
-	const class FScene* Scene;
-	const FSceneView* ReferenceView;
-	const FSceneViewFamily& ReferenceViewFamily;
-
-PRAGMA_DISABLE_DEPRECATION_WARNINGS
-	FRDGBuilder& GraphBuilder;
-	FRHICommandList& RHICmdList;
-	FRayTracingMeshResourceCollector& RayTracingMeshResourceCollector;
-	TArray<FRayTracingDynamicGeometryUpdateParams> DynamicRayTracingGeometriesToUpdate;
-	FGlobalDynamicVertexBuffer DynamicVertexBuffer;
-	FGlobalDynamicIndexBuffer DynamicIndexBuffer;
-	FGlobalDynamicReadBuffer& DynamicReadBuffer;
-
-	ENGINE_API FRayTracingMaterialGatheringContext(
-		const FScene* InScene,
-		const FSceneView* InReferenceView,
-		const FSceneViewFamily& InReferenceViewFamily,
-		FRDGBuilder& InGraphBuilder,
-		FRayTracingMeshResourceCollector& InRayTracingMeshResourceCollector,
-		FGPUScenePrimitiveCollector& InDynamicPrimitiveCollector,
-		FGlobalDynamicReadBuffer& InGlobalDynamicReadBuffer);
-
-	ENGINE_API virtual ~FRayTracingMaterialGatheringContext();
-PRAGMA_ENABLE_DEPRECATION_WARNINGS
-
-	ENGINE_API void SetPrimitive(const FPrimitiveSceneProxy* InPrimitiveSceneProxy);
-
-	ENGINE_API void Reset();
-
-	ENGINE_API void AddReferencedGeometryGroup(RayTracing::FGeometryGroupHandle GeometryGroup);
-
-	ENGINE_API const TSet<RayTracing::FGeometryGroupHandle>& GetReferencedGeometryGroups() const;
-
-private:
-	TSet<RayTracing::FGeometryGroupHandle> ReferencedGeometryGroups;
-
-	const bool bUsingReferenceBasedResidency;
-};
-
 namespace RayTracing
 {
 	class FDynamicRayTracingInstancesContext;
