Don't keep references to classes after all matching subsystems have been removed

This fixes a crash with merged modules when after many subsystems have properly deinitialized, we'll still try to access their class.
In these cases the class might have been removed and its code unloaded, so we have to clean up.

#rnx

diff --git a/Engine/Source/Runtime/Engine/Private/Subsystems/SubsystemCollection.cpp b/Engine/Source/Runtime/Engine/Private/Subsystems/SubsystemCollection.cpp
--- a/Engine/Source/Runtime/Engine/Private/Subsystems/SubsystemCollection.cpp
+++ b/Engine/Source/Runtime/Engine/Private/Subsystems/SubsystemCollection.cpp
@@ -336,6 +336,8 @@ void FSubsystemCollectionBase::RemoveAndDeinitializeSubsystem(USubsystem* Subsys
 
 	const UClass* SubsystemClass = Subsystem->GetClass();
 
+	TArray<UClass*> ClassesToRemove;
+
 	for (TPair<UClass*, TUniquePtr<FSubsystemArray>>& Pair : SubsystemArrayMap)
 	{
 		const bool bIsInterface = Pair.Key->IsChildOf<UInterface>();
@@ -346,11 +348,21 @@ void FSubsystemCollectionBase::RemoveAndDeinitializeSubsystem(USubsystem* Subsys
 				*Subsystem->GetPathName(),
 				*Pair.Key->GetPathName());
 			Pair.Value->Subsystems.Remove(Subsystem);
+
+			if (Pair.Value->Subsystems.Num() == 0)
+			{
+				ClassesToRemove.Add(Pair.Key);
+			}
 		}
 	}
 
 	Subsystem->Deinitialize();
 	Subsystem->InternalOwningSubsystem = nullptr;
+
+	for (UClass* ClassToRemove : ClassesToRemove)
+	{
+		SubsystemArrayMap.Remove(ClassToRemove);
+	}
 }
 
 void FSubsystemCollectionBase::ActivateExternalSubsystem(UClass* SubsystemClass)
