New Material Translator:  Texture sample upgrades, plus support for 8 more expressions used in EngineTest project.
* Fix for inconsistent application of EMaterialSamplerType in texture samples.  Where texture object or texture object parameter nodes were used, was incorrectly using the sampler type from the expression node, rather than the linked object.
* Enabled validation of EMaterialSamplerType, matching validation behavior of the original translator.
* Fix for derivative calculations involving texture samples.  StageSwitch nodes would throw assertions if encountered when evaluating analytic derivatives, and these are automatically generated for texture reads.  Added pass through for StageSwitch derivatives to handle this.  Added code to fall back to hardware derivatives in cases where analytic derivatives are impossible (texture reads and custom nodes), rather than returning zero.  This matches the behavior of the original translator, now producing consistent code and visual results on test Nanite materials using derivatives.  Finally, needed to change from "ddx/ddy" to "DDX/DDY" -- necessary to avoid compile errors in shader passes that don't support hardware derivatives.
* Support for compiling for previous frame in vertex shader WorldPositionOffset.
* Fix for SampleTexture from vertex shader.  Need to generate SampleLevel, rather than attempting to use analytic derivatives, which only work in the compute shader.
* Throw error when explicit DDX / DDY nodes are used from vertex shader (as in original translator).
* Throw error when SampleTexture is used from vertex shader with the MipBias option (as in original translator).

#rnx

diff --git a/Engine/Source/Runtime/Engine/Private/Materials/MaterialExpressionsIR.cpp b/Engine/Source/Runtime/Engine/Private/Materials/MaterialExpressionsIR.cpp
--- a/Engine/Source/Runtime/Engine/Private/Materials/MaterialExpressionsIR.cpp
+++ b/Engine/Source/Runtime/Engine/Private/Materials/MaterialExpressionsIR.cpp
@@ -143,6 +143,7 @@
 #include "Materials/MaterialExpressionParticleRelativeTime.h"
 #include "Materials/MaterialExpressionParticleSize.h"
 #include "Materials/MaterialExpressionParticleSpeed.h"
+#include "Materials/MaterialExpressionParticleSubUV.h"
 #include "Materials/MaterialExpressionParticleSubUVProperties.h"
 #include "Materials/MaterialExpressionPathTracingBufferTexture.h"
 #include "Materials/MaterialExpressionPathTracingQualitySwitch.h"
@@ -605,7 +606,7 @@ void UMaterialExpressionVirtualTextureFeatureSwitch::Build(MIR::FEmitter& Em)
 	Em.Output(0, UseVirtualTexturing(Em.GetShaderPlatform()) ? YesValue : NoValue);
 }
 
-static FValueRef EmitInlineHLSL(MIR::FEmitter& Em, FName ExternalCodeIdentifier, TConstArrayView<FValueRef> InArguments)
+static FValueRef EmitInlineHLSL(MIR::FEmitter& Em, FName ExternalCodeIdentifier, TConstArrayView<FValueRef> InArguments = {})
 {
 	const FMaterialExternalCodeDeclaration* ExternalCodeDeclaration = MaterialExternalCodeRegistry::Get().FindExternalCode(ExternalCodeIdentifier);
 	if (!ExternalCodeDeclaration)
@@ -629,7 +630,7 @@ static void EmitExternalCodeConditionalReplace(MIR::FEmitter& Em, FName External
 	DefaultValue = Em.Cast(DefaultValue, CommonType);
 	ReplaceValue = Em.Cast(ReplaceValue, CommonType);
 
-	Em.Output(0, Em.Branch(EmitInlineHLSL(Em, ExternalCodeIdentifier, {}), ReplaceValue, DefaultValue));
+	Em.Output(0, Em.Branch(EmitInlineHLSL(Em, ExternalCodeIdentifier), ReplaceValue, DefaultValue));
 }
 
 #define REPLACE_INPUTS(DEFAULT, REPLACE) DEFAULT, REPLACE, TEXT(#DEFAULT), TEXT(#REPLACE)
@@ -911,7 +912,7 @@ void UMaterialExpressionSquareRoot::Build(MIR::FEmitter& Em)
 	Em.Output(0, Em.Sqrt(Em.Input(&Input)));
 }
 
-static FValueRef EmitInlineHLSL(MIR::FEmitter& Em, const UMaterialExpressionExternalCodeBase& InExternalCodeExpression, int32 InExternalCodeIdentifierIndex, TConstArrayView<FValueRef> InArguments, MIR::EValueFlags ValueFlags = MIR::EValueFlags::None)
+static FValueRef EmitInlineHLSL(MIR::FEmitter& Em, const UMaterialExpressionExternalCodeBase& InExternalCodeExpression, int32 InExternalCodeIdentifierIndex, TConstArrayView<FValueRef> InArguments = {}, MIR::EValueFlags ValueFlags = MIR::EValueFlags::None)
 {
 	checkf(InExternalCodeIdentifierIndex >= 0 && InExternalCodeIdentifierIndex < InExternalCodeExpression.ExternalCodeIdentifiers.Num(),
 		TEXT("External code identifier index (%d) out of bounds; Upper bound is %d"), InExternalCodeIdentifierIndex, InExternalCodeExpression.ExternalCodeIdentifiers.Num());
@@ -1266,38 +1267,41 @@ static MIR::ETextureReadMode TextureGatherModeToMIR(::ETextureGatherMode Mode)
 	}
 }
 
-static void BuildTextureSampleExpression(MIR::FEmitter& Em, UMaterialExpressionTextureSample* Expr, FValueRef Texture, EMaterialValueType TextureType)
+// Returns this value's texture sampler type if it has one (SAMPLERTYPE_MAX otherwise).
+static EMaterialSamplerType GetValueMaterialSamplerType(FValueRef Value)
 {
-	FValueRef TexCoords = Em.TryInput(&Expr->Coordinates);
-	if (!TexCoords)
+	if (auto TextureObject = MIR::As<MIR::FTextureObject>(Value))
 	{
-		TexCoords = Em.ExternalInput(MIR::TexCoordIndexToExternalInput(Expr->ConstCoordinate));
+		return TextureObject->SamplerType;
 	}
+	else if (auto UniformParameter = MIR::As<MIR::FUniformParameter>(Value))
+	{
+		return UniformParameter->SamplerType;
+	}
+	return SAMPLERTYPE_MAX;
+}
 
+static FValueRef BuildTextureSample(MIR::FEmitter& Em, UMaterialExpressionTextureSample* Expr, FValueRef Texture, EMaterialValueType TextureType, FValueRef TexCoords, bool bAutomaticViewMipBias)
+{
 	FValueRef TextureRead{};
 	if (Expr->GatherMode != TGM_None)
 	{
 		if (Expr->MipValueMode != TMVM_None)
 		{
 			Em.Errorf(TEXT("Texture gather does not support mipmap overrides (it implicitly accesses a specific mip)."));
-			return;
+			return TextureRead;
 		}
 
-		TextureRead = Em.TextureGather(Texture, TexCoords, TextureGatherModeToMIR(Expr->GatherMode), Expr->SamplerSource, Expr->SamplerType);
-	} 
+		TextureRead = Em.TextureGather(Texture, TexCoords, TextureGatherModeToMIR(Expr->GatherMode), Expr->SamplerSource, GetValueMaterialSamplerType(Texture));
+	}
 	else
 	{
-		// Determine if automatic view mip bias should be used, by trying to acquire a its input as a static boolean.
-		bool bAutomaticViewMipBias = Em.ToConstantBool(Em.InputDefaultBool(&Expr->AutomaticViewMipBiasValue, Expr->AutomaticViewMipBias));
-	
 		// If not 2D texture, disable AutomaticViewMipBias.
 		if (!(TextureType & (MCT_Texture2D | MCT_TextureVirtual | MCT_TextureMeshPaint)))
 		{
 			bAutomaticViewMipBias = false;
 		}
 
-		UE_MIR_CHECKPOINT(Em);
-
 		// Get the mip value level (either through the expression input or using the given constant if disconnected).
 		FValueRef MipValue;
 		if (Expr->MipValueMode == TMVM_MipLevel || Expr->MipValueMode == TMVM_MipBias)
@@ -1308,27 +1312,59 @@ static void BuildTextureSampleExpression(MIR::FEmitter& Em, UMaterialExpressionT
 		switch (Expr->MipValueMode)
 		{
 			case TMVM_None:
-				TextureRead = Em.TextureSample(Texture, TexCoords, bAutomaticViewMipBias, Expr->SamplerSource, Expr->SamplerType);
+				TextureRead = Em.TextureSample(Texture, TexCoords, bAutomaticViewMipBias, Expr->SamplerSource, GetValueMaterialSamplerType(Texture));
 				break;
 		
 			case TMVM_MipBias:
-				TextureRead = Em.TextureSampleBias(Texture, TexCoords, MipValue, bAutomaticViewMipBias, Expr->SamplerSource, Expr->SamplerType);
+				TextureRead = Em.TextureSampleBias(Texture, TexCoords, MipValue, bAutomaticViewMipBias, Expr->SamplerSource, GetValueMaterialSamplerType(Texture));
 				break;
 		
 			case TMVM_MipLevel:
-				TextureRead = Em.TextureSampleLevel(Texture, TexCoords, MipValue, bAutomaticViewMipBias, Expr->SamplerSource, Expr->SamplerType);
+				TextureRead = Em.TextureSampleLevel(Texture, TexCoords, MipValue, bAutomaticViewMipBias, Expr->SamplerSource, GetValueMaterialSamplerType(Texture));
 				break;
 
 			case TMVM_Derivative:
 			{
 				FValueRef TexCoordsDdx = Em.Cast(Em.Input(&Expr->CoordinatesDX), TexCoords->Type);
 				FValueRef TexCoordsDdy = Em.Cast(Em.Input(&Expr->CoordinatesDY), TexCoords->Type);
-				TextureRead = Em.TextureSampleGrad(Texture, TexCoords, TexCoordsDdx, TexCoordsDdy, bAutomaticViewMipBias, Expr->SamplerSource, Expr->SamplerType);
+				TextureRead = Em.TextureSampleGrad(Texture, TexCoords, TexCoordsDdx, TexCoordsDdy, bAutomaticViewMipBias, Expr->SamplerSource, GetValueMaterialSamplerType(Texture));
 				break;
 			}
 		}
 	}
 
+	return TextureRead;
+}
+
+static FValueRef BuildTextureValue(MIR::FEmitter& Em, UMaterialExpressionTextureSample* Expr)
+{
+	FValueRef TextureValue = Em.TryInput(&Expr->TextureObject);
+	if (!TextureValue)
+	{
+		if (!Expr->Texture)
+		{
+			Em.Error(TEXT("No texture specified for this expression."));
+			return Em.Poison();
+		}
+
+		TextureValue = Em.TextureObject(Expr->Texture.Get(), Expr->SamplerType);
+	}
+	return TextureValue;
+}
+
+static void BuildTextureSampleExpression(MIR::FEmitter& Em, UMaterialExpressionTextureSample* Expr, FValueRef Texture, EMaterialValueType TextureType)
+{
+	FValueRef TexCoords = Em.TryInput(&Expr->Coordinates);
+	if (!TexCoords)
+	{
+		TexCoords = Em.ExternalInput(MIR::TexCoordIndexToExternalInput(Expr->ConstCoordinate));
+	}
+
+	// Determine if automatic view mip bias should be used, by trying to acquire its input as a static boolean.
+	const bool bAutomaticViewMipBias = Em.ToConstantBool(Em.InputDefaultBool(&Expr->AutomaticViewMipBiasValue, Expr->AutomaticViewMipBias));
+
+	FValueRef TextureRead = BuildTextureSample(Em, Expr, Texture, TextureType, TexCoords, bAutomaticViewMipBias);
+
 	Em.Output(0, Em.Swizzle(TextureRead, MIR::FSwizzleMask::XYZ()));
 	Em.Output(1, Em.Subscript(TextureRead, 0));
 	Em.Output(2, Em.Subscript(TextureRead, 1));
@@ -1337,44 +1373,82 @@ static void BuildTextureSampleExpression(MIR::FEmitter& Em, UMaterialExpressionT
 	Em.Output(5, TextureRead); 
 }
 
-// Returns this value's textures if it has one (nullptr otherwise).
-static EMaterialSamplerType GetValueMaterialSamplerType(FValueRef Value)
+void UMaterialExpressionTextureSample::Build(MIR::FEmitter& Em)
 {
-	if (auto TextureObject = MIR::As<MIR::FTextureObject>(Value))
-	{
-		return TextureObject->SamplerType;
-	}
-	else if (auto UniformParameter = MIR::As<MIR::FUniformParameter>(Value))
-	{
-		return UniformParameter->SamplerType;
-	}
-	return SAMPLERTYPE_MAX;
+	FValueRef TextureValue = BuildTextureValue(Em, this);
+	UE_MIR_CHECKPOINT(Em);
+
+	BuildTextureSampleExpression(Em, this, TextureValue, Em.GetTextureFromValue(TextureValue)->GetMaterialType());
 }
 
-void UMaterialExpressionTextureSample::Build(MIR::FEmitter& Em)
+// Inherits from UMaterialExpressionTextureSample, but uses different particle specific UVs, and optionally supports blending two different texture samples.
+void UMaterialExpressionParticleSubUV::Build(MIR::FEmitter& Em)
 {
-	FValueRef TextureValue = Em.TryInput(&TextureObject);
-	if (!TextureValue)
+	FValueRef TextureValue = BuildTextureValue(Em, this);
+	UE_MIR_CHECKPOINT(Em);
+
+	EMaterialValueType TextureType = Em.GetTextureFromValue(TextureValue)->GetMaterialType();
+
+	// Although the parent UMaterialExpressionTextureSample class includes an automatic view mip bias flag, it is specifically ignored by ParticleSubUV.
+	const bool bAutomaticViewMipBias = false;
+
+	static FName NAME_ParticleSubUVCoords0("ParticleSubUVCoords0");
+	static FName NAME_ParticleSubUVCoords1("ParticleSubUVCoords1");
+	static FName NAME_ParticleSubUVLerp("ParticleSubUVLerp");
+
+	FValueRef TexCoords0 = EmitInlineHLSL(Em, NAME_ParticleSubUVCoords0);
+	FValueRef Sample0 = BuildTextureSample(Em, this, TextureValue, TextureType, TexCoords0, bAutomaticViewMipBias);
+
+	if (bBlend)
 	{
-		if (!Texture)
-		{
-			Em.Error(TEXT("No texture specified for this expression."));
-			return;
-		}
+		FValueRef TexCoords1 = EmitInlineHLSL(Em, NAME_ParticleSubUVCoords1);
+		FValueRef Sample1 = BuildTextureSample(Em, this, TextureValue, TextureType, TexCoords1, bAutomaticViewMipBias);
 
-		TextureValue = Em.TextureObject(Texture.Get(), SamplerType);
+		FValueRef SubImageLerp = EmitInlineHLSL(Em, NAME_ParticleSubUVLerp);
+
+		Sample0 = Em.Lerp(Sample0, Sample1, SubImageLerp);
 	}
 
-	UE_MIR_CHECKPOINT(Em);
+	// Same outputs as UMaterialExpressionTextureSample
+	Em.Output(0, Em.Swizzle(Sample0, MIR::FSwizzleMask::XYZ()));
+	Em.Output(1, Em.Subscript(Sample0, 0));
+	Em.Output(2, Em.Subscript(Sample0, 1));
+	Em.Output(3, Em.Subscript(Sample0, 2));
+	Em.Output(4, Em.Subscript(Sample0, 3));
+	Em.Output(5, Sample0);
+}
 
-	UTexture* DefaultTexture = Em.GetTextureFromValue(TextureValue);
-	if (!DefaultTexture)
+// Inherits from UMaterialExpressionTextureSample, but does extra math on the sample afterwards.  Note that this was an HLSL utility function in the
+// original translator, but uses ops here.  The main advantage of the ops version is that it uses the standard texture sampling code path, rather than
+// sampling the texture in the utility function, meaning it supports all sampling features (the original would break if using non-standard sampling).
+void UMaterialExpressionAntialiasedTextureMask::Build(MIR::FEmitter& Em)
+{
+	// Check if a texture is assigned and the right type.
+	FString ErrorMessage;
+	if (!TextureIsValid(Texture, ErrorMessage))
 	{
-		Em.Error(TEXT("Missing texture object from input"));
+		Em.Errorf(TEXT("%s"), *ErrorMessage);
 		return;
 	}
 
-	BuildTextureSampleExpression(Em, this, TextureValue, DefaultTexture->GetMaterialType());
+	FValueRef TexCoords = Em.TryInput(&Coordinates);
+	if (!TexCoords)
+	{
+		TexCoords = Em.ExternalInput(MIR::TexCoordIndexToExternalInput(ConstCoordinate));
+	}
+
+	FValueRef TextureValue = BuildTextureValue(Em, this);
+	UE_MIR_CHECKPOINT(Em);
+	FValueRef Sample1 = BuildTextureSample(Em, this, TextureValue, MCT_Texture2D, TexCoords, /*bAutomaticViewMipBias=*/ false);
+
+	FValueRef ThresholdConst = Em.ConstantFloat(Threshold);
+
+	// Logic below ported from the AntialiasedTextureMask HLSL function.
+	Sample1 = Em.Subscript(Sample1, FMath::Clamp(Channel, 0, 3));
+
+	FValueRef TexDDLength = Em.Max(Em.Abs(Em.PartialDerivative(Sample1, MIR::EDerivativeAxis::X)), Em.Abs(Em.PartialDerivative(Sample1, MIR::EDerivativeAxis::Y)));
+	FValueRef Top = Em.Subtract(Sample1, ThresholdConst);
+	Em.Output(0, Em.Add(Em.Divide(Top, TexDDLength), ThresholdConst));
 }
 
 static FValueRef BuildTextureObjectParameter(MIR::FEmitter& Em, UMaterialExpressionTextureSampleParameter* Expr)
@@ -1386,6 +1460,19 @@ static FValueRef BuildTextureObjectParameter(MIR::FEmitter& Em, UMaterialExpress
 		return nullptr;
 	}
 
+	if (!Expr->Texture)
+	{
+		Em.Error(TEXT("Requires valid texture"));
+		return nullptr;
+	}
+
+	FString SamplerTypeError;
+	if (!UMaterialExpressionTextureBase::VerifySamplerType(Em.GetShaderPlatform(), Em.GetTargetPlatform(), Expr->Texture, Expr->SamplerType, SamplerTypeError))
+	{
+		Em.Errorf(TEXT("%s"), *SamplerTypeError);
+		return nullptr;
+	}
+
 	FValueRef ParameterValue = Em.Parameter(Expr->GetParameterName(), Param, Expr->SamplerType);
 	if (!ParameterValue->Type.IsTexture())
 	{
@@ -1572,7 +1659,7 @@ void UMaterialExpressionTime::Build(MIR::FEmitter& Em)
 	// When pausing the game is ignored for this time expression, use real-time instead of game-time.
 	if (!bOverride_Period)
 	{
-		FStringView InlinedCode = bIgnorePause ? TEXTVIEW("View.<PREV>RealTime") : TEXTVIEW("View.<PREV>GameTime");
+		FStringView InlinedCode = bIgnorePause ? TEXTVIEW("View.<PREVFRAME>RealTime") : TEXTVIEW("View.<PREVFRAME>GameTime");
 		Em.Output(0, Em.InlineHLSL(ScalarFloatType, InlinedCode, {}, MIR::EValueFlags::SubstituteTags));
 	}
 	else if (Period == 0.0f)
@@ -1584,7 +1671,7 @@ void UMaterialExpressionTime::Build(MIR::FEmitter& Em)
 		// Note: Don't use IR intrinsic for Fmod() here to avoid conversion to fp16 on mobile.
 		// We want full 32 bit float precision until the fmod when using a period.
 		FValueRef PeriodValue = Em.ConstantFloat(Period);
-		FStringView InlinedCode = bIgnorePause ? TEXTVIEW("fmod(View.<PREV>RealTime, $0)") : TEXTVIEW("fmod(View.<PREV>GameTime, $0)");
+		FStringView InlinedCode = bIgnorePause ? TEXTVIEW("fmod(View.<PREVFRAME>RealTime, $0)") : TEXTVIEW("fmod(View.<PREVFRAME>GameTime, $0)");
 		Em.Output(0, Em.InlineHLSL(ScalarFloatType, InlinedCode, { PeriodValue }, MIR::EValueFlags::SubstituteTags));
 	}
 }
@@ -2451,6 +2538,30 @@ static FValueRef EmitWorldPosition(MIR::FEmitter& Em, EWorldPositionIncludedOffs
 	return WorldPosition;
 }
 
+static FValueRef EmitLocalPosition(MIR::FEmitter& Em, ELocalPositionOrigin LocalOrigin, EPositionIncludedOffsets LocalShaderOffset)
+{
+	// Make sure ELocalPositionOrigin / EPositionIncludedOffsets and corresponding elements of EExternalInput stay in sync, so math using enum is valid.
+	static_assert((int32)ELocalPositionOrigin::Instance == 0);
+	static_assert((int32)ELocalPositionOrigin::Primitive == 2);
+	static_assert((int32)EPositionIncludedOffsets::IncludeOffsets == 0);
+	static_assert((int32)EPositionIncludedOffsets::ExcludeOffsets == 1);
+	static_assert((int32)MIR::EExternalInput::LocalPosition_InstanceNoOffsets - (int32)MIR::EExternalInput::LocalPosition_Instance == (int32)EPositionIncludedOffsets::ExcludeOffsets);
+	static_assert((int32)MIR::EExternalInput::LocalPosition_Primitive - (int32)MIR::EExternalInput::LocalPosition_Instance == (int32)ELocalPositionOrigin::Primitive);
+
+	// ELocalPositionOrigin::InstancePreSkinning just uses an external code declaration, and doesn't have variations for offsets.
+	if (LocalOrigin == ELocalPositionOrigin::InstancePreSkinning)
+	{
+		static FName NAME_PreSkinnedPosition("PreSkinnedPosition");
+		return EmitInlineHLSL(Em, NAME_PreSkinnedPosition);
+	}
+
+	// Otherwise, there's a 2x2 configuration of Origin and Offset type.  Given that ELocalPositionOrigin::InstancePreSkinning is unused, valid origin values are 0 or 2,
+	// so we can just add the values together to get a unique index from [0..3]
+	int32 OriginIndex = FMath::Clamp((int32)LocalOrigin, 0, 2);
+	int32 OffsetIndex = FMath::Clamp((int32)LocalShaderOffset, 0, 1);
+	return Em.ExternalInput((MIR::EExternalInput)((int32)MIR::EExternalInput::LocalPosition_Instance + OriginIndex + OffsetIndex));
+}
+
 // Adapted from FHLSLMaterialTranslator::GetWorldPositionOrDefault.  WorldPosition input is optional, a default is provided if not set.
 static FValueRef EmitWorldPositionOrDefault(MIR::FEmitter& Em, FValueRef WorldPosition, EPositionOrigin PositionOrigin)
 {
@@ -2480,6 +2591,10 @@ void UMaterialExpressionWorldPosition::Build(MIR::FEmitter& Em)
 	Em.Output(2, Em.Subscript(WorldPosition, 2));
 }
 
+void UMaterialExpressionLocalPosition::Build(MIR::FEmitter& Em)
+{
+	Em.Output(0, EmitLocalPosition(Em, LocalOrigin, IncludedOffsets));
+}
 
 void UMaterialExpressionMakeMaterialAttributes::Build(MIR::FEmitter& Em)
 {
@@ -3012,7 +3127,7 @@ void UMaterialExpressionDepthFade::Build(MIR::FEmitter& Em)
 	FValueRef FadeDistanceValue = Em.Max(Em.InputDefaultFloat(&FadeDistance, FadeDistanceDefault), Em.ConstantFloat(UE_DELTA));
 
 	static FName NAME_PixelDepth("PixelDepth");
-	FValueRef PixelDepth = EmitInlineHLSL(Em, NAME_PixelDepth, {});
+	FValueRef PixelDepth = EmitInlineHLSL(Em, NAME_PixelDepth);
 	// On mobile scene depth is limited to 65500 
 	// to avoid false fading on objects that are close or exceed this limit we clamp pixel depth to (65500 - FadeDistance)
 	if (Em.GetFeatureLevel() <= ERHIFeatureLevel::ES3_1)
@@ -3046,7 +3161,7 @@ void UMaterialExpressionDistanceFieldApproxAO::Build(MIR::FEmitter& Em)
 	if (!NormalValue)
 	{
 		static FName NAME_VertexNormal("VertexNormal");
-		NormalValue = EmitInlineHLSL(Em, NAME_VertexNormal, {});
+		NormalValue = EmitInlineHLSL(Em, NAME_VertexNormal);
 	}
 
 	FValueRef BaseDistanceValue = Em.InputDefaultFloat(&BaseDistance, BaseDistanceDefault);
@@ -3119,7 +3234,7 @@ void UMaterialExpressionFresnel::Build(MIR::FEmitter& Em)
 		NormalArg = Em.InlineHLSL(MIR::FType::MakeFloatVector(3), TEXTVIEW("Parameters.WorldNormal"), {}, MIR::EValueFlags::None, MIR::EGraphProperties::ReadsPixelNormal);
 	}
 
-	FValueRef DotArg = Em.Dot(NormalArg, EmitInlineHLSL(Em, NAME_CameraVector, {}));
+	FValueRef DotArg = Em.Dot(NormalArg, EmitInlineHLSL(Em, NAME_CameraVector));
 	FValueRef MaxArg = Em.Max(Em.ConstantFloat(0.f), DotArg);
 	FValueRef MinusArg = Em.Subtract(Em.ConstantFloat(1.f), MaxArg);
 	FValueRef ExponentArg = Em.InputDefaultFloat(&ExponentIn, Exponent);
@@ -3148,13 +3263,13 @@ void UMaterialExpressionReflectionVectorWS::Build(MIR::FEmitter& Em)
 		}
 
 		// Normal * dot(Normal, CameraVector) * 2.0 - CameraVector;
-		FValueRef CameraVector = EmitInlineHLSL(Em, NAME_CameraVector, {});
+		FValueRef CameraVector = EmitInlineHLSL(Em, NAME_CameraVector);
 		Em.Output(0, Em.Subtract(Em.Multiply(NormalValue, Em.Multiply(Em.Dot(NormalValue, CameraVector), Em.ConstantFloat(2.0f))), CameraVector));
 	}
 	else
 	{
 		static FName NAME_ReflectionVector("ReflectionVector");
-		Em.Output(0, EmitInlineHLSL(Em, NAME_ReflectionVector, {}));
+		Em.Output(0, EmitInlineHLSL(Em, NAME_ReflectionVector));
 	}
 }
 
@@ -3190,7 +3305,7 @@ void UMaterialExpressionRotator::Build(MIR::FEmitter& Em)
 	FValueRef TimeValue = Em.TryInput(&Time);
 	if (!TimeValue)
 	{
-		TimeValue = Em.InlineHLSL(MIR::FType::MakeFloatScalar(), TEXTVIEW("View.<PREV>GameTime"), {}, MIR::EValueFlags::SubstituteTags);
+		TimeValue = Em.InlineHLSL(MIR::FType::MakeFloatScalar(), TEXTVIEW("View.<PREVFRAME>GameTime"), {}, MIR::EValueFlags::SubstituteTags);
 	}
 	TimeValue = Em.Multiply(TimeValue, Em.ConstantFloat(Speed));
 
@@ -3260,7 +3375,7 @@ void UMaterialExpressionSkyAtmosphereViewLuminance::Build(MIR::FEmitter& Em)
 	FValueRef WorldDirectionValue = Em.TryInput(&WorldDirection);
 	if (!WorldDirectionValue)
 	{
-		WorldDirectionValue = Em.Multiply(Em.ConstantFloat(-1.0f), EmitInlineHLSL(Em, NAME_CameraVector, {}));
+		WorldDirectionValue = Em.Multiply(Em.ConstantFloat(-1.0f), EmitInlineHLSL(Em, NAME_CameraVector));
 	}
 	Em.Output(0, EmitInlineHLSL(Em, *this, 0, { WorldDirectionValue }));
 }
@@ -3457,10 +3572,10 @@ void UMaterialExpressionBounds::Build(MIR::FEmitter& Em)
 	static_assert(MEILB_PreSkinnedLocal == 2);
 
 	int32 OutputOffset = Type * 4;
-	Em.Output(0, EmitInlineHLSL(Em, *this, OutputOffset + 0, {}));
-	Em.Output(1, EmitInlineHLSL(Em, *this, OutputOffset + 1, {}));
-	Em.Output(2, EmitInlineHLSL(Em, *this, OutputOffset + 2, {}));
-	Em.Output(3, EmitInlineHLSL(Em, *this, OutputOffset + 3, {}));
+	Em.Output(0, EmitInlineHLSL(Em, *this, OutputOffset + 0));
+	Em.Output(1, EmitInlineHLSL(Em, *this, OutputOffset + 1));
+	Em.Output(2, EmitInlineHLSL(Em, *this, OutputOffset + 2));
+	Em.Output(3, EmitInlineHLSL(Em, *this, OutputOffset + 3));
 }
 
 void UMaterialExpressionBumpOffset::Build(MIR::FEmitter& Em)
@@ -3476,7 +3591,7 @@ void UMaterialExpressionBumpOffset::Build(MIR::FEmitter& Em)
 	Em.Output(0,
 		Em.Add(
 			Em.Multiply(
-				Em.Swizzle(EmitTransformVectorBase(Em, EmitInlineHLSL(Em, NAME_CameraVector, {}), MCB_World, MCB_Tangent, false, nullptr, nullptr), MIR::FSwizzleMask(MIR::EVectorComponent::X, MIR::EVectorComponent::Y)),
+				Em.Swizzle(EmitTransformVectorBase(Em, EmitInlineHLSL(Em, NAME_CameraVector), MCB_World, MCB_Tangent, false, nullptr, nullptr), MIR::FSwizzleMask(MIR::EVectorComponent::X, MIR::EVectorComponent::Y)),
 				Em.Add(
 					Em.Multiply(
 						HeightRatioArg,
@@ -3588,7 +3703,7 @@ void UMaterialExpressionPanner::Build(MIR::FEmitter& Em)
 	}
 	else
 	{
-		TimeArg = Em.InlineHLSL(MIR::FType::MakeFloatScalar(), TEXTVIEW("View.<PREV>GameTime"), {}, MIR::EValueFlags::SubstituteTags);
+		TimeArg = Em.InlineHLSL(MIR::FType::MakeFloatScalar(), TEXTVIEW("View.<PREVFRAME>GameTime"), {}, MIR::EValueFlags::SubstituteTags);
 	}
 	
 	FValueRef SpeedVectorArg = Em.InputDefaultFloat2(&Speed, FVector2f(SpeedX, SpeedY));
@@ -3637,4 +3752,138 @@ void UMaterialExpressionPerInstanceCustomData3Vector::Build(MIR::FEmitter& Em)
 	Em.Output(0, EmitInlineHLSL(Em, *this, 0, { DataIndexArgument, DefaultArgument }));
 }
 
+void UMaterialExpressionPreviousFrameSwitch::Build(MIR::FEmitter& Em)
+{
+	FValueRef CurrentFrameValue = Em.Input(&CurrentFrame);
+	FValueRef PreviousFrameValue = Em.Input(&PreviousFrame);
+	UE_MIR_CHECKPOINT(Em);
+
+	Em.Output(0, Em.Branch(Em.ExternalInput(MIR::EExternalInput::CompilingPreviousFrame), PreviousFrameValue, CurrentFrameValue));
+}
+
+void UMaterialExpressionHairAttributes::Build(MIR::FEmitter& Em)
+{
+	FValueRef HairUV = Em.InlineHLSL(MIR::FType::MakeFloatVector(2), TEXTVIEW("MaterialExpressionGetHairUV(Parameters)"));
+	Em.Output(0, Em.Subscript(HairUV, 0));
+	Em.Output(1, Em.Subscript(HairUV, 1));
+
+	FValueRef HairDimensions = Em.InlineHLSL(MIR::FType::MakeFloatVector(2), TEXTVIEW("MaterialExpressionGetHairDimensions(Parameters)"));
+	Em.Output(2, Em.Subscript(HairDimensions, 0));		// Length
+	Em.Output(3, Em.Subscript(HairDimensions, 1));		// Radius
+
+	Em.Output(4, Em.InlineHLSL(MIR::FType::MakeFloatScalar(), TEXTVIEW("MaterialExpressionGetHairSeed(Parameters)")));
+	Em.Output(5, Em.InlineHLSL(MIR::FType::MakeFloatVector(3), TEXTVIEW("MaterialExpressionGetHairTangent(Parameters, $0)"), { Em.ConstantBool(bUseTangentSpace) }));
+	Em.Output(6, Em.InlineHLSL(MIR::FType::MakeFloatVector(2), TEXTVIEW("MaterialExpressionGetHairRootUV(Parameters)")));
+	Em.Output(7, Em.InlineHLSL(MIR::FType::MakeFloatVector(3), TEXTVIEW("MaterialExpressionGetHairBaseColor(Parameters)")));
+	Em.Output(8, Em.InlineHLSL(MIR::FType::MakeFloatScalar(), TEXTVIEW("MaterialExpressionGetHairRoughness(Parameters)")));
+	Em.Output(9, Em.InlineHLSL(MIR::FType::MakeFloatScalar(), TEXTVIEW("MaterialExpressionGetHairDepth(Parameters)")));
+	Em.Output(10, Em.InlineHLSL(MIR::FType::MakeFloatScalar(), TEXTVIEW("MaterialExpressionGetHairCoverage(Parameters)")));
+	Em.Output(11, Em.InlineHLSL(MIR::FType::MakeFloatVector(4), TEXTVIEW("MaterialExpressionGetHairAuxilaryData(Parameters)")));
+	Em.Output(12, Em.InlineHLSL(MIR::FType::MakeFloatVector(2), TEXTVIEW("MaterialExpressionGetAtlasUVs(Parameters)")));
+	Em.Output(13, Em.InlineHLSL(MIR::FType::MakeFloatScalar(), TEXTVIEW("MaterialExpressionGetHairGroupIndex(Parameters)")));
+	Em.Output(14, Em.InlineHLSL(MIR::FType::MakeFloatScalar(), TEXTVIEW("MaterialExpressionGetHairAO(Parameters)")));
+	Em.Output(15, Em.Subscript(Em.InlineHLSL(MIR::FType::MakeFloatVector(3), TEXTVIEW("MaterialExpressionGetHairClumpID(Parameters)")), 0));
+}
+
+void UMaterialExpressionHairColor::Build(MIR::FEmitter& Em)
+{
+	FValueRef MelaninInput = Em.InputDefaultFloat(&Melanin, 0.5f);
+	FValueRef RednessInput = Em.InputDefaultFloat(&Redness, 0.0f);
+	FValueRef DyeColorInput = Em.InputDefaultFloat3(&DyeColor, FVector3f(1.f, 1.f, 1.f));
+
+	Em.Output(0, Em.InlineHLSL(MIR::FType::MakeFloatVector(3), TEXTVIEW("MaterialExpressionGetHairColorFromMelanin($0, $1, $2)"), { MelaninInput, RednessInput, DyeColorInput }));
+}
+
+void UMaterialExpressionMapARPassthroughCameraUV::Build(MIR::FEmitter& Em)
+{
+	FValueRef UV = Em.Input(&Coordinates);
+	UE_MIR_CHECKPOINT(Em);
+
+	FValueRef UVPair0 = Em.InlineHLSL(MIR::FType::MakeFloatVector(4), TEXTVIEW("ResolvedView.XRPassthroughCameraUVs[0]"));
+	FValueRef UVPair1 = Em.InlineHLSL(MIR::FType::MakeFloatVector(4), TEXTVIEW("ResolvedView.XRPassthroughCameraUVs[1]"));
+
+	FValueRef ULerp = Em.Lerp(UVPair0, UVPair1, Em.Subscript(UV, 0));
+	Em.Output(0, Em.Lerp(Em.Swizzle(ULerp, MIR::FSwizzleMask(MIR::EVectorComponent::X, MIR::EVectorComponent::Y)), Em.Swizzle(ULerp, MIR::FSwizzleMask(MIR::EVectorComponent::Z, MIR::EVectorComponent::W)), Em.Subscript(UV, 1)));
+}
+
+void UMaterialExpressionSwitch::Build(MIR::FEmitter& Em)
+{
+	FValueRef CompiledDefault = Em.InputDefaultFloat(&Default, ConstDefault);
+
+	// If no other inputs, just return the default
+	if (Inputs.Num() == 0)
+	{
+		Em.Output(0, CompiledDefault);
+		return;
+	}
+
+	// Only the "x" component of the switch value is used.
+	FValueRef CompiledSwitchValue = Em.InputDefaultFloat(&SwitchValue, ConstSwitchValue);
+	if (CompiledSwitchValue->Type.IsVector())
+	{
+		CompiledSwitchValue = Em.Subscript(CompiledSwitchValue, 0);
+	}
+
+	// Compile the inputs.
+	TArray<FValueRef> CompiledInputs;
+	CompiledInputs.SetNumUninitialized(Inputs.Num());
+
+	for (int32 i = 0; i < Inputs.Num(); i++)
+	{
+		CompiledInputs[i] = Em.Input(&Inputs[i].Input);
+	}
+	UE_MIR_CHECKPOINT(Em);		// Make sure inputs are connected.
+
+	// Get common type of inputs.  Done as a separate loop, to avoid spurious errors for unconnected inputs, which otherwise also produce "No common type" errors.
+	MIR::FType CommonType = CompiledDefault->Type;
+	for (int32 i = 0; i < Inputs.Num(); i++)
+	{
+		CommonType = Em.GetCommonType(CommonType, CompiledInputs[i]->Type);
+	}
+	UE_MIR_CHECKPOINT(Em);		// Make sure inputs have a valid common type.
+
+	// If the switch value is a constant, we can directly pass the corresponding input as the result.
+	const MIR::FConstant* CompiledSwitchValueConstant = CompiledSwitchValue->As<MIR::FConstant>();
+	if (CompiledSwitchValueConstant)
+	{
+		int32 InputIndex = 0;
+		switch (CompiledSwitchValue->Type.AsPrimitive()->ScalarKind)
+		{
+		case MIR::EScalarKind::Bool:	InputIndex = CompiledSwitchValueConstant->Boolean ? 1 : 0;  break;
+		case MIR::EScalarKind::Int:		InputIndex = CompiledSwitchValueConstant->Integer;  break;
+		case MIR::EScalarKind::Float:	InputIndex = FMath::FloorToInt(CompiledSwitchValueConstant->Float);  break;
+		case MIR::EScalarKind::Double:	InputIndex = (int32)FMath::FloorToInt(CompiledSwitchValueConstant->Double);  break;
+		default:  UE_MIR_UNREACHABLE();
+		}
+
+		if (Inputs.IsValidIndex(InputIndex))
+		{
+			Em.Output(0, Em.Cast(CompiledInputs[InputIndex], CommonType));
+		}
+		else
+		{
+			Em.Output(0, Em.Cast(CompiledDefault, CommonType));
+		}
+		return;
+	}
+
+	// Floor the switch value if it's a float, to prepare for comparisons.
+	if (CompiledSwitchValue->Type.IsAnyFloat())
+	{
+		CompiledSwitchValue = Em.Floor(CompiledSwitchValue);
+	}
+
+	// Generate a switch statement as a chain of if..else branches.  We scan backwards, so the comparisons end up in order,
+	// factoring in that each Branch is a parent of the previous Branch, and so the last Branch added is the first that gets
+	// executed.  The first previous Branch (final else case) starts out as the default.
+	FValueRef PreviousBranch = Em.Cast(CompiledDefault, CommonType);
+
+	for (int32 i = Inputs.Num() - 1; i >= 0; i--)
+	{
+		PreviousBranch = Em.Branch(Em.Equals(CompiledSwitchValue, Em.ConstantInt(i)), Em.Cast(CompiledInputs[i], CommonType), PreviousBranch);
+	}
+
+	Em.Output(0, PreviousBranch);
+}
+
 #endif // WITH_EDITOR
diff --git a/Engine/Source/Runtime/Engine/Private/Materials/MaterialIR.cpp b/Engine/Source/Runtime/Engine/Private/Materials/MaterialIR.cpp
--- a/Engine/Source/Runtime/Engine/Private/Materials/MaterialIR.cpp
+++ b/Engine/Source/Runtime/Engine/Private/Materials/MaterialIR.cpp
@@ -512,6 +512,10 @@ const TCHAR* LexToString(EExternalInput Input)
 		case EExternalInput::WorldPosition_AbsoluteNoOffsets: return TEXT("WorldPosition_AbsoluteNoOffsets");
 		case EExternalInput::WorldPosition_CameraRelative: return TEXT("WorldPosition_CameraRelative");
 		case EExternalInput::WorldPosition_CameraRelativeNoOffsets: return TEXT("WorldPosition_CameraRelativeNoOffsets");
+		case EExternalInput::LocalPosition_Instance:			return TEXT("LocalPosition_Instance");
+		case EExternalInput::LocalPosition_InstanceNoOffsets:	return TEXT("LocalPosition_InstanceNoOffsets");
+		case EExternalInput::LocalPosition_Primitive:			return TEXT("LocalPosition_Primitive");
+		case EExternalInput::LocalPosition_PrimitiveNoOffsets:	return TEXT("LocalPosition_PrimitiveNoOffsets");
 
 		case EExternalInput::TexCoord0_Ddx: return TEXT("TexCoord0_Ddx");
 		case EExternalInput::TexCoord1_Ddx: return TEXT("TexCoord1_Ddx");
@@ -525,6 +529,10 @@ const TCHAR* LexToString(EExternalInput Input)
 		case EExternalInput::WorldPosition_AbsoluteNoOffsets_Ddx: return TEXT("WorldPosition_AbsoluteNoOffsets_Ddx");
 		case EExternalInput::WorldPosition_CameraRelative_Ddx: return TEXT("WorldPosition_CameraRelative_Ddx");
 		case EExternalInput::WorldPosition_CameraRelativeNoOffsets_Ddx: return TEXT("WorldPosition_CameraRelativeNoOffsets_Ddx");
+		case EExternalInput::LocalPosition_Instance_Ddx:			return TEXT("LocalPosition_Instance_Ddx");
+		case EExternalInput::LocalPosition_InstanceNoOffsets_Ddx:	return TEXT("LocalPosition_InstanceNoOffsets_Ddx");
+		case EExternalInput::LocalPosition_Primitive_Ddx:			return TEXT("LocalPosition_Primitive_Ddx");
+		case EExternalInput::LocalPosition_PrimitiveNoOffsets_Ddx:	return TEXT("LocalPosition_PrimitiveNoOffsets_Ddx");
 
 		case EExternalInput::TexCoord0_Ddy: return TEXT("TexCoord0_Ddy");
 		case EExternalInput::TexCoord1_Ddy: return TEXT("TexCoord1_Ddy");
@@ -538,6 +546,10 @@ const TCHAR* LexToString(EExternalInput Input)
 		case EExternalInput::WorldPosition_AbsoluteNoOffsets_Ddy: return TEXT("WorldPosition_AbsoluteNoOffsets_Ddy");
 		case EExternalInput::WorldPosition_CameraRelative_Ddy: return TEXT("WorldPosition_CameraRelative_Ddy");
 		case EExternalInput::WorldPosition_CameraRelativeNoOffsets_Ddy: return TEXT("WorldPosition_CameraRelativeNoOffsets_Ddy");
+		case EExternalInput::LocalPosition_Instance_Ddy:			return TEXT("LocalPosition_Instance_Ddy");
+		case EExternalInput::LocalPosition_InstanceNoOffsets_Ddy:	return TEXT("LocalPosition_InstanceNoOffsets_Ddy");
+		case EExternalInput::LocalPosition_Primitive_Ddy:			return TEXT("LocalPosition_Primitive_Ddy");
+		case EExternalInput::LocalPosition_PrimitiveNoOffsets_Ddy:	return TEXT("LocalPosition_PrimitiveNoOffsets_Ddy");
 
 		case EExternalInput::ActorPosition_Absolute: return TEXT("ActorPosition_Absolute");
 		case EExternalInput::ActorPosition_CameraRelative: return TEXT("ActorPosition_CameraRelative");
@@ -547,6 +559,7 @@ const TCHAR* LexToString(EExternalInput Input)
 		case EExternalInput::ViewMaterialTextureDerivativeMultiply: return TEXT("ViewMaterialTextureDerivativeMultiply");
 		case EExternalInput::GlobalDistanceField: return TEXT("GlobalDistanceField");
 		case EExternalInput::DynamicParticleParameterIndex: return TEXT("DynamicParticleParameterIndex");
+		case EExternalInput::CompilingPreviousFrame: return TEXT("CompilingPreviousFrame");
 		default: UE_MIR_UNREACHABLE();
 	}
 }
@@ -583,11 +596,26 @@ FType GetExternalInputType(EExternalInput Id)
 		case EExternalInput::WorldPosition_CameraRelativeNoOffsets_Ddx:
 		case EExternalInput::WorldPosition_CameraRelative_Ddy:
 		case EExternalInput::WorldPosition_CameraRelativeNoOffsets_Ddy:
+		case EExternalInput::LocalPosition_Instance:
+		case EExternalInput::LocalPosition_InstanceNoOffsets:
+		case EExternalInput::LocalPosition_Primitive:
+		case EExternalInput::LocalPosition_PrimitiveNoOffsets:
+		case EExternalInput::LocalPosition_Instance_Ddx:
+		case EExternalInput::LocalPosition_InstanceNoOffsets_Ddx:
+		case EExternalInput::LocalPosition_Primitive_Ddx:
+		case EExternalInput::LocalPosition_PrimitiveNoOffsets_Ddx:
+		case EExternalInput::LocalPosition_Instance_Ddy:
+		case EExternalInput::LocalPosition_InstanceNoOffsets_Ddy:
+		case EExternalInput::LocalPosition_Primitive_Ddy:
+		case EExternalInput::LocalPosition_PrimitiveNoOffsets_Ddy:
 			return FType::MakeFloatVector(3);
 
 		case EExternalInput::DynamicParticleParameterIndex:
 			return FType::MakeIntScalar();
 
+		case EExternalInput::CompilingPreviousFrame:
+			return FType::MakeBoolScalar();
+
 		default: UE_MIR_UNREACHABLE();
 	}
 }
diff --git a/Engine/Source/Runtime/Engine/Private/Materials/MaterialIREmitter.cpp b/Engine/Source/Runtime/Engine/Private/Materials/MaterialIREmitter.cpp
--- a/Engine/Source/Runtime/Engine/Private/Materials/MaterialIREmitter.cpp
+++ b/Engine/Source/Runtime/Engine/Private/Materials/MaterialIREmitter.cpp
@@ -5,6 +5,7 @@
 #include "Materials/MaterialIRModule.h"
 #include "Materials/MaterialIRModuleBuilder.h"
 #include "Materials/MaterialIRInternal.h"
+#include "Materials/MaterialExpressionTextureBase.h"
 #include "Materials/MaterialExternalCodeRegistry.h"
 #include "Materials/MaterialAggregate.h"
 #include "Shader/ShaderTypes.h"
@@ -971,6 +972,13 @@ FValueRef FEmitter::ShadingModel(EMaterialShadingModel Id)
 
 FValueRef FEmitter::TextureObject(UTexture* Texture, EMaterialSamplerType SamplerType)
 {
+	FString SamplerTypeError;
+	if (!UMaterialExpressionTextureBase::VerifySamplerType(GetShaderPlatform(), GetTargetPlatform(), Texture, SamplerType, SamplerTypeError))
+	{
+		Errorf(TEXT("%s"), *SamplerTypeError);
+		return Poison();
+	}
+
 	FTextureObject Proto = MakePrototype<FTextureObject>(FType::MakeTexture());
 	Proto.Texture = Texture;
 	Proto.SamplerType = SamplerType;
@@ -3297,37 +3305,44 @@ FValueRef FEmitter::TextureSample(FValueRef Texture, FValueRef TexCoord, bool bA
 
 	SamplerType = MapSamplerTypeForTexture(Texture, SamplerType);
 
-	FTextureRead PrototypeHw = MakePrototype<FTextureRead>(FType::MakeFloatVector(4));
-	PrototypeHw.TextureObject = Texture;
-	PrototypeHw.TexCoord = TexCoord;
-	PrototypeHw.Mode = ETextureReadMode::MipAuto;
-	PrototypeHw.SamplerSourceMode = SamplerSourceMode;
-	PrototypeHw.SamplerType = SamplerType;
+	FTextureRead PrototypePixel = MakePrototype<FTextureRead>(FType::MakeFloatVector(4));
+	PrototypePixel.TextureObject = Texture;
+	PrototypePixel.TexCoord = TexCoord;
+	PrototypePixel.Mode = ETextureReadMode::MipAuto;
+	PrototypePixel.SamplerSourceMode = SamplerSourceMode;
+	PrototypePixel.SamplerType = SamplerType;
+
+	FTextureRead PrototypeCompute = PrototypePixel;
+	PrototypeCompute.Mode = ETextureReadMode::Derivatives;
+	PrototypeCompute.TexCoordDdx = AnalyticalPartialDerivative(TexCoord, EDerivativeAxis::X);
+	PrototypeCompute.TexCoordDdy = AnalyticalPartialDerivative(TexCoord, EDerivativeAxis::Y);
 
-	FTextureRead PrototypeAn = PrototypeHw;
-	PrototypeAn.Mode = ETextureReadMode::Derivatives;
-	PrototypeAn.TexCoordDdx = AnalyticalPartialDerivative(TexCoord, EDerivativeAxis::X);
-	PrototypeAn.TexCoordDdy = AnalyticalPartialDerivative(TexCoord, EDerivativeAxis::Y);	
+	FTextureRead PrototypeVertex = PrototypePixel;
+	PrototypeVertex.Mode = ETextureReadMode::MipLevel;
+	PrototypeVertex.MipValue = ConstantZero(EScalarKind::Float);
 
 	if (bAutomaticViewMipBias)
 	{
 		FValueRef ViewMaterialTextureMipBias = ExternalInput(EExternalInput::ViewMaterialTextureMipBias);
-		PrototypeHw.Mode = ETextureReadMode::MipBias;
-		PrototypeHw.MipValue = ViewMaterialTextureMipBias;
+		PrototypePixel.Mode = ETextureReadMode::MipBias;
+		PrototypePixel.MipValue = ViewMaterialTextureMipBias;
 
 		FValueRef Exp2ViewMaterialTextureMipBias = Operator(UO_Exponential2, ViewMaterialTextureMipBias);
-		PrototypeAn.TexCoordDdx = Operator(BO_Multiply, PrototypeAn.TexCoordDdx, Exp2ViewMaterialTextureMipBias);
-		PrototypeAn.TexCoordDdy = Operator(BO_Multiply, PrototypeAn.TexCoordDdy, Exp2ViewMaterialTextureMipBias);
+		PrototypeCompute.TexCoordDdx = Operator(BO_Multiply, PrototypeCompute.TexCoordDdx, Exp2ViewMaterialTextureMipBias);
+		PrototypeCompute.TexCoordDdy = Operator(BO_Multiply, PrototypeCompute.TexCoordDdy, Exp2ViewMaterialTextureMipBias);
 	}
 
 	if (IsVirtualSamplerType(SamplerType))
 	{
-		PrototypeHw.VTPageTable = VTPageTableLoadFromSamplerSource(*this, Texture, SamplerSourceMode, PrototypeHw.TexCoord, {}, {});
-		PrototypeAn.VTPageTable = VTPageTableLoadFromSamplerSource(*this, Texture, SamplerSourceMode, PrototypeAn.TexCoord, PrototypeAn.TexCoordDdx, PrototypeAn.TexCoordDdy);
+		PrototypePixel.VTPageTable = VTPageTableLoadFromSamplerSource(*this, Texture, SamplerSourceMode, PrototypePixel.TexCoord, {}, {});
+		PrototypeCompute.VTPageTable = VTPageTableLoadFromSamplerSource(*this, Texture, SamplerSourceMode, PrototypeCompute.TexCoord, PrototypeCompute.TexCoordDdx, PrototypeCompute.TexCoordDdy);
+		PrototypeVertex.VTPageTable = VTPageTableLoadFromSamplerSource(*this, Texture, SamplerSourceMode, PrototypeVertex.TexCoord, {}, {});
 	}
 
-	FStageSwitch StageSwitch = MakePrototype<FStageSwitch>(PrototypeHw.Type);
-	StageSwitch.SetArgs(EmitPrototype(*this, PrototypeHw), EmitPrototype(*this, PrototypeAn));
+	FStageSwitch StageSwitch = MakePrototype<FStageSwitch>(PrototypePixel.Type);
+	StageSwitch.Args[Stage_Vertex] = EmitPrototype(*this, PrototypeVertex);
+	StageSwitch.Args[Stage_Pixel] = EmitPrototype(*this, PrototypePixel);
+	StageSwitch.Args[Stage_Compute] = EmitPrototype(*this, PrototypeCompute);
 
 	return EmitPrototype(*this, StageSwitch);
 }
@@ -3376,29 +3391,40 @@ FValueRef FEmitter::TextureSampleBias(FValueRef Texture, FValueRef TexCoord, FVa
 
 	SamplerType = MapSamplerTypeForTexture(Texture, SamplerType);
 
-	FTextureRead PrototypeHw = MakePrototype<FTextureRead>(FType::MakeFloatVector(4));
-	PrototypeHw.TextureObject = Texture;
-	PrototypeHw.TexCoord = TexCoord;
-	PrototypeHw.MipValue = MipBias;
-	PrototypeHw.Mode = ETextureReadMode::MipBias;
-	PrototypeHw.SamplerSourceMode = SamplerSourceMode;
-	PrototypeHw.SamplerType = SamplerType;
+	FTextureRead PrototypePixel = MakePrototype<FTextureRead>(FType::MakeFloatVector(4));
+	PrototypePixel.TextureObject = Texture;
+	PrototypePixel.TexCoord = TexCoord;
+	PrototypePixel.MipValue = MipBias;
+	PrototypePixel.Mode = ETextureReadMode::MipBias;
+	PrototypePixel.SamplerSourceMode = SamplerSourceMode;
+	PrototypePixel.SamplerType = SamplerType;
 
-	FTextureRead PrototypeAn = PrototypeHw;
-	PrototypeAn.Mode = ETextureReadMode::Derivatives;
+	FTextureRead PrototypeCompute = PrototypePixel;
+	PrototypeCompute.Mode = ETextureReadMode::Derivatives;
 
 	FValueRef Exp2MipBias = Operator(UO_Exponential2, MipBias);
-	PrototypeAn.TexCoordDdx = Operator(BO_Multiply, AnalyticalPartialDerivative(TexCoord, EDerivativeAxis::X), Exp2MipBias);
-	PrototypeAn.TexCoordDdy = Operator(BO_Multiply, AnalyticalPartialDerivative(TexCoord, EDerivativeAxis::Y), Exp2MipBias);	
+	PrototypeCompute.TexCoordDdx = Operator(BO_Multiply, AnalyticalPartialDerivative(TexCoord, EDerivativeAxis::X), Exp2MipBias);
+	PrototypeCompute.TexCoordDdy = Operator(BO_Multiply, AnalyticalPartialDerivative(TexCoord, EDerivativeAxis::Y), Exp2MipBias);
+
+	// We throw a deliberate error downstream for this expression if referenced from a vertex shader.  We use the hardware
+	// derivative value type (disallowed in vertex shader, already tracked by the value analyzer) to communicate the error.
+	// We could alternately make a one-off type just to communicate this error, but this works, and is somewhat logical,
+	// given the root cause of this sampler type not working is due to missing derivative support in vertex shaders.
+	FHardwarePartialDerivative PrototypeVertex = MakePrototype<FHardwarePartialDerivative>(PrototypePixel.Type);
+	PrototypeVertex.Arg = ConstantFloat(0.0f);
+	PrototypeVertex.Axis = EDerivativeAxis::X;
+	PrototypeVertex.Source = EDerivativeSource::TextureSampleBias;
 
 	if (IsVirtualSamplerType(SamplerType))
 	{
-		PrototypeHw.VTPageTable = VTPageTableLoadFromSamplerSource(*this, Texture, SamplerSourceMode, PrototypeHw.TexCoord, {}, {});
-		PrototypeAn.VTPageTable = VTPageTableLoadFromSamplerSource(*this, Texture, SamplerSourceMode, PrototypeAn.TexCoord, PrototypeAn.TexCoordDdx, PrototypeAn.TexCoordDdy);
+		PrototypePixel.VTPageTable = VTPageTableLoadFromSamplerSource(*this, Texture, SamplerSourceMode, PrototypePixel.TexCoord, {}, {});
+		PrototypeCompute.VTPageTable = VTPageTableLoadFromSamplerSource(*this, Texture, SamplerSourceMode, PrototypeCompute.TexCoord, PrototypeCompute.TexCoordDdx, PrototypeCompute.TexCoordDdy);
 	}
 
-	FStageSwitch StageSwitch = MakePrototype<FStageSwitch>(PrototypeHw.Type);
-	StageSwitch.SetArgs(EmitPrototype(*this, PrototypeHw), EmitPrototype(*this, PrototypeAn));
+	FStageSwitch StageSwitch = MakePrototype<FStageSwitch>(PrototypePixel.Type);
+	StageSwitch.Args[Stage_Vertex] = EmitPrototype(*this, PrototypeVertex);
+	StageSwitch.Args[Stage_Pixel] = EmitPrototype(*this, PrototypePixel);
+	StageSwitch.Args[Stage_Compute] = EmitPrototype(*this, PrototypeCompute);
 
 	return EmitPrototype(*this, StageSwitch);
 }
@@ -3469,17 +3495,19 @@ FValueRef FEmitter::PartialDerivative(FValueRef Value, EDerivativeAxis Axis)
 	FHardwarePartialDerivative HwDerivativeProto = MakePrototype<FHardwarePartialDerivative>(Value->Type);
 	HwDerivativeProto.Arg = Value;
 	HwDerivativeProto.Axis = Axis;
+	HwDerivativeProto.Source = EDerivativeSource::Derivative;
 	FValueRef HwDerivative = EmitPrototype(*this, HwDerivativeProto);
 
 	// Compute the analytical derivative for stages that don't support hardware derivatives.
 	FValueRef AnalyticalDerivative = AnalyticalPartialDerivative(Value, Axis);
 
-	// Emit the stage switch instruction so that hardware derivatives are used on stages
-	// that support it and analytical deriatives in the other stages.
+	// Emit the stage switch instruction so that analytical derivatives are used on stages that support it and hardware
+	// derivatives in the other stages.  Note that hardware derivatives throw errors on the vertex stage downstream, but
+	// we don't know if the expression is reached in that stage until the ValueAnalyzer runs.
 	FValueRef StageValues[NumStages] = {};
 	for (int i = 0; i < NumStages; ++i)
 	{
-		StageValues[i] = (i == Stage_Pixel) ? HwDerivative : AnalyticalDerivative;
+		StageValues[i] = (i == Stage_Compute) ? AnalyticalDerivative : HwDerivative;
 	}
 	return StageSwitch(Value->Type, StageValues);
 }
@@ -3799,6 +3827,26 @@ FValueRef FEmitter::AnalyticalPartialDerivative(FValueRef Value, EDerivativeAxis
 				AnalyticalPartialDerivative(AsBranch->FalseArg, Axis));
 		}
 
+		case VK_StageSwitch:
+		{
+			// For StageSwitch, we want to pass through and generate derivatives for its input.  We only need to do this for
+			// the compute stage, because the analytic derivative code path is unreachable for the pixel and vertex stages.
+			// To reach the analytic derivative code path in the first place, there will have been a higher level stage switch,
+			// which will already have chosen a different hardware derivative path for the pixel shader, or thrown an error for
+			// the vertex shader, where explicit derivatives are disallowed completely.
+			// 
+			// For the other stages, we can pass a poison value, to detect if this assumption is violated in the future.
+			// Because this is a non-error unreachable poison, we don't call FEmitter::Poison, as that will trigger an unwanted
+			// breakpoint when using the debug feature that breaks on poison values.
+			FStageSwitch* AsStageSwitch = Value->As<FStageSwitch>();
+			TStaticArray<FValueRef, NumStages> StageDerivatives;
+			for (int32 StageIndex = 0; StageIndex < NumStages; StageIndex++)
+			{
+				StageDerivatives[StageIndex] = StageIndex == Stage_Compute ? AnalyticalPartialDerivative(AsStageSwitch->Args[StageIndex], Axis) : FPoison::Get();
+			}
+			return StageSwitch(Value->Type, StageDerivatives);
+		}
+
 		case VK_Subscript:
 		{
 			FSubscript* AsSubscript = Value->As<FSubscript>();
@@ -3823,13 +3871,29 @@ FValueRef FEmitter::AnalyticalPartialDerivative(FValueRef Value, EDerivativeAxis
 		}
 
 		// These values are uniform (constant), thus their value is always zero.
-		// Custom HLSL calls are not derivative aware, so we also return constant zero for VK_Call.
 		case VK_Constant:
 		case VK_UniformParameter:
 		case VK_PreshaderParameter:
+			return Cast(ConstantZero(Value->Type.GetPrimitive().ScalarKind), Value->Type);
+
+		// These values don't work with analytic derivatives, and force hardware derivatives (or zero if the shader model doesn't support compute shader derivatives).
 		case VK_TextureRead:
 		case VK_Call:
-			return Cast(ConstantZero(Value->Type.GetPrimitive().ScalarKind), Value->Type);
+		{
+			if (GetFeatureLevel() >= ERHIFeatureLevel::Type::SM6)
+			{
+				// Make the hardware derivative instruction.
+				FHardwarePartialDerivative HwDerivativeProto = MakePrototype<FHardwarePartialDerivative>(Value->Type);
+				HwDerivativeProto.Arg = Value;
+				HwDerivativeProto.Axis = Axis;
+				HwDerivativeProto.Source = EDerivativeSource::AnalyticDerivative;
+				return EmitPrototype(*this, HwDerivativeProto);
+			}
+			else
+			{
+				return Cast(ConstantZero(Value->Type.GetPrimitive().ScalarKind), Value->Type);
+			}
+		}
 
 		default:
 			UE_MIR_UNREACHABLE();
diff --git a/Engine/Source/Runtime/Engine/Private/Materials/MaterialIRToHLSLTranslator.cpp b/Engine/Source/Runtime/Engine/Private/Materials/MaterialIRToHLSLTranslator.cpp
--- a/Engine/Source/Runtime/Engine/Private/Materials/MaterialIRToHLSLTranslator.cpp
+++ b/Engine/Source/Runtime/Engine/Private/Materials/MaterialIRToHLSLTranslator.cpp
@@ -254,12 +254,19 @@ struct FPrivate : FMaterialIRToHLSLTranslation
 	// Current material stage being translated (e.g., vertex, pixel)
 	MIR::EStage CurrentStage;
 
+	// Set when generating world position offset HLSL using previous frame data.
+	// Affects token replacement for <PREV> in external code, and code generation for certain external inputs.
+	bool bCompilingPreviousFrame = false;
+
 	// Generated HLSL for pixel attribute evaluation
 	FString PixelAttributesHLSL;
 	
 	// Generated HLSL for world position offset calculation
 	FString WorldPositionOffsetHLSL;
 
+	// Generated HLSL for previous world position offset calculation
+	FString PreviousWorldPositionOffsetHLSL;
+
 	// HLSL code for evaluating normal material attributes per stage
 	FString EvaluateNormalMaterialAttributeHLSL[MIR::NumStages];
 	
@@ -290,6 +297,13 @@ struct FPrivate : FMaterialIRToHLSLTranslation
 		LowerBlock(Module->GetRootBlock(MIR::Stage_Vertex));
 
 		WorldPositionOffsetHLSL = MoveTemp(Printer.Buffer);
+
+		// Same as above, but with previous frame flag set.
+		bCompilingPreviousFrame = true;
+		BeginStage(MIR::Stage_Vertex);
+		LowerBlock(Module->GetRootBlock(MIR::Stage_Vertex));
+		PreviousWorldPositionOffsetHLSL = MoveTemp(Printer.Buffer);
+		bCompilingPreviousFrame = false;
 	}
 
 	void GenerateOtherStageHLSL(MIR::EStage Stage)
@@ -473,9 +487,6 @@ struct FPrivate : FMaterialIRToHLSLTranslation
 
 	void LowerExternalInput(const MIR::FExternalInput* ExternalInput)
 	{
-		// @todo-laura.hermanns - Must be replaced as soon as this state is available in MIR
-		constexpr bool bCompilingPreviousFrame = false;
-
 		int32 ExternalInputIndex = (int32)ExternalInput->Id;
 		if (MIR::IsExternalInputTexCoord(ExternalInput->Id))
 		{
@@ -502,8 +513,8 @@ struct FPrivate : FMaterialIRToHLSLTranslation
 			//
 			Printer << TEXTVIEW("Get");
 
-			// If compiling for the previous frame in the vertex shader
-			if (CurrentStage == MIR::EStage::Stage_Vertex && bCompilingPreviousFrame)
+			// If compiling for the previous frame
+			if (bCompilingPreviousFrame)
 			{
 				Printer << TEXTVIEW("Prev");
 			}
@@ -532,20 +543,77 @@ struct FPrivate : FMaterialIRToHLSLTranslation
 			// Expression emitter assumes these to be LWC for analytic derivative evaluation, so we promote them on load
 			Printer << TEXTVIEW("WSPromote(Parameters.WorldPosition_DDY)");
 		}
+		else if (ExternalInput->Id >= MIR::EExternalInput::LocalPosition_Instance && ExternalInput->Id <= MIR::EExternalInput::LocalPosition_PrimitiveNoOffsets)
+		{
+			// Various function permutations exist to fetch local position.  "Prev" permutations are only available in the vertex shader, while
+			// "MoMaterialOffsets" permutations are only available in the pixel shader, so we need to cobble together the permutation string
+			// factoring in those limitations.  Logic adapted from FHLSLMaterialTranslator::LocalPosition (and similar to WorldPosition above).
+			// 
+			// Format:  Get[Prev]Position[Instance|Primitive]Space[_NoMaterialOffsets](Parameters)
+			//
+			Printer << TEXTVIEW("Get");
+
+			// If compiling for the previous frame
+			if (bCompilingPreviousFrame)
+			{
+				Printer << TEXTVIEW("Prev");
+			}
+
+			Printer << TEXTVIEW("Position");
+			Printer << (ExternalInput->Id <= MIR::EExternalInput::LocalPosition_InstanceNoOffsets ? TEXTVIEW("Instance") : TEXTVIEW("Primitive"));
+			Printer << TEXTVIEW("Space");
+
+			if (CurrentStage == MIR::EStage::Stage_Pixel && (ExternalInput->Id == MIR::EExternalInput::LocalPosition_InstanceNoOffsets || ExternalInput->Id == MIR::EExternalInput::LocalPosition_PrimitiveNoOffsets))
+			{
+				Printer << TEXTVIEW("_NoMaterialOffsets");
+			}
+
+			Printer << TEXTVIEW("(Parameters)");
+		}
+		else if (ExternalInput->Id >= MIR::EExternalInput::LocalPosition_Instance_Ddx && ExternalInput->Id <= MIR::EExternalInput::LocalPosition_PrimitiveNoOffsets_Ddx)
+		{
+			// Local position derivatives reuse the world position derivatives, transforming them to local space
+			if (ExternalInput->Id <= MIR::EExternalInput::LocalPosition_InstanceNoOffsets_Ddx)
+			{
+				Printer << TEXTVIEW("mul(Parameters.WorldPosition_DDX, DFToFloat3x3(GetWorldToInstanceDF(Parameters)))");
+			}
+			else
+			{
+				Printer << (bCompilingPreviousFrame ?
+					TEXTVIEW("mul(Parameters.WorldPosition_DDX, DFToFloat3x3(GetPrevWorldToLocalDF(Parameters)))") :
+					TEXTVIEW("mul(Parameters.WorldPosition_DDX, DFToFloat3x3(GetWorldToLocalDF(Parameters)))"));
+			}
+		}
+		else if (ExternalInput->Id >= MIR::EExternalInput::LocalPosition_Instance_Ddy && ExternalInput->Id <= MIR::EExternalInput::LocalPosition_PrimitiveNoOffsets_Ddy)
+		{
+			// Local position derivatives reuse the world position derivatives, transforming them to local space
+			if (ExternalInput->Id <= MIR::EExternalInput::LocalPosition_InstanceNoOffsets_Ddy)
+			{
+				Printer << TEXTVIEW("mul(Parameters.WorldPosition_DDY, DFToFloat3x3(GetWorldToInstanceDF(Parameters)))");
+			}
+			else
+			{
+				Printer << (bCompilingPreviousFrame ?
+					TEXTVIEW("mul(Parameters.WorldPosition_DDY, DFToFloat3x3(GetPrevWorldToLocalDF(Parameters)))") :
+					TEXTVIEW("mul(Parameters.WorldPosition_DDY, DFToFloat3x3(GetWorldToLocalDF(Parameters)))"));
+			}
+		}
 		else if (ExternalInput->Id == MIR::EExternalInput::ActorPosition_Absolute)
 		{
-			const bool bPrevious = (CurrentStage == MIR::EStage::Stage_Vertex && bCompilingPreviousFrame);
-			Printer << (bPrevious ? TEXTVIEW("GetPreviousActorWorldPosition(Parameters)") : TEXTVIEW("GetActorWorldPosition(Parameters)"));
+			Printer << (bCompilingPreviousFrame ? TEXTVIEW("GetPreviousActorWorldPosition(Parameters)") : TEXTVIEW("GetActorWorldPosition(Parameters)"));
 		}
 		else if (ExternalInput->Id == MIR::EExternalInput::ActorPosition_CameraRelative)
 		{
-			const bool bPrevious = (CurrentStage == MIR::EStage::Stage_Vertex && bCompilingPreviousFrame);
-			Printer << (bPrevious ? TEXTVIEW("GetPreviousActorTranslatedWorldPosition(Parameters)") : TEXTVIEW("GetActorTranslatedWorldPosition(Parameters)"));
+			Printer << (bCompilingPreviousFrame ? TEXTVIEW("GetPreviousActorTranslatedWorldPosition(Parameters)") : TEXTVIEW("GetActorTranslatedWorldPosition(Parameters)"));
 		}
 		else if (ExternalInput->Id == MIR::EExternalInput::DynamicParticleParameterIndex)
 		{
 			Printer << ExternalInput->UserData;
 		}
+		else if (ExternalInput->Id == MIR::EExternalInput::CompilingPreviousFrame)
+		{
+			Printer << (bCompilingPreviousFrame ? TEXTVIEW("true") : TEXTVIEW("false"));
+		}
 		else
 		{
 			FStringView Code;
@@ -668,11 +736,10 @@ struct FPrivate : FMaterialIRToHLSLTranslation
 		FString SubstitutedCode;
 		if (ExternalCode->HasFlags(MIR::EValueFlags::SubstituteTags))
 		{
-			// @todo-laura.hermanns - Must be replaced as soon as this state is available in MIR
-			constexpr bool bCompilingPreviousFrame = false;
 			SubstitutedCode = FString{ FinalCode };
 			SubstitutedCode.ReplaceInline(TEXT("<PREV>"), bCompilingPreviousFrame ? TEXT("Prev") : TEXT(""));
 			SubstitutedCode.ReplaceInline(TEXT("<PREVIOUS>"), bCompilingPreviousFrame ? TEXT("Previous") : TEXT(""));
+			SubstitutedCode.ReplaceInline(TEXT("<PREVFRAME>"), bCompilingPreviousFrame ? TEXT("PrevFrame") : TEXT(""));
 
 			FinalCode = SubstitutedCode;
 		}
@@ -1878,7 +1945,7 @@ struct FPrivate : FMaterialIRToHLSLTranslation
 		}
 		else
 		{
-			Printer << (StageSwitch->Axis == MIR::EDerivativeAxis::X ? TEXTVIEW("ddx(") : TEXTVIEW("ddy(")) << LowerValue(StageSwitch->Arg) << TEXTVIEW(")");
+			Printer << (StageSwitch->Axis == MIR::EDerivativeAxis::X ? TEXTVIEW("DDX(") : TEXTVIEW("DDY(")) << LowerValue(StageSwitch->Arg) << TEXTVIEW(")");
 		}
 	}
 
@@ -2072,7 +2139,7 @@ struct FPrivate : FMaterialIRToHLSLTranslation
 		Params.Add(TEXT("get_material_translucent_multiple_scattering_extinction"), FString::Printf(TEXT(TAB "return MaterialFloat3(%.5f, %.5f, %.5f)"), Extinction.R, Extinction.G, Extinction.B));
 
 		Params.Add(TEXT("get_material_world_position_offset_raw"), WorldPositionOffsetHLSL);
-		Params.Add(TEXT("get_material_previous_world_position_offset_raw"), WorldPositionOffsetHLSL);
+		Params.Add(TEXT("get_material_previous_world_position_offset_raw"), PreviousWorldPositionOffsetHLSL);
 	
 		FString EvaluateMaterialDeclaration;
 		EvaluateMaterialDeclaration.Append(TEXT("void EvaluateVertexMaterialAttributes(in out FMaterialVertexParameters Parameters)\n{\n"));
diff --git a/Engine/Source/Runtime/Engine/Private/Materials/MaterialIRValueAnalyzer.cpp b/Engine/Source/Runtime/Engine/Private/Materials/MaterialIRValueAnalyzer.cpp
--- a/Engine/Source/Runtime/Engine/Private/Materials/MaterialIRValueAnalyzer.cpp
+++ b/Engine/Source/Runtime/Engine/Private/Materials/MaterialIRValueAnalyzer.cpp
@@ -70,6 +70,21 @@ static void AnalyzeExternalInput(FMaterialIRValueAnalyzer& Analyzer, MIR::FExter
 	}
 }
 
+static void AnalyzeInStageHardwarePartialDerivative(FMaterialIRValueAnalyzer& Analyzer, MIR::FHardwarePartialDerivative* PartialDerivative, MIR::EStage Stage)
+{
+	if (Stage == MIR::Stage_Vertex)
+	{
+		if (PartialDerivative->Source == MIR::EDerivativeSource::TextureSampleBias)
+		{
+			Analyzer.Module->AddError(nullptr, TEXT("(Node TextureSample) MipBias is only supported in the pixel shader"));
+		}
+		else
+		{
+			Analyzer.Module->AddError(nullptr, FString::Printf(TEXT("Invalid DD%c node used in vertex shader input!"), 'X' + (int32)PartialDerivative->Axis));
+		}
+	}
+}
+
 static EMaterialShaderFrequency MapToMaterialShaderFrequencyOrAny(MIR::EStage Stage)
 {
 	switch (Stage)
@@ -812,6 +827,7 @@ void FMaterialIRValueAnalyzer::AnalyzeInStage(MIR::FValue* Value, MIR::EStage St
 		EXPAND_CASE(InlineHLSL);
 		EXPAND_CASE(VTPageTableRead);
 		EXPAND_CASE(ScreenTexture);
+		EXPAND_CASE(HardwarePartialDerivative);
 		default: break;
 	}
 
diff --git a/Engine/Source/Runtime/Engine/Public/Materials/MaterialExpressionAntialiasedTextureMask.h b/Engine/Source/Runtime/Engine/Public/Materials/MaterialExpressionAntialiasedTextureMask.h
--- a/Engine/Source/Runtime/Engine/Public/Materials/MaterialExpressionAntialiasedTextureMask.h
+++ b/Engine/Source/Runtime/Engine/Public/Materials/MaterialExpressionAntialiasedTextureMask.h
@@ -38,6 +38,7 @@ class UMaterialExpressionAntialiasedTextureMask : public UMaterialExpressionText
 	//~ End UMaterialExpressionTextureSampleParameter Interface
 
 	//~ Begin UMaterialExpression Interface
+	virtual void Build(MIR::FEmitter& Emitter) override;
 	virtual int32 Compile(class FMaterialCompiler* Compiler, int32 OutputIndex) override;
 	virtual void GetCaption(TArray<FString>& OutCaptions) const override;
 	//~ End UMaterialExpression Interface
diff --git a/Engine/Source/Runtime/Engine/Public/Materials/MaterialExpressionHairAttributes.h b/Engine/Source/Runtime/Engine/Public/Materials/MaterialExpressionHairAttributes.h
--- a/Engine/Source/Runtime/Engine/Public/Materials/MaterialExpressionHairAttributes.h
+++ b/Engine/Source/Runtime/Engine/Public/Materials/MaterialExpressionHairAttributes.h
@@ -20,6 +20,7 @@ class UMaterialExpressionHairAttributes : public UMaterialExpression
 
 	//~ Begin UMaterialExpression Interface
 #if WITH_EDITOR
+	virtual void Build(MIR::FEmitter& Emitter) override;
 	virtual int32 Compile(class FMaterialCompiler* Compiler, int32 OutputIndex) override;
 	virtual void GetCaption(TArray<FString>& OutCaptions) const override;
 #endif
diff --git a/Engine/Source/Runtime/Engine/Public/Materials/MaterialExpressionHairColor.h b/Engine/Source/Runtime/Engine/Public/Materials/MaterialExpressionHairColor.h
--- a/Engine/Source/Runtime/Engine/Public/Materials/MaterialExpressionHairColor.h
+++ b/Engine/Source/Runtime/Engine/Public/Materials/MaterialExpressionHairColor.h
@@ -28,6 +28,7 @@ class UMaterialExpressionHairColor : public UMaterialExpression
 
 	//~ Begin UMaterialExpression Interface
 #if WITH_EDITOR
+	virtual void Build(MIR::FEmitter& Emitter) override;
 	virtual int32 Compile(class FMaterialCompiler* Compiler, int32 OutputIndex) override;
 	virtual void GetCaption(TArray<FString>& OutCaptions) const override;
 #endif
diff --git a/Engine/Source/Runtime/Engine/Public/Materials/MaterialExpressionLocalPosition.h b/Engine/Source/Runtime/Engine/Public/Materials/MaterialExpressionLocalPosition.h
--- a/Engine/Source/Runtime/Engine/Public/Materials/MaterialExpressionLocalPosition.h
+++ b/Engine/Source/Runtime/Engine/Public/Materials/MaterialExpressionLocalPosition.h
@@ -46,6 +46,7 @@ class UMaterialExpressionLocalPosition : public UMaterialExpression
 	
 	//~ Begin UMaterialExpression Interface
 #if WITH_EDITOR
+	virtual void Build(MIR::FEmitter& Emitter) override;
 	virtual int32 Compile(class FMaterialCompiler* Compiler, int32 OutputIndex) override;
 	virtual void GetCaption(TArray<FString>& OutCaptions) const override;
 	virtual void GetExpressionToolTip(TArray<FString>& OutToolTip) override;
diff --git a/Engine/Source/Runtime/Engine/Public/Materials/MaterialExpressionMapARPassthroughCameraUV.h b/Engine/Source/Runtime/Engine/Public/Materials/MaterialExpressionMapARPassthroughCameraUV.h
--- a/Engine/Source/Runtime/Engine/Public/Materials/MaterialExpressionMapARPassthroughCameraUV.h
+++ b/Engine/Source/Runtime/Engine/Public/Materials/MaterialExpressionMapARPassthroughCameraUV.h
@@ -24,6 +24,7 @@ public:
 
 	//~ Begin UMaterialExpression Interface
 #if WITH_EDITOR
+	virtual void Build(MIR::FEmitter& Emitter) override;
 	virtual int32 Compile(class FMaterialCompiler* Compiler, int32 OutputIndex) override;
 	virtual void GetCaption(TArray<FString>& OutCaptions) const override;
 #endif
diff --git a/Engine/Source/Runtime/Engine/Public/Materials/MaterialExpressionParticleSubUV.h b/Engine/Source/Runtime/Engine/Public/Materials/MaterialExpressionParticleSubUV.h
--- a/Engine/Source/Runtime/Engine/Public/Materials/MaterialExpressionParticleSubUV.h
+++ b/Engine/Source/Runtime/Engine/Public/Materials/MaterialExpressionParticleSubUV.h
@@ -20,6 +20,7 @@ public:
 
 	//~ Begin UMaterialExpression Interface
 #if WITH_EDITOR
+	virtual void Build(MIR::FEmitter& Emitter) override;
 	virtual int32 Compile(class FMaterialCompiler* Compiler, int32 OutputIndex) override;
 	virtual void GetCaption(TArray<FString>& OutCaptions) const override;
 	virtual int32 GetWidth() const override;	
diff --git a/Engine/Source/Runtime/Engine/Public/Materials/MaterialExpressionPreviousFrameSwitch.h b/Engine/Source/Runtime/Engine/Public/Materials/MaterialExpressionPreviousFrameSwitch.h
--- a/Engine/Source/Runtime/Engine/Public/Materials/MaterialExpressionPreviousFrameSwitch.h
+++ b/Engine/Source/Runtime/Engine/Public/Materials/MaterialExpressionPreviousFrameSwitch.h
@@ -22,6 +22,7 @@ public:
 
 	//~ Begin UMaterialExpression Interface
 #if WITH_EDITOR
+	virtual void Build(MIR::FEmitter& Emitter) override;
 	virtual int32 Compile(class FMaterialCompiler* Compiler, int32 OutputIndex) override;
 	virtual void GetCaption(TArray<FString>& OutCaptions) const override;
 	virtual void GetExpressionToolTip(TArray<FString>& OutToolTip) override;
diff --git a/Engine/Source/Runtime/Engine/Public/Materials/MaterialExpressionSwitch.h b/Engine/Source/Runtime/Engine/Public/Materials/MaterialExpressionSwitch.h
--- a/Engine/Source/Runtime/Engine/Public/Materials/MaterialExpressionSwitch.h
+++ b/Engine/Source/Runtime/Engine/Public/Materials/MaterialExpressionSwitch.h
@@ -69,6 +69,7 @@ class UMaterialExpressionSwitch : public UMaterialExpression
 	
 	//~ Begin UMaterialExpression Interface
 #if WITH_EDITOR
+	virtual void Build(MIR::FEmitter& Emitter) override;
 	virtual int32 Compile(class FMaterialCompiler* Compiler, int32 OutputIndex) override;
 	virtual void GetCaption(TArray<FString>& OutCaptions) const override;
 	virtual FText GetCreationName() const override { return FText::FromString(TEXT("Switch")); }
diff --git a/Engine/Source/Runtime/Engine/Public/Materials/MaterialIR.h b/Engine/Source/Runtime/Engine/Public/Materials/MaterialIR.h
--- a/Engine/Source/Runtime/Engine/Public/Materials/MaterialIR.h
+++ b/Engine/Source/Runtime/Engine/Public/Materials/MaterialIR.h
@@ -360,6 +360,15 @@ enum class EDerivativeAxis
 	Y, // Corresponding to d/dy
 };
 
+// Specify the source expression type of a derivative, useful for generating actionable user facing errors if derivatives
+// are used where not allowed (vertex shader).
+enum class EDerivativeSource
+{
+	Derivative,				// DDX or DDY expression
+	TextureSampleBias,		// TextureSample expression with MipBias option selected
+	AnalyticDerivative,		// Generated during calculation of an analytic derivative
+};
+
 // Enumeration of all supported material external inputs.
 enum class EExternalInput
 {
@@ -377,6 +386,10 @@ enum class EExternalInput
 	WorldPosition_AbsoluteNoOffsets,		// (LWC)	WPT_ExcludeAllShaderOffsets		Absolute World Position, excluding Material Shader Offsets
 	WorldPosition_CameraRelative,			// (float)	WPT_CameraRelative				Camera Relative World Position, including Material Shader Offsets
 	WorldPosition_CameraRelativeNoOffsets,	// (float)	WPT_CameraRelativeNoOffsets		Camera Relative World Position, excluding Material Shader Offsets
+	LocalPosition_Instance,
+	LocalPosition_InstanceNoOffsets,
+	LocalPosition_Primitive,
+	LocalPosition_PrimitiveNoOffsets,
 
 	TexCoord0_Ddx,
 	TexCoord1_Ddx,
@@ -390,6 +403,10 @@ enum class EExternalInput
 	WorldPosition_AbsoluteNoOffsets_Ddx,	// LWC (cast from float on load)
 	WorldPosition_CameraRelative_Ddx,
 	WorldPosition_CameraRelativeNoOffsets_Ddx,
+	LocalPosition_Instance_Ddx,
+	LocalPosition_InstanceNoOffsets_Ddx,
+	LocalPosition_Primitive_Ddx,
+	LocalPosition_PrimitiveNoOffsets_Ddx,
 
 	TexCoord0_Ddy,
 	TexCoord1_Ddy,
@@ -403,13 +420,17 @@ enum class EExternalInput
 	WorldPosition_AbsoluteNoOffsets_Ddy,	// LWC (cast from float on load)
 	WorldPosition_CameraRelative_Ddy,
 	WorldPosition_CameraRelativeNoOffsets_Ddy,
+	LocalPosition_Instance_Ddy,
+	LocalPosition_InstanceNoOffsets_Ddy,
+	LocalPosition_Primitive_Ddy,
+	LocalPosition_PrimitiveNoOffsets_Ddy,
 
 	// Ranges of External inputs with derivatives.  There must be an equal number of value, ddx, and ddy variations.  TexCoord0 must remain first,
 	// as code that translates texture coordinate indices to external input ID and back assumes this.
 	WithDerivatives_First = TexCoord0,
-	WithDerivatives_LastVal = WorldPosition_CameraRelativeNoOffsets,
-	WithDerivatives_LastDdx = WorldPosition_CameraRelativeNoOffsets_Ddx,
-	WithDerivatives_LastDdy = WorldPosition_CameraRelativeNoOffsets_Ddy,
+	WithDerivatives_LastVal = LocalPosition_PrimitiveNoOffsets,
+	WithDerivatives_LastDdx = LocalPosition_PrimitiveNoOffsets_Ddx,
+	WithDerivatives_LastDdy = LocalPosition_PrimitiveNoOffsets_Ddy,
 	WithDerivatives_Last = WithDerivatives_LastDdy,
 
 	ActorPosition_Absolute,					// LWC
@@ -420,6 +441,7 @@ enum class EExternalInput
 	ViewMaterialTextureDerivativeMultiply,
 	GlobalDistanceField,					// Flag to mark that global distance fields are used, for value analyzer step.  Doesn't generate code.
 	DynamicParticleParameterIndex,			// UserData stores the parameter index.
+	CompilingPreviousFrame,					// Boolean value bCompilingPreviousFrame, set when compiling the vertex shader WorldPositionOffsets for the previous frame.
 
 	Count,
 };
@@ -980,6 +1002,9 @@ struct FHardwarePartialDerivative : TInstruction<VK_HardwarePartialDerivative, 1
 
 	// The direction of partial derivative
 	EDerivativeAxis Axis;
+
+	// The source expression that generated this partial derivative (for error reporting)
+	EDerivativeSource Source;
 };
 
 // A "no operation" instruction is a dummy instruction used to perform analysis on a value (and its dependencies)
diff --git a/Engine/Source/Runtime/Engine/Public/Materials/MaterialIREmitter.h b/Engine/Source/Runtime/Engine/Public/Materials/MaterialIREmitter.h
--- a/Engine/Source/Runtime/Engine/Public/Materials/MaterialIREmitter.h
+++ b/Engine/Source/Runtime/Engine/Public/Materials/MaterialIREmitter.h
@@ -606,7 +606,7 @@ public:
 	// @param InputValues An array of input values required by the HLSL code.
 	// @param ValueFlags Flags that specify additional properties for this value (default is None).
 	// @param UsedGraphProperties Properties used during graph processing (default is None).
-	FValueRef InlineHLSL(FType Type, FStringView Code, TConstArrayView<FValueRef> InputValues, EValueFlags ValueFlags = EValueFlags::None, EGraphProperties UsedGraphProperties = EGraphProperties::None);
+	FValueRef InlineHLSL(FType Type, FStringView Code, TConstArrayView<FValueRef> InputValues = {}, EValueFlags ValueFlags = EValueFlags::None, EGraphProperties UsedGraphProperties = EGraphProperties::None);
 
 	// Emit a new HLSL function inline within the material shader based on an external code declaration.
 	// @param InExternalCodeDeclaration The external HLSL code declaration to use for emission.
