ChildActorComponent: fix child actor components serialized from OFPA files during PIE not properly added to the world since CL 38476333.
The original change was relying on the fact that duplicated actors will recreate child actors, but didn't take into account actors serialized from disk.

diff --git a/Engine/Source/Runtime/Engine/Private/WorldPartition/WorldPartitionLevelHelper.cpp b/Engine/Source/Runtime/Engine/Private/WorldPartition/WorldPartitionLevelHelper.cpp
--- a/Engine/Source/Runtime/Engine/Private/WorldPartition/WorldPartitionLevelHelper.cpp
+++ b/Engine/Source/Runtime/Engine/Private/WorldPartition/WorldPartitionLevelHelper.cpp
@@ -793,6 +793,15 @@ bool FWorldPartitionLevelHelper::LoadActorsInternal(FLoadActorsParams&& InParams
 					// Rename through UObject to avoid changing Actor's external packaging and folder properties
 					Actor->UObject::Rename(*AddActorContainerID(PackageObjectMapping->ContainerID, Actor->GetName()), DestLevel, REN_NonTransactional | REN_DoNotDirty | REN_DontCreateRedirectors);
 
+					// Handle child actors
+					Actor->ForEachComponent<UChildActorComponent>(true, [DestLevel = DestLevel, PackageObjectMapping](UChildActorComponent* ChildActorComponent)
+					{
+						if (AActor* ChildActor = ChildActorComponent->GetChildActor())
+						{
+							ChildActor->UObject::Rename(*AddActorContainerID(PackageObjectMapping->ContainerID, ChildActor->GetName()), DestLevel, REN_NonTransactional | REN_DoNotDirty | REN_DontCreateRedirectors);
+						}
+					});
+
 					// Apply Pre-ConstructionScript Properties
 					TArray<FActorPropertyOverride> ActorPropertyOverrides;
 					for (auto ActorOverrideMapping : PackageObjectMapping->PropertyOverrides)
@@ -870,6 +879,16 @@ bool FWorldPartitionLevelHelper::LoadActorsInternal(FLoadActorsParams&& InParams
 
 					DestLevel->Actors.Add(Actor);
 					checkf(Actor->GetLevel() == DestLevel, TEXT("Levels mismatch, got : %s, expected: %s\nActor: %s\nActorFullName: %s\nActorPackage: %s"), *DestLevel->GetFullName(), *Actor->GetLevel()->GetFullName(), *Actor->GetActorNameOrLabel(), *Actor->GetFullName(), *Actor->GetPackage()->GetFullName());
+
+					// Handle child actors
+					Actor->ForEachComponent<UChildActorComponent>(true, [DestLevel = DestLevel](UChildActorComponent* ChildActorComponent)
+					{
+						if (AActor* ChildActor = ChildActorComponent->GetChildActor())
+						{
+							DestLevel->Actors.Add(ChildActor);
+							check(ChildActor->GetLevel() == DestLevel);
+						}
+					});
 				}
 
 				UE_LOG(LogWorldPartition, Verbose, TEXT(" ==> Loaded %s (remaining: %d)"), *Actor->GetFullName(), LoadProgress->NumPendingLoadRequests);
