- Remove Empty from FFXSystemSet::OnMarkPendingKill
- This contends with the render thread so we leave it for the destructor which will run on the RT

#rnx

diff --git a/Engine/Source/Runtime/Engine/Private/Particles/FXSystemSet.cpp b/Engine/Source/Runtime/Engine/Private/Particles/FXSystemSet.cpp
--- a/Engine/Source/Runtime/Engine/Private/Particles/FXSystemSet.cpp
+++ b/Engine/Source/Runtime/Engine/Private/Particles/FXSystemSet.cpp
@@ -208,7 +208,6 @@ void FFXSystemSet::OnMarkPendingKill()
 	{
 		FXSystem->OnMarkPendingKill();
 	}
-	FXSystems.Empty();
 
 	FFXSystemInterface::OnMarkPendingKill();
 }
