Fixed crash when enabling debug display of data layer load time (wp.Runtime.ToggleDrawDataLayersLoadTime) and playing a replay containing old (now deleted) Data Layer Names.

diff --git a/Engine/Source/Runtime/Engine/Private/WorldPartition/DataLayer/DataLayerManager.cpp b/Engine/Source/Runtime/Engine/Private/WorldPartition/DataLayer/DataLayerManager.cpp
--- a/Engine/Source/Runtime/Engine/Private/WorldPartition/DataLayer/DataLayerManager.cpp
+++ b/Engine/Source/Runtime/Engine/Private/WorldPartition/DataLayer/DataLayerManager.cpp
@@ -516,9 +516,18 @@ void UDataLayerManager::DrawDataLayersStatus(UCanvas* Canvas, FVector2D& Offset)
 
 		TArray<const UDataLayerInstance*> LoadingDataLayers;
 		LoadingDataLayers.Reserve(LoadedDataLayers.Num() + ActiveDataLayers.Num());
-		auto CopyLambda = [this](FName DataLayerName) { return GetDataLayerInstanceFromName(DataLayerName); };
-		Algo::Transform(LoadedDataLayers, LoadingDataLayers, CopyLambda);
-		Algo::Transform(ActiveDataLayers, LoadingDataLayers, CopyLambda);
+		auto TransformAndAppendDataLayers = [this](const TSet<FName>& InDataLayerNames, TArray<const UDataLayerInstance*>& OutDataLayerInstances)
+		{
+			for (FName DataLayerName : InDataLayerNames)
+			{
+				if (const UDataLayerInstance* DataLayerInstance = GetDataLayerInstanceFromName(DataLayerName))
+				{
+					OutDataLayerInstances.Add(DataLayerInstance);
+				}
+			}
+		};
+		TransformAndAppendDataLayers(LoadedDataLayers, LoadingDataLayers);
+		TransformAndAppendDataLayers(ActiveDataLayers, LoadingDataLayers);
 
 		for (const UDataLayerInstance* DataLayerInstance : LoadingDataLayers)
 		{
