Ensure that FDefaultMaterialInstance's CacheUniformExpression is cancelled on destroy if still pending. Under certain conditions the FDefaultMaterialInstance could be destroyed and subsequent dispatch of CacheUniformExpression would result in a crash due to a dangling reference. This replicates functionality from FMaterialInstanceResource for which this guard is already implemented.

diff --git a/Engine/Source/Runtime/Engine/Private/Materials/Material.cpp b/Engine/Source/Runtime/Engine/Private/Materials/Material.cpp
--- a/Engine/Source/Runtime/Engine/Private/Materials/Material.cpp
+++ b/Engine/Source/Runtime/Engine/Private/Materials/Material.cpp
@@ -417,6 +417,16 @@ public:
 		});
 	}
 
+	virtual void StartCacheUniformExpressions() const override
+	{
+		Material->StartCacheUniformExpressions();
+	}
+
+	virtual void FinishCacheUniformExpressions() const override
+	{
+		Material->FinishCacheUniformExpressions();
+	}
+
 	virtual const FMaterialRenderProxy* GetFallback(ERHIFeatureLevel::Type InFeatureLevel) const override
 	{
 		const FMaterialRenderProxy* Fallback = &GetFallbackRenderProxy();
@@ -1387,6 +1397,16 @@ void UMaterial::GetUsedTexturesAndIndices(TArray<UTexture*>& OutTextures, TArray
 	}
 }
 
+void UMaterial::StartCacheUniformExpressions() const
+{
+	bCachingUniformExpressions = true;
+}
+
+void UMaterial::FinishCacheUniformExpressions() const
+{
+	bCachingUniformExpressions = false;
+}
+
 #if !(UE_BUILD_SHIPPING || UE_BUILD_TEST)
 void UMaterial::LogMaterialsAndTextures(FOutputDevice& Ar, int32 Indent) const
 {
@@ -5943,11 +5963,11 @@ void UMaterial::BeginDestroy()
 	if (DefaultMaterialInstance || ResourcesToDestroy.Num() > 0)
 	{
 		ReleasedByRT = false;
-		FMaterialRenderProxy* LocalResource = DefaultMaterialInstance;
-		FThreadSafeBool* Released = &ReleasedByRT;
+
 		ENQUEUE_RENDER_COMMAND(BeginDestroyCommand)(
-			[ResourcesToDestroy = MoveTemp(ResourcesToDestroy), LocalResource, Released](FRHICommandListImmediate& RHICmdList) mutable
+			[ResourcesToDestroy = MoveTemp(ResourcesToDestroy), this](FRHICommandListImmediate& RHICmdList) mutable
 		{
+			FMaterialRenderProxy* LocalResource = DefaultMaterialInstance;
 			if (LocalResource)
 			{
 				LocalResource->MarkForGarbageCollection();
@@ -5962,7 +5982,14 @@ void UMaterial::BeginDestroy()
 			// Release the references before assigning the bool below.
 			ResourcesToDestroy.Empty();
 
-			*Released = true;
+			// And remove from deferred uniform expression cache queue if it's in that
+			if (bCachingUniformExpressions)
+			{
+				LocalResource->CancelCacheUniformExpressions();
+				bCachingUniformExpressions = false;
+			}
+
+			ReleasedByRT = true;
 		});
 	}
 }
diff --git a/Engine/Source/Runtime/Engine/Public/Materials/Material.h b/Engine/Source/Runtime/Engine/Public/Materials/Material.h
--- a/Engine/Source/Runtime/Engine/Public/Materials/Material.h
+++ b/Engine/Source/Runtime/Engine/Public/Materials/Material.h
@@ -1206,6 +1206,9 @@ private:
 	/** Flag used to guarantee that the RT is finished using various resources in this UMaterial before cleanup. */
 	FThreadSafeBool ReleasedByRT;
 
+	/** This instance is caching uniform expressions (the render thread hasn't canceled its creation). */
+	mutable bool bCachingUniformExpressions = false;
+
 #if WITH_EDITORONLY_DATA
 	UPROPERTY()
 	bool bSavedCachedExpressionData_DEPRECATED;
@@ -1426,6 +1429,9 @@ public:
 	ENGINE_API bool SetStaticSwitchParameterValueEditorOnly(FName ParameterName, bool OutValue, FGuid OutExpressionGuid);
 #endif // WITH_EDITOR
 
+	/** Tracking of in-flight uniform expression cache update operations for the DefaultMaterialInstance, for thread safety destroying the resource. */
+	void StartCacheUniformExpressions() const;
+	void FinishCacheUniformExpressions() const;
 
 #if !(UE_BUILD_SHIPPING || UE_BUILD_TEST)
 	/**
