Add ensure and comments to prevent creation of registered FTickableGameObjects on worker threads which is an unsafe race condition.
Add option to FTickableObjectBase constructor to allow creating with a disabled tick, which is safe on a worker thread.
Convert some types that can be created on the async loading thread to only enable after initialization, this will fix some random crashes.

diff --git a/Engine/Source/Runtime/Engine/Private/Subsystems/WorldSubsystem.cpp b/Engine/Source/Runtime/Engine/Private/Subsystems/WorldSubsystem.cpp
--- a/Engine/Source/Runtime/Engine/Private/Subsystems/WorldSubsystem.cpp
+++ b/Engine/Source/Runtime/Engine/Private/Subsystems/WorldSubsystem.cpp
@@ -54,6 +54,7 @@ PRAGMA_ENABLE_DEPRECATION_WARNINGS
 // ----------------------------------------------------------------------------------
 
 UTickableWorldSubsystem::UTickableWorldSubsystem()
+	: FTickableGameObject(ETickableTickType::Never)
 {
 
 }
@@ -111,5 +112,5 @@ void UTickableWorldSubsystem::BeginDestroy()
 {
 	Super::BeginDestroy();
 
-	ensureMsgf(!bInitialized, TEXT("Tickable subsystem %s was destroyed while still initialized! Check for missing Super call"), *GetFullName());
+	ensureMsgf(!bInitialized, TEXT("Tickable subsystem %s was destroyed while still initialized! Check for missing Super::Deinitialize call"), *GetFullName());
 }
diff --git a/Engine/Source/Runtime/Engine/Private/Tickable.cpp b/Engine/Source/Runtime/Engine/Private/Tickable.cpp
--- a/Engine/Source/Runtime/Engine/Private/Tickable.cpp
+++ b/Engine/Source/Runtime/Engine/Private/Tickable.cpp
@@ -6,11 +6,14 @@
 
 DECLARE_CYCLE_STAT(TEXT("TickableGameObjects Time"), STAT_TickableGameObjectsTime, STATGROUP_Game);
 
-void FTickableObjectBase::FTickableStatics::QueueTickableObjectForAdd(FTickableObjectBase* InTickable)
+void FTickableObjectBase::FTickableStatics::QueueTickableObjectForAdd(FTickableObjectBase* InTickable, ETickableTickType NewTickType)
 {
+	// Objects set to never tick at creation will not be registered
+	check(NewTickType != ETickableTickType::Never);
+
 	// This only needs to lock the new object queue
 	UE::TScopeLock NewTickableObjectsLock(NewTickableObjectsCritical);
-	NewTickableObjects.Add(InTickable, ETickableTickType::NewObject);
+	NewTickableObjects.Add(InTickable, NewTickType);
 }
 
 void FTickableObjectBase::FTickableStatics::SetTickTypeForTickableObject(FTickableObjectBase* InTickable, ETickableTickType NewTickType)
@@ -18,9 +21,6 @@ void FTickableObjectBase::FTickableStatics::SetTickTypeForTickableObject(FTickab
 	UE::TScopeLock TickableObjectsLock(TickableObjectsCritical);
 	UE::TScopeLock NewTickableObjectsLock(NewTickableObjectsCritical);
 	
-	// Existing entries should never be set to new object
-	check(NewTickType != ETickableTickType::NewObject);
-
 	const int32 Pos = TickableObjects.IndexOfByKey(InTickable);
 	if (NewTickType == ETickableTickType::Never)
 	{
@@ -49,8 +49,12 @@ void FTickableObjectBase::FTickableStatics::SetTickTypeForTickableObject(FTickab
 			// If this is registered, it was removed from the new list in BeginTicking
 			check(!NewTickableObjects.Contains(InTickable));
 
-			// This will modify behavior for the current frame if it has not ticked yet
-			TickableObjects[Pos].TickType = NewTickType;
+			// Existing entries should never be set to new object
+			if (ensure(NewTickType != ETickableTickType::NewObject))
+			{
+				// This will modify behavior for the current frame if it has not ticked yet
+				TickableObjects[Pos].TickType = NewTickType;
+			}
 		}
 		else
 		{
@@ -126,31 +130,31 @@ void FTickableObjectBase::SimpleTickObjects(FTickableStatics& Statics, TFunction
 
 
 // FTickableGameObject implementation
-
-FTickableGameObject::FTickableGameObject()
+FTickableGameObject::FTickableGameObject(ETickableTickType StartingTickType)
 {
-	FTickableStatics& Statics = GetStatics();
+	if (StartingTickType != ETickableTickType::Never)
+	{
+		// It is only safe to create tickable game objects on the game thread, as otherwise there is a race condition between object initialize and the game thread tick
+		// If you hit this ensure, change the constructor to use FTickableGameObject(ETickableTickType::Never) and call SetTickableTickType after initialization
+		ensure(IsInGameThread());
 
-	// Queue for creation, this can get called very early in startup
-	Statics.QueueTickableObjectForAdd(this);
+		// Queue for creation, this can get called very early in startup
+		FTickableStatics& Statics = GetStatics();
+		Statics.QueueTickableObjectForAdd(this, StartingTickType);
+	}
 }
 
 FTickableGameObject::~FTickableGameObject()
 {
+	// This won't do anything if it was already set to never tick
 	// Depending on destruction order this could create a new statics object during shutdown, but the removal request will be ignored
-	FTickableStatics& Statics = GetStatics();
-	
-	// This won't do anything if it was already set to Never
-	Statics.SetTickTypeForTickableObject(this, ETickableTickType::Never);
+	SetTickableTickType(ETickableTickType::Never);
 }
 
 void FTickableGameObject::SetTickableTickType(ETickableTickType NewTickType)
 {
-	if (ensure(NewTickType != ETickableTickType::NewObject))
-	{
-		FTickableStatics& Statics = GetStatics();
-		Statics.SetTickTypeForTickableObject(this, NewTickType);
-	}
+	FTickableStatics& Statics = GetStatics();
+	Statics.SetTickTypeForTickableObject(this, NewTickType);
 }
 
 FTickableObjectBase::FTickableStatics& FTickableGameObject::GetStatics()
diff --git a/Engine/Source/Runtime/Engine/Public/Tickable.h b/Engine/Source/Runtime/Engine/Public/Tickable.h
--- a/Engine/Source/Runtime/Engine/Public/Tickable.h
+++ b/Engine/Source/Runtime/Engine/Public/Tickable.h
@@ -69,7 +69,7 @@ protected:
 
 
 		/** Call from anywhere to lock and add to the new objects list */
-		ENGINE_API void QueueTickableObjectForAdd(FTickableObjectBase* InTickable);
+		ENGINE_API void QueueTickableObjectForAdd(FTickableObjectBase* InTickable, ETickableTickType NewTickType = ETickableTickType::NewObject);
 
 		/** Call from anywhere to lock critical sections and set the tick type for an object which could modify both arrays */
 		ENGINE_API void SetTickTypeForTickableObject(FTickableObjectBase* TickableObject, ETickableTickType NewTickType);
@@ -141,9 +141,14 @@ public:
 	UE_NONCOPYABLE(FTickableGameObject);
 
 	/**
-	 * Registers this instance with the static array of tickable objects.	
+	 * Registers this instance with the static array of tickable objects, if it can ever tick.
+	 * By default this function can only be called on the game thread, to make sure it doesn't call functions before initialization is complete.
+	 * With the default NewObject tick type, it will call GetTickableTickType at the start of the next tick to determine the desired type.
+	 * If this is something like a UObject that could be created on a different thread (like for async loading), construct with a Never tick type and enable tick later.
+	 * 
+	 * @param StartingTickType The initial ticking state of this object, override it to start disabled or skip the call to GetTickableTickType
 	 */
-	ENGINE_API FTickableGameObject();
+	ENGINE_API FTickableGameObject(ETickableTickType StartingTickType = ETickableTickType::NewObject);
 
 	/**
 	 * Removes this instance from the static array of tickable objects.
@@ -187,6 +192,7 @@ public:
 	/**
 	 * Call to modify the tickable type of this instance. 
 	 * This can be used to enable or disable tick even if GetTickableTickType has already been called.
+	 * If this is a UObject that could be destroyed on a different thread, call this with Never during BeginDestroy (or earlier).
 	 * 
 	 * @param NewTickType	The new tick type of this instance
 	 */
