Only put snapping state into the clipboard if it's non-default. This avoids confusing older node importers, allows the Migrate asset action to behave as expected.

#rnx

diff --git a/Engine/Source/Runtime/Engine/Private/EdGraph/EdGraphPin.cpp b/Engine/Source/Runtime/Engine/Private/EdGraph/EdGraphPin.cpp
--- a/Engine/Source/Runtime/Engine/Private/EdGraph/EdGraphPin.cpp
+++ b/Engine/Source/Runtime/Engine/Private/EdGraph/EdGraphPin.cpp
@@ -1229,16 +1229,23 @@ bool UEdGraphPin::ExportTextItem(FString& ValueStr, int32 PortFlags) const
 
 		// Intentionally not exporting bSavePinIfOrphaned as it is transient
 		// Intentionally not exporting bWasTrashed as it is transient
+		// avoid exporting bHasSnappedChild unless it's true, as it defaults to false
+		// and older code that processes T3D nodes may be confused by extra properties:
+		if(bHasSnappedChild)
+		{
+			ValueStr += UEdGraphPin::bHasSnappedChildName + TEXT("=");
+			bool LocalHasSnappedChildName = bHasSnappedChild;
+			BoolPropCDO->ExportTextItem_Direct(ValueStr, &LocalHasSnappedChildName, nullptr, nullptr, PortFlags, nullptr);
+			ValueStr += PinHelpers::ExportTextPropDelimiter;
+		}
 		
-		ValueStr += UEdGraphPin::bHasSnappedChildName + TEXT("=");
-		bool LocalHasSnappedChildName = bHasSnappedChild;
-		BoolPropCDO->ExportTextItem_Direct(ValueStr, &LocalHasSnappedChildName, nullptr, nullptr, PortFlags, nullptr);
-		ValueStr += PinHelpers::ExportTextPropDelimiter;
-		
-		ValueStr += UEdGraphPin::bHasSnappedParentName + TEXT("=");
-		bool LocalHasSnappedParentName = bHasSnappedParent;
-		BoolPropCDO->ExportTextItem_Direct(ValueStr, &LocalHasSnappedParentName, nullptr, nullptr, PortFlags, nullptr);
-		ValueStr += PinHelpers::ExportTextPropDelimiter;
+		if(bHasSnappedParent)
+		{
+			ValueStr += UEdGraphPin::bHasSnappedParentName + TEXT("=");
+			bool LocalHasSnappedParentName = bHasSnappedParent;
+			BoolPropCDO->ExportTextItem_Direct(ValueStr, &LocalHasSnappedParentName, nullptr, nullptr, PortFlags, nullptr);
+			ValueStr += PinHelpers::ExportTextPropDelimiter;
+		}
 	}
 #endif //WITH_EDITORONLY_DATA
 	ValueStr += ")";
