[MIR] Introduced the concept of a "EntryPoint" (in FMaterialIRModule::FEntryPoint) to generalize the concept of execution stage.

Before this change, the translator was only capable of generating a fixed number of execution contexts for each shader stage (vertex, pixel and compute). This restriction has been lifted with the ability to create any number of execution stages, dubbed "EntryPoint". Shader stages are now builtin procedures always created within a module.

This change is key to support advanced features such as custom outputs and the ability to generate subgraphs of IR instructions inside local, user-specified HLSL functions, that will in turn be used in the future to add support for material layers.

This change also includes refactoring of the Builder to collect all non IR building operations into a single "FinalizeArtifacts" step. Added a plenty of comments explaining the flow of operations in the builder.

#rnx

diff --git a/Engine/Source/Runtime/Engine/Private/Materials/MaterialIR.cpp b/Engine/Source/Runtime/Engine/Private/Materials/MaterialIR.cpp
--- a/Engine/Source/Runtime/Engine/Private/Materials/MaterialIR.cpp
+++ b/Engine/Source/Runtime/Engine/Private/Materials/MaterialIR.cpp
@@ -24,18 +24,6 @@ const TCHAR* LexToString(EStage Stage)
 	return nullptr;
 }
 
-bool MaterialOutputEvaluatesInStage(EMaterialProperty Property, EStage Stage)
-{	
-	if (Property == MP_WorldPositionOffset)
-	{
-		return Stage == Stage_Vertex;
-	}
-	else
-	{
-		return Stage != Stage_Vertex;
-	}
-}
-
 const TCHAR* LexToString(EValueKind Kind)
 {
 	switch (Kind)
@@ -752,21 +740,21 @@ bool FComposite::AreComponentsConstant() const
 	return true;
 }
 
-FBlock* FInstruction::GetDesiredBlockForUse(EStage Stage, int32 UseIndex)
+FBlock* FInstruction::GetTargetBlockForUse(int32 EntryPointIndex, int32 UseIndex)
 {
 	if (auto Branch = As<FBranch>())
 	{
 		switch (UseIndex)
 		{
-			case 0: return Block[Stage]; 				// ConditionArg goes into the same block as this instruction's
-			case 1: return &Branch->TrueBlock[Stage]; 	// TrueArg
-			case 2: return &Branch->FalseBlock[Stage]; 	// FalseArg
+			case 0: return Linkage[EntryPointIndex].Block; 		// ConditionArg goes into the same block as this instruction's
+			case 1: return &Branch->TrueBlock[EntryPointIndex]; 	// TrueArg
+			case 2: return &Branch->FalseBlock[EntryPointIndex]; // FalseArg
 			default: UE_MIR_UNREACHABLE();
 		}
 	}
 
 	// By default, dependencies can go in the same block as this instruction
-	return Block[Stage];
+	return Linkage[EntryPointIndex].Block;
 }
 
 bool IsComparisonOperator(EOperator Op)
diff --git a/Engine/Source/Runtime/Engine/Private/Materials/MaterialIRDebug.cpp b/Engine/Source/Runtime/Engine/Private/Materials/MaterialIRDebug.cpp
--- a/Engine/Source/Runtime/Engine/Private/Materials/MaterialIRDebug.cpp
+++ b/Engine/Source/Runtime/Engine/Private/Materials/MaterialIRDebug.cpp
@@ -37,23 +37,28 @@ struct FDebugDumpIRUseGraphState
 			"edge [fontname=\"Consolas\"]\n\n"
 		));
 
-		for (int32 StageIndex = 0; StageIndex < MIR::NumStages; ++StageIndex)
+		for (int32 EntryPointIndex = 0; EntryPointIndex < Module.GetNumEntryPoints(); ++EntryPointIndex)
 		{
-			for (const FSetMaterialOutput* Output : Module.GetOutputs((MIR::EStage)StageIndex))
+			const FMaterialIRModule::FEntryPoint& EntryPoint = Module.GetEntryPoint(EntryPointIndex);
+
+			for (const FValue* Output : EntryPoint.Outputs)
 			{
-				ValueStack.Push(Output);
+				if (Output)
+				{
+					ValueStack.Push(Output);
+				}
 			}
 
 			while (!ValueStack.IsEmpty())
 			{
-				DumpValue((MIR::EStage)StageIndex, ValueStack.Pop());
+				DumpValue(EntryPointIndex, EntryPoint.Stage, ValueStack.Pop());
 			}
 		}
 
 		Out.Appendf(TEXT("\n}\n"));
 	}
 
-	void DumpValue(MIR::EStage Stage, const FValue* Value)
+	void DumpValue(uint32 EntryPointIndex, MIR::EStage Stage, const FValue* Value)
 	{
 		const bool bDumpInstructionSequence = CVarDumpMaterialIRUseGraph_EnableNext.GetValueOnAnyThread();
 
@@ -69,9 +74,9 @@ struct FDebugDumpIRUseGraphState
 		Out.Append(TEXT(">]\n"));
 
 		const FInstruction* Instr = AsInstruction(Value);
-		if (bDumpInstructionSequence && Instr && Instr->Next[Stage])
+		if (bDumpInstructionSequence && Instr && Instr->Linkage[EntryPointIndex].Next)
 		{
-			Out.Appendf(TEXT("\"%p\" -> \"%p\" [color=\"red\"]\n"), Instr, Instr->Next[Stage]);
+			Out.Appendf(TEXT("\"%p\" -> \"%p\" [color=\"red\"]\n"), Instr, Instr->Linkage[EntryPointIndex].Next);
 		}
 
 		int32 UseIndex = -1;
@@ -99,7 +104,7 @@ struct FDebugDumpIRUseGraphState
 			if (bDumpInstructionSequence && Instr)
 			{
 				const FInstruction* UseInstr = AsInstruction(Use);
-				if (UseInstr && UseInstr->Block[Stage] != Instr->Block[Stage])
+				if (UseInstr && UseInstr->Linkage[EntryPointIndex].Block != Instr->Linkage[EntryPointIndex].Block)
 				{
 					Out.Appendf(TEXT("\"%p\" -> \"%p\" [color=\"red\", style=\"dashed\"]\n"), UseInstr, Instr);
 				}
@@ -209,31 +214,33 @@ struct FDebugDumpIRState
 	const FMaterialIRModule* Module{};
 
 	// Output string containing the generated result.
-	FString Out; 
+	FString Out{}; 
 
 	// String used for temporary operations. Clear before use.
-	FString Temp;
+	FString Temp{};
 
 	// Maps values to an incrementing id. Used to give values a "name" for future referencing (e.g. "%6")
-	TMap<const FValue*, uint32> ValueToIdMap;
+	TMap<const FValue*, uint32> ValueToIdMap{};
 
 	// Counter used to assign an id to encountered values.
-	uint32 InstrIdCounter = 0;
+	uint32 InstrIdCounter{};
 
 	// Array of encountered parameters. Used later on to generate the a recap of all referenced parameters.
-	TArray<TPair<uint32, const FValue*>> ReferencedParameters;
+	TArray<TPair<uint32, const FValue*>> ReferencedParameters{};
 	
 	// Array of encountered user HLSL functions.
-	TArray<const MIR::FFunctionHLSL*> ReferencedFunctionHLSLs;
+	TArray<const MIR::FFunctionHLSL*> ReferencedFunctionHLSLs{};
+
+	int32 CurrentEntryPointIndex{};
 
 	// Stage we're currently emitting to.
-	MIR::EStage CurrentStage;
+	MIR::EStage CurrentStage{};
 
 	// Prints a block of instructions to `Out`. Indentation indicates how many levels of indentation to put
 	// to the left of printed instruction.
 	void AppendBlock(const FBlock& Block, int32 Indentation)
 	{
-		for (FInstruction* Instr = Block.Instructions; Instr; Instr = Instr->Next[CurrentStage])
+		for (FInstruction* Instr = Block.Instructions; Instr; Instr = Instr->GetNext(CurrentEntryPointIndex))
 		{
 			// Format the left column (e.g. "%4 = ") string if this instruction is referenceable.
 			Temp.Empty();
@@ -277,8 +284,8 @@ struct FDebugDumpIRState
 				Out.Appendf(TEXT("%s "), *Use->Type.GetSpelling());
 
 				// If this use is in a block different from current's, dump the block in "{}" first.
-				const FBlock* UseBlock = Instr->GetDesiredBlockForUse(CurrentStage, UseIndex);
-				if (UseBlock != Instr->Block[CurrentStage] && UseBlock->Instructions)
+				const FBlock* UseBlock = Instr->GetTargetBlockForUse(CurrentEntryPointIndex, UseIndex);
+				if (UseBlock != Instr->GetBlock(CurrentEntryPointIndex) && UseBlock->Instructions)
 				{
 					Out.Append(TEXT("{\n"));
 					AppendBlock(*UseBlock, Indentation + 1);
@@ -520,11 +527,14 @@ FString DebugDumpIR(FStringView MaterialName, const FMaterialIRModule& Module)
 
 	// Dump the IR instructions in the root block.
 	static const TCHAR* Header = TEXT("; Material \"%s\" translated IR dump\n\n");
-	for (int32 StageIndex = 0; StageIndex < MIR::NumStages; ++StageIndex)
+	for (int32 EntryPointIndex = 0; EntryPointIndex < Module.GetNumEntryPoints(); ++EntryPointIndex)
 	{
-		State.CurrentStage = (MIR::EStage)StageIndex;
-		State.Out.Appendf(TEXT("\n; Stage %d \"%s\"\n"), StageIndex, MIR::LexToString(State.CurrentStage));
-		State.AppendBlock(Module.GetRootBlock((MIR::EStage)StageIndex), 0);
+		const FMaterialIRModule::FEntryPoint& EntryPoint = Module.GetEntryPoint(EntryPointIndex);
+		State.Out.Appendf(TEXT("\n; Entry Point %d \"%s\" (stage \"%s\")\n"), EntryPointIndex, EntryPoint.Name.GetData(), MIR::LexToString(EntryPoint.Stage));
+
+		State.CurrentEntryPointIndex = EntryPointIndex;
+		State.CurrentStage = EntryPoint.Stage;
+		State.AppendBlock(EntryPoint.RootBlock, 0);
 	}
 
 	// Print referenced material parameters recap if any
diff --git a/Engine/Source/Runtime/Engine/Private/Materials/MaterialIREmitter.cpp b/Engine/Source/Runtime/Engine/Private/Materials/MaterialIREmitter.cpp
--- a/Engine/Source/Runtime/Engine/Private/Materials/MaterialIREmitter.cpp
+++ b/Engine/Source/Runtime/Engine/Private/Materials/MaterialIREmitter.cpp
@@ -1173,29 +1173,23 @@ FValueRef FEmitter::ScreenTexture(EScreenTexture TextureKind)
 
 FSetMaterialOutput* FEmitter::SetMaterialOutput(EMaterialProperty InProperty, FValue* Arg)
 {
-	FSetMaterialOutput Proto = MakePrototype<FSetMaterialOutput>({});
+	FSetMaterialOutput Proto = MakePrototype<FSetMaterialOutput>(FType::MakeVoid());
 	Proto.Property           = InProperty;
 	Proto.Arg                = Arg;
-
-	// Initialize the instruction block to the root of each stage it is evaluated in.
-	for (int i = 0; i < NumStages; ++i)
-	{
-		if (MaterialOutputEvaluatesInStage(InProperty, (EStage)i))
-		{
-			Proto.Block[i] = &Module->RootBlocks[i];
-		}
-	}
-
-	// Make the instruction
+	
 	FSetMaterialOutput* Instr = static_cast<FSetMaterialOutput*>(EmitPrototype(*this, Proto).Value);
 
 	// Add the instruction to list of outputs of the stages it is evaluated in.
-	for (int i = 0; i < NumStages; ++i)
+	if (InProperty == MP_WorldPositionOffset)
 	{
-		if (MaterialOutputEvaluatesInStage(InProperty, (EStage)i))
-		{
-			Module->Outputs[i].Add(Instr);
-		}
+		Module->GetEntryPoint(MIR::Stage_Vertex).Outputs[0] = Instr;
+	}
+	else
+	{
+		int32 StageOutputIndex = UMaterialAggregate::MaterialPropertyToAttributeIndex(InProperty);
+		check(StageOutputIndex != -1);
+		Module->GetEntryPoint(MIR::Stage_Pixel).Outputs[StageOutputIndex] = Instr;
+		Module->GetEntryPoint(MIR::Stage_Compute).Outputs[StageOutputIndex] = Instr;
 	}
 
 	return Instr;
diff --git a/Engine/Source/Runtime/Engine/Private/Materials/MaterialIRModule.cpp b/Engine/Source/Runtime/Engine/Private/Materials/MaterialIRModule.cpp
--- a/Engine/Source/Runtime/Engine/Private/Materials/MaterialIRModule.cpp
+++ b/Engine/Source/Runtime/Engine/Private/Materials/MaterialIRModule.cpp
@@ -21,26 +21,42 @@ FMaterialIRModule::~FMaterialIRModule()
 
 void FMaterialIRModule::Empty()
 {
-	Allocator.Flush();
-	FMemory::Memset(RootBlocks, 0, sizeof(RootBlocks));
-
-	for (int32 i = 0; i < MIR::NumStages; ++i)
-	{
-		Outputs[i].Empty();
-	}
-
+	Errors.Empty();
+	ShadingModelsFromCompilation = {};
+	FunctionHLSLs.Empty();
+	ParameterCollections.Empty();
+	EnvironmentDefines.Empty();
+	UserStrings.Empty();
+	ParameterIdToData.Empty();
+	ParameterInfoToId.Empty();
+	Statistics = {};
+	MIR::ZeroArray(TArrayView<MIR::FValue*>{ PropertyValues, MP_MAX });
+	EntryPoints.Empty();
 	Values.Empty();
+	CompilationOutput = {};
+	Allocator.Flush();
 
-	// Reset module statistics.
+	// Initialize statistics used external input usage mask.
 	for (int i = 0; i < MIR::NumStages; ++i)
 	{
 		Statistics.ExternalInputUsedMask[i].Init(false, (int)MIR::EExternalInput::Count);
 	}
+}
 
-	Statistics.NumVertexTexCoords = 0;
-	Statistics.NumPixelTexCoords = 0;
+int32 FMaterialIRModule::AddEntryPoint(FStringView Name, MIR::EStage Stage, int32 NumOutputs)
+{
+	uint32 Index = EntryPoints.Num();
+
+	EntryPoints.Push({
+		.Name = Name,
+		.Stage = Stage,
+		.RootBlock = {},
+		.Outputs = AllocateArray<MIR::FValue*>(NumOutputs),
+	});
 
-	Statistics.DynamicParticleParameterMask = 0;
+	MIR::ZeroArray(EntryPoints.Last().Outputs);
+
+	return Index;
 }
 
 FStringView FMaterialIRModule::InternString(FStringView InString)
@@ -87,17 +103,9 @@ void FMaterialIRModule::AddShadingModel(EMaterialShadingModel InShadingModel)
 	ShadingModelsFromCompilation.AddShadingModel(InShadingModel);
 }
 
-// Similar to FHLSLMaterialTranslator::IsMaterialPropertyUsed.
-bool FMaterialIRModule::IsMaterialPropertyUsed(MIR::EStage Stage, EMaterialProperty InProperty) const
+bool FMaterialIRModule::IsMaterialPropertyUsed(EMaterialProperty InProperty) const
 {
-	for (const MIR::FSetMaterialOutput* Output : GetOutputs(Stage))
-	{
-		if (Output->Property == InProperty)
-		{
-			return !Output->Arg->EqualsConstant(FMaterialAttributeDefinitionMap::GetDefaultValue(InProperty));
-		}
-	}
-	return false;
+	return PropertyValues[InProperty] && !PropertyValues[InProperty]->EqualsConstant(FMaterialAttributeDefinitionMap::GetDefaultValue(InProperty));
 }
 
 #endif // #if WITH_EDITOR
diff --git a/Engine/Source/Runtime/Engine/Private/Materials/MaterialIRModuleBuilder.cpp b/Engine/Source/Runtime/Engine/Private/Materials/MaterialIRModuleBuilder.cpp
--- a/Engine/Source/Runtime/Engine/Private/Materials/MaterialIRModuleBuilder.cpp
+++ b/Engine/Source/Runtime/Engine/Private/Materials/MaterialIRModuleBuilder.cpp
@@ -72,66 +72,68 @@ struct FAnalysisContext
 	}
 };
 
-static MIR::FValue* FindIRModuleOutputValue(const FMaterialIRModule& Module, MIR::EStage Stage, EMaterialProperty InProperty)
-{
-	for (const MIR::FSetMaterialOutput* Output : Module.GetOutputs(Stage))
-	{
-		if (Output->Property == InProperty)
-		{
-			return Output->Arg;
-		}
-	}
-	return nullptr;
-}
-
-static bool HasNonZeroIRModuleOutput(const FMaterialIRModule& Module, MIR::EStage Stage, EMaterialProperty InProperty)
-{
-	MIR::FValue* Value = FindIRModuleOutputValue(Module, Stage, InProperty);
-	return Value && !Value->AreAllExactlyZero();
-}
-
 struct FMaterialIRModuleBuilderImpl
 {
 	FMaterialIRModuleBuilder* Builder;
 	FMaterialIRModule* Module;
-	MIR::FEmitter* Emitter;
+	MIR::FEmitter Emitter;
 	TArray<FAnalysisContext> AnalysisContextStack;
 	FMaterialIRValueAnalyzer ValueAnalyzer;
 	MIR::FValue* DefaultMaterialAggregate;
 	FColorMaterialInput PreviewInput;
 
-	void Step_Initialize()
+	void Step_Initialize(FMaterialIRModuleBuilder* InBuilder, FMaterialIRModule* InModule)
 	{
+		this->Builder = InBuilder;
+		this->Module = InModule;
+
+		// Setup the Builder implementation
+		ValueAnalyzer.Setup(Builder->Material, Module, &Module->CompilationOutput,  Builder->TargetInsights);
+
+		// Empty the module and set it up
 		Module->Empty();
 		Module->ShaderPlatform = Builder->ShaderPlatform;
 		Module->TargetPlatform = Builder->TargetPlatform;
 		Module->FeatureLevel = Builder->FeatureLevel;
 		Module->QualityLevel = Builder->QualityLevel;
 		
-		Emitter->Initialize();
-		AnalysisContextStack.Emplace();
+		// Declare a entry point to evaluate the vertex stage.
+		Module->AddEntryPoint(TEXTVIEW("VertexStage"), MIR::Stage_Vertex, 1);
+
+		// Declare the entry points to evaluate both the pixel and compute stages.
+		int32 Num = UMaterialAggregate::GetMaterialAttributesProperties().Num();
+		Module->AddEntryPoint(TEXTVIEW("PixelStage"), MIR::Stage_Pixel, Num);
+		Module->AddEntryPoint(TEXTVIEW("ComputeStage"), MIR::Stage_Compute, Num);
 
-		DefaultMaterialAggregate = Emitter->Aggregate(UMaterialAggregate::GetMaterialAttributes());
+		// Setup the emitter and initialize it
+		Emitter.BuilderImpl = this;
+		Emitter.Material = Builder->Material;
+		Emitter.Module = Module;
+		Emitter.StaticParameterSet = &Builder->StaticParameters;
+		Emitter.Initialize();
 
+		// Create an IR value to hold the material attributes aggregate default
+		DefaultMaterialAggregate = Emitter.Aggregate(UMaterialAggregate::GetMaterialAttributes());
+
+		// Push the root analysis context
+		AnalysisContextStack.Emplace();
+
+		// Set the preview input expression
 		PreviewInput.Expression = Builder->PreviewExpression;
 	}
 
-	void Step_GenerateOutputInstructions()
+	void Step_PushRootExpressionDependencies()
 	{
 		//If we are processing from the Material Attributes output node, push the dependency here.
-		const bool bUseMaterialAttributes = Builder->Material->bUseMaterialAttributes;
-		if (bUseMaterialAttributes)
+		if (Builder->Material->bUseMaterialAttributes)
 		{
 			FMaterialInputDescription Input;
 			GetExpressionInputDescription(MP_MaterialAttributes, Input);
 			PushDependency(Input.Input->Expression);
 		}
-
-		// Helper local function to process a single material attribute output.
-		auto PrepareSingleMaterialAttribute = [&](EMaterialProperty Property)
+		else
 		{
-			//If we are not using the Material Attributes output node, then process the properties as normal
-			if (!bUseMaterialAttributes)
+			for (EMaterialProperty Property : UMaterialAggregate::GetMaterialAttributesProperties())
 			{
 				// Read the material input associated to this property
 				FMaterialInputDescription Input;
@@ -140,20 +142,10 @@ struct FMaterialIRModuleBuilderImpl
 				// Push the connected expression to this material attribute input as a dependency, if any.
 				PushDependency(Input.Input->Expression);
 			}
-
-			// Create the set material output instruction (leave it unset for now; it is set in Step_FlowValuesIntoMaterialOutputs)
-			// Note: Material Attributes need to still map the output correctly at the end so this is always set.
-			Emitter->SetMaterialOutput(Property, nullptr);
-		};
-
-		// Prepare the material attributes.
-		for (EMaterialProperty Property : UMaterialAggregate::GetMaterialAttributesProperties())
-		{
-			PrepareSingleMaterialAttribute(Property);
 		}
 	}
 
-	void Step_BuildMaterialExpressionsToIRGraph()
+	bool Step_BuildMaterialExpressionsToIRGraph()
 	{
 		while (true)
 		{
@@ -181,29 +173,31 @@ struct FMaterialIRModuleBuilderImpl
 				break;
 			}
 		}
+
+		return Module->IsValid();
 	}
 
 	void BuildTopMaterialExpression()
 	{
 		FAnalysisContext& CurrContext = AnalysisContextStack.Last();
-		Emitter->Expression = CurrContext.ExpressionStack.Last();
+		Emitter.Expression = CurrContext.ExpressionStack.Last();
 
 		// If expression is clean, nothing to be done.
-		if (CurrContext.BuiltExpressions.Contains(Emitter->Expression))
+		if (CurrContext.BuiltExpressions.Contains(Emitter.Expression))
 		{
 			CurrContext.ExpressionStack.Pop(EAllowShrinking::No);
 			return;
 		}
 
 		// Push to the expression stack all dependencies that still need to be analyzed.
-		for (FExpressionInputIterator It{ Emitter->Expression }; It; ++It)
+		for (FExpressionInputIterator It{ Emitter.Expression }; It; ++It)
 		{
 			PushDependency(It.Input->Expression);
 		}
 
 		// Named reroute usage nodes should make add their declaration as dependency, as they
 		// do not have any direct input (they simply forward to the declaration input).
-		if (auto NamedReroute = Cast<UMaterialExpressionNamedRerouteUsage>(Emitter->Expression))
+		if (auto NamedReroute = Cast<UMaterialExpressionNamedRerouteUsage>(Emitter.Expression))
 		{
 			if (NamedReroute->IsDeclarationValid())
 			{
@@ -212,24 +206,24 @@ struct FMaterialIRModuleBuilderImpl
 		}
 
 		// If on top of the stack there's a different expression, we have a dependency to analyze first.
-		if (CurrContext.ExpressionStack.Last() != Emitter->Expression)
+		if (CurrContext.ExpressionStack.Last() != Emitter.Expression)
 		{
 			return;
 		}
 
 		// Take the top expression out of the stack as ready for analysis. Also mark it as built.
 		CurrContext.ExpressionStack.Pop();
-		CurrContext.BuiltExpressions.Add(Emitter->Expression);
+		CurrContext.BuiltExpressions.Add(Emitter.Expression);
 
 		// Flow the value into this expression's inputs from their connected outputs.
-		for (FExpressionInputIterator It{ Emitter->Expression}; It; ++It)
+		for (FExpressionInputIterator It{ Emitter.Expression}; It; ++It)
 		{
 			FlowValueThroughConnection(It.Input);
 		}
 		
 		// Usage reroute nodes should forward the value coming in from their declaration input (since
 		// they don't have an input of their own).
-		if (auto NamedReroute = Cast<UMaterialExpressionNamedRerouteUsage>(Emitter->Expression))
+		if (auto NamedReroute = Cast<UMaterialExpressionNamedRerouteUsage>(Emitter.Expression))
 		{
 			if (NamedReroute->IsDeclarationValid())
 			{
@@ -237,7 +231,7 @@ struct FMaterialIRModuleBuilderImpl
 			}
 		}
 
-		if (auto Call = Cast<UMaterialExpressionMaterialFunctionCall>(Emitter->Expression))
+		if (auto Call = Cast<UMaterialExpressionMaterialFunctionCall>(Emitter.Expression))
 		{
 			// Function calls are handled internally as they manipulate the analysis context stack.
 			PushFunctionCall(Call);
@@ -246,10 +240,10 @@ struct FMaterialIRModuleBuilderImpl
 		{
 			// Invoke the expression build function. This will perform semantic analysis, error reporting and
 			// emit IR values for its outputs (which will flow into connected expressions inputs).
-			Emitter->Expression->Build(*Emitter);
+			Emitter.Expression->Build(Emitter);
 
 			// Populate the insight information about this expression pins.
-			AddExpressionConnectionInsights(Emitter->Expression);
+			AddExpressionConnectionInsights(Emitter.Expression);
 		}
 	}
 	
@@ -297,18 +291,18 @@ struct FMaterialIRModuleBuilderImpl
 
 			if (FunctionInput.ExpressionInput->bUsePreviewValueAsDefault)
 			{
-				MIR::FValue* Value = Emitter->TryInput(Call->GetInput(i));
-				CallInputValues[i] = Value ? Emitter->Cast(Value, Type) : nullptr;
+				MIR::FValue* Value = Emitter.TryInput(Call->GetInput(i));
+				CallInputValues[i] = Value ? Emitter.Cast(Value, Type) : nullptr;
 			}
 			else
 			{
-				MIR::FValue* Value = Emitter->Input(Call->GetInput(i));
-				CallInputValues[i] = Emitter->Cast(Value, Type);
+				MIR::FValue* Value = Emitter.Input(Call->GetInput(i));
+				CallInputValues[i] = Emitter.Cast(Value, Type);
 			}
 		}
 
 		// If some error occurred (e.g. some function input wasn't linked in) early out.
-		if (Emitter->CurrentExpressionHasErrors())
+		if (Emitter.CurrentExpressionHasErrors())
 		{
 			return;
 		}
@@ -359,7 +353,7 @@ struct FMaterialIRModuleBuilderImpl
 			FFunctionExpressionOutput& FunctionOutput = Call->FunctionOutputs[i];
 
 			// Get the value flowing into the function output inside the function in the current context.
-			MIR::FValue* Value = Emitter->Input(FunctionOutput.ExpressionOutput->GetInput(0));
+			MIR::FValue* Value = Emitter.Input(FunctionOutput.ExpressionOutput->GetInput(0));
 
 			// And flow it to the relative function *call* output in the parent context.
 			ParentContext.SetOutputValue(Call->GetOutput(i), Value);
@@ -372,23 +366,21 @@ struct FMaterialIRModuleBuilderImpl
 		AddExpressionConnectionInsights(Call);
 	}
 
-	void Step_FlowValuesIntoMaterialOutputs()
+	bool Step_EmitSetMaterialPropertyInstructions()
 	{
-		// Get the material attribute input info.
-		const bool bUseMaterialAttributes = Builder->Material->bUseMaterialAttributes;
-		FMaterialInputDescription MaterialAttributesInput;
-		MIR::FValue* MaterialAttributesValue{};
-		if (bUseMaterialAttributes)
+		// First, if the material is flagged to use the material attributes aggregate, read its value now
+		// so that we can extract its individual attributes later.
+		MIR::FValue* MaterialAttributesValue = nullptr;
+		if (Builder->Material->bUseMaterialAttributes)
 		{
-			ensure(GetExpressionInputDescription(MP_MaterialAttributes, MaterialAttributesInput));
-
-			FlowValueThroughConnection(MaterialAttributesInput.Input);
+			FMaterialInputDescription InputDesc;
+			GetExpressionInputDescription(MP_MaterialAttributes, InputDesc);
 
 			// Fetch the value from the material attributes input.
-			MaterialAttributesValue = MIR::Internal::FetchValueFromExpressionInput(this, MaterialAttributesInput.Input);
+			FlowValueThroughConnection(InputDesc.Input);
+			MaterialAttributesValue = MIR::Internal::FetchValueFromExpressionInput(this, InputDesc.Input);
 			if (!MaterialAttributesValue)
 			{
-				// If none there, use the default aggregate
 				MaterialAttributesValue = DefaultMaterialAggregate;
 			}
 
@@ -396,94 +388,83 @@ struct FMaterialIRModuleBuilderImpl
 			check(!MaterialAttributesValue->IsPoison() && MaterialAttributesValue->Type.AsAggregate() == UMaterialAggregate::GetMaterialAttributes());
 		}
 
-		for (uint32 StageIndex = 0; StageIndex < MIR::NumStages; ++StageIndex)
+		for (EMaterialProperty Property : UMaterialAggregate::GetMaterialAttributesProperties())
 		{
-			for (MIR::FSetMaterialOutput* Output : Module->Outputs[StageIndex])
-			{
-				FMaterialInputDescription Input;
-				ensure(GetExpressionInputDescription(Output->Property, Input));
+			// Get the input description of this material property (input, type, default value, etc).
+			FMaterialInputDescription InputDesc;
+			GetExpressionInputDescription(Property, InputDesc);
 
-				if (bUseMaterialAttributes)
-				{
-					// If material is in material attributes mode, extract this property's attribute from the material attributes aggregate value
-					// and manually flow it into the this property's material expression input pin.
-					MIR::FValue* PropertyValue = Emitter->Subscript(MaterialAttributesValue, UMaterialAggregate::MaterialPropertyToAttributeIndex(Output->Property));
-					MIR::Internal::BindValueToExpressionInput(this, Input.Input, PropertyValue);
-				}
+			// This holds the value being set to this property.
+			MIR::FValue* PropertyValue = nullptr;
 
-				// Flow, if any, the value coming in through the connection to the material attribute pin.
-				FlowValueThroughConnection(Input.Input);
+			// If the material attributes value is valid, extract this property attribute from the material attributes aggregate value
+			// and manually flow it into the this property material expression input pin.
+			if (MaterialAttributesValue)
+			{
+				PropertyValue = Emitter.Subscript(MaterialAttributesValue, UMaterialAggregate::MaterialPropertyToAttributeIndex(Property));
+			}
+			else
+			{
+				// Otherwise grab the value from the individual attribute pin.
+				FlowValueThroughConnection(InputDesc.Input);
+				PropertyValue = MIR::Internal::FetchValueFromExpressionInput(this, InputDesc.Input);
+			}
 
-				if (MIR::FValue* Value = MIR::Internal::FetchValueFromExpressionInput(this, Input.Input))
+			if (PropertyValue)
+			{
+				// If this property is the emissive color and we're previewing the material, apply gamma correction to the previewed value.
+				if (Property == MP_EmissiveColor && PreviewInput.IsConnected())
 				{
-					// If this property is the emissive color and we're previewing the material, apply gamma correction to the previewed value.
-					if (Output->Property == MP_EmissiveColor && PreviewInput.IsConnected())
-					{
-						MIR::FValue* Zero = Emitter->ConstantZero(MIR::EScalarKind::Float);
-
-						// Get preview expression back into gamma corrected space, as DrawTile does not do this adjustment.
-						Value = Emitter->Pow(Emitter->Max(Value, Zero), Emitter->ConstantFloat(1.0f / 2.2f));
+					MIR::FValue* Zero = Emitter.ConstantZero(MIR::EScalarKind::Float);
 
-						// Preview should display scalars as red, so if this is a scalar, create a vector padded with zeroes
-						if (Value->Type.IsScalar())
-						{
-							Value = Emitter->Vector3(Value, Zero, Zero);
-						}
+					// Get preview expression back into gamma corrected space, as DrawTile does not do this adjustment.
+					PropertyValue = Emitter.Pow(Emitter.Max(PropertyValue, Zero), Emitter.ConstantFloat(1.0f / 2.2f));
 
-						Output->Arg = Value;
-					}
-					else
+					// Preview should display scalars as red, so if this is a scalar, create a vector padded with zeroes
+					if (PropertyValue->Type.IsScalar())
 					{
-						// If a value is flowing in through the connection, cast it to this material attribute type and assign it.
-						// Special case for shading model, because UE::Shader::EValueType doesn't include a native shading model type, and uses int instead.
-						MIR::FType OutputArgType = Output->Property == MP_ShadingModel ? MIR::FType::MakeShadingModel() : MIR::FType::FromShaderType(Input.Type);
-						Output->Arg = Emitter->Cast(Value, OutputArgType);
+						PropertyValue = Emitter.Vector3(PropertyValue, Zero, Zero);
 					}
 				}
-				else if (Input.bUseConstant)
-				{
-					// If input is marked to use constant, assign this output to the specified constant value.
-					Output->Arg = Emitter->ConstantFromShaderValue(Input.ConstantValue);
-				}
 				else
 				{
-					// Otherwise, fallback to assigning this material output to its default value.
-					Output->Arg = Emitter->Subscript(DefaultMaterialAggregate, UMaterialAggregate::MaterialPropertyToAttributeIndex(Output->Property));
+					// If a value is flowing in through the connection, cast it to this material attribute type and assign it.
+					// Special case for shading model, because UE::Shader::EValueType doesn't include a native shading model type, and uses int instead.
+					MIR::FType OutputArgType = Property == MP_ShadingModel ? MIR::FType::MakeShadingModel() : MIR::FType::FromShaderType(InputDesc.Type);
+					PropertyValue = Emitter.Cast(PropertyValue, OutputArgType);
 				}
+			}
+			else if (InputDesc.bUseConstant)
+			{
+				// If input is marked to use constant, assign this output to the specified constant value.
+				PropertyValue = Emitter.ConstantFromShaderValue(InputDesc.ConstantValue);
+			}
+			else
+			{
+				// Otherwise, fallback to assigning this material output to its default value.
+				PropertyValue = Emitter.Subscript(DefaultMaterialAggregate, UMaterialAggregate::MaterialPropertyToAttributeIndex(Property));
+			}
+			
+			// Quit if some error occurred in the operations above.
+			if (!Module->IsValid())
+			{
+				return false;
+			}
 
-				// The value being set to this material output is now valid.
-				check(Output->Arg);
-
-				// Connect output properties with substrate parameter if a FrontMaterial property is provided
-				if (Output->Property == MP_FrontMaterial)
-				{
-					if (MIR::FPromoteSubstrateParameter* SubstrateParameter = Output->Arg->As<MIR::FPromoteSubstrateParameter>())
-					{
-						for (int32 PropertyIndex = 0; PropertyIndex < MP_MAX; ++PropertyIndex)
-						{
-							SubstrateParameter->PropertyArgs[PropertyIndex] = FindIRModuleOutputValue(*Module, MIR::EStage::Stage_Pixel, static_cast<EMaterialProperty>(PropertyIndex));
-						}
-
-						// Store separate value for normal vector transformed into world-space for Substrate
-						MIR::FValue* TangentSpaceNormal = SubstrateParameter->PropertyArgs[MP_Normal];
-						MIR::FValue* WorldSpaceNormal = Emitter->InlineHLSL(TangentSpaceNormal->Type, TEXT("TransformTangentNormalToWorld(Parameters.TangentToWorld, $0)"), { TangentSpaceNormal });
-						SubstrateParameter->WorldSpaceTangentsAndNormals[0] = WorldSpaceNormal;
-						SubstrateParameter->WorldSpaceTangentsAndNormals[1] = WorldSpaceNormal;
-						SubstrateParameter->WorldSpaceTangentsAndNormals[2] = WorldSpaceNormal;
-						SubstrateParameter->WorldSpaceTangentsAndNormals[3] = WorldSpaceNormal;
-					}
-				}
+			// The value being set to this material output is now valid.
+			check(PropertyValue);
 
-				//
-				if (!Module->IsValid())
-				{
-					return;
-				}
+			// Set this property value onto the module
+			Module->SetPropertyValue(Property, PropertyValue);
 
-				// Finally, push this connection insight
-				PushConnectionInsight(Builder->Material, (int)Output->Property, Input.Input->Expression, Input.Input->OutputIndex, Output->Arg->Type);
-			}
+			// Emit the SetMaterialProperty instruction
+			Emitter.SetMaterialOutput(Property, PropertyValue);
+			
+			// Finally, push this connection insight
+			PushConnectionInsight(Builder->Material, (int32)Property, InputDesc.Input->Expression, InputDesc.Input->OutputIndex, PropertyValue->Type);
 		}
+
+		return Module->IsValid();
 	}
 
 	bool GetExpressionInputDescription(EMaterialProperty Property, FMaterialInputDescription& Input)
@@ -515,21 +496,27 @@ struct FMaterialIRModuleBuilderImpl
 		}
 	}
 
-	void Step_AnalyzeIRGraph()
+	bool Step_AnalyzeIRGraph()
 	{
 		TArray<MIR::FValue*> ValueStack{};
+		TSet<MIR::FValue*> VisitedValues{};
 
-		for (uint32 StageIndex = 0; StageIndex < MIR::NumStages; ++StageIndex)
+		// Analyze the nodes in each entry point.
+		for (int32 EntryPointIndex = 0; EntryPointIndex < Module->GetNumEntryPoints(); ++EntryPointIndex)
 		{
-			MIR::EStage CurrentStage = (MIR::EStage)StageIndex;
-
-			// Clear the value stack but preserve its allocated memory.
+			// Reset bookkeeping to process new output subgraph
 			ValueStack.Empty(ValueStack.Max());
+			VisitedValues.Empty(Module->Values.Num());
 
-			// Push each output in the current stage to the value stack.
-			for (MIR::FSetMaterialOutput* Output : Module->Outputs[StageIndex])
+			FMaterialIRModule::FEntryPoint& EntryPoint = Module->GetEntryPoint(EntryPointIndex);
+
+			// Push this SetOutput instruction the value stack for processing.
+			for (MIR::FValue* Output : EntryPoint.Outputs)
 			{
-				ValueStack.Push(Output);
+				if (Output)
+				{
+					ValueStack.Push(Output);
+				}
 			}
 
 			// Process until the value stack is empty.
@@ -540,17 +527,17 @@ struct FMaterialIRModuleBuilderImpl
 				// Module building should have interrupted before if poison values were generated.
 				check(!Value->IsPoison());
 
-				// If this instruction has already been analyzed for this stage, nothing else is left to do for it. Continue.
-				if (Value->IsAnalyzed(CurrentStage))
+				// If this instruction has already been analyzed for this entry point, nothing else is left to do for it. Continue.
+				if (VisitedValues.Contains(Value))
 				{
 					ValueStack.Pop();
 					continue;
 				}
 
 				// Before analyzing this value, make sure all used values are analyzed first.
-				for (MIR::FValue* Use : Value->GetUsesForStage(CurrentStage))
+				for (MIR::FValue* Use : Value->GetUsesForStage(EntryPoint.Stage))
 				{
-					if (Use && !Use->IsAnalyzed(CurrentStage))
+					if (Use && !VisitedValues.Contains(Use))
 					{
 						ValueStack.Push(Use);
 					}
@@ -566,114 +553,87 @@ struct FMaterialIRModuleBuilderImpl
 				ValueStack.Pop();
 
 				// Go through each use instruction and increment its counter of users (this instruction).
-				for (MIR::FValue* Use : Value->GetUsesForStage(CurrentStage))
+				for (MIR::FValue* Use : Value->GetUsesForStage(EntryPoint.Stage))
 				{
 					// If this used value is an instruction, update its counter of users (in current stage).
 					if (MIR::FInstruction* UseInstr = MIR::AsInstruction(Use))
 					{
-						UseInstr->NumUsers[StageIndex] += 1;
+						UseInstr->Linkage[EntryPointIndex].NumUsers += 1;
 					}
 				}
 
-				// If this is the first time this value is analyzed, let the analyzer process it.
+				// If this is the first time this value is analyzed in any entry point, let the analyzer process it.
 				// Note that individual value processing is independent from the stage it runs on so we can perform it only once.
-				if ((Value->Flags & MIR::EValueFlags::AnalyzedInAnyStageMask) == MIR::EValueFlags::None)
+				if (!Value->HasFlags(MIR::EValueFlags::AnalyzedInAnyStage))
 				{
-					ValueAnalyzer.Analyze(Value);
-				}
+					Value->SetFlags(MIR::EValueFlags::AnalyzedInAnyStage);
 
-				ValueAnalyzer.AnalyzeInStage(Value, CurrentStage);
+					// Flow the graph properties downstream from the value's uses into this value.
+					for (MIR::FValue* Use : Value->GetUses())
+					{
+						if (Use)
+						{
+							Value->GraphProperties |= Use->GraphProperties;
+						}
+					}
 
-				// Mark the used instruction as analyzed for this stage.
-				Value->Flags |= MIR::EValueFlags(1 << StageIndex);
-			}
-		}
-	}
+					// Allocate the entry points linkage information for this instruction.
+					if (MIR::FInstruction* Instr = MIR::AsInstruction(Value))
+					{
+						Instr->Linkage = Module->AllocateArray<MIR::FInstructionLinkage>(Module->GetNumEntryPoints());
+						MIR::ZeroArray(Instr->Linkage);
+					}
 
-	void Step_ConsolidateEnvironmentDefines()
-	{
-		// Keep defines if a combined condition is met. Otherwise, remove them from the environemnt defines set.
-		auto KeepDefineConditionally = [this](FName Name, bool bConditionToKeepDefine) -> void
-			{
-				if (!bConditionToKeepDefine)
-				{
-					ValueAnalyzer.EnvironmentDefines.Remove(Name);
+					// Then analyze the instruction based on its kind.
+					ValueAnalyzer.Analyze(Value);
 				}
-			};
 
-		KeepDefineConditionally(TEXT("USES_PER_INSTANCE_CUSTOM_DATA"), ValueAnalyzer.Material->bUsedWithInstancedStaticMeshes);
-		KeepDefineConditionally(TEXT("NEEDS_PER_INSTANCE_RANDOM_PS"), ValueAnalyzer.Material->bUsedWithInstancedStaticMeshes);
-		KeepDefineConditionally(TEXT("USES_PER_INSTANCE_FADE_AMOUNT"), ValueAnalyzer.Material->bUsedWithInstancedStaticMeshes);
-
-		// Derive additional defines from the final state of the IR analysis
-		DeriveEnvironmentDefinesFromAnalysisState();
-
-		// Move final environment defines from analyzer into output module
-		Module->EnvironmentDefines = MoveTemp(ValueAnalyzer.EnvironmentDefines);
-	}
-
-	void DeriveEnvironmentDefinesFromAnalysisState()
-	{
-		// Virtual texturing defines
-		int32 NumVirtualTextureFeedbackRequests = 0;
-		for (const FMaterialIRValueAnalyzer::FVTStackEntry& VTStack : ValueAnalyzer.VTStacks)
-		{
-			if (VTStack.bGenerateFeedback)
-			{
-				++NumVirtualTextureFeedbackRequests;
-			}
-		}
-		Module->CompilationOutput.NumVirtualTextureFeedbackRequests = NumVirtualTextureFeedbackRequests;
-	}
+				// Analyze this instruction in this entry point's stage if it's the first time it's encountered.
+				MIR::EValueFlags StageFlag = MIR::EValueFlags(1 << (int32)EntryPoint.Stage);
+				if (!Value->HasFlags(StageFlag))
+				{
+					Value->SetFlags(StageFlag);
+					ValueAnalyzer.AnalyzeInStage(Value, EntryPoint.Stage);
+				}
 
-	void Step_AnalyzeBuiltinDefines()
-	{
-		// Match various defines against the material configuration
-		if (Module->EnvironmentDefines.Contains(TEXT("MIR.SceneDepth")))
-		{
-			// @todo-jason.hoerner - Support for material instance blend mode overrides needed.
-			if (ValueAnalyzer.Material->MaterialDomain != MD_PostProcess && !IsTranslucentBlendMode(ValueAnalyzer.Material->BlendMode))
-			{
-				Module->AddError(nullptr, TEXT("Only transparent or postprocess materials can read from scene depth."));
+				// Mark the used instruction as analyzed for this entry point.
+				VisitedValues.Add(Value);
 			}
 		}
 
-		// Remove all environment defines that have the "MIR." prefix as they are not meant to propagate into the set of compiler environment defines.
-		TCHAR DefineMIRPrefix[5] = {};
-		for (auto Iter = Module->EnvironmentDefines.CreateIterator(); Iter; ++Iter)
-		{
-			Iter->ToStringTruncate(DefineMIRPrefix, UE_ARRAY_COUNT(DefineMIRPrefix));
-			if (FCString::Strncmp(DefineMIRPrefix, TEXT("MIR."), UE_ARRAY_COUNT(DefineMIRPrefix)) == 0)
-			{
-				Iter.RemoveCurrent();
-			}
-		}
+		return Module->IsValid();
 	}
 
 	void Step_LinkInstructions()
 	{
 		TArray<MIR::FInstruction*> InstructionStack{};
 
-		for (uint32 StageIndex = 0; StageIndex < MIR::NumStages; ++StageIndex)
+		for (int32 EntryPointIndex = 0; EntryPointIndex < Module->GetNumEntryPoints(); ++EntryPointIndex)
 		{
 			// This function walks the instruction graph and puts each instruction into the inner most possible block.
 			InstructionStack.Empty(InstructionStack.Max());
 
-			// Initialize each output's block in this stage and add it to the instruction stack.
-			for (MIR::FSetMaterialOutput* Output : Module->Outputs[StageIndex])
-			{
-				Output->Block[StageIndex] = &Module->RootBlocks[StageIndex];
+			FMaterialIRModule::FEntryPoint& Procedure = Module->GetEntryPoint(EntryPointIndex);
 
-				InstructionStack.Add(Output);
+			// Push all entry point final outputs onto the instruction stack to begin.
+			// Note: the first output on the stack will be the first to be evaluated in the entry point root block.
+			for (MIR::FValue* Output : Procedure.Outputs)
+			{
+				if (MIR::FInstruction* Instr = MIR::AsInstruction(Output))
+				{
+					Instr->Linkage[EntryPointIndex].Block = &Procedure.RootBlock;
+					InstructionStack.Push(Instr);
+				}
 			}
 
 			while (!InstructionStack.IsEmpty())
 			{
 				MIR::FInstruction* Instr = InstructionStack.Pop();
+				MIR::FBlock* InstrBlock = Instr->Linkage[EntryPointIndex].Block;
 
 				// Push the instruction to its block in reverse order (push front)
-				Instr->Next[StageIndex] = Instr->Block[StageIndex] ->Instructions;
-				Instr->Block[StageIndex]->Instructions = Instr;
+				Instr->Linkage[EntryPointIndex].Next = InstrBlock->Instructions;
+				InstrBlock->Instructions = Instr;
 
 				//
 				if (Instr->As<MIR::FNop>())
@@ -681,38 +641,40 @@ struct FMaterialIRModuleBuilderImpl
 					continue;
 				}
 
-				TConstArrayView<MIR::FValue*> Uses = Instr->GetUsesForStage((MIR::EStage)StageIndex);
+				TConstArrayView<MIR::FValue*> Uses = Instr->GetUsesForStage(Procedure.Stage);
 				for (int32 UseIndex = 0; UseIndex < Uses.Num(); ++UseIndex)
 				{
-					MIR::FValue* Use = Uses[UseIndex];
-					MIR::FInstruction* UseInstr = MIR::AsInstruction(Use);
+					MIR::FInstruction* UseInstr = MIR::AsInstruction(Uses[UseIndex]);
+
 					if (!UseInstr)
 					{
 						continue;
 					}
 
 					// Get the block into which the dependency instruction should go.
-					MIR::FBlock* TargetBlock = Instr->GetDesiredBlockForUse((MIR::EStage)StageIndex, UseIndex);
+					MIR::FBlock* TargetBlock = Instr->GetTargetBlockForUse(EntryPointIndex, UseIndex);
 
 					// Update dependency's block to be a child of current instruction's block.
-					if (TargetBlock != Instr->Block[StageIndex])
+					if (TargetBlock != InstrBlock)
 					{
-						TargetBlock->Parent = Instr->Block[StageIndex];
-						TargetBlock->Level = Instr->Block[StageIndex]->Level + 1;
+						TargetBlock->Parent = InstrBlock;
+						TargetBlock->Level = InstrBlock->Level + 1;
 					}
 
 					// Set the dependency's block to the common block betwen its current block and this one.
-					UseInstr->Block[StageIndex] = UseInstr->Block[StageIndex]
-						? UseInstr->Block[StageIndex]->FindCommonParentWith(TargetBlock)
+					MIR::FInstructionLinkage& UseLinkage = UseInstr->Linkage[EntryPointIndex];
+
+					UseLinkage.Block = UseLinkage.Block
+						? UseLinkage.Block->FindCommonParentWith(TargetBlock)
 						: TargetBlock;
 
 					// Increase the number of times this dependency instruction has been considered.
 					// When all of its users have processed, we can carry on visiting this instruction.
-					++UseInstr->NumProcessedUsers[StageIndex];
-					check(UseInstr->NumProcessedUsers[StageIndex] <= UseInstr->NumUsers[StageIndex]);
+					++UseLinkage.NumProcessedUsers;
+					check(UseLinkage.NumProcessedUsers <= UseLinkage.NumUsers);
 
 					// If all dependants have been processed, we can carry the processing from this dependency.
-					if (UseInstr->NumProcessedUsers[StageIndex] == UseInstr->NumUsers[StageIndex])
+					if (UseLinkage.NumProcessedUsers == UseLinkage.NumUsers)
 					{
 						InstructionStack.Push(UseInstr);
 					}
@@ -721,11 +683,77 @@ struct FMaterialIRModuleBuilderImpl
 		}
 	}
 
-	void Step_Finalize()
+	// This final step generates all output non-IR-graph data such as setting up the CompilationOutput and the shader environment defines.
+	void Step_FinalizeArtifacts()
+	{
+		ConsolidateEnvironmentDefines();
+		AnalyzeBuiltinDefines();
+		WriteCompilationOutput();
+		GenerateDebugInsights();
+	}
+
+	void ConsolidateEnvironmentDefines()
+	{
+		// Keep defines if a combined condition is met. Otherwise, remove them from the environemnt defines set.
+		auto KeepDefineConditionally = [this](FName Name, bool bConditionToKeepDefine) -> void
+		{
+			if (!bConditionToKeepDefine)
+			{
+				ValueAnalyzer.EnvironmentDefines.Remove(Name);
+			}
+		};
+
+		KeepDefineConditionally(TEXT("USES_PER_INSTANCE_CUSTOM_DATA"), ValueAnalyzer.Material->bUsedWithInstancedStaticMeshes);
+		KeepDefineConditionally(TEXT("NEEDS_PER_INSTANCE_RANDOM_PS"), ValueAnalyzer.Material->bUsedWithInstancedStaticMeshes);
+		KeepDefineConditionally(TEXT("USES_PER_INSTANCE_FADE_AMOUNT"), ValueAnalyzer.Material->bUsedWithInstancedStaticMeshes);
+
+		// Derive additional defines from the final state of the IR analysis
+		{
+			// Virtual texturing defines
+			int32 NumVirtualTextureFeedbackRequests = 0;
+			for (const FMaterialIRValueAnalyzer::FVTStackEntry& VTStack : ValueAnalyzer.VTStacks)
+			{
+				if (VTStack.bGenerateFeedback)
+				{
+					++NumVirtualTextureFeedbackRequests;
+				}
+			}
+			Module->CompilationOutput.NumVirtualTextureFeedbackRequests = NumVirtualTextureFeedbackRequests;
+		}
+
+		// Move final environment defines from analyzer into output module
+		Module->EnvironmentDefines = MoveTemp(ValueAnalyzer.EnvironmentDefines);
+	}
+
+	void AnalyzeBuiltinDefines()
+	{
+		// Match various defines against the material configuration
+		if (Module->EnvironmentDefines.Contains(TEXT("MIR.SceneDepth")))
+		{
+			// @todo-jason.hoerner - Support for material instance blend mode overrides needed.
+			if (ValueAnalyzer.Material->MaterialDomain != MD_PostProcess && !IsTranslucentBlendMode(ValueAnalyzer.Material->BlendMode))
+			{
+				Module->AddError(nullptr, TEXT("Only transparent or postprocess materials can read from scene depth."));
+			}
+		}
+
+		// Remove all environment defines that have the "MIR." prefix as they are not meant to propagate into the set of compiler environment defines.
+		TCHAR DefineMIRPrefix[5] = {};
+		for (auto Iter = Module->EnvironmentDefines.CreateIterator(); Iter; ++Iter)
+		{
+			Iter->ToStringTruncate(DefineMIRPrefix, UE_ARRAY_COUNT(DefineMIRPrefix));
+			if (FCString::Strncmp(DefineMIRPrefix, TEXT("MIR."), UE_ARRAY_COUNT(DefineMIRPrefix)) == 0)
+			{
+				Iter.RemoveCurrent();
+			}
+		}
+	}
+
+	void WriteCompilationOutput()
 	{
 		FMaterialCompilationOutput& CompilationOutput = Module->CompilationOutput;
-		CompilationOutput.NumUsedUVScalars = Module->Statistics.NumPixelTexCoords * 2;
 
+		CompilationOutput.NumUsedUVScalars = Module->Statistics.NumPixelTexCoords * 2;
 		CompilationOutput.UniformExpressionSet.SetParameterCollections(Module->ParameterCollections);
 
 		int32 NumPostProcessInputs = CompilationOutput.GetNumPostProcessInputsUsed();
@@ -734,8 +762,8 @@ struct FMaterialIRModuleBuilderImpl
 			Module->AddError(nullptr, FString::Printf(TEXT("Maximum Scene Texture post process inputs exceeded (%d > %d), between SceneTexture nodes with PostProcessInputs or UserSceneTexture nodes."), NumPostProcessInputs, kPostProcessMaterialInputCountMax));
 		}
 
-		CompilationOutput.bUsesPixelDepthOffset = HasNonZeroIRModuleOutput(*Module, MIR::EStage::Stage_Pixel, MP_PixelDepthOffset);
-		CompilationOutput.bUsesWorldPositionOffset = HasNonZeroIRModuleOutput(*Module, MIR::EStage::Stage_Vertex, MP_WorldPositionOffset);
+		CompilationOutput.bUsesPixelDepthOffset = MaterialPropertyHasNonZeroValue(MP_PixelDepthOffset);
+		CompilationOutput.bUsesWorldPositionOffset = MaterialPropertyHasNonZeroValue(MP_WorldPositionOffset);
 
 		// @todo-jason.hoerner - Fill in CompilationOutput.bUsesDisplacement.
 		// @todo-jason.hoerner - Custom output support.
@@ -762,7 +790,7 @@ struct FMaterialIRModuleBuilderImpl
 			CompilationOutput,
 			Module->GetCompiledShadingModels(),
 			ValueAnalyzer.Material->GetBlendMode(),					// @todo-jason.hoerner - we need to handle material instances overriding the blend mode
-			Module->IsMaterialPropertyUsed(MIR::EStage::Stage_Pixel, MP_FrontMaterial),
+			Module->IsMaterialPropertyUsed(MP_FrontMaterial),
 			Module->GetShaderPlatform(),
 			ValidationErrors);
 
@@ -771,10 +799,39 @@ struct FMaterialIRModuleBuilderImpl
 			Module->AddError(nullptr, *ValidationError);
 		}
 	}
+
+	void GenerateDebugInsights()
+	{
+		if (!Builder->Material->MaterialInsight)
+		{
+			return;
+		}
+
+		// Dump the module IR to string and store it inside the material insights.
+		Builder->Material->MaterialInsight.Get()->IRString = MIR::DebugDumpIR(Builder->Material->GetFullName(), *Module);
+		
+		// Dump the requested debugging information
+		switch (CVarMaterialIRDebugDumpLevel.GetValueOnGameThread())
+		{
+			case 2:
+			{
+				MIR::DebugDumpIRUseGraph(*Module);
+				// fallthrough
+			}
+
+			case 1:
+			{
+				// Save the dump to file
+				FString FilePath = FPaths::Combine(FPaths::ProjectSavedDir(), "Materials", TEXT("IRDump.txt"));
+				FFileHelper::SaveStringToFile(Builder->Material->MaterialInsight.Get()->IRString, *FilePath);
+				// fallthrough
+			}
+		}
+	}
 	
-	/* Other functions */
+	/* Auxiliary functions */
 
-	//
+	// Adds an expression connection insight to the MaterialInsights instance, if any.
 	void AddExpressionConnectionInsights(UMaterialExpression* Expression)
 	{
 		if (!Builder->TargetInsights)
@@ -792,7 +849,7 @@ struct FMaterialIRModuleBuilderImpl
 		}
 	}
 	
-	//
+	// Adds a connection insight to the MaterialInsights instance, if any.
 	void PushConnectionInsight(const UObject* InputObject, int InputIndex, const UMaterialExpression* OutputExpression, int OutputIndex, MIR::FType Type)
 	{
 		if (!Builder->TargetInsights || Type.IsPoison())
@@ -810,82 +867,88 @@ struct FMaterialIRModuleBuilderImpl
 		
 		Builder->TargetInsights->ConnectionInsights.Push(Insight);
 	}
+
+	// Returns whether a material property (e.g. MP_BaseColor) has a value assigned that isn't a constant zero.
+	// Used to determine if a property is being used.
+	bool MaterialPropertyHasNonZeroValue(EMaterialProperty InProperty)
+	{
+		return Module->GetPropertyValue(InProperty) && !Module->GetPropertyValue(InProperty)->AreAllExactlyZero();
+	}
 };
 
+namespace MIR::Internal {
+
+	FValue* FetchValueFromExpressionInput(FMaterialIRModuleBuilderImpl* Builder, const FExpressionInput* Input)
+	{
+		return Builder->AnalysisContextStack.Last().GetInputValue(Input);
+	}
+
+	void BindValueToExpressionInput(FMaterialIRModuleBuilderImpl* Builder, const FExpressionInput* Input, FValue* Value)
+	{
+		Builder->AnalysisContextStack.Last().SetInputValue(Input, Value);
+	}
+
+	void BindValueToExpressionOutput(FMaterialIRModuleBuilderImpl* Builder, const FExpressionOutput* Output, FValue* Value)
+	{
+		Builder->AnalysisContextStack.Last().SetOutputValue(Output, Value);
+	}
+
+} // namespace MIR::Internal
+
 bool FMaterialIRModuleBuilder::Build(FMaterialIRModule* TargetModule)
 {
+	FMaterialIRModuleBuilderImpl Impl;
+
 	FMemMark MemMark { FMemStack::Get() };
-	FMaterialIRModuleBuilderImpl Impl{ this, TargetModule };
 
-	// Setup the emitter
-	MIR::FEmitter Emitter;
-	Emitter.BuilderImpl = &Impl;
-	Emitter.Material = Material;
-	Emitter.Module = TargetModule;
-	Emitter.StaticParameterSet = &StaticParameters;
-
-	// Setup the Builder implementation
-	Impl.Emitter = &Emitter;
-	Impl.ValueAnalyzer.Setup(Material, TargetModule, &TargetModule->CompilationOutput, TargetInsights);
-	
-	Impl.Step_Initialize();
-	Impl.Step_GenerateOutputInstructions();
-	Impl.Step_BuildMaterialExpressionsToIRGraph();
+	// Initialize the module to a blank slate, initialize the builder auxiliary data
+	// and the emitter for IR values emission.
+	Impl.Step_Initialize(this, TargetModule);
+
+	// Identify the material property output pins and push their value-producing
+	// expressions onto the analysis context stack to start crawling from them.
+	Impl.Step_PushRootExpressionDependencies();
 
-	if (!TargetModule->IsValid())
+	// Main step. It crawls the expression graph and calls the Build() function on each
+	// visited expression in order to emit the IR values that implement that expression
+	// semantics. At the end of this step the IR values graph has been built, but is still
+	// missing the root SetMaterialProperty instructions.
+	if(!Impl.Step_BuildMaterialExpressionsToIRGraph())
 	{
 		return false;
 	}
 
-	Impl.Step_FlowValuesIntoMaterialOutputs();
-
-	if (!TargetModule->IsValid())
+	// Materials dont have a final "output expression", so this step grabs the values flowing
+	// into the material property output pins (if any) and generates SetMaterialProperty instructions
+    // handling details such as default values, preview material expression, etc.
+	if (!Impl.Step_EmitSetMaterialPropertyInstructions())
 	{
 		return false;
 	}
 
-	Impl.Step_AnalyzeIRGraph();
-	Impl.Step_ConsolidateEnvironmentDefines();
-	Impl.Step_AnalyzeBuiltinDefines();
-	Impl.Step_LinkInstructions();
-	Impl.Step_Finalize();
-
-	check(Material->MaterialInsight.IsValid());
-	Material->MaterialInsight.Get()->IRString = MIR::DebugDumpIR(Material->GetFullName(), *TargetModule);
-
-	// Dump debugging information if requested 
-	switch (CVarMaterialIRDebugDumpLevel.GetValueOnGameThread())
+	// Now that the full IR graph has been produced, starting from the output instructions of each
+	// entry point, crawl the IR graph backwards in order to let each value analyze itself. A value
+	// is analyzed only after all its dependencies (its uses) have been analyzed first, so that when
+	// a value is analyzed it is guaranteed to have all the information to properly analyze itself.
+	// In this step is performed semantic analysis, where a value can potentially throw new errors depending
+	// on the semantic context it is placed in (for instance, an instruction can be executed only in
+	// specific stages will throw an error if it finds itself being executed in an incorrect stage.
+	if (!Impl.Step_AnalyzeIRGraph())
 	{
-		case 2: MIR::DebugDumpIRUseGraph(*TargetModule); // fallthrough
-		case 1:
-		{
-			// Save the dump to file
-			FString FilePath = FPaths::Combine(FPaths::ProjectSavedDir(), "Materials", TEXT("IRDump.txt"));
-			FFileHelper::SaveStringToFile(Material->MaterialInsight.Get()->IRString, *FilePath);
-			// fallthrough
-		}
+		return false;
 	}
 
-	return TargetModule->IsValid();
-}
-
-namespace MIR::Internal {
-
-FValue* FetchValueFromExpressionInput(FMaterialIRModuleBuilderImpl* Builder, const FExpressionInput* Input)
-{
-	return Builder->AnalysisContextStack.Last().GetInputValue(Input);
-}
+	// The IR graph has now been been fully produced and is valid. Proceed to link instructions together,
+	// placing each instruction into its own parent block. This is done in a way to put instructions in the
+	// narrowest possible scope that still puts them in an execution order that will occur after its dependencies
+	// have occurred.
+	Impl.Step_LinkInstructions();
 
-void BindValueToExpressionInput(FMaterialIRModuleBuilderImpl* Builder, const FExpressionInput* Input, FValue* Value)
-{
-	Builder->AnalysisContextStack.Last().SetInputValue(Input, Value);
-}
+	// Finally, populate all other non IR-graph artifacts such as the CompilationOutput and the EnvironmentDefines
+	// data structures.
+	Impl.Step_FinalizeArtifacts();
 
-void BindValueToExpressionOutput(FMaterialIRModuleBuilderImpl* Builder, const FExpressionOutput* Output, FValue* Value)
-{
-	Builder->AnalysisContextStack.Last().SetOutputValue(Output, Value);
+	return true;
 }
 
-} // namespace MIR::Internal
-
 #endif // #if WITH_EDITOR
diff --git a/Engine/Source/Runtime/Engine/Private/Materials/MaterialIRToHLSLTranslator.cpp b/Engine/Source/Runtime/Engine/Private/Materials/MaterialIRToHLSLTranslator.cpp
--- a/Engine/Source/Runtime/Engine/Private/Materials/MaterialIRToHLSLTranslator.cpp
+++ b/Engine/Source/Runtime/Engine/Private/Materials/MaterialIRToHLSLTranslator.cpp
@@ -209,19 +209,19 @@ static const TCHAR* GetShadingModelParameterName(EMaterialShadingModel InModel)
 	}
 }
 
-static bool InstructionUsesPhiValue(const MIR::FInstruction* Instr, MIR::EStage Stage)
+static bool InstructionUsesPhiValue(const MIR::FInstruction* Instr, int32 EntryPointIndex)
 {
 	if (auto Branch = Instr->As<MIR::FBranch>())
 	{
-		return Branch->TrueBlock[Stage].Instructions || Branch->FalseBlock[Stage].Instructions;
+		return Branch->TrueBlock[EntryPointIndex].Instructions || Branch->FalseBlock[EntryPointIndex].Instructions;
 	}
 	return false;
 }
 
-static bool IsFoldable(const MIR::FInstruction* Instr, MIR::EStage Stage)
+static bool IsFoldable(const MIR::FInstruction* Instr, int32 EntryPointIndex)
 {
 	// Instructions that use Phi values cannot be folded/inlined.
-	if (InstructionUsesPhiValue(Instr, Stage))
+	if (InstructionUsesPhiValue(Instr, EntryPointIndex))
 	{
 		return false;
 	}
@@ -229,12 +229,13 @@ static bool IsFoldable(const MIR::FInstruction* Instr, MIR::EStage Stage)
 	// Don't fold instructions that become very long to improve readability
 	switch (Instr->Kind)
 	{
-		case MIR::VK_VTPageTableRead: return false;
-		case MIR::VK_Call: return false;
-		default: break;
+		case MIR::VK_SetMaterialOutput:
+		case MIR::VK_VTPageTableRead:
+		case MIR::VK_Call:
+			return false;
+		default:
+			return true;
 	}
-
-	return true;
 }
 
 static FStringView GVector4SwizzleSubset[4] = { TEXTVIEW(".x"), TEXTVIEW(".xy"), TEXTVIEW(".xyz"), {} };
@@ -242,27 +243,30 @@ static FStringView GVector4SwizzleSubset[4] = { TEXTVIEW(".x"), TEXTVIEW(".xy"),
 struct FPrivate : FMaterialIRToHLSLTranslation
 {
 	// Utility for emitting HLSL code
-	FHLSLPrinter Printer; 
+	FHLSLPrinter Printer{}; 
 
 	// Number of local variables generated during translation.
 	int32 NumLocals{};
 	
 	// Mapping from instructions to their corresponding local variable index
 	// Used to map instructions that have more than one use to a local identifier like "_42".
-	TMap<const MIR::FInstruction*, uint32> InstrToLocalIndex;
+	TMap<const MIR::FInstruction*, uint32> InstrToLocalIndex{};
+
+	// The index of the current entry point being generated.
+	int32 CurrentEntryPointIndex{};
 
 	// Current material stage being translated (e.g., vertex, pixel)
-	MIR::EStage CurrentStage;
+	MIR::EStage CurrentStage{};
 
 	// Set when generating world position offset HLSL using previous frame data.
 	// Affects token replacement for <PREV> in external code, and code generation for certain external inputs.
 	bool bCompilingPreviousFrame = false;
 
 	// Generated HLSL for pixel attribute evaluation
-	FString PixelAttributesHLSL;
+	FString PixelAttributesHLSL{};
 	
 	// Generated HLSL for world position offset calculation
-	FString WorldPositionOffsetHLSL;
+	FString WorldPositionOffsetHLSL{};
 
 	// Generated HLSL for previous world position offset calculation
 	FString PreviousWorldPositionOffsetHLSL;
@@ -277,7 +281,7 @@ struct FPrivate : FMaterialIRToHLSLTranslation
 	{
 		for (EMaterialProperty Property : UMaterialAggregate::GetMaterialAttributesProperties())
 		{
-			if (!MIR::MaterialOutputEvaluatesInStage(Property, MIR::Stage_Pixel))
+			if (Property == MP_WorldPositionOffset)
 			{
 				continue;
 			}
@@ -290,69 +294,61 @@ struct FPrivate : FMaterialIRToHLSLTranslation
 		}
 	}
 	
-	void GenerateVertexStageHLSL()
+	// Generates the full HLSL of the specified entry point.
+	const FMaterialIRModule::FEntryPoint& GenerateEntryPoint(int32 EntryPointIndex)
 	{
-		BeginStage(MIR::Stage_Vertex);
+		Printer = {};
+		Printer.Tabs = 1;
+		Printer << Indentation;
+
+		const FMaterialIRModule::FEntryPoint& EntryPoint = Module->GetEntryPoint(EntryPointIndex);
+
+		CurrentEntryPointIndex = EntryPointIndex;
+		CurrentStage = EntryPoint.Stage;
 
-		LowerBlock(Module->GetRootBlock(MIR::Stage_Vertex));
+		LowerBlock(EntryPoint.RootBlock);
 
+		return EntryPoint;
+	}
+	
+	// Generates the HLSL that sets the outputs of the vertex stage.	
+	void GenerateVertexStageHLSL()
+	{
+		// Generate current frame vertex stage HLSL
+		GenerateEntryPoint(MIR::Stage_Vertex);
 		WorldPositionOffsetHLSL = MoveTemp(Printer.Buffer);
 
-		// Same as above, but with previous frame flag set.
+		// Generate previous frame vertex stage HLSL
 		bCompilingPreviousFrame = true;
-		BeginStage(MIR::Stage_Vertex);
-		LowerBlock(Module->GetRootBlock(MIR::Stage_Vertex));
+		GenerateEntryPoint(MIR::Stage_Vertex);
 		PreviousWorldPositionOffsetHLSL = MoveTemp(Printer.Buffer);
 		bCompilingPreviousFrame = false;
 	}
 
-	void GenerateOtherStageHLSL(MIR::EStage Stage)
+	// Generates the HLSL that sets the outputs of a non-vertex stage.
+	void GenerateNonVertexStageHLSL(MIR::EStage Stage)
 	{
-		BeginStage(Stage);
-
-		LowerBlock(Module->GetRootBlock(CurrentStage));
+		GenerateEntryPoint(Stage);
 
 		if (Stage == MIR::Stage_Pixel)
 		{
 			Printer << TEXTVIEW("PixelMaterialInputs.Subsurface = 0") << EndOfStatement;
-
-			if (Substrate::IsSubstrateEnabled())
-			{
-				Printer << TEXTVIEW("// Substrate normal and tangent") << NewLine;
-				Printer << TEXTVIEW("Parameters.SharedLocalBases.Normals[0] = TransformTangentNormalToWorld(Parameters.TangentToWorld, PixelMaterialInputs.Normal)") << EndOfStatement;
-				Printer << TEXTVIEW("#if MATERIAL_TANGENTSPACENORMAL") << NewLine;
-				Printer << TEXTVIEW("Parameters.SharedLocalBases.Normals[0] *= Parameters.TwoSidedSign") << EndOfStatement;
-				Printer << TEXTVIEW("#endif") << NewLine;
-				Printer << TEXTVIEW("Parameters.SharedLocalBases.Tangents[0] = Parameters.TangentToWorld[0]") << EndOfStatement;
-				Printer << TEXTVIEW("#if MATERIAL_TANGENTSPACENORMAL") << NewLine;
-				Printer << TEXTVIEW("Parameters.SharedLocalBases.Tangents[0] *= Parameters.TwoSidedSign") << EndOfStatement;
-				Printer << TEXTVIEW("#endif") << NewLine;
-				Printer << TEXTVIEW("Parameters.SharedLocalBases.Count = 1") << EndOfStatement;
-			}
 		}
 
 		EvaluateOtherMaterialAttributesHLSL[CurrentStage] = MoveTemp(Printer.Buffer);
 	}
 	
-	void BeginStage(MIR::EStage Stage)
-	{
-		Printer = {};
-		Printer.Tabs = 1;
-		Printer << Indentation;
-		CurrentStage = Stage;
-	}
-	
 	ENoOp LowerBlock(const MIR::FBlock& Block)
 	{
 		int32 OldNumLocals = NumLocals;
-		for (MIR::FInstruction* Instr = Block.Instructions; Instr; Instr = Instr->Next[CurrentStage])
+		for (MIR::FInstruction* Instr = Block.Instructions; Instr; Instr = Instr->GetNext(CurrentEntryPointIndex))
 		{
-            if (Instr->NumUsers[CurrentStage] == 1 && IsFoldable(Instr, CurrentStage))
+            if (Instr->GetNumUsers(CurrentEntryPointIndex) == 1 && IsFoldable(Instr, CurrentEntryPointIndex))
 			{
                 continue;
             }
             
-			if (Instr->NumUsers[CurrentStage] >= 1 && !Instr->As<MIR::FCall>())
+			if (Instr->GetNumUsers(CurrentEntryPointIndex) >= 1 && !Instr->As<MIR::FCall>())
 			{
 				// Allocate a new local index for this instruction
 				uint32 LocalIndex = NumLocals++;
@@ -365,7 +361,7 @@ struct FPrivate : FMaterialIRToHLSLTranslation
 				
 				// If this instruction doesn't use a phi value, we'll immediately assign its local to its result so output the "="
 				// E.g., if Instr is a branch, we will set its value inside the generate "if" {} scopes, so no need for a "=" now.
-                if (!InstructionUsesPhiValue(Instr, CurrentStage))
+                if (!InstructionUsesPhiValue(Instr, CurrentEntryPointIndex))
 				{
                     Printer << TEXTVIEW(" = ");
                 }
@@ -404,7 +400,7 @@ struct FPrivate : FMaterialIRToHLSLTranslation
 		if (const MIR::FInstruction* Instr = MIR::AsInstruction(Value))
 		{ 
 			// If this instruction has only one user (dependant) and is foldable, then generate the full HLSL for the instruction in place.
-			if (Instr->NumUsers[CurrentStage] <= 1 && IsFoldable(Instr, CurrentStage))
+			if (Instr->GetNumUsers(CurrentEntryPointIndex) <= 1 && IsFoldable(Instr, CurrentEntryPointIndex))
 			{
 				Printer << LowerInstruction(Instr);
 			}
@@ -1281,15 +1277,15 @@ struct FPrivate : FMaterialIRToHLSLTranslation
 		}
 	}
 
-	void LowerSetMaterialOutput(const MIR::FSetMaterialOutput* Output)
+	void LowerSetMaterialOutput(const MIR::FSetMaterialOutput* SetMaterialOutput)
 	{
-		if (Output->Property == MP_WorldPositionOffset)
+		if (SetMaterialOutput->Property == MP_WorldPositionOffset)
 		{
 			Printer << TEXTVIEW("return ");
 		}
 		else
 		{
-			if (const FMaterialAggregateAttribute* PropertyAggregate = UMaterialAggregate::GetMaterialAttribute(Output->Property))
+			if (const FMaterialAggregateAttribute* PropertyAggregate = UMaterialAggregate::GetMaterialAttribute(SetMaterialOutput->Property))
 			{
 				Printer << TEXTVIEW("PixelMaterialInputs.") << PropertyAggregate->Name.ToString() << TEXTVIEW(" = ");
 			}
@@ -1298,36 +1294,8 @@ struct FPrivate : FMaterialIRToHLSLTranslation
 				UE_MIR_UNREACHABLE();
 			}			
 		}
-
-		if (Substrate::IsSubstrateEnabled())
-		{
-			// Lower to default value when Substrate is enabled
-			switch (Output->Property)
-			{
-				case MP_EmissiveColor:		Printer << TEXTVIEW("0"); break;
-				case MP_Opacity:			Printer << TEXTVIEW("1.0f"); break;
-				case MP_OpacityMask:		Printer << TEXTVIEW("1.0f"); break;
-				case MP_BaseColor:			Printer << TEXTVIEW("0"); break;
-				case MP_Metallic:			Printer << TEXTVIEW("0"); break;
-				case MP_Specular:			Printer << TEXTVIEW("0"); break;
-				case MP_Roughness:			Printer << TEXTVIEW("0"); break;
-				case MP_Anisotropy:			Printer << TEXTVIEW("0"); break;
-				case MP_Normal:				LowerValue(Output->Arg); break;
-				case MP_Tangent:			Printer << TEXTVIEW("0"); break;
-				case MP_SubsurfaceColor:	Printer << TEXTVIEW("0"); break;
-				case MP_AmbientOcclusion:	Printer << TEXTVIEW("1.0f"); break;
-				case MP_Refraction:			Printer << TEXTVIEW("0"); break;
-				case MP_PixelDepthOffset:	Printer << TEXTVIEW("0.0f"); break;
-				case MP_ShadingModel:		Printer << TEXTVIEW("0"); break;
-				case MP_SurfaceThickness:	Printer << TEXTVIEW("0.01f"); break;
-				case MP_Displacement:		Printer << TEXTVIEW("-1.0f"); break;
-				default: LowerValue(Output->Arg); break;
-			}
-		}
-		else
-		{
-			LowerValue(Output->Arg);
-		}
+			
+		LowerValue(SetMaterialOutput->Arg);
 	}
 
 	void LowerOperator(const MIR::FOperator* Operator)
@@ -1529,7 +1497,7 @@ struct FPrivate : FMaterialIRToHLSLTranslation
 
 	void LowerBranch(const MIR::FBranch* Branch)
 	{
-		if (IsFoldable(Branch, CurrentStage))
+		if (IsFoldable(Branch, CurrentEntryPointIndex))
 		{
 			if (Branch->TrueArg->Type.IsDouble())
 			{
@@ -1549,11 +1517,11 @@ struct FPrivate : FMaterialIRToHLSLTranslation
 		{
 			Printer << EndOfStatement;
 			Printer << TEXTVIEW("if (") << LowerValue(Branch->ConditionArg) << TEXTVIEW(")") << NewLine << OpenBrace;
-			Printer << LowerBlock(Branch->TrueBlock[CurrentStage]);
+			Printer << LowerBlock(Branch->TrueBlock[CurrentEntryPointIndex]);
 			Printer << TEXTVIEW("_") << InstrToLocalIndex[Branch] << TEXTVIEW(" = ") << LowerValue(Branch->TrueArg) << EndOfStatement;
 			Printer << CloseBrace << NewLine;
 			Printer << TEXTVIEW("else") << NewLine << OpenBrace;
-			Printer << LowerBlock(Branch->FalseBlock[CurrentStage]);
+			Printer << LowerBlock(Branch->FalseBlock[CurrentEntryPointIndex]);
 			Printer << TEXTVIEW("_") << InstrToLocalIndex[Branch] << TEXTVIEW(" = ") << LowerValue(Branch->FalseArg) << EndOfStatement;
 			Printer << CloseBrace;
 		}
@@ -2294,10 +2262,10 @@ struct FPrivate : FMaterialIRToHLSLTranslation
 		ensure(ShadingModels.IsValid());
 
 		// Logic from FHLSLMaterialTranslator::TranslateMaterial
-		bool bOpacityPropertyIsUsed = Module->IsMaterialPropertyUsed(MIR::EStage::Stage_Pixel, MP_Opacity);
+		bool bOpacityPropertyIsUsed = Module->IsMaterialPropertyUsed(MP_Opacity);
 
 		bool bUsesCurvature = Module->GetFeatureLevel() == ERHIFeatureLevel::ES3_1 &&
-			((ShadingModels.HasShadingModel(MSM_SubsurfaceProfile) && Module->IsMaterialPropertyUsed(MIR::EStage::Stage_Pixel, MP_CustomData0))
+			((ShadingModels.HasShadingModel(MSM_SubsurfaceProfile) && Module->IsMaterialPropertyUsed(MP_CustomData0))
 			|| (ShadingModels.HasShadingModel(MSM_Eye) && bOpacityPropertyIsUsed));
 
 		int32 NumActiveShadingModels = 0;
@@ -2467,8 +2435,8 @@ void FMaterialIRToHLSLTranslation::Run(TMap<FString, FString>& OutParameters, FS
 	FPrivate Private{ *this };
 	Private.GeneratePixelAttributesHLSL();
 	Private.GenerateVertexStageHLSL();
-	Private.GenerateOtherStageHLSL(MIR::Stage_Pixel);
-	Private.GenerateOtherStageHLSL(MIR::Stage_Compute);
+	Private.GenerateNonVertexStageHLSL(MIR::Stage_Pixel);
+	Private.GenerateNonVertexStageHLSL(MIR::Stage_Compute);
 	Private.GenerateTemplateStringParameters(OutParameters);
 	Private.GenerateCustomFunctionsHLSL(OutParameters);
 	Private.GetShaderCompilerEnvironment(OutEnvironment);
diff --git a/Engine/Source/Runtime/Engine/Private/Materials/MaterialIRValueAnalyzer.cpp b/Engine/Source/Runtime/Engine/Private/Materials/MaterialIRValueAnalyzer.cpp
--- a/Engine/Source/Runtime/Engine/Private/Materials/MaterialIRValueAnalyzer.cpp
+++ b/Engine/Source/Runtime/Engine/Private/Materials/MaterialIRValueAnalyzer.cpp
@@ -548,6 +548,15 @@ static void AnalyzeShadingModel(FMaterialIRValueAnalyzer& Analyzer, MIR::FShadin
 	}
 }
 
+static void AnalyzeBranch(FMaterialIRValueAnalyzer& Analyzer, MIR::FBranch* Branch)
+{
+	Branch->TrueBlock = Analyzer.Module->AllocateArray<MIR::FBlock>(Analyzer.Module->GetNumEntryPoints());
+	MIR::ZeroArray(Branch->TrueBlock);
+
+	Branch->FalseBlock = Analyzer.Module->AllocateArray<MIR::FBlock>(Analyzer.Module->GetNumEntryPoints());
+	MIR::ZeroArray(Branch->FalseBlock);
+}
+
 static void AnalyzeTextureRead(FMaterialIRValueAnalyzer& Analyzer, MIR::FTextureRead* TextureRead)
 {
 	// Ensure a uniform parameter is allocated when a texture read instruction accesses the texture.
@@ -779,14 +788,6 @@ void FMaterialIRValueAnalyzer::Setup(UMaterial* InMaterial, FMaterialIRModule* I
 
 void FMaterialIRValueAnalyzer::Analyze(MIR::FValue* Value)
 {
-	for (MIR::FValue* Use : Value->GetUses())
-	{
-		if (Use)
-		{
-			Value->GraphProperties |= Use->GraphProperties;
-		}
-	}
-
 	#define EXPAND_CASE(ValueType, ...) \
 	case MIR::VK_##ValueType: \
 	{ \
@@ -800,6 +801,7 @@ void FMaterialIRValueAnalyzer::Analyze(MIR::FValue* Value)
 		EXPAND_CASE(ExternalInput);
 		EXPAND_CASE(MaterialParameterCollection);
 		EXPAND_CASE(UniformParameter);
+		EXPAND_CASE(Branch);
 		EXPAND_CASE(TextureRead);
 		EXPAND_CASE(PreshaderParameter);
 		EXPAND_CASE(SetMaterialOutput);
diff --git a/Engine/Source/Runtime/Engine/Public/Materials/MaterialIR.h b/Engine/Source/Runtime/Engine/Public/Materials/MaterialIR.h
--- a/Engine/Source/Runtime/Engine/Public/Materials/MaterialIR.h
+++ b/Engine/Source/Runtime/Engine/Public/Materials/MaterialIR.h
@@ -41,9 +41,6 @@ enum EStage
 // Returns the string representation of given stage.
 const TCHAR* LexToString(EStage Stage);
 
-// Returns whether specified property is evaluated in specified stage.
-bool MaterialOutputEvaluatesInStage(EMaterialProperty Property, EStage Stage);
-
 // A collection of bit flags for a specific Value instance.
 enum class EValueFlags : uint8
 {
@@ -58,26 +55,24 @@ enum class EValueFlags : uint8
 	// Value has been analyzed for the compute stage.
 	AnalyzedForStageCompute = 1 << 2,
 
-	// Mask that includes all "analyzed in stage" flags.
-	AnalyzedInAnyStageMask = AnalyzedForStageVertex | AnalyzedForStagePixel | AnalyzedForStageCompute,
+	// Value has been analyzed in some stage.
+	AnalyzedInAnyStage = 1 << 3,
 
 	// Whether to substitute abstract tokens like "<PREV>" in code expressions of InlineHLSL instructions.
-	SubstituteTags = 1 << 3,
+	SubstituteTags = 1 << 4,
 
 	// Whether 'FInlineHLSL::Code' field is used. Otherwise, 'ExternalCodeDeclaration' is used.
-	HasDynamicHLSLCode = 1 << 4, 
+	HasDynamicHLSLCode = 1 << 5, 
 
 	// External code declaration must emit the secondary definition for the DDX derivative.
-	DerivativeDDX = 1 << 5,
+	DerivativeDDX = 1 << 6,
 
 	// External code declaration must emit the secondary definition for the DDY derivative.
-	DerivativeDDY = 1 << 6,
+	DerivativeDDY = 1 << 7,
 };
 
 ENUM_CLASS_FLAGS(EValueFlags);
 
-static_assert(uint32(1 << NumStages) - 1 == uint32(EValueFlags::AnalyzedInAnyStageMask), "Discrepancy between the number of stages in NumStages and the flags in EValueFlags.");
-
 // A collection of graph properties used by a value.
 //
 // If a graph property is true, it means that either the value itself makes has of that
@@ -586,6 +581,24 @@ struct FBlock
 	FBlock* FindCommonParentWith(MIR::FBlock* Other);
 };
 
+//
+struct FInstructionLinkage
+{
+	// The next instruction executing after this one.
+	FInstruction* Next;
+
+	// This instruction parent block.
+	FBlock* Block;
+
+	// How many users (i.e., dependencies) this instruction has.
+	uint32 NumUsers;
+
+	// The number of users processed during instruction graph linking in each stage of execution.
+	// Note: This information combined with NumUsers is used by the builder to push instructions
+	//      in the appropriate block automatically.
+	uint32 NumProcessedUsers;
+};
+
 // Base struct of an instruction value.
 //
 // An instruction is a value in a well defined order of evaluation.
@@ -599,22 +612,22 @@ struct FBlock
 //      an instruction, and are instead automatically populated by the builder.
 struct FInstruction : FValue
 {
-	// The next instruction executing after this one in all stages of execution.
-	FInstruction* Next[NumStages];
+	// Array of linkage information (how this instruction is connected inside the IR graph)
+	// for each entry point registered in the module.
+	TArrayView<FInstructionLinkage> Linkage;
 
-	// This instruction parent block in all stages of execution.
-	FBlock* Block[NumStages];
+	// Returns the block in which the dependency with specified index should execute.
+	FBlock* GetTargetBlockForUse(int32 EntryPointIndex, int32 UseIndex);
 
-	// How many users (i.e., dependencies) this instruction has in each stage of execution.
-	uint32 NumUsers[NumStages];
+	// Returns the next instruction in specified entry point. 
+	FInstruction* GetNext(int32 EntryPointIndex) const { return Linkage[EntryPointIndex].Next; }
 
-	// The number of users processed during instruction graph linking in each stage of execution.
-	// Note: This information combined with NumUsers is used by the builder to push instructions
-	//      in the appropriate block automatically.
-	uint32 NumProcessedUsers[NumStages];
+	// Returns the number of user sin specified entry point.
+	uint32 GetNumUsers(int32 EntryPointIndex) const { return Linkage[EntryPointIndex].NumUsers; }
+
+	// Returns the block this instruction belongs to in specified entry point.
+	const FBlock* GetBlock(int32 EntryPointIndex) const { return Linkage[EntryPointIndex].Block; }
 
-	// Returns the block in which the dependency with specified index should execute.
-	FBlock* GetDesiredBlockForUse(EStage Stage, int32 UseIndex);
 };
 
 template <EValueKind TTypeKind, uint32 TNumStaticUses = 0>
@@ -854,11 +867,11 @@ struct FBranch : TInstruction<VK_Branch, 3>
 	// Value this branch evaluates to when the condition is false.
 	FValue* FalseArg;
 
-	// The inner block (in each execution stage) the subgraph evaluating `TrueArg` should be placed in.
-	FBlock TrueBlock[NumStages];
+	// The inner block (in each module entry point) the subgraph evaluating `TrueArg` should be placed in.
+	TArrayView<FBlock> TrueBlock;
 
-	// The inner block (in each execution stage) the subgraph evaluating `FalseArg` should be placed in.
-	FBlock FalseBlock[NumStages];
+	// The inner block (in each module entry point) the subgraph evaluating `FalseArg` should be placed in.
+	TArrayView<FBlock> FalseBlock;
 };
 
 // A subscript instruction.
@@ -984,7 +997,7 @@ struct FVTPageTableRead : TInstruction<VK_VTPageTableRead, 4>
 // Utility value for selecting a different value based on the execution stage.
 struct FStageSwitch : TInstruction<VK_StageSwitch, 1>
 {
-	// The argument for to be bypassed in each stage. e
+	// The argument for to be bypassed in each stage.
 	FValue* Args[NumStages];
 
 	// Use the specified value argument in the pixel stage, and another specified
@@ -1093,8 +1106,10 @@ struct FFunctionHLSL : FFunction
 	// The HLSL code of the body of this function.
 	FStringView Code;
 
+	// Array of #defines to declare in the material shader.
 	TConstArrayView<FFunctionHLSLDefine> Defines;
 
+	// Array of include directives to declare in the material shader.
 	TConstArrayView<FStringView> Includes;
 
 	// Whether this function equals specified (they have the same data).
@@ -1112,7 +1127,6 @@ struct FCall : TInstruction<VK_Call, 0>
 	
 	// The number of arguments.
 	int NumArguments;
-
 };
 
 // Instruction that evaluates some user-function call output parameter result.
diff --git a/Engine/Source/Runtime/Engine/Public/Materials/MaterialIRCommon.h b/Engine/Source/Runtime/Engine/Public/Materials/MaterialIRCommon.h
--- a/Engine/Source/Runtime/Engine/Public/Materials/MaterialIRCommon.h
+++ b/Engine/Source/Runtime/Engine/Public/Materials/MaterialIRCommon.h
@@ -90,6 +90,13 @@ struct TTemporaryArray : public TArrayView<T>
 	}
 };
 
+template <typename T>
+void ZeroArray(TArrayView<T> Array)
+{
+	static_assert(TIsTriviallyCopyAssignable<T>::Value);
+	FMemory::Memzero(Array.GetData(), Array.Num() * sizeof(T));
+}
+
 }
 
 #define UE_MIR_UNREACHABLE() { check(!"Unreachable"); UE_ASSUME(false); }
diff --git a/Engine/Source/Runtime/Engine/Public/Materials/MaterialIRModule.h b/Engine/Source/Runtime/Engine/Public/Materials/MaterialIRModule.h
--- a/Engine/Source/Runtime/Engine/Public/Materials/MaterialIRModule.h
+++ b/Engine/Source/Runtime/Engine/Public/Materials/MaterialIRModule.h
@@ -46,6 +46,25 @@ public:
 		uint32 DynamicParticleParameterMask;
 	};
 
+	// An entry point is a sequence of instructions that evaluates a set of "output" instructions.
+	// A way to think of an entry point is as the "main scope" or the "body" of a C function.
+	// Each shader stage (Vertex, Pixel and Compute) gets its own entry point with one or more outputs.
+	// Other entry points are crated by user specified "local functions", such as custom outputs.
+	struct FEntryPoint
+	{
+		// This entry point's user name.
+		FStringView Name;
+		
+		// The shader stage this entry point is intended to run in.
+		MIR::EStage Stage;
+		
+		// The root block containing this entry point's instructions.
+		MIR::FBlock RootBlock;
+
+		// An array of terminal values this entry point evaluates.
+		TArrayView<MIR::FValue*> Outputs;
+	};
+
 public:
 	FMaterialIRModule();
 	~FMaterialIRModule();
@@ -68,11 +87,21 @@ public:
 	// Retrieves the material compilation output.
 	const FMaterialCompilationOutput& GetCompilationOutput() const { return CompilationOutput; }
 
-	// Returns the material outputs for a given stage.
-	TArrayView<const MIR::FSetMaterialOutput* const> GetOutputs(MIR::EStage Stage) const { return Outputs[Stage]; }
+	// Adds a new entry point with specified name, stage and expected number of evaluated output values.
+	int32 AddEntryPoint(FStringView Name, MIR::EStage Stage, int32 NumOutputs);
 
-	// Retrieves the root block (i.e. the "main" scope) for a specific shader stage.
-	const MIR::FBlock& GetRootBlock(MIR::EStage Stage) const { return RootBlocks[Stage]; }
+	// Returns the number of procedures inthe module.
+	int32 GetNumEntryPoints() const { return EntryPoints.Num(); }
+
+	// Gets the entry point of given Index previously added with AddEntryPoint().
+	FEntryPoint& GetEntryPoint(int32 Index) { return EntryPoints[Index]; }
+	const FEntryPoint& GetEntryPoint(int32 Index) const { return EntryPoints[Index]; }
+
+	// Sets the value of a material property for future reference.
+	void SetPropertyValue(EMaterialProperty Property, MIR::FValue* Value) { PropertyValues[Property] = Value; }
+	
+	// Retrieves the value a meterial property is assigned to.
+	const MIR::FValue* GetPropertyValue(EMaterialProperty Property) const { return PropertyValues[Property]; }
 
 	// Returns a list of all environment define names this module requires to be enabled for shader compilation.
 	const TSet<FName>& GetEnvironmentDefines() const { return EnvironmentDefines; }
@@ -105,7 +134,7 @@ public:
 	FMaterialShadingModelField GetCompiledShadingModels() const { return this->ShadingModelsFromCompilation; }
 
 	// Returns true if the given property has a non default value.  Similar to FHLSLMaterialTranslator::IsMaterialPropertyUsed.
-	bool IsMaterialPropertyUsed(MIR::EStage Stage, EMaterialProperty InProperty) const;
+	bool IsMaterialPropertyUsed(EMaterialProperty InProperty) const;
 
 	// Stores a user-defined string and returns a view of it.
 	FStringView InternString(FStringView InString);
@@ -118,6 +147,21 @@ public:
 
 	// Reports a translation error.
 	void AddError(UMaterialExpression* Expression, FString Message);
+	
+	// Allocates a chunk of memory using the module's internal memory block, to be used
+	// by the user as they see fit.
+	// Note: The returned array has the same lifetime as this module.
+	void* Allocate(size_t Size, size_t Alignment) { return Allocator.Alloc(Size, Alignment); }
+
+	// Allocates an array of elements using the module's internal memory block, to be used
+	// by the user as they see fit.
+	// Note: The returned array has the same lifetime as this module.
+	template <typename T>
+	TArrayView<T> AllocateArray(int32 Count)
+	{
+		static_assert(TIsTriviallyCopyAssignable<T>::Value);
+		return { (T*)Allocate(sizeof(T) * Count, alignof(T)), Count };
+	}
 
 private:
 	// The chunk allocator used to contain child allocations (e.g., values, interned strings, etc)
@@ -141,11 +185,11 @@ private:
 	// List of all the IR values contained in this module.
 	TArray<MIR::FValue*> Values;
 
-	// Output nodes per stage.
-	TArray<MIR::FSetMaterialOutput*> Outputs[MIR::NumStages];
+	// Array of entry points in this module.
+	TArray<FEntryPoint> EntryPoints;
 
-	// Root blocks per stage.
-	MIR::FBlock RootBlocks[MIR::NumStages];
+	// The final value assigned to each material property.
+	MIR::FValue* PropertyValues[MP_MAX];
 
 	// Compilation statistics.
 	FStatistics Statistics;
