Minor improvements to my last MIR change suggested by BugHawk.

#rnx

diff --git a/Engine/Source/Runtime/Engine/Private/Materials/MaterialIRModule.cpp b/Engine/Source/Runtime/Engine/Private/Materials/MaterialIRModule.cpp
--- a/Engine/Source/Runtime/Engine/Private/Materials/MaterialIRModule.cpp
+++ b/Engine/Source/Runtime/Engine/Private/Materials/MaterialIRModule.cpp
@@ -48,7 +48,7 @@ int32 FMaterialIRModule::AddEntryPoint(FStringView Name, MIR::EStage Stage, int3
 	uint32 Index = EntryPoints.Num();
 
 	EntryPoints.Push({
-		.Name = Name,
+		.Name = InternString(Name),
 		.Stage = Stage,
 		.RootBlock = {},
 		.Outputs = AllocateArray<MIR::FValue*>(NumOutputs),
diff --git a/Engine/Source/Runtime/Engine/Public/Materials/MaterialIR.h b/Engine/Source/Runtime/Engine/Public/Materials/MaterialIR.h
--- a/Engine/Source/Runtime/Engine/Public/Materials/MaterialIR.h
+++ b/Engine/Source/Runtime/Engine/Public/Materials/MaterialIR.h
@@ -593,7 +593,7 @@ struct FInstructionLinkage
 	// How many users (i.e., dependencies) this instruction has.
 	uint32 NumUsers;
 
-	// The number of users processed during instruction graph linking in each stage of execution.
+	// The number of users processed during instruction graph linking in each entry point.
 	// Note: This information combined with NumUsers is used by the builder to push instructions
 	//      in the appropriate block automatically.
 	uint32 NumProcessedUsers;
