Fixed replay issues with WorldDataLayers actor when running in standalone.
Also removed the re-initialization of runtime datalayer states in AWorldDataLayers::RewindForReplay and AWorldDataLayers::PostRegisterAllComponents as this was causing data layers to be in the incorrect state. Now fully rely on replication to set the runtime states.

diff --git a/Engine/Source/Runtime/Engine/Private/WorldPartition/DataLayer/WorldDataLayers.cpp b/Engine/Source/Runtime/Engine/Private/WorldPartition/DataLayer/WorldDataLayers.cpp
--- a/Engine/Source/Runtime/Engine/Private/WorldPartition/DataLayer/WorldDataLayers.cpp
+++ b/Engine/Source/Runtime/Engine/Private/WorldPartition/DataLayer/WorldDataLayers.cpp
@@ -82,28 +82,6 @@ void AWorldDataLayers::GetLifetimeReplicatedProps(TArray< FLifetimeProperty >& O
 	DOREPLIFETIME_WITH_PARAMS_FAST(AWorldDataLayers, RepEffectiveActiveDataLayerNames, Params);
 }
 
-void AWorldDataLayers::RewindForReplay()
-{
-	Super::RewindForReplay();
-
-	// Same as PostRegisterAllComponents when rewinding we want to reset our state to the initial state and rely on the Replay/Replication.
-	ResetDataLayerRuntimeStates();
-	InitializeDataLayerRuntimeStates();
-}
-
-void AWorldDataLayers::PostRegisterAllComponents()
-{
-	Super::PostRegisterAllComponents();
-
-	// When running a Replay we want to reset our state to the initial state and rely on the Replay/Replication.
-	// Unfortunately this can't be tested in the PostLoad as the World doesn't have a demo driver yet.
-	if (GetWorld()->IsPlayingReplay())
-	{
-		ResetDataLayerRuntimeStates();
-		InitializeDataLayerRuntimeStates();
-	}
-}
-
 #define AWORLDDATALAYERS_UPDATE_REPLICATED_DATALAYERS(ReplicatedArray, SourceArray) \
 	MARK_PROPERTY_DIRTY_FROM_NAME(AWorldDataLayers, ReplicatedArray, this); \
 	ReplicatedArray = SourceArray;
@@ -294,7 +272,7 @@ bool AWorldDataLayers::SetDataLayerRuntimeState(const UDataLayerInstance* InData
 		}
 
 		// Update replicated properties
-		if (!bIsLocalDataLayer && (NetMode == NM_DedicatedServer || NetMode == NM_ListenServer))
+		if (!bIsLocalDataLayer && (NetMode == NM_DedicatedServer || NetMode == NM_ListenServer || NetMode == NM_Standalone))
 		{
 			FlushNetDormancy();
 			AWORLDDATALAYERS_UPDATE_REPLICATED_DATALAYERS(RepActiveDataLayerNames, ActiveDataLayerNames.Array());
@@ -433,7 +411,7 @@ void AWorldDataLayers::ResolveEffectiveRuntimeState(const UDataLayerInstance* In
 	if (EffectiveStates.SetDataLayerEffectiveRuntimeState(DataLayerName, bIsLocalDataLayer, NewEffectiveRuntimeState, OldEffectiveRuntimeState))
 	{
 		// Update Replicated Properties
-		if (!bIsLocalDataLayer && (NetMode == NM_DedicatedServer || NetMode == NM_ListenServer))
+		if (!bIsLocalDataLayer && (NetMode == NM_DedicatedServer || NetMode == NM_ListenServer || NetMode == NM_Standalone))
 		{
 			FlushNetDormancy();
 			AWORLDDATALAYERS_UPDATE_REPLICATED_DATALAYERS(RepEffectiveActiveDataLayerNames, EffectiveStates.GetReplicatedEffectiveActiveDataLayerNames().Array());
diff --git a/Engine/Source/Runtime/Engine/Public/WorldPartition/DataLayer/WorldDataLayers.h b/Engine/Source/Runtime/Engine/Public/WorldPartition/DataLayer/WorldDataLayers.h
--- a/Engine/Source/Runtime/Engine/Public/WorldPartition/DataLayer/WorldDataLayers.h
+++ b/Engine/Source/Runtime/Engine/Public/WorldPartition/DataLayer/WorldDataLayers.h
@@ -92,8 +92,6 @@ public:
 	//~ End UObject Interface
 
 	//~ Begin AActor Interface
-	ENGINE_API virtual void RewindForReplay() override;
-	ENGINE_API virtual void PostRegisterAllComponents() override;
 #if WITH_EDITOR
 	ENGINE_API virtual void PreEditUndo() override;
 	ENGINE_API virtual void PostEditUndo() override;
