Stop game state from replicating using the proxy parent connection since it's not associated any with connected clients.

diff --git a/Engine/Plugins/Runtime/MultiServerReplication/Source/MultiServerReplication/Private/MultiServerProxy.cpp b/Engine/Plugins/Runtime/MultiServerReplication/Source/MultiServerReplication/Private/MultiServerProxy.cpp
--- a/Engine/Plugins/Runtime/MultiServerReplication/Source/MultiServerReplication/Private/MultiServerProxy.cpp
+++ b/Engine/Plugins/Runtime/MultiServerReplication/Source/MultiServerReplication/Private/MultiServerProxy.cpp
@@ -51,6 +51,12 @@ namespace UE::MultiServerProxy::Private
 		TEXT("net.proxy.EnableDisconnectionSupport"),
 		bEnableDisconnectionSupport,
 		TEXT("When true, when a client disconnects from the proxy it will close the connection to the backend server."));
+
+	bool bEnableParentConnectionReplication = false;
+	static FAutoConsoleVariableRef CVarEnableParentConnectionReplication(
+		TEXT("net.proxy.EnableParentConnectionReplication"),
+		bEnableParentConnectionReplication,
+		TEXT("When true, the proxy will configure the parent connection to the backend game server to replicate state."));
 }
 
 namespace UE::MultiServerProxy::Metrics
@@ -447,6 +453,16 @@ void UProxyBackendNetConnection::HandleClientPlayer(APlayerController* NewPlayer
 			}
 		}
 	}
+
+	// The parent connection is not associated with a proxy client connection so it doesn't need to replicate state down to the
+	// proxy. Only child connections are associated with clients. For this reason, replication to this connection on the backend
+	// game server can be disabled.
+	if (!UE::MultiServerProxy::Private::bEnableParentConnectionReplication)
+	{
+		bool bEnableReplication = false;
+		UFunction* Function = ANoPawnPlayerController::StaticClass()->FindFunctionByName(TEXT("ServerEnableReplicationToConnection"));
+		GameServerConnection->Driver->ProcessRemoteFunction(Cast<ANoPawnPlayerController>(NewPlayerController), Function, &bEnableReplication, nullptr, nullptr);
+	}
 }
 
 void UProxyBackendChildNetConnection::HandleClientPlayer(APlayerController* NewPlayerController, UNetConnection* GameServerConnection)
diff --git a/Engine/Source/Runtime/Engine/Classes/GameFramework/PlayerController.h b/Engine/Source/Runtime/Engine/Classes/GameFramework/PlayerController.h
--- a/Engine/Source/Runtime/Engine/Classes/GameFramework/PlayerController.h
+++ b/Engine/Source/Runtime/Engine/Classes/GameFramework/PlayerController.h
@@ -2366,4 +2366,14 @@ public:
 	/** Set a view target position for use in replication relevancy. */
 	UFUNCTION(Server, Unreliable)
 	void ServerSetViewTargetPosition(FVector ViewTargetPos);
+
+	/** When disabled, any connection using this player controller won't replicate actors down to it's client. Enabled by default. */
+	UFUNCTION(Server, Reliable)
+	void ServerEnableReplicationToConnection(bool bEnable);
+
+	bool IsEnableReplicationToConnection() const { return bEnableReplicationToConnection; }
+
+private:
+	
+	bool bEnableReplicationToConnection = true;
 };
diff --git a/Engine/Source/Runtime/Engine/Private/Net/NetworkMetricsDefs.cpp b/Engine/Source/Runtime/Engine/Private/Net/NetworkMetricsDefs.cpp
--- a/Engine/Source/Runtime/Engine/Private/Net/NetworkMetricsDefs.cpp
+++ b/Engine/Source/Runtime/Engine/Private/Net/NetworkMetricsDefs.cpp
@@ -107,6 +107,7 @@ const FName RawPingClientAvg("RawPingClientAvg");
 const FName NumPlayerControllersForReplication("NumPlayerControllersForReplication");
 const FName NumNoPawnPlayerControllersForReplication("NumNoPawnPlayerControllersForReplication");
 const FName NumPawnsForReplication("NumPawnsForReplication");
+const FName NumReplicationViewers("NumReplicationViewers");
 
 }
 
diff --git a/Engine/Source/Runtime/Engine/Private/NetDriver.cpp b/Engine/Source/Runtime/Engine/Private/NetDriver.cpp
--- a/Engine/Source/Runtime/Engine/Private/NetDriver.cpp
+++ b/Engine/Source/Runtime/Engine/Private/NetDriver.cpp
@@ -182,6 +182,7 @@ int32 GNumSaturatedConnections; // Counter for how many connections are skipped/
 int32 GNumSharedSerializationHit;
 int32 GNumSharedSerializationMiss;
 int32 GNumSkippedObjectEmptyUpdates;
+int32 GNumReplicationViewers;
 
 extern int32 GNetRPCDebug;
 
@@ -2070,6 +2071,9 @@ void UNetDriver::SetupNetworkMetrics()
 
 	// The number of instances of APawn registered for replication.
 	GetMetrics()->CreateInt(UE::Net::Metric::NumPawnsForReplication, 0);
+
+	// (Generic only) The number of FNetViewers used during replication.
+	GetMetrics()->CreateInt(UE::Net::Metric::NumReplicationViewers, 0);
 }
 
 void UNetDriver::SetupNetworkMetricsListeners(bool bIsClient)
@@ -5779,17 +5783,40 @@ void UNetDriver::ServerReplicateActors_ForConnection(UE::Net::FServerReplicateAc
 #if WITH_SERVER_CODE
 	const int32 LocalNumSaturated = GNumSaturatedConnections;
 
+	auto ShouldUseConnectionAsViewer = [](UNetConnection* Connection)
+	{
+		if (Connection)
+		{
+			if (ANoPawnPlayerController* NoPawnPlayerController = Cast<ANoPawnPlayerController>(Connection->PlayerController))
+			{
+				return NoPawnPlayerController->IsEnableReplicationToConnection();
+			}
+		}
+
+		return true;
+	};
+
 	// Make a list of viewers this connection should consider (this connection and children of this connection)
 	Params.ConnectionViewers.Reset();
-	new(Params.ConnectionViewers)FNetViewer(Params.Connection, Params.DeltaSeconds);
+
+	if (ShouldUseConnectionAsViewer(Params.Connection))
+	{
+		new(Params.ConnectionViewers)FNetViewer(Params.Connection, Params.DeltaSeconds);
+	}
+
 	for (int32 ViewerIndex = 0; ViewerIndex < Params.Connection->Children.Num(); ViewerIndex++)
 	{
 		if (Params.Connection->Children[ViewerIndex]->ViewTarget != NULL)
 		{
-			new(Params.ConnectionViewers)FNetViewer(Params.Connection->Children[ViewerIndex], Params.DeltaSeconds);
+			if (ShouldUseConnectionAsViewer(Params.Connection->Children[ViewerIndex]))
+			{
+				new(Params.ConnectionViewers)FNetViewer(Params.Connection->Children[ViewerIndex], Params.DeltaSeconds);
+			}
 		}
 	}
 
+	GNumReplicationViewers += Params.ConnectionViewers.Num();
+
 	// send ClientAdjustment if necessary
 	// we do this here so that we send a maximum of one per packet to that client; there is no value in stacking additional corrections
 	if (Params.Connection->PlayerController)
@@ -6015,6 +6042,7 @@ struct FScopedNetDriverStats
 			GNumReplicateActorCalls = 0;
 			GNumSaturatedConnections = 0;
 			GNumSkippedObjectEmptyUpdates = 0;
+			GNumReplicationViewers = 0;
 
 			GReplicationGatherPrioritizeTimeSeconds = 0.f;
 			GServerReplicateActorTimeSeconds = 0.f;
@@ -6055,6 +6083,7 @@ struct FScopedNetDriverStats
 			NetDriver->GetMetrics()->SetFloat(UE::Net::Metric::NumClientUpdateLevelVisibility, static_cast<float>(GNumClientUpdateLevelVisibility));
 			NetDriver->GetMetrics()->SetInt(UE::Net::Metric::SharedSerializationPropertyHit, GNumSharedSerializationHit);
 			NetDriver->GetMetrics()->SetInt(UE::Net::Metric::SharedSerializationPropertyMiss, GNumSharedSerializationMiss);
+			NetDriver->GetMetrics()->SetInt(UE::Net::Metric::NumReplicationViewers, GNumReplicationViewers);
 
 			int32 NumOpenChannels = 0;
 			int32 NumTickingChannels = 0;
diff --git a/Engine/Source/Runtime/Engine/Private/PlayerController.cpp b/Engine/Source/Runtime/Engine/Private/PlayerController.cpp
--- a/Engine/Source/Runtime/Engine/Private/PlayerController.cpp
+++ b/Engine/Source/Runtime/Engine/Private/PlayerController.cpp
@@ -6693,5 +6693,15 @@ void ANoPawnPlayerController::ServerSetViewTargetPosition_Implementation(FVector
 	Actor->SetActorLocation(ViewTargetPos);
 }
 
+void ANoPawnPlayerController::ServerEnableReplicationToConnection_Implementation(bool bEnable)
+{
+	UE_LOG(LogNet, Log, TEXT("%s replication for connections using %s (previously %s)."),
+		   bEnable ? TEXT("Enabled") : TEXT("Disabled"),
+		   *GetName(),
+		   bEnableReplicationToConnection ? TEXT("enabled") : TEXT("disabled"));
+
+	bEnableReplicationToConnection = bEnable;
+}
+
 #undef LOCTEXT_NAMESPACE
 
diff --git a/Engine/Source/Runtime/Engine/Public/Net/NetworkMetricsDefs.h b/Engine/Source/Runtime/Engine/Public/Net/NetworkMetricsDefs.h
--- a/Engine/Source/Runtime/Engine/Public/Net/NetworkMetricsDefs.h
+++ b/Engine/Source/Runtime/Engine/Public/Net/NetworkMetricsDefs.h
@@ -109,4 +109,5 @@ extern ENGINE_API const FName RawPingClientAvg;
 extern ENGINE_API const FName NumPlayerControllersForReplication;
 extern ENGINE_API const FName NumNoPawnPlayerControllersForReplication;
 extern ENGINE_API const FName NumPawnsForReplication;
+extern ENGINE_API const FName NumReplicationViewers;
 }
