Optimization : Reduced number of streaming cell to analyze when computing streaming performance. Instead of going through all active cells, only those that are pending to be added to the world are computed. So when everything is stable and added to the world, there no streaming cell to analyze.

diff --git a/Engine/Source/Runtime/Engine/Private/WorldPartition/WorldPartitionRuntimeHash.cpp b/Engine/Source/Runtime/Engine/Private/WorldPartition/WorldPartitionRuntimeHash.cpp
--- a/Engine/Source/Runtime/Engine/Private/WorldPartition/WorldPartitionRuntimeHash.cpp
+++ b/Engine/Source/Runtime/Engine/Private/WorldPartition/WorldPartitionRuntimeHash.cpp
@@ -621,6 +621,7 @@ EWorldPartitionStreamingPerformance UWorldPartitionRuntimeHash::GetStreamingPerf
 
 EWorldPartitionStreamingPerformance UWorldPartitionRuntimeHash::GetStreamingPerformance(const TSet<const UWorldPartitionRuntimeCell*>& CellsToActivate, bool& bOutShouldBlock) const
 {
+	TRACE_CPUPROFILER_EVENT_SCOPE(UWorldPartitionRuntimeHash::GetStreamingPerformance)
 	EWorldPartitionStreamingPerformance StreamingPerformance = EWorldPartitionStreamingPerformance::Good;
 	bOutShouldBlock = false;
 
diff --git a/Engine/Source/Runtime/Engine/Private/WorldPartition/WorldPartitionStreamingPolicy.cpp b/Engine/Source/Runtime/Engine/Private/WorldPartition/WorldPartitionStreamingPolicy.cpp
--- a/Engine/Source/Runtime/Engine/Private/WorldPartition/WorldPartitionStreamingPolicy.cpp
+++ b/Engine/Source/Runtime/Engine/Private/WorldPartition/WorldPartitionStreamingPolicy.cpp
@@ -397,6 +397,7 @@ void UWorldPartitionStreamingPolicy::UpdateStreamingStateInternal(const FUpdateS
 
 	check(InParams.FrameActivateCells.IsEmpty());
 	check(InParams.FrameLoadCells.IsEmpty());
+	check(InParams.FrameAnalysePerformanceCells.IsEmpty());
 	check(OutTargetState.IsEmpty());
 
 	FWorldPartitionStreamingContext Context(InParams.DataLayersLogicOperator, InParams.GetWorldDataLayersEffectiveStates(), InParams.PolicyUpdateStreamingStateEpoch);
@@ -406,6 +407,7 @@ void UWorldPartitionStreamingPolicy::UpdateStreamingStateInternal(const FUpdateS
 		// Reset frame cells to avoid reallocation at every Update call
 		InParams.FrameActivateCells.Reset();
 		InParams.FrameLoadCells.Reset();
+		InParams.FrameAnalysePerformanceCells.Reset();
 	};
 
 	if (InParams.bCanStream)
@@ -664,7 +666,16 @@ void UWorldPartitionStreamingPolicy::UpdateStreamingStateInternal(const FUpdateS
 		}
 		else
 		{
-			OutTargetState.StreamingPerformance = InParams.RuntimeHash->GetStreamingPerformance(InParams.FrameActivateCells, OutTargetState.bBlockOnSlowStreaming);
+			TRACE_CPUPROFILER_EVENT_SCOPE(UWorldPartitionStreamingPolicy::GetStreamingPerformance);
+			const TSet<const UWorldPartitionRuntimeCell*>& PendingAddToWorldCells = InParams.CurrentState.ActivatedCells.GetPendingAddToWorldCells();
+			for (const UWorldPartitionRuntimeCell* Cell : InParams.FrameActivateCells)
+			{
+				if (PendingAddToWorldCells.Contains(Cell) || !InParams.CurrentState.ActivatedCells.Contains(Cell))
+				{
+					InParams.FrameAnalysePerformanceCells.Add(Cell);
+				}
+			}
+			OutTargetState.StreamingPerformance = InParams.RuntimeHash->GetStreamingPerformance(InParams.FrameAnalysePerformanceCells, OutTargetState.bBlockOnSlowStreaming);
 		}
 	}
 }
@@ -1276,6 +1287,7 @@ UWorldPartitionStreamingPolicy::FUpdateStreamingStateParams::FUpdateStreamingSta
 	, ServerDisallowedStreamingOutDataLayers(InPolicy->GetServerDisallowedStreamingOutDataLayers())
 	, FrameActivateCells(InPolicy->FrameActivateCells)
 	, FrameLoadCells(InPolicy->FrameLoadCells)
+	, FrameAnalysePerformanceCells(InPolicy->FrameAnalysePerformanceCells)
 	, WorldDataLayersEffectiveStatesRef(FWorldDataLayersEffectiveStatesAccessor::Get(InPolicy->GetOuterUWorldPartition()->GetTypedOuter<UWorld>()->GetWorldDataLayers()))
 {
 	check(IsInGameThread());
diff --git a/Engine/Source/Runtime/Engine/Public/WorldPartition/WorldPartitionStreamingPolicy.h b/Engine/Source/Runtime/Engine/Public/WorldPartition/WorldPartitionStreamingPolicy.h
--- a/Engine/Source/Runtime/Engine/Public/WorldPartition/WorldPartitionStreamingPolicy.h
+++ b/Engine/Source/Runtime/Engine/Public/WorldPartition/WorldPartitionStreamingPolicy.h
@@ -191,6 +191,7 @@ private:
 		const TSet<FName>& ServerDisallowedStreamingOutDataLayers;
 		TSet<const UWorldPartitionRuntimeCell*>& FrameActivateCells;
 		TSet<const UWorldPartitionRuntimeCell*>& FrameLoadCells;
+		TSet<const UWorldPartitionRuntimeCell*>& FrameAnalysePerformanceCells;
 
 		const FWorldDataLayersEffectiveStates& GetWorldDataLayersEffectiveStates() const { return WorldDataLayersEffectiveStatesCopy.Get(WorldDataLayersEffectiveStatesRef); }
 	private:
@@ -256,6 +257,7 @@ private:
 	// Used internally by UpdateStreamingStateInternal (avoids re-allocations)
 	mutable TSet<const UWorldPartitionRuntimeCell*> FrameActivateCells;
 	mutable TSet<const UWorldPartitionRuntimeCell*> FrameLoadCells;
+	mutable TSet<const UWorldPartitionRuntimeCell*> FrameAnalysePerformanceCells;
 
 	// CVars to control update optimization
 	static bool IsUpdateOptimEnabled;
