Reduce time spent computing world's CSV stats NumLevelsLoading, NumLevelsMakingVisible and NumLevelsMakingInvisible.

diff --git a/Engine/Source/Runtime/Engine/Private/LevelTick.cpp b/Engine/Source/Runtime/Engine/Private/LevelTick.cpp
--- a/Engine/Source/Runtime/Engine/Private/LevelTick.cpp
+++ b/Engine/Source/Runtime/Engine/Private/LevelTick.cpp
@@ -1958,23 +1958,10 @@ void UWorld::Tick( ELevelTick TickType, float DeltaSeconds )
 
 	if (IsGameWorld() && FCsvProfiler::Get()->IsCapturing())
 	{
-		for (ULevelStreaming* LevelStreaming : StreamingLevels)
+		if (ULevelStreamingProfilingSubsystem* LevelStreamingProfilingSubsystem = GetSubsystem<ULevelStreamingProfilingSubsystem>())
 		{
-			switch (LevelStreaming->GetLevelStreamingStatus())
-			{
-			case LEVEL_Loading:
-				CSV_CUSTOM_STAT(LevelStreaming, NumLevelsLoading, 1, ECsvCustomStatOp::Accumulate);
-				break;
-			case LEVEL_MakingVisible:
-				CSV_CUSTOM_STAT(LevelStreaming, NumLevelsMakingVisible, 1, ECsvCustomStatOp::Accumulate);
-				break;
-			case LEVEL_MakingInvisible:
-				CSV_CUSTOM_STAT(LevelStreaming, NumLevelsMakingInvisible, 1, ECsvCustomStatOp::Accumulate);
-				break;
-			}
+			LevelStreamingProfilingSubsystem->RecordLevelStreamingCountsToCSV();
 		}
-
-		CSV_CUSTOM_STAT(LevelStreaming, NumLevelsPendingPurge, FLevelStreamingGCHelper::GetNumLevelsPendingPurge(), ECsvCustomStatOp::Set);
 	}
 #endif // (CSV_PROFILER_STATS && !UE_BUILD_SHIPPING)
 
diff --git a/Engine/Source/Runtime/Engine/Private/ProfilingDebugging/LevelStreamingProfilingSubsystem.cpp b/Engine/Source/Runtime/Engine/Private/ProfilingDebugging/LevelStreamingProfilingSubsystem.cpp
--- a/Engine/Source/Runtime/Engine/Private/ProfilingDebugging/LevelStreamingProfilingSubsystem.cpp
+++ b/Engine/Source/Runtime/Engine/Private/ProfilingDebugging/LevelStreamingProfilingSubsystem.cpp
@@ -4,6 +4,7 @@
 
 #include "Engine/Level.h"
 #include "Engine/LevelStreaming.h"
+#include "Engine/LevelStreamingGCHelper.h"
 #include "HAL/FileManager.h"
 #include "Misc/Paths.h"
 #include "Streaming/LevelStreamingDelegates.h"
@@ -11,6 +12,7 @@
 #include "WorldPartition/WorldPartition.h"
 #include "WorldPartition/WorldPartitionLevelStreamingDynamic.h"
 #include "Stats/Stats.h"
+#include "Algo/ForEach.h"
 
 #if !WITH_EDITOR
 #include "Engine/World.h"
@@ -372,6 +374,24 @@ void ULevelStreamingProfilingSubsystem::OnLevelStreamingTargetStateChanged(UWorl
 	}
 }
 
+#if (CSV_PROFILER_STATS && !UE_BUILD_SHIPPING)
+void ULevelStreamingProfilingSubsystem::RecordLevelStreamingCountsToCSV()
+{	
+	if (FCsvProfiler::Get()->IsCapturing())
+	{
+		int32 NumLevelsLoading = WorldLevelStreamingStats.LevelsLoading.Num();
+		int32 NumLevelsMakingVisible = 0;
+		Algo::ForEach(WorldLevelStreamingStats.LevelsMakingVisible, [&NumLevelsMakingVisible, this](const ULevelStreaming* LevelStreaming) { NumLevelsMakingVisible += GetWorld()->IsLevelMakingVisible(LevelStreaming->GetLoadedLevel()) ? 1 : 0; });
+		int32 NumLevelsMakingInvisible = WorldLevelStreamingStats.LevelsMakingInvisible.Num();
+
+		CSV_CUSTOM_STAT(LevelStreaming, NumLevelsLoading, NumLevelsLoading, ECsvCustomStatOp::Set);
+		CSV_CUSTOM_STAT(LevelStreaming, NumLevelsMakingVisible, NumLevelsMakingVisible, ECsvCustomStatOp::Set);
+		CSV_CUSTOM_STAT(LevelStreaming, NumLevelsMakingInvisible, NumLevelsMakingInvisible, ECsvCustomStatOp::Set);
+		CSV_CUSTOM_STAT(LevelStreaming, NumLevelsPendingPurge, FLevelStreamingGCHelper::GetNumLevelsPendingPurge(), ECsvCustomStatOp::Set);
+	}
+}
+#endif
+
 void ULevelStreamingProfilingSubsystem::OnLevelStreamingStateChanged(UWorld* World, const ULevelStreaming* StreamingLevel, ULevel* LevelIfLoaded, ELevelStreamingState PrevState, ELevelStreamingState NewState)
 {
 	if (World != GetWorld())
@@ -379,6 +399,25 @@ void ULevelStreamingProfilingSubsystem::OnLevelStreamingStateChanged(UWorld* Wor
 		return;
 	}
 
+#if (CSV_PROFILER_STATS && !UE_BUILD_SHIPPING)
+	if (NewState != PrevState)
+	{
+		WorldLevelStreamingStats.Remove(StreamingLevel);
+		if (NewState == ELevelStreamingState::Loading)
+		{
+			WorldLevelStreamingStats.LevelsLoading.Add(StreamingLevel);
+		}
+		else if (NewState == ELevelStreamingState::MakingVisible)
+		{
+			WorldLevelStreamingStats.LevelsMakingVisible.Add(StreamingLevel);
+		}
+		else if (NewState == ELevelStreamingState::MakingInvisible)
+		{
+			WorldLevelStreamingStats.LevelsMakingInvisible.Add(StreamingLevel);
+		}
+	}
+#endif
+
 	UE_LOG(LogLevelStreamingProfiling, Verbose, TEXT("OnLevelStreamingStateChanged: %s %s. %s -> %s (LoadedLevel: %s)"),
 		*StreamingLevel->GetName(), 
 		*StreamingLevel->GetWorldAssetPackageName(),
diff --git a/Engine/Source/Runtime/Engine/Public/ProfilingDebugging/LevelStreamingProfilingSubsystem.h b/Engine/Source/Runtime/Engine/Public/ProfilingDebugging/LevelStreamingProfilingSubsystem.h
--- a/Engine/Source/Runtime/Engine/Public/ProfilingDebugging/LevelStreamingProfilingSubsystem.h
+++ b/Engine/Source/Runtime/Engine/Public/ProfilingDebugging/LevelStreamingProfilingSubsystem.h
@@ -9,6 +9,7 @@
 #include "Subsystems/WorldSubsystem.h"
 #include "Tasks/Task.h"
 #include "UObject/ObjectKey.h"
+#include "UObject/UObjectAnnotation.h"
 
 #include "LevelStreamingProfilingSubsystem.generated.h"
 
@@ -54,6 +55,10 @@ public:
 	/* Returns the squared distance (e.g. from world partition cell bounds) at which a level is considered to have streamed in too late. */
 	static ENGINE_API double GetLateStreamingDistanceSquared();
 
+#if (CSV_PROFILER_STATS && !UE_BUILD_SHIPPING)
+	ENGINE_API void RecordLevelStreamingCountsToCSV();
+#endif
+
 protected:
 	// WorldSubsystem interface
 	ENGINE_API virtual bool DoesSupportWorldType(const EWorldType::Type WorldType) const;
@@ -196,6 +201,23 @@ private:
 
 	// Possibly-executing task dependend on our recorded stats. Cannot start recording again until this task is complete.
 	UE::Tasks::TTask<void> ReportWritingTask;
+
+#if (CSV_PROFILER_STATS && !UE_BUILD_SHIPPING)
+	struct FWorldLevelStreamingStats
+	{
+		void Remove(const ULevelStreaming* StreamingLevel)
+		{
+			LevelsLoading.Remove(StreamingLevel);
+			LevelsMakingVisible.Remove(StreamingLevel);
+			LevelsMakingInvisible.Remove(StreamingLevel);
+		}
+		TSet<const ULevelStreaming*> LevelsLoading;
+		TSet<const ULevelStreaming*> LevelsMakingVisible;
+		TSet<const ULevelStreaming*> LevelsMakingInvisible;
+	};
+
+	FWorldLevelStreamingStats WorldLevelStreamingStats;
+#endif
 };
 
 
