Turn off DataTableJSON.ExportUsingPropertyVisitor to fix 100% crash when loading living world data tables containing world conditions

diff --git a/Engine/Source/Runtime/Engine/Private/DataTableJSON.cpp b/Engine/Source/Runtime/Engine/Private/DataTableJSON.cpp
--- a/Engine/Source/Runtime/Engine/Private/DataTableJSON.cpp
+++ b/Engine/Source/Runtime/Engine/Private/DataTableJSON.cpp
@@ -65,7 +65,7 @@ namespace UE::DataTableJSON
 	 * in case issues get reported and validate potential regression.
 	 * We could eventually get rid of the FieldIterator implementation.
 	 */
-	bool bExportUsingPropertyVisitor = true;
+	bool bExportUsingPropertyVisitor = false;
 	static FAutoConsoleVariableRef CVarRuntimeValidationContext(
 		TEXT("DataTableJSON.ExportUsingPropertyVisitor"),
 		bExportUsingPropertyVisitor,
