Add LoadGameIfExistsAsync to SaveGameSystem which combines the exists and load checks into one async call to avoid an extra frame delay.
Change GameplayStatics to use this for async loads which fixes a warning when trying to load a nonexistent save game.

diff --git a/Engine/Source/Runtime/Engine/Private/GameplayStatics.cpp b/Engine/Source/Runtime/Engine/Private/GameplayStatics.cpp
--- a/Engine/Source/Runtime/Engine/Private/GameplayStatics.cpp
+++ b/Engine/Source/Runtime/Engine/Private/GameplayStatics.cpp
@@ -2512,7 +2512,7 @@ void UGameplayStatics::AsyncLoadGameFromSlot(const FString& SlotName, const int3
 	{
 		FPlatformUserId PlatformUserId = FPlatformMisc::GetPlatformUserForUserIndex(UserIndex);
 
-		SaveSystem->LoadGameAsync(false, *SlotName, PlatformUserId,
+		SaveSystem->LoadGameIfExistsAsync(false, *SlotName, PlatformUserId,
 			[LoadedDelegate, UserIndex](const FString& SlotName, FPlatformUserId PlatformUserId, bool bSuccess, const TArray<uint8>& Data)
 			{
 				check(IsInGameThread());
diff --git a/Engine/Source/Runtime/Engine/Private/SaveGameSystem.cpp b/Engine/Source/Runtime/Engine/Private/SaveGameSystem.cpp
--- a/Engine/Source/Runtime/Engine/Private/SaveGameSystem.cpp
+++ b/Engine/Source/Runtime/Engine/Private/SaveGameSystem.cpp
@@ -94,6 +94,40 @@ void ISaveGameSystem::LoadGameAsync(bool bAttemptToUseUI, const TCHAR* Name, FPl
 	);
 }
 
+void ISaveGameSystem::LoadGameIfExistsAsync(bool bAttemptToUseUI, const TCHAR* Name, FPlatformUserId PlatformUserId, FSaveGameAsyncLoadCompleteCallback Callback)
+{
+	FString SlotName(Name);
+
+	// start the load operation on a background thread.
+	AsyncTaskPipe.Launch(UE_SOURCE_LOCATION,
+		[this, bAttemptToUseUI, SlotName, PlatformUserId, Callback]()
+		{
+			// load the savegame if it exists
+			TSharedRef<TArray<uint8>> Data = MakeShared<TArray<uint8>>();
+			int32 UserIndex = FPlatformMisc::GetUserIndexForPlatformUser(PlatformUserId);
+
+			const ESaveExistsResult ExistsResult = DoesSaveGameExistWithResult(*SlotName, UserIndex);
+			bool bResult = false;
+
+			if (ExistsResult == ESaveExistsResult::OK)
+			{
+				bResult = LoadGame(bAttemptToUseUI, *SlotName, UserIndex, Data.Get());
+			}
+
+			// trigger the callback on the game thread
+			if (Callback)
+			{
+				OnAsyncComplete(
+					[SlotName, PlatformUserId, bResult, Callback, Data]()
+					{
+						Callback(SlotName, PlatformUserId, bResult, Data.Get());
+					}
+				);
+			}
+		}
+	);
+}
+
 
 void ISaveGameSystem::DeleteGameAsync(bool bAttemptToUseUI, const TCHAR* Name, FPlatformUserId PlatformUserId, FSaveGameAsyncOpCompleteCallback Callback)
 {
diff --git a/Engine/Source/Runtime/Engine/Public/SaveGameSystem.h b/Engine/Source/Runtime/Engine/Public/SaveGameSystem.h
--- a/Engine/Source/Runtime/Engine/Public/SaveGameSystem.h
+++ b/Engine/Source/Runtime/Engine/Public/SaveGameSystem.h
@@ -83,6 +83,9 @@ public:
 	/* Load the named savegame asyncronously. Default implementation runs blocking version on a background thread */
 	ENGINE_API virtual void LoadGameAsync(bool bAttemptToUseUI, const TCHAR* Name, FPlatformUserId PlatformUserId, FSaveGameAsyncLoadCompleteCallback Callback);
 
+	/* Combines DoesSaveGameExistAsync and LoadGameAsync into one operation for simpler usage, treats all errors as the save game not existing */
+	ENGINE_API virtual void LoadGameIfExistsAsync(bool bAttemptToUseUI, const TCHAR* Name, FPlatformUserId PlatformUserId, FSaveGameAsyncLoadCompleteCallback Callback);
+
 	/* Delete the named savegame asyncronously. Default implementation runs blocking version on a background thread */
 	ENGINE_API virtual void DeleteGameAsync(bool bAttemptToUseUI, const TCHAR* Name, FPlatformUserId PlatformUserId, FSaveGameAsyncOpCompleteCallback Callback);
 
