Removed invalid comment (reported by BugHawk)

diff --git a/Engine/Source/Runtime/Engine/Private/WorldPartition/DataLayer/WorldDataLayers.cpp b/Engine/Source/Runtime/Engine/Private/WorldPartition/DataLayer/WorldDataLayers.cpp
--- a/Engine/Source/Runtime/Engine/Private/WorldPartition/DataLayer/WorldDataLayers.cpp
+++ b/Engine/Source/Runtime/Engine/Private/WorldPartition/DataLayer/WorldDataLayers.cpp
@@ -64,9 +64,6 @@ AWorldDataLayers::AWorldDataLayers(const FObjectInitializer& ObjectInitializer)
 	bAlwaysRelevant = true;
 	bReplicates = true;
 	SetNetDormancy(DORM_Initial);
-
-	// Avoid actor from being Destroyed/Recreated when scrubbing a replay
-	// instead AWorldDataLayers::RewindForReplay() gets called to reset this actors state
 	bReplayRewindable = true;
 }
 
