Optimized UMaterialInterface::FindTextureStreamingDataIndexRange.

On Win64 with FName compression enabled, reduced cost from 90us down to 4us.
We hit this very often during render asset streaming (csv RenderAssetStreaming).
High cost is due to usage of FNameLexicalLess.

#tests PIE Hermes_Terrain

diff --git a/Engine/Source/Runtime/Engine/Private/Materials/MaterialInterface.cpp b/Engine/Source/Runtime/Engine/Private/Materials/MaterialInterface.cpp
--- a/Engine/Source/Runtime/Engine/Private/Materials/MaterialInterface.cpp
+++ b/Engine/Source/Runtime/Engine/Private/Materials/MaterialInterface.cpp
@@ -1907,18 +1907,21 @@ bool UMaterialInterface::FindTextureStreamingDataIndexRange(FName TextureName, i
 		return false;
 	}
 
-	const int32 MatchingIndex = Algo::BinarySearchBy(TextureStreamingData, TextureName, &FMaterialTextureInfo::TextureName, FNameLexicalLess());
-	if (MatchingIndex != INDEX_NONE)
-	{
-		// Find the range of entries for this texture. 
-		// This is possible because the same texture could be bound to several register and also be used with different sampling UV.
-		LowerIndex = MatchingIndex;
-		HigherIndex = MatchingIndex;
-		while (HigherIndex + 1 < TextureStreamingData.Num() && TextureStreamingData[HigherIndex + 1].TextureName == TextureName)
+	const int32 NumTextureSteamingData = TextureStreamingData.Num();
+	for (int32 Index = 0; Index < NumTextureSteamingData; ++Index)
+	{
+		if (TextureStreamingData[Index].TextureName == TextureName)
 		{
-			++HigherIndex;
+			// Find the range of entries for this texture. 
+			// This is possible because the same texture could be bound to several register and also be used with different sampling UV.
+			LowerIndex = Index;
+			while (Index + 1 < NumTextureSteamingData && TextureStreamingData[Index + 1].TextureName == TextureName)
+			{
+				++Index;
+			}
+			HigherIndex = Index;
+			return true;
 		}
-		return true;
 	}
 	return false;
 }
