Fix subtitles disappearing by removing old IModularFeature queueing code.

diff --git a/Engine/Source/Runtime/Engine/Private/SoundWave.cpp b/Engine/Source/Runtime/Engine/Private/SoundWave.cpp
--- a/Engine/Source/Runtime/Engine/Private/SoundWave.cpp
+++ b/Engine/Source/Runtime/Engine/Private/SoundWave.cpp
@@ -3548,16 +3548,6 @@ void USoundWave::Parse(FAudioDevice* AudioDevice, const UPTRINT NodeWaveInstance
 		}
 
 		WaveInstance = &HandleStart(ActiveSound, NodeWaveInstanceHash);
-
-		IModularFeatures& ModularFeatures = IModularFeatures::Get();
-		if (ModularFeatures.IsModularFeatureAvailable(ISubtitlesAndClosedCaptionsModule::GetModularFeatureName()))
-		{
-			const ISubtitlesAndClosedCaptionsModule& SubtitlesModule = ModularFeatures.GetModularFeature<ISubtitlesAndClosedCaptionsModule>(ISubtitlesAndClosedCaptionsModule::GetModularFeatureName());
-			if (const UAssetUserData * Subtitle{ GetAssetUserDataOfClass(SubtitlesModule.GetAssetUserDataClass()) })
-			{
-				FSubtitlesAndClosedCaptionsDelegates::QueueSubtitle.ExecuteIfBound(FQueueSubtitleParameters{*Subtitle}, ESubtitleTiming::InternallyTimed);
-			}
-		}
 	}
 	else
 	{
