Optimized USkeleton::GetBlendProfile by reducing string/fname conversion.

On Win64 test client with FName compression enabled, cost is down from 25us to 1us.

diff --git a/Engine/Source/Runtime/Engine/Private/Animation/Skeleton.cpp b/Engine/Source/Runtime/Engine/Private/Animation/Skeleton.cpp
--- a/Engine/Source/Runtime/Engine/Private/Animation/Skeleton.cpp
+++ b/Engine/Source/Runtime/Engine/Private/Animation/Skeleton.cpp
@@ -2232,14 +2232,12 @@ void USkeleton::GetAssetRegistryTags(FAssetRegistryTagsContext Context) const
 
 UBlendProfile* USkeleton::GetBlendProfile(const FName& InProfileName)
 {
-	TObjectPtr<UBlendProfile>* FoundProfile = BlendProfiles.FindByPredicate([InProfileName](const UBlendProfile* Profile)
+	for (const TObjectPtr<UBlendProfile>& Profile : BlendProfiles)
 	{
-		return Profile->GetFName() == InProfileName;
-	});
-
-	if(FoundProfile)
-	{
-		return *FoundProfile;
+		if (Profile->GetFName() == InProfileName)
+		{
+			return Profile;
+		}
 	}
 	return nullptr;
 }
