- Destroy Cascade GPU simulations during OnMarkPendingKill to ensure they are gone before the instances are destroyed

diff --git a/Engine/Source/Runtime/Engine/Private/Particles/FXSystem.cpp b/Engine/Source/Runtime/Engine/Private/Particles/FXSystem.cpp
--- a/Engine/Source/Runtime/Engine/Private/Particles/FXSystem.cpp
+++ b/Engine/Source/Runtime/Engine/Private/Particles/FXSystem.cpp
@@ -264,16 +264,8 @@ FFXSystem::FFXSystem(ERHIFeatureLevel::Type InFeatureLevel, EShaderPlatform InSh
 
 FFXSystem::~FFXSystem()
 {
-	for (FVectorFieldInstance* VectorFieldInstance : VectorFields)
-	{
-		if (VectorFieldInstance)
-		{
-			delete VectorFieldInstance;
-		}
-	}
-	VectorFields.Empty();
-
-	DestroyGPUSimulation();
+	ensure(GPUSimulations.Num() == 0);
+	ensure(VectorFields.Num() == 0);
 }
 
 const FName FFXSystem::Name(TEXT("FFXSystem"));
@@ -304,6 +296,28 @@ void FFXSystem::Tick(UWorld* World, float DeltaSeconds)
 	}
 }
 
+void FFXSystem::OnMarkPendingKill()
+{
+	FFXSystemInterface::OnMarkPendingKill();
+
+	// Enqueue command to destroy simulations as destruction order can matter for Cascade
+	ENQUEUE_RENDER_COMMAND(FOnMarkPendingKill)(
+		[this](FRHICommandListImmediate&)
+		{
+			for (FVectorFieldInstance* VectorFieldInstance : VectorFields)
+			{
+				if (VectorFieldInstance)
+				{
+					delete VectorFieldInstance;
+				}
+			}
+			VectorFields.Empty();
+
+			DestroyGPUSimulation();
+		}
+	);
+}
+
 #if WITH_EDITOR
 void FFXSystem::Suspend()
 {
diff --git a/Engine/Source/Runtime/Engine/Private/Particles/FXSystemPrivate.h b/Engine/Source/Runtime/Engine/Private/Particles/FXSystemPrivate.h
--- a/Engine/Source/Runtime/Engine/Private/Particles/FXSystemPrivate.h
+++ b/Engine/Source/Runtime/Engine/Private/Particles/FXSystemPrivate.h
@@ -199,6 +199,7 @@ public:
 
 	// Begin FFXSystemInterface.
 	virtual void Tick(UWorld* World, float DeltaSeconds) override;
+	virtual void OnMarkPendingKill() override;
 #if WITH_EDITOR
 	virtual void Suspend() override;
 	virtual void Resume() override;
