SplineComponent: Flesh out change notifications

Necessary for an upcoming change which accelerates spline component visualization

diff --git a/Engine/Source/Runtime/Engine/Classes/Components/SplineComponent.h b/Engine/Source/Runtime/Engine/Classes/Components/SplineComponent.h
--- a/Engine/Source/Runtime/Engine/Classes/Components/SplineComponent.h
+++ b/Engine/Source/Runtime/Engine/Classes/Components/SplineComponent.h
@@ -393,8 +393,14 @@ public:
 	ENGINE_API void SetSpline(const FSpline& InSpline);
 	ENGINE_API void SetSpline(const FSplineCurves& InSplineCurves);
 
-	UE_EXPERIMENTAL(5.7, "Spline change notification is experimental while the API is under active development.")
-	FSimpleMulticastDelegate& GetOnSplineSet() { return OnSplineSet; }
+	UE_EXPERIMENTAL(5.7, "Spline change notifications are experimental while the API is under active development.")
+	FSimpleMulticastDelegate& GetOnSplineChanged() { return OnSplineChanged; }
+
+	UE_EXPERIMENTAL(5.7, "Spline update notifications are experimental while the API is under active development.")
+	FSimpleMulticastDelegate& GetOnSplineUpdated() { return OnSplineUpdated; }
+
+	UE_EXPERIMENTAL(5.7, "Spline display setting change notifications are experimental while the API is under active development.")
+	FSimpleMulticastDelegate& GetOnSplineDisplayChanged() { return OnSplineDisplayChanged; }
 	
 	/** Reset all spline data to match the current state of SplineCurves */
 	void SynchronizeSplines() { SetSpline(WarninglessSplineCurves()); }
@@ -991,7 +997,14 @@ protected:
 	/** Returns the length of the specified spline segment up to the parametric value given */
 	ENGINE_API float GetSegmentLength(const int32 Index, const float Param = 1.0f) const;
 
-	FSimpleMulticastDelegate OnSplineSet;
+	/** Delegate broadcast any time the spline is mutated. */
+	FSimpleMulticastDelegate OnSplineChanged;
+
+	/** Delegate broadcast any time UpdateSpline is called (regardless of whether or not the spline has changed). */
+	FSimpleMulticastDelegate OnSplineUpdated;
+
+	/** Delegate broadcast any time the spline's visualization properties have changed. */
+	FSimpleMulticastDelegate OnSplineDisplayChanged;
 	
 private:
 	
diff --git a/Engine/Source/Runtime/Engine/Private/Components/SplineComponent.cpp b/Engine/Source/Runtime/Engine/Private/Components/SplineComponent.cpp
--- a/Engine/Source/Runtime/Engine/Private/Components/SplineComponent.cpp
+++ b/Engine/Source/Runtime/Engine/Private/Components/SplineComponent.cpp
@@ -191,7 +191,7 @@ void USplineComponent::SetDefaultSpline()
     WarninglessSplineCurves().Rotation.Points.Emplace(EndParam, FQuat::Identity, FQuat::Identity, FQuat::Identity, CIM_CurveAuto);
     WarninglessSplineCurves().Scale.Points.Emplace(EndParam, FVector(1.0f), FVector::ZeroVector, FVector::ZeroVector, CIM_CurveAuto);
 
-	OnSplineSet.Broadcast();
+	OnSplineChanged.Broadcast();
 }
 
 void USplineComponent::GetLifetimeReplicatedProps(TArray<FLifetimeProperty>& OutLifetimeProps) const
@@ -391,6 +391,8 @@ void USplineComponent::UpdateSpline()
 		MarkRenderStateDirty();
 	}
 #endif
+
+	OnSplineUpdated.Broadcast();
 }
 
 void USplineComponent::SetOverrideConstructionScript(bool bInOverride)
@@ -873,7 +875,7 @@ void USplineComponent::SetClosedLoop(bool bInClosedLoop, bool bUpdateSpline)
 		UpdateSpline();
 	}
 
-	OnSplineSet.Broadcast();
+	OnSplineChanged.Broadcast();
 }
 
 void USplineComponent::SetClosedLoopAtPosition(bool bInClosedLoop, float Key, bool bUpdateSpline)
@@ -887,7 +889,7 @@ void USplineComponent::SetClosedLoopAtPosition(bool bInClosedLoop, float Key, bo
 		UpdateSpline();
 	}
 
-	OnSplineSet.Broadcast();
+	OnSplineChanged.Broadcast();
 }
 
 bool USplineComponent::IsClosedLoop() const
@@ -1013,7 +1015,7 @@ void USplineComponent::AddPoint(const FSplinePoint& InSplinePoint, bool bUpdateS
 		UpdateSpline();
 	}
 
-	OnSplineSet.Broadcast();
+	OnSplineChanged.Broadcast();
 }
 
 void USplineComponent::AddPoints(const TArray<FSplinePoint>& InSplinePoints, bool bUpdateSpline)
@@ -1034,7 +1036,7 @@ void USplineComponent::AddPoints(const TArray<FSplinePoint>& InSplinePoints, boo
 		UpdateSpline();
 	}
 
-	OnSplineSet.Broadcast();
+	OnSplineChanged.Broadcast();
 }
 
 
@@ -1079,7 +1081,7 @@ void USplineComponent::AddSplinePoint(const FVector& Position, ESplineCoordinate
 		UpdateSpline();
 	}
 
-	OnSplineSet.Broadcast();
+	OnSplineChanged.Broadcast();
 }
 
 void USplineComponent::AddSplinePointAtIndex(const FVector& Position, int32 Index, ESplineCoordinateSpace::Type CoordinateSpace, bool bUpdateSpline)
@@ -1132,7 +1134,7 @@ void USplineComponent::AddSplinePointAtIndex(const FVector& Position, int32 Inde
 		UpdateSpline();
 	}
 
-	OnSplineSet.Broadcast();
+	OnSplineChanged.Broadcast();
 }
 
 void USplineComponent::RemoveSplinePoint(int32 Index, bool bUpdateSpline)
@@ -1174,7 +1176,7 @@ void USplineComponent::RemoveSplinePoint(int32 Index, bool bUpdateSpline)
 		UpdateSpline();
 	}
 
-	OnSplineSet.Broadcast();
+	OnSplineChanged.Broadcast();
 }
 
 void USplineComponent::SetSplinePoints(const TArray<FVector>& Points, ESplineCoordinateSpace::Type CoordinateSpace, bool bUpdateSpline)
@@ -1228,7 +1230,7 @@ void USplineComponent::SetSplinePoints(const TArray<FVector>& Points, ESplineCoo
 		UpdateSpline();
 	}
 
-	OnSplineSet.Broadcast();
+	OnSplineChanged.Broadcast();
 }
 
 void USplineComponent::SetLocationAtSplinePoint(int32 PointIndex, const FVector& InLocation, ESplineCoordinateSpace::Type CoordinateSpace, bool bUpdateSpline)
@@ -1248,7 +1250,7 @@ void USplineComponent::SetLocationAtSplinePoint(int32 PointIndex, const FVector&
 			UpdateSpline();
 		}
 
-		OnSplineSet.Broadcast();
+		OnSplineChanged.Broadcast();
 	}
 }
 
@@ -1280,7 +1282,7 @@ void USplineComponent::SetTangentsAtSplinePoint(int32 PointIndex, const FVector&
 			UpdateSpline();
 		}
 
-		OnSplineSet.Broadcast();
+		OnSplineChanged.Broadcast();
 	}
 }
 
@@ -1302,7 +1304,7 @@ void USplineComponent::SetUpVectorAtSplinePoint(int32 PointIndex, const FVector&
 			UpdateSpline();
 		}
 
-		OnSplineSet.Broadcast();
+		OnSplineChanged.Broadcast();
 	}
 }
 
@@ -1345,7 +1347,7 @@ void USplineComponent::SetQuaternionAtSplinePoint(int32 PointIndex, const FQuat&
 			UpdateSpline();
 		}
 
-		OnSplineSet.Broadcast();
+		OnSplineChanged.Broadcast();
 	}
 }
 
@@ -1363,7 +1365,7 @@ void USplineComponent::SetScaleAtSplinePoint(int32 PointIndex, const FVector& In
 			UpdateSpline();
 		}
 
-		OnSplineSet.Broadcast();
+		OnSplineChanged.Broadcast();
 	}
 }
 
@@ -1397,7 +1399,7 @@ void USplineComponent::SetSplinePointType(int32 PointIndex, ESplinePointType::Ty
 			UpdateSpline();
 		}
 
-		OnSplineSet.Broadcast();
+		OnSplineChanged.Broadcast();
 	}
 }
 
@@ -1613,7 +1615,7 @@ void USplineComponent::SetDefaultUpVector(const FVector& UpVector, ESplineCoordi
 
 	UpdateSpline();
 
-	OnSplineSet.Broadcast();
+	OnSplineChanged.Broadcast();
 }
 
 FVector USplineComponent::GetDefaultUpVector(ESplineCoordinateSpace::Type CoordinateSpace) const
@@ -2498,7 +2500,7 @@ FPrimitiveSceneProxy* USplineComponent::CreateSceneProxy()
 		{
 			CSV_SCOPED_TIMING_STAT_GLOBAL(FSplineSceneProxy_FSplineSceneProxy)
 			
-			constexpr int32 NumStepsPerSegment = 40;
+			constexpr int32 NumStepsPerSegment = 21;
 			
 			const int32 NumPoints = SplineInfo.Points.Num();
 			const int32 NumSegments = FMath::Max(SplineInfo.bIsLooped ? NumPoints : NumPoints - 1, 0);
@@ -2633,15 +2635,20 @@ FPrimitiveSceneProxy* USplineComponent::CreateSceneProxy()
 				// If not the first keypoint, draw a line to the previous keypoint.
 				if (KeyIdx > 0)
 				{
+					// This initial value of 1.4f is an experimentally determined value which gives the same effective segment thickness as the
+					// previous PDI rendering when the level editor thickness adjustment setting was at its default value (0). PDI special cases
+					// lines of 0 thickness while our approach does not, so if we are not artificially increasing it here we will get invisible
+					// lines by default.
+					float SegmentLineThickness = 1.4f;
+#if WITH_EDITOR
+					SegmentLineThickness += GetDefault<ULevelEditorViewportSettings>()->SplineLineThicknessAdjustment;
+#endif
+					
 					// For constant interpolation - don't draw ticks - just draw dotted line.
 					if (SplineInfo.Points[KeyIdx - 1].InterpMode == CIM_Constant)
 					{
 						FVector OldPos = OldKeyPos;
 						
-						float SegmentLineThickness = 1.4f;
-#if WITH_EDITOR
-						SegmentLineThickness += GetDefault<ULevelEditorViewportSettings>()->SplineLineThicknessAdjustment * 6.f;
-#endif
 						const FVector SegmentVector = SplineInfo.Points[KeyIdx].OutVal - OldKeyPos;
 						const float SegmentLength = SegmentVector.Length();
 						const FVector SegmentDirection = SegmentVector.GetSafeNormal();
@@ -2657,7 +2664,7 @@ FPrimitiveSceneProxy* USplineComponent::CreateSceneProxy()
 							// wasteful, we do it for 2 reasons:
 							// 1) Consistent logic for subdivision makes preallocating the vertex buffers trivial
 							// 2) We are rendering lines in a much more efficient way than before
-							const bool bHideLine = SplineInfo.Points[KeyIdx - 1].InterpMode == CIM_Constant && StepIdx % 2;
+							const bool bHideLine = SplineInfo.Points[KeyIdx - 1].InterpMode == CIM_Constant && StepIdx % 2 == 0;
 							AppendLineToBuffers(OldPos, NewPos, LineColor.ToFColorSRGB(), bHideLine ? 0.f : SegmentLineThickness);
 
 							OldPos = NewPos;
@@ -2667,15 +2674,6 @@ FPrimitiveSceneProxy* USplineComponent::CreateSceneProxy()
 					{
 						FVector OldPos = OldKeyPos;
 
-						// This initial value of 1.4f is an experimentally determined value which gives the same effective segment thickness as the
-						// previous PDI rendering when the level editor thickness adjustment setting was at its default value (0). PDI special cases
-						// lines of 0 thickness while our approach does not, so if we are not artificially increasing it here we will get invisible
-						// lines by default.
-						float SegmentLineThickness = 1.4f;
-#if WITH_EDITOR
-						SegmentLineThickness += GetDefault<ULevelEditorViewportSettings>()->SplineLineThicknessAdjustment;
-#endif
-
 						for (int32 StepIdx = 1; StepIdx <= NumStepsPerSegment; StepIdx++)
 						{
 							const float Key = (KeyIdx - 1) + (StepIdx / static_cast<float>(NumStepsPerSegment));
@@ -3099,7 +3097,7 @@ void USplineComponent::SetSpline(const FSpline& InSpline)
 	WarninglessSplineCurves() = FSplineCurves::FromSplineInterface(InSpline);
 	Spline = InSpline;
 	
-	OnSplineSet.Broadcast();
+	OnSplineChanged.Broadcast();
 }
 
 void USplineComponent::SetSpline(const FSplineCurves& InSplineCurves)
@@ -3107,7 +3105,7 @@ void USplineComponent::SetSpline(const FSplineCurves& InSplineCurves)
 	WarninglessSplineCurves() = InSplineCurves;
 	Spline = InSplineCurves;
 
-	OnSplineSet.Broadcast();
+	OnSplineChanged.Broadcast();
 }
 
 int32 USplineComponent::GetVersion() const
@@ -3192,8 +3190,6 @@ void USplineComponent::ApplyComponentInstanceData(FSplineInstanceData* SplineIns
 	bSplineHasBeenEdited = SplineInstanceData->bSplineHasBeenEdited;
 
 	UpdateSpline();
-
-	OnSplineSet.Broadcast();
 }
 
 void USplineComponent::ApplyComponentInstanceData(struct FSplineComponentInstanceData* SplineInstanceData, const bool bPostUCS)
@@ -3238,8 +3234,6 @@ void USplineComponent::ApplyComponentInstanceData(struct FSplineComponentInstanc
 	bSplineHasBeenEdited = SplineInstanceData->bSplineHasBeenEdited;
 
 	UpdateSpline();
-
-	OnSplineSet.Broadcast();
 }
 
 #if WITH_EDITOR
@@ -3247,13 +3241,13 @@ void USplineComponent::PostEditChangeChainProperty(FPropertyChangedChainEvent& P
 {
 	if (PropertyChangedEvent.Property != nullptr)
 	{
+		const FName PropertyName(PropertyChangedEvent.Property->GetFName());
+		
 		static const FName ReparamStepsPerSegmentName = GET_MEMBER_NAME_CHECKED(USplineComponent, ReparamStepsPerSegment);
 		static const FName StationaryEndpointsName = GET_MEMBER_NAME_CHECKED(USplineComponent, bStationaryEndpoints);
 		static const FName DefaultUpVectorName = GET_MEMBER_NAME_CHECKED(USplineComponent, DefaultUpVector);
 		static const FName ClosedLoopName = GET_MEMBER_NAME_CHECKED(USplineComponent, bClosedLoop);
-		static const FName SplineHasBeenEditedName = GET_MEMBER_NAME_CHECKED(USplineComponent, bSplineHasBeenEdited);
 
-		const FName PropertyName(PropertyChangedEvent.Property->GetFName());
 		if (PropertyName == ReparamStepsPerSegmentName ||
 			PropertyName == StationaryEndpointsName ||
 			PropertyName == DefaultUpVectorName ||
@@ -3261,8 +3255,28 @@ void USplineComponent::PostEditChangeChainProperty(FPropertyChangedChainEvent& P
 		{
 			UpdateSpline();
 
-			OnSplineSet.Broadcast();
+			OnSplineChanged.Broadcast();
 		}
+
+		static const FName EditorUnselectedSplineSegmentColorName = GET_MEMBER_NAME_CHECKED(USplineComponent, EditorUnselectedSplineSegmentColor);
+		static const FName EditorSelectedSplineSegmentColorName = GET_MEMBER_NAME_CHECKED(USplineComponent, EditorSelectedSplineSegmentColor);
+		static const FName EditorTangentColorName = GET_MEMBER_NAME_CHECKED(USplineComponent, EditorTangentColor);
+		static const FName AllowDiscontinuousSplineName = GET_MEMBER_NAME_CHECKED(USplineComponent, bAllowDiscontinuousSpline);
+		static const FName AdjustTangentsOnSnapName = GET_MEMBER_NAME_CHECKED(USplineComponent, bAdjustTangentsOnSnap);
+		static const FName ShouldVisualizeScaleName = GET_MEMBER_NAME_CHECKED(USplineComponent, bShouldVisualizeScale);
+		static const FName ScaleVisualizationWidthName = GET_MEMBER_NAME_CHECKED(USplineComponent, ScaleVisualizationWidth);
+
+		if (PropertyName == EditorUnselectedSplineSegmentColorName ||
+			PropertyName == EditorSelectedSplineSegmentColorName ||
+			PropertyName == EditorTangentColorName ||
+			PropertyName == AllowDiscontinuousSplineName ||
+			PropertyName == AdjustTangentsOnSnapName ||
+			PropertyName == ShouldVisualizeScaleName ||
+			PropertyName == ScaleVisualizationWidthName)
+		{
+			OnSplineDisplayChanged.Broadcast();
+		}
+		
 	}
 
 	Super::PostEditChangeChainProperty(PropertyChangedEvent);
