Make sure to remove duplicate actor Data Layers. Also avoid EDL from being part of the non-external data layer array.

diff --git a/Engine/Source/Runtime/Engine/Private/ActorEditor.cpp b/Engine/Source/Runtime/Engine/Private/ActorEditor.cpp
--- a/Engine/Source/Runtime/Engine/Private/ActorEditor.cpp
+++ b/Engine/Source/Runtime/Engine/Private/ActorEditor.cpp
@@ -41,6 +41,7 @@
 #include "UObject/SoftObjectPtr.h"
 #include "Algo/Transform.h"
 #include "Algo/RemoveIf.h"
+#include "Algo/Unique.h"
 #include "Misc/DataValidation.h"
 #include "DeletedObjectPlaceholder.h"
 #include "DataStorage/Features.h"
@@ -1860,11 +1861,17 @@ void AActor::FixupDataLayers(bool bRevertChangesOnLockedDataLayer /*= false*/)
 	}
 	
 	// Cleanup Data Layer assets we can't reference (Private DLs)
-	DataLayerAssets.SetNum(Algo::RemoveIf(DataLayerAssets, [this, &InvalidDataLayerAssets](const TSoftObjectPtr<UDataLayerAsset>& AssetPath)
+	DataLayerAssets.SetNum(Algo::RemoveIf(DataLayerAssets, [this, &InvalidDataLayerAssets](const TSoftObjectPtr<UDataLayerAsset>& DataLayerAsset)
 	{
-		return InvalidDataLayerAssets.Contains(AssetPath) || !UDataLayerAsset::CanBeReferencedByActor(AssetPath, this);
+		return InvalidDataLayerAssets.Contains(DataLayerAsset) ||
+			!UDataLayerAsset::CanBeReferencedByActor(DataLayerAsset, this) ||
+			// Make sure there are no EDLs in DataLayerAssets
+			(DataLayerAsset.IsValid() && DataLayerAsset->IsA<UExternalDataLayerAsset>());
 	}));
 
+	// Remove duplicates
+	DataLayerAssets.SetNum(Algo::Unique(DataLayerAssets));
+
 	// Use Actor's DataLayerManager since the fixup is relative to this level
 	UDataLayerManager* DataLayerManager = UDataLayerManager::GetDataLayerManager(this);
 	if (!DataLayerManager || !DataLayerManager->CanResolveDataLayers())
diff --git a/Engine/Source/Runtime/Engine/Public/WorldPartition/DataLayer/DataLayerInstanceNames.h b/Engine/Source/Runtime/Engine/Public/WorldPartition/DataLayer/DataLayerInstanceNames.h
--- a/Engine/Source/Runtime/Engine/Public/WorldPartition/DataLayer/DataLayerInstanceNames.h
+++ b/Engine/Source/Runtime/Engine/Public/WorldPartition/DataLayer/DataLayerInstanceNames.h
@@ -19,25 +19,16 @@ struct FDataLayerInstanceNames
 	}
 
 	FDataLayerInstanceNames(const TArray<FName>& InDataLayers, bool bInIsFirstDataLayerExternal)
-	: bIsFirstDataLayerExternal(bInIsFirstDataLayerExternal)
-	, DataLayers(InDataLayers)
 	{
-#if WITH_EDITOR
-		bIsForcedEmptyNonExternalDataLayers = false;
-#endif
+		const int32 NonExternalDataLayersOffset = bInIsFirstDataLayerExternal && InDataLayers.Num() ? 1 : 0;
+		const int32 NumNonExternalDataLayers = InDataLayers.Num() - NonExternalDataLayersOffset;
+		const FName ExternalDataLayer = (NonExternalDataLayersOffset > 0) ? InDataLayers[0] : NAME_None;
+		Initialize(MakeArrayView(NumNonExternalDataLayers ? &InDataLayers[NonExternalDataLayersOffset] : nullptr, NumNonExternalDataLayers), ExternalDataLayer);
 	}
 
 	FDataLayerInstanceNames(const TArray<FName>& InNonExternalDataLayers, FName InExternalDataLayer)
 	{
-#if WITH_EDITOR
-		bIsForcedEmptyNonExternalDataLayers = false;
-#endif
-		bIsFirstDataLayerExternal = !InExternalDataLayer.IsNone();
-		if (bIsFirstDataLayerExternal)
-		{
-			DataLayers.Add(InExternalDataLayer);
-		}
-		DataLayers.Append(InNonExternalDataLayers);
+		Initialize(InNonExternalDataLayers, InExternalDataLayer);
 	}
 
 	const FName GetExternalDataLayer() const
@@ -75,6 +66,28 @@ struct FDataLayerInstanceNames
 	const TArray<FName>& ToArray() const { return DataLayers; }
 
 private:
+
+	void Initialize(TArrayView<const FName> InNonExternalDataLayers, FName InExternalDataLayer)
+	{
+#if WITH_EDITOR
+		bIsForcedEmptyNonExternalDataLayers = false;
+#endif
+		bIsFirstDataLayerExternal = !InExternalDataLayer.IsNone();
+
+		DataLayers.Reserve(InNonExternalDataLayers.Num() + (bIsFirstDataLayerExternal ? 1 : 0));
+		if (bIsFirstDataLayerExternal)
+		{
+			DataLayers.Add(InExternalDataLayer);
+		}
+		for (FName DataLayer : InNonExternalDataLayers)
+		{
+			if (DataLayer != InExternalDataLayer)
+			{
+				DataLayers.AddUnique(DataLayer);
+			}
+		}
+	}
+
 #if WITH_EDITOR
 	bool IsForcedEmptyNonExternalDataLayers() const { return bIsForcedEmptyNonExternalDataLayers; }
 	void ForceEmptyNonExternalDataLayers()
