[MIR] Implement runtime virtual textures (RVT) in new material translator.

This generalizes some of the code from UMaterialExpressionRuntimeVirtualTextureSample::Compile() to be shared with the new material translator
by introducing FRuntimeVirtualTextureUnpackProperties to pass properties and validation to the compile functions.
A few more functions have been moved to UE::MaterialTranslatorUtils namespace.

This also adds MIR::FRuntimeVirtualTextureObject as RVTs do not share the same UTexture class.
In order to share functionality in MIR between UTexture and URuntimeVirtualTexture, some functions return UObject now instead of UTexture, similar to FMaterialParameterValue::AsTextureObject().

Several parameters of all the TextureSample* functions in the MIR emitter are now using FTextureSampleBaseAttributes to more easily pass a larger set of parameters that are shared across all texture sampling functions.

This change improves generally VT support as the encoding of page table and stack layers in MIR now align more closely to the old translator.

#test Ran NewTranslator test suite in EngineTest without regression
#rnx

diff --git a/Engine/Source/Runtime/Engine/Private/Materials/MaterialExpressionNoise.cpp b/Engine/Source/Runtime/Engine/Private/Materials/MaterialExpressionNoise.cpp
--- a/Engine/Source/Runtime/Engine/Private/Materials/MaterialExpressionNoise.cpp
+++ b/Engine/Source/Runtime/Engine/Private/Materials/MaterialExpressionNoise.cpp
@@ -7,6 +7,7 @@
 #include "Materials/MaterialExpressionNoise.h"
 #include "Materials/MaterialExpressionVectorNoise.h"
 #include "Materials/HLSLMaterialTranslator.h"
+#include "Materials/MaterialSharedPrivate.h"
 #if WITH_EDITOR
 #include "MaterialGraph/MaterialGraphSchema.h"
 #endif
@@ -16,11 +17,6 @@
 #define LOCTEXT_NAMESPACE "MaterialExpressionNoise"
 
 
-// Extern function from MaterialExpressions.cpp
-FName GetWorldPositionInputName(EPositionOrigin PositionOrigin);
-EWorldPositionIncludedOffsets GetWorldPositionTypeWithOrigin(EPositionOrigin PositionOrigin, bool IncludeOffsets);
-
-
 ///////////////////////////////////////////////////////////////////////////////
 // UMaterialExpressionNoise
 ///////////////////////////////////////////////////////////////////////////////
@@ -90,7 +86,7 @@ FName UMaterialExpressionNoise::GetInputName(int32 InputIndex) const
 {
 	if (GetInput(InputIndex) == &Position)
 	{
-		return GetWorldPositionInputName(WorldPositionOriginType);
+		return UE::MaterialTranslatorUtils::GetWorldPositionInputName(WorldPositionOriginType);
 	}
 
 	return Super::GetInputName(InputIndex);
@@ -120,7 +116,7 @@ int32 UMaterialExpressionNoise::Compile(class FMaterialCompiler* Compiler, int32
 	}
 	else
 	{
-		PositionInput = Compiler->WorldPosition(GetWorldPositionTypeWithOrigin(WorldPositionOriginType, true));
+		PositionInput = Compiler->WorldPosition(UE::MaterialTranslatorUtils::GetWorldPositionTypeWithOrigin(WorldPositionOriginType, true));
 	}
 
 	int32 FilterWidthInput;
@@ -203,7 +199,7 @@ FName UMaterialExpressionVectorNoise::GetInputName(int32 InputIndex) const
 {
 	if (GetInput(InputIndex) == &Position)
 	{
-		return GetWorldPositionInputName(WorldPositionOriginType);
+		return UE::MaterialTranslatorUtils::GetWorldPositionInputName(WorldPositionOriginType);
 	}
 
 	return Super::GetInputName(InputIndex);
@@ -233,7 +229,7 @@ int32 UMaterialExpressionVectorNoise::Compile(class FMaterialCompiler* Compiler,
 	}
 	else
 	{
-		PositionInput = Compiler->WorldPosition(GetWorldPositionTypeWithOrigin(WorldPositionOriginType, true));
+		PositionInput = Compiler->WorldPosition(UE::MaterialTranslatorUtils::GetWorldPositionTypeWithOrigin(WorldPositionOriginType, true));
 	}
 
 	return Compiler->VectorNoise(PositionInput, WorldPositionOriginType, Quality, NoiseFunction, bTiling, TileSize);
diff --git a/Engine/Source/Runtime/Engine/Private/Materials/MaterialExpressions.cpp b/Engine/Source/Runtime/Engine/Private/Materials/MaterialExpressions.cpp
--- a/Engine/Source/Runtime/Engine/Private/Materials/MaterialExpressions.cpp
+++ b/Engine/Source/Runtime/Engine/Private/Materials/MaterialExpressions.cpp
@@ -753,28 +753,6 @@ int32 CompileSubstrateBlendFunction(FMaterialCompiler* Compiler, const int32 A,
 	return INDEX_NONE;
 }
 
-extern EWorldPositionIncludedOffsets GetWorldPositionTypeWithOrigin(EPositionOrigin PositionOrigin, bool IncludeOffsets = true)
-{
-	switch (PositionOrigin)
-	{
-	case EPositionOrigin::Absolute: return IncludeOffsets ? WPT_Default : WPT_ExcludeAllShaderOffsets;
-	case EPositionOrigin::CameraRelative: return IncludeOffsets ? WPT_CameraRelative : WPT_CameraRelativeNoOffsets;
-	default: checkNoEntry();
-	}
-	return WPT_Default;
-}
-
-extern FName GetWorldPositionInputName(EPositionOrigin PositionOrigin)
-{
-	switch (PositionOrigin)
-	{
-	case EPositionOrigin::CameraRelative: return TEXT("Translated World Position");
-	case EPositionOrigin::Absolute: return TEXT("World Position");
-	default: checkNoEntry();
-	}
-	return FName();
-}
-
 void FMaterialExpressionCollection::AddExpression(UMaterialExpression* InExpression)
 {
 	Expressions.AddUnique(InExpression);
@@ -3185,6 +3163,22 @@ bool UMaterialExpressionRuntimeVirtualTextureSample::InitVirtualTextureDependent
 	return bChanged;
 }
 
+ESamplerSourceMode UMaterialExpressionRuntimeVirtualTextureSample::GetSamplerSourceMode() const
+{
+	// Convert texture address mode to matching sampler source mode.
+	// Would be better if ESamplerSourceMode had a Mirror enum that we could also use...
+	switch (TextureAddressMode)
+	{
+	case RVTTA_Clamp:
+		return SSM_Clamp_WorldGroupSettings;
+	case RVTTA_Wrap:
+		return SSM_Wrap_WorldGroupSettings;
+	case RVTTA_MAX:
+		break;
+	}
+	return SSM_Clamp_WorldGroupSettings;
+}
+
 void UMaterialExpressionRuntimeVirtualTextureSample::InitOutputs()
 {
 #if WITH_EDITORONLY_DATA
@@ -3242,7 +3236,7 @@ FName UMaterialExpressionRuntimeVirtualTextureSample::GetInputName(int32 InputIn
 	switch (InputIndex)
 	{
 	case 1:
-		return GetWorldPositionInputName(WorldPositionOriginType);
+		return UE::MaterialTranslatorUtils::GetWorldPositionInputName(WorldPositionOriginType);
 	case 2:
 		if (MipValueMode == RVTMVM_MipLevel)
 		{
@@ -3324,18 +3318,19 @@ void UMaterialExpressionRuntimeVirtualTextureSample::PostEditChangeProperty(FPro
 	Super::PostEditChangeProperty(PropertyChangedEvent);
 }
 
-int32 UMaterialExpressionRuntimeVirtualTextureSample::Compile(class FMaterialCompiler* Compiler, int32 OutputIndex)
+bool UMaterialExpressionRuntimeVirtualTextureSample::IsParameter() const
 {
-	// Is this a valid UMaterialExpressionRuntimeVirtualTextureSampleParameter?
-	const bool bIsParameter = HasAParameterName() && GetParameterName().IsValid() && !GetParameterName().IsNone();
+	return HasAParameterName() && GetParameterName().IsValid() && !GetParameterName().IsNone();
+}
 
-	// Check validity of current virtual texture
-	bool bIsVirtualTextureValid = VirtualTexture != nullptr;
-	if (!bIsVirtualTextureValid)
+bool UMaterialExpressionRuntimeVirtualTextureSample::ValidateVirtualTextureParameters(FString& OutError) const
+{
+	if (!VirtualTexture)
 	{
-		if (bIsParameter)
+		if (IsParameter())
 		{
-			return Compiler->Error(TEXT("Missing input Virtual Texture"));
+			OutError = TEXT("Missing input Virtual Texture");
+			return false;
 		}
 	}
 	else if (VirtualTexture->GetMaterialType() != MaterialType)
@@ -3344,167 +3339,206 @@ int32 UMaterialExpressionRuntimeVirtualTextureSample::Compile(class FMaterialCom
 		FString MaterialTypeDisplayName = Enum->GetDisplayNameTextByValue((int64)MaterialType).ToString();
 		FString TextureTypeDisplayName = Enum->GetDisplayNameTextByValue((int64)VirtualTexture->GetMaterialType()).ToString();
 
-		Compiler->Errorf(TEXT("%Material type is '%s', should be '%s' to match %s"),
+		OutError = FString::Printf(TEXT("Material type is '%s', should be '%s' to match %s"),
 			*MaterialTypeDisplayName,
 			*TextureTypeDisplayName,
 			*VirtualTexture->GetName());
 
-		bIsVirtualTextureValid = false;
+		return false;
 	}
 	else if (VirtualTexture->GetSinglePhysicalSpace() != bSinglePhysicalSpace)
 	{
-		Compiler->Errorf(TEXT("%Page table packing is '%d', should be '%d' to match %s"),
+		OutError = FString::Printf(TEXT("Page table packing is '%d', should be '%d' to match %s"),
 			bSinglePhysicalSpace ? 1 : 0,
 			VirtualTexture->GetSinglePhysicalSpace() ? 1 : 0,
 			*VirtualTexture->GetName());
 
-		bIsVirtualTextureValid = false;
+		return false;
 	}
 	else if ((VirtualTexture->GetAdaptivePageTable()) != bAdaptive)
 	{
-		Compiler->Errorf(TEXT("Adaptive page table is '%d', should be '%d' to match %s"),
+		OutError = FString::Printf(TEXT("Adaptive page table is '%d', should be '%d' to match %s"),
 			bAdaptive ? 1 : 0,
 			VirtualTexture->GetAdaptivePageTable() ? 1 : 0,
 			*VirtualTexture->GetName());
 
-		bIsVirtualTextureValid = false;
+		return false;
 	}
 
-	// Calculate the virtual texture layer and sampling/unpacking functions for this output
-	// Fallback to a sensible default value if the output isn't valid for the bound virtual texture
-	uint32 UnpackTarget = 0;
-	uint32 UnpackMask = 0;
-	EVirtualTextureUnpackType UnpackType = EVirtualTextureUnpackType::None;
-
-	bool bIsBaseColorValid = false;
-	bool bIsSpecularValid = false;
-	bool bIsRoughnessValid = false;
-	bool bIsNormalValid = false;
-	bool bIsWorldHeightValid = false;
-	bool bIsMaskValid = false;
-	bool bIsMask4Valid = false;
-	bool bIsDisplacementValid = false;
+	return true;
+}
+
+bool UMaterialExpressionRuntimeVirtualTextureSample::GetRVTUnpackProperties(int32 OutputIndex, bool bIsVirtualTextureValid, FRuntimeVirtualTextureUnpackProperties& Output) const
+{
+	Output = FRuntimeVirtualTextureUnpackProperties{};
 
 	switch (MaterialType)
 	{
-	case ERuntimeVirtualTextureMaterialType::BaseColor: bIsBaseColorValid = true; break;
-	case ERuntimeVirtualTextureMaterialType::BaseColor_Normal_Roughness: bIsBaseColorValid = bIsNormalValid = bIsRoughnessValid = true; break;
-	case ERuntimeVirtualTextureMaterialType::BaseColor_Normal_Specular: bIsRoughnessValid = bIsBaseColorValid = bIsNormalValid = bIsSpecularValid = true; break;
-	case ERuntimeVirtualTextureMaterialType::BaseColor_Normal_Specular_YCoCg: bIsRoughnessValid = bIsBaseColorValid = bIsNormalValid = bIsSpecularValid = true; break;
-	case ERuntimeVirtualTextureMaterialType::BaseColor_Normal_Specular_Mask_YCoCg: bIsRoughnessValid = bIsBaseColorValid = bIsNormalValid = bIsSpecularValid = bIsMaskValid = true; break;
-	case ERuntimeVirtualTextureMaterialType::Mask4: bIsMask4Valid = true; break;
-	case ERuntimeVirtualTextureMaterialType::WorldHeight: bIsWorldHeightValid = true; break;
-	case ERuntimeVirtualTextureMaterialType::Displacement: bIsDisplacementValid = true; break;
+	case ERuntimeVirtualTextureMaterialType::BaseColor: Output.bIsBaseColorValid = true; break;
+	case ERuntimeVirtualTextureMaterialType::BaseColor_Normal_Roughness: Output.bIsBaseColorValid = Output.bIsNormalValid = Output.bIsRoughnessValid = true; break;
+	case ERuntimeVirtualTextureMaterialType::BaseColor_Normal_Specular: Output.bIsRoughnessValid = Output.bIsBaseColorValid = Output.bIsNormalValid = Output.bIsSpecularValid = true; break;
+	case ERuntimeVirtualTextureMaterialType::BaseColor_Normal_Specular_YCoCg: Output.bIsRoughnessValid = Output.bIsBaseColorValid = Output.bIsNormalValid = Output.bIsSpecularValid = true; break;
+	case ERuntimeVirtualTextureMaterialType::BaseColor_Normal_Specular_Mask_YCoCg: Output.bIsRoughnessValid = Output.bIsBaseColorValid = Output.bIsNormalValid = Output.bIsSpecularValid = Output.bIsMaskValid = true; break;
+	case ERuntimeVirtualTextureMaterialType::Mask4: Output.bIsMask4Valid = true; break;
+	case ERuntimeVirtualTextureMaterialType::WorldHeight: Output.bIsWorldHeightValid = true; break;
+	case ERuntimeVirtualTextureMaterialType::Displacement: Output.bIsDisplacementValid = true; break;
 	}
 
 	switch (OutputIndex)
 	{
-	case 0: 
-		if (bIsVirtualTextureValid && bIsBaseColorValid)
+	case 0:
+		if (bIsVirtualTextureValid && Output.bIsBaseColorValid)
 		{
 			switch (MaterialType)
 			{
-			case ERuntimeVirtualTextureMaterialType::BaseColor_Normal_Specular_YCoCg:
-			case ERuntimeVirtualTextureMaterialType::BaseColor_Normal_Specular_Mask_YCoCg: UnpackType = EVirtualTextureUnpackType::BaseColorYCoCg; break;
-			case ERuntimeVirtualTextureMaterialType::BaseColor_Normal_Roughness: UnpackType = EVirtualTextureUnpackType::BaseColorSRGB; break;
-			default: UnpackTarget = 0; UnpackMask = 0x7; break;
+			case ERuntimeVirtualTextureMaterialType::BaseColor_Normal_Specular_YCoCg: [[fallthrough]];
+			case ERuntimeVirtualTextureMaterialType::BaseColor_Normal_Specular_Mask_YCoCg: Output.UnpackType = EVirtualTextureUnpackType::BaseColorYCoCg; break;
+			case ERuntimeVirtualTextureMaterialType::BaseColor_Normal_Roughness: Output.UnpackType = EVirtualTextureUnpackType::BaseColorSRGB; break;
+			default: Output.UnpackTarget = 0; Output.UnpackMask = 0x7; break;
 			}
 		}
 		else
 		{
-			return Compiler->Constant3(0.f, 0.f, 0.f);
+			Output.ConstantVector = { 0.f, 0.f, 0.f };
 		}
 		break;
 	case 1:
-		if (bIsVirtualTextureValid && bIsSpecularValid)
+		if (bIsVirtualTextureValid && Output.bIsSpecularValid)
 		{
 			switch (MaterialType)
 			{
-			case ERuntimeVirtualTextureMaterialType::BaseColor_Normal_Specular:  UnpackTarget = 1; UnpackMask = 0x1; break;
-			case ERuntimeVirtualTextureMaterialType::BaseColor_Normal_Specular_YCoCg:
-			case ERuntimeVirtualTextureMaterialType::BaseColor_Normal_Specular_Mask_YCoCg: UnpackTarget = 2; UnpackMask = 0x1; break;
+			case ERuntimeVirtualTextureMaterialType::BaseColor_Normal_Specular: Output.UnpackTarget = 1; Output.UnpackMask = 0x1; break;
+			case ERuntimeVirtualTextureMaterialType::BaseColor_Normal_Specular_YCoCg: [[fallthrough]];
+			case ERuntimeVirtualTextureMaterialType::BaseColor_Normal_Specular_Mask_YCoCg: Output.UnpackTarget = 2; Output.UnpackMask = 0x1; break;
 			}
 		}
 		else
 		{
-			return Compiler->Constant(0.5f);
+			Output.ConstantVector = { 0.5f };
 		}
 		break;
 	case 2:
-		if (bIsVirtualTextureValid && bIsRoughnessValid)
+		if (bIsVirtualTextureValid && Output.bIsRoughnessValid)
 		{
 			switch (MaterialType)
 			{
-			case ERuntimeVirtualTextureMaterialType::BaseColor_Normal_Roughness: UnpackTarget = 1; UnpackMask = 0x2; break;
-			case ERuntimeVirtualTextureMaterialType::BaseColor_Normal_Specular: UnpackTarget = 1; UnpackMask = 0x2; break;
-			case ERuntimeVirtualTextureMaterialType::BaseColor_Normal_Specular_YCoCg:
-			case ERuntimeVirtualTextureMaterialType::BaseColor_Normal_Specular_Mask_YCoCg: UnpackTarget = 2; UnpackMask = 0x2; break;
+			case ERuntimeVirtualTextureMaterialType::BaseColor_Normal_Roughness: Output.UnpackTarget = 1; Output.UnpackMask = 0x2; break;
+			case ERuntimeVirtualTextureMaterialType::BaseColor_Normal_Specular: Output.UnpackTarget = 1; Output.UnpackMask = 0x2; break;
+			case ERuntimeVirtualTextureMaterialType::BaseColor_Normal_Specular_YCoCg: [[fallthrough]];
+			case ERuntimeVirtualTextureMaterialType::BaseColor_Normal_Specular_Mask_YCoCg: Output.UnpackTarget = 2; Output.UnpackMask = 0x2; break;
 			}
 		}
 		else
 		{
-			return Compiler->Constant(0.5f);
+			Output.ConstantVector = { 0.5f };
 		}
 		break;
 	case 3:
-		if (bIsVirtualTextureValid && bIsNormalValid)
+		if (bIsVirtualTextureValid && Output.bIsNormalValid)
 		{
 			switch (MaterialType)
 			{
-			case ERuntimeVirtualTextureMaterialType::BaseColor_Normal_Roughness: UnpackType = EVirtualTextureUnpackType::NormalBGR565; break;
-			case ERuntimeVirtualTextureMaterialType::BaseColor_Normal_Specular: UnpackType = EVirtualTextureUnpackType::NormalBC3BC3; break;
-			case ERuntimeVirtualTextureMaterialType::BaseColor_Normal_Specular_YCoCg:
-			case ERuntimeVirtualTextureMaterialType::BaseColor_Normal_Specular_Mask_YCoCg: UnpackType = EVirtualTextureUnpackType::NormalBC5BC1; break;
+			case ERuntimeVirtualTextureMaterialType::BaseColor_Normal_Roughness: Output.UnpackType = EVirtualTextureUnpackType::NormalBGR565; break;
+			case ERuntimeVirtualTextureMaterialType::BaseColor_Normal_Specular: Output.UnpackType = EVirtualTextureUnpackType::NormalBC3BC3; break;
+			case ERuntimeVirtualTextureMaterialType::BaseColor_Normal_Specular_YCoCg: [[fallthrough]];
+			case ERuntimeVirtualTextureMaterialType::BaseColor_Normal_Specular_Mask_YCoCg: Output.UnpackType = EVirtualTextureUnpackType::NormalBC5BC1; break;
 			}
 		}
 		else
 		{
-			return Compiler->Constant3(0.f, 0.f, 1.f);
+			Output.ConstantVector = { 0.f, 0.f, 1.f };
 		}
 		break;
 	case 4:
-		if (bIsVirtualTextureValid && bIsWorldHeightValid)
+		if (bIsVirtualTextureValid && Output.bIsWorldHeightValid)
 		{
-			UnpackType = EVirtualTextureUnpackType::HeightR16;
+			Output.UnpackType = EVirtualTextureUnpackType::HeightR16;
 		}
 		else
 		{
-			return Compiler->Constant(0.f);
+			Output.ConstantVector = { 0.f };
 		}
 		break;
 	case 5:
-		if (bIsVirtualTextureValid && bIsMaskValid)
+		if (bIsVirtualTextureValid && Output.bIsMaskValid)
 		{
-			UnpackTarget = 2; UnpackMask = 0x8;
+			Output.UnpackTarget = 2; Output.UnpackMask = 0x8;
 		}
 		else
 		{
-			return Compiler->Constant(1.f);
+			Output.ConstantVector = { 1.f };
 		}
 		break;
 	case 6:
-		if (bIsVirtualTextureValid && bIsDisplacementValid)
+		if (bIsVirtualTextureValid && Output.bIsDisplacementValid)
 		{
-			UnpackType = EVirtualTextureUnpackType::DisplacementR16;
+			Output.UnpackType = EVirtualTextureUnpackType::DisplacementR16;
 		}
 		else
 		{
-			return Compiler->Constant(0.f);
+			Output.ConstantVector = { 0.f };
 		}
 		break;
 	case 7:
-		if (bIsVirtualTextureValid && bIsMask4Valid)
+		if (bIsVirtualTextureValid && Output.bIsMask4Valid)
 		{
-			UnpackTarget = 0; UnpackMask = 0xf; break;
+			Output.UnpackTarget = 0; Output.UnpackMask = 0xf;
 		}
 		else
 		{
-			return Compiler->Constant4(0.f, 0.f, 0.f, 0.f);
+			Output.ConstantVector = { 0.f, 0.f, 0.f, 0.f };
 		}
+		break;
 	default:
+		return false;
+	}
+	return true;
+}
+
+int32 UMaterialExpressionRuntimeVirtualTextureSample::Compile(class FMaterialCompiler* Compiler, int32 OutputIndex)
+{
+	// Is this a valid UMaterialExpressionRuntimeVirtualTextureSampleParameter?
+	const bool bIsParameter = IsParameter();
+
+	// Check validity of current virtual texture
+	FString TextureValidityError;
+	const bool bIsVirtualTextureValid = ValidateVirtualTextureParameters(TextureValidityError);
+	if (!bIsVirtualTextureValid)
+	{
+		Compiler->Error(*TextureValidityError);
+		if (!VirtualTexture)
+		{
+			return INDEX_NONE;
+		}
+	}
+
+	// Calculate the virtual texture layer and sampling/unpacking functions for this output
+	// Fallback to a sensible default value if the output isn't valid for the bound virtual texture
+	FRuntimeVirtualTextureUnpackProperties UnpackProperties;
+	if (!GetRVTUnpackProperties(OutputIndex, bIsVirtualTextureValid, UnpackProperties))
+	{
 		return INDEX_NONE;
 	}
+
+	if (!UnpackProperties.ConstantVector.IsEmpty())
+	{
+		switch (UnpackProperties.ConstantVector.Num())
+		{
+		case 1:
+			return Compiler->Constant(UnpackProperties.ConstantVector[0]);
+		case 3:
+			return Compiler->Constant3(UnpackProperties.ConstantVector[0], UnpackProperties.ConstantVector[1], UnpackProperties.ConstantVector[2]);
+		case 4:
+			return Compiler->Constant4(UnpackProperties.ConstantVector[0], UnpackProperties.ConstantVector[1], UnpackProperties.ConstantVector[2], UnpackProperties.ConstantVector[3]);
+		default:
+			checkNoEntry();
+		}
+	}
 	
+	const uint32 UnpackTarget = UnpackProperties.UnpackTarget;
+	const uint32 UnpackMask = UnpackProperties.UnpackMask;
+	const EVirtualTextureUnpackType UnpackType = UnpackProperties.UnpackType;
+
 	// Compile the texture object references
 	const int32 TextureLayerCount = URuntimeVirtualTexture::GetLayerCount(MaterialType);
 	check(TextureLayerCount <= RuntimeVirtualTexture::MaxTextureLayers);
@@ -3562,7 +3596,7 @@ int32 UMaterialExpressionRuntimeVirtualTextureSample::Compile(class FMaterialCom
 		}
 		else
 		{
-			WorldPositionIndex = Compiler->WorldPosition(GetWorldPositionTypeWithOrigin(WorldPositionOriginType));
+			WorldPositionIndex = Compiler->WorldPosition(UE::MaterialTranslatorUtils::GetWorldPositionTypeWithOrigin(WorldPositionOriginType));
 			ensure(WorldPositionIndex != INDEX_NONE);
 		}
 		
@@ -3651,23 +3685,11 @@ int32 UMaterialExpressionRuntimeVirtualTextureSample::Compile(class FMaterialCom
 		MipValue1Index = Compiler->AppendVector(UDdy, VDdy);
 	}
 
-	// Convert texture address mode to matching sampler source mode.
-	// Would be better if ESamplerSourceMode had a Mirror enum that we could also use...
-	ESamplerSourceMode SamplerSourceMode = SSM_Clamp_WorldGroupSettings;
-	switch (TextureAddressMode)
-	{
-	case RVTTA_Clamp:
-		SamplerSourceMode = SSM_Clamp_WorldGroupSettings;
-		break;
-	case RVTTA_Wrap:
-		SamplerSourceMode = SSM_Wrap_WorldGroupSettings;
-		break;
-	}
-
 	// We can support disabling feedback for MipLevel mode.
 	const bool bForceEnableFeedback = TextureMipLevelMode != TMVM_MipLevel;
 
 	// Compile the texture sample code
+	const ESamplerSourceMode SamplerSourceMode = GetSamplerSourceMode();
 	const bool bAutomaticMipViewBias = true;
 	int32 SampleCodeIndex[RuntimeVirtualTexture::MaxTextureLayers] = { INDEX_NONE };
 	for (int32 TexureLayerIndex = 0; TexureLayerIndex < TextureLayerCount; TexureLayerIndex++)
@@ -18704,7 +18726,7 @@ FName UMaterialExpressionDistanceToNearestSurface::GetInputName(int32 InputIndex
 {
 	if (GetInput(InputIndex) == &Position)
 	{
-		return GetWorldPositionInputName(WorldPositionOriginType);
+		return UE::MaterialTranslatorUtils::GetWorldPositionInputName(WorldPositionOriginType);
 	}
 
 	return Super::GetInputName(InputIndex);
@@ -18734,7 +18756,7 @@ int32 UMaterialExpressionDistanceToNearestSurface::Compile(class FMaterialCompil
 	}
 	else 
 	{
-		PositionArg = Compiler->WorldPosition(GetWorldPositionTypeWithOrigin(WorldPositionOriginType));
+		PositionArg = Compiler->WorldPosition(UE::MaterialTranslatorUtils::GetWorldPositionTypeWithOrigin(WorldPositionOriginType));
 	}
 
 	return Compiler->DistanceToNearestSurface(PositionArg, WorldPositionOriginType);
@@ -18754,7 +18776,7 @@ FName UMaterialExpressionDistanceFieldGradient::GetInputName(int32 InputIndex) c
 {
 	if (GetInput(InputIndex) == &Position)
 	{
-		return GetWorldPositionInputName(WorldPositionOriginType);
+		return UE::MaterialTranslatorUtils::GetWorldPositionInputName(WorldPositionOriginType);
 	}
 
 	return Super::GetInputName(InputIndex);
@@ -18784,7 +18806,7 @@ int32 UMaterialExpressionDistanceFieldGradient::Compile(class FMaterialCompiler*
 	}
 	else 
 	{
-		PositionArg = Compiler->WorldPosition(GetWorldPositionTypeWithOrigin(WorldPositionOriginType));
+		PositionArg = Compiler->WorldPosition(UE::MaterialTranslatorUtils::GetWorldPositionTypeWithOrigin(WorldPositionOriginType));
 	}
 
 	return Compiler->DistanceFieldGradient(PositionArg, WorldPositionOriginType);
@@ -18803,7 +18825,7 @@ FName UMaterialExpressionDistanceFieldApproxAO::GetInputName(int32 InputIndex) c
 {
 	if (GetInput(InputIndex) == &Position)
 	{
-		return GetWorldPositionInputName(WorldPositionOriginType);
+		return UE::MaterialTranslatorUtils::GetWorldPositionInputName(WorldPositionOriginType);
 	}
 
 	return Super::GetInputName(InputIndex);
@@ -18833,7 +18855,7 @@ int32 UMaterialExpressionDistanceFieldApproxAO::Compile(class FMaterialCompiler*
 	}
 	else
 	{
-		PositionArg = Compiler->WorldPosition(GetWorldPositionTypeWithOrigin(WorldPositionOriginType));
+		PositionArg = Compiler->WorldPosition(UE::MaterialTranslatorUtils::GetWorldPositionTypeWithOrigin(WorldPositionOriginType));
 	}
 
 	int32 NormalArg = INDEX_NONE;
@@ -18867,7 +18889,7 @@ FName UMaterialExpressionSamplePhysicsVectorField::GetInputName(int32 InputIndex
 {
 	if (GetInput(InputIndex) == &WorldPosition)
 	{
-		return GetWorldPositionInputName(WorldPositionOriginType);
+		return UE::MaterialTranslatorUtils::GetWorldPositionInputName(WorldPositionOriginType);
 	}
 
 	return Super::GetInputName(InputIndex);
@@ -18897,7 +18919,7 @@ int32 UMaterialExpressionSamplePhysicsVectorField::Compile(class FMaterialCompil
 	}
 	else
 	{
-		PositionArg = Compiler->WorldPosition(GetWorldPositionTypeWithOrigin(WorldPositionOriginType));
+		PositionArg = Compiler->WorldPosition(UE::MaterialTranslatorUtils::GetWorldPositionTypeWithOrigin(WorldPositionOriginType));
 	}
 
 	return Compiler->SamplePhysicsField(PositionArg, WorldPositionOriginType, EFieldOutputType::Field_Output_Vector, static_cast<uint8>(FieldTarget));
@@ -18916,7 +18938,7 @@ FName UMaterialExpressionSamplePhysicsScalarField::GetInputName(int32 InputIndex
 {
 	if (GetInput(InputIndex) == &WorldPosition)
 	{
-		return GetWorldPositionInputName(WorldPositionOriginType);
+		return UE::MaterialTranslatorUtils::GetWorldPositionInputName(WorldPositionOriginType);
 	}
 
 	return Super::GetInputName(InputIndex);
@@ -18946,7 +18968,7 @@ int32 UMaterialExpressionSamplePhysicsScalarField::Compile(class FMaterialCompil
 	}
 	else
 	{
-		PositionArg = Compiler->WorldPosition(GetWorldPositionTypeWithOrigin(WorldPositionOriginType));
+		PositionArg = Compiler->WorldPosition(UE::MaterialTranslatorUtils::GetWorldPositionTypeWithOrigin(WorldPositionOriginType));
 	}
 
 	return Compiler->SamplePhysicsField(PositionArg, WorldPositionOriginType, EFieldOutputType::Field_Output_Scalar, static_cast<uint8>(FieldTarget));
@@ -18965,7 +18987,7 @@ FName UMaterialExpressionSamplePhysicsIntegerField::GetInputName(int32 InputInde
 {
 	if (GetInput(InputIndex) == &WorldPosition)
 	{
-		return GetWorldPositionInputName(WorldPositionOriginType);
+		return UE::MaterialTranslatorUtils::GetWorldPositionInputName(WorldPositionOriginType);
 	}
 
 	return Super::GetInputName(InputIndex);
@@ -18995,7 +19017,7 @@ int32 UMaterialExpressionSamplePhysicsIntegerField::Compile(class FMaterialCompi
 	}
 	else
 	{
-		PositionArg = Compiler->WorldPosition(GetWorldPositionTypeWithOrigin(WorldPositionOriginType));
+		PositionArg = Compiler->WorldPosition(UE::MaterialTranslatorUtils::GetWorldPositionTypeWithOrigin(WorldPositionOriginType));
 	}
 
 	return Compiler->SamplePhysicsField(PositionArg, WorldPositionOriginType, EFieldOutputType::Field_Output_Integer, static_cast<uint8>(FieldTarget));
@@ -19696,7 +19718,7 @@ FName UMaterialExpressionAtmosphericFogColor::GetInputName(int32 InputIndex) con
 {
 	if (GetInput(InputIndex) == &WorldPosition)
 	{
-		return GetWorldPositionInputName(WorldPositionOriginType);
+		return UE::MaterialTranslatorUtils::GetWorldPositionInputName(WorldPositionOriginType);
 	}
 
 	return Super::GetInputName(InputIndex);
@@ -20113,7 +20135,7 @@ FName UMaterialExpressionSkyAtmosphereLightIlluminance::GetInputName(int32 Input
 {
 	if (GetInput(InputIndex) == &WorldPosition)
 	{
-		return GetWorldPositionInputName(WorldPositionOriginType);
+		return UE::MaterialTranslatorUtils::GetWorldPositionInputName(WorldPositionOriginType);
 	}
 
 	return Super::GetInputName(InputIndex);
@@ -20142,7 +20164,7 @@ int32 UMaterialExpressionSkyAtmosphereLightIlluminance::Compile(class FMaterialC
 	}
 	else
 	{
-		WorldPositionInput = Compiler->WorldPosition(GetWorldPositionTypeWithOrigin(WorldPositionOriginType));
+		WorldPositionInput = Compiler->WorldPosition(UE::MaterialTranslatorUtils::GetWorldPositionTypeWithOrigin(WorldPositionOriginType));
 	}
 	return Compiler->SkyAtmosphereLightIlluminance(WorldPositionInput, WorldPositionOriginType, LightIndex);
 }
@@ -20244,7 +20266,7 @@ FName UMaterialExpressionSkyAtmosphereAerialPerspective::GetInputName(int32 Inpu
 {
 	if (GetInput(InputIndex) == &WorldPosition)
 	{
-		return GetWorldPositionInputName(WorldPositionOriginType);
+		return UE::MaterialTranslatorUtils::GetWorldPositionInputName(WorldPositionOriginType);
 	}
 
 	return Super::GetInputName(InputIndex);
@@ -20273,7 +20295,7 @@ int32 UMaterialExpressionSkyAtmosphereAerialPerspective::Compile(class FMaterial
 	}
 	else
 	{
-		WorldPositionInput = Compiler->WorldPosition(GetWorldPositionTypeWithOrigin(WorldPositionOriginType));
+		WorldPositionInput = Compiler->WorldPosition(UE::MaterialTranslatorUtils::GetWorldPositionTypeWithOrigin(WorldPositionOriginType));
 	}
 	return Compiler->SkyAtmosphereAerialPerspective(WorldPositionInput, WorldPositionOriginType);
 }
diff --git a/Engine/Source/Runtime/Engine/Private/Materials/MaterialExpressionsIR.cpp b/Engine/Source/Runtime/Engine/Private/Materials/MaterialExpressionsIR.cpp
--- a/Engine/Source/Runtime/Engine/Private/Materials/MaterialExpressionsIR.cpp
+++ b/Engine/Source/Runtime/Engine/Private/Materials/MaterialExpressionsIR.cpp
@@ -5,6 +5,7 @@
 #include "Curves/CurveLinearColorAtlas.h"
 #include "Engine/Engine.h"
 #include "Engine/Font.h"
+#include "MaterialIRInternal.h"
 #include "MaterialDomain.h"
 #include "MaterialExpressionIO.h"
 #include "Materials/Material.h"
@@ -1267,20 +1268,6 @@ static MIR::ETextureReadMode TextureGatherModeToMIR(::ETextureGatherMode Mode)
 	}
 }
 
-// Returns this value's texture sampler type if it has one (SAMPLERTYPE_MAX otherwise).
-static EMaterialSamplerType GetValueMaterialSamplerType(FValueRef Value)
-{
-	if (auto TextureObject = MIR::As<MIR::FTextureObject>(Value))
-	{
-		return TextureObject->SamplerType;
-	}
-	else if (auto UniformParameter = MIR::As<MIR::FUniformParameter>(Value))
-	{
-		return UniformParameter->SamplerType;
-	}
-	return SAMPLERTYPE_MAX;
-}
-
 static FValueRef BuildTextureSample(MIR::FEmitter& Em, UMaterialExpressionTextureSample* Expr, FValueRef Texture, EMaterialValueType TextureType, FValueRef TexCoords, bool bAutomaticViewMipBias)
 {
 	FValueRef TextureRead{};
@@ -1292,7 +1279,7 @@ static FValueRef BuildTextureSample(MIR::FEmitter& Em, UMaterialExpressionTextur
 			return TextureRead;
 		}
 
-		TextureRead = Em.TextureGather(Texture, TexCoords, TextureGatherModeToMIR(Expr->GatherMode), Expr->SamplerSource, GetValueMaterialSamplerType(Texture));
+		TextureRead = Em.TextureGather(Texture, TexCoords, TextureGatherModeToMIR(Expr->GatherMode), { Expr->SamplerSource });
 	}
 	else
 	{
@@ -1312,22 +1299,22 @@ static FValueRef BuildTextureSample(MIR::FEmitter& Em, UMaterialExpressionTextur
 		switch (Expr->MipValueMode)
 		{
 			case TMVM_None:
-				TextureRead = Em.TextureSample(Texture, TexCoords, bAutomaticViewMipBias, Expr->SamplerSource, GetValueMaterialSamplerType(Texture));
+				TextureRead = Em.TextureSample(Texture, TexCoords, bAutomaticViewMipBias, { Expr->SamplerSource });
 				break;
 		
 			case TMVM_MipBias:
-				TextureRead = Em.TextureSampleBias(Texture, TexCoords, MipValue, bAutomaticViewMipBias, Expr->SamplerSource, GetValueMaterialSamplerType(Texture));
+				TextureRead = Em.TextureSampleBias(Texture, TexCoords, MipValue, bAutomaticViewMipBias, { Expr->SamplerSource });
 				break;
 		
 			case TMVM_MipLevel:
-				TextureRead = Em.TextureSampleLevel(Texture, TexCoords, MipValue, bAutomaticViewMipBias, Expr->SamplerSource, GetValueMaterialSamplerType(Texture));
+				TextureRead = Em.TextureSampleLevel(Texture, TexCoords, MipValue, bAutomaticViewMipBias, { Expr->SamplerSource });
 				break;
 
 			case TMVM_Derivative:
 			{
 				FValueRef TexCoordsDdx = Em.Cast(Em.Input(&Expr->CoordinatesDX), TexCoords->Type);
 				FValueRef TexCoordsDdy = Em.Cast(Em.Input(&Expr->CoordinatesDY), TexCoords->Type);
-				TextureRead = Em.TextureSampleGrad(Texture, TexCoords, TexCoordsDdx, TexCoordsDdy, bAutomaticViewMipBias, Expr->SamplerSource, GetValueMaterialSamplerType(Texture));
+				TextureRead = Em.TextureSampleGrad(Texture, TexCoords, TexCoordsDdx, TexCoordsDdy, bAutomaticViewMipBias, { Expr->SamplerSource });
 				break;
 			}
 		}
@@ -1375,10 +1362,28 @@ static void BuildTextureSampleExpression(MIR::FEmitter& Em, UMaterialExpressionT
 
 void UMaterialExpressionTextureSample::Build(MIR::FEmitter& Em)
 {
-	FValueRef TextureValue = BuildTextureValue(Em, this);
+	FValueRef TextureValue = Em.TryInput(&TextureObject);
+	if (!TextureValue)
+	{
+		if (!Texture)
+		{
+			Em.Error(TEXT("No texture specified for this expression."));
+			return;
+		}
+
+		TextureValue = Em.TextureObject(Texture.Get(), SamplerType);
+	}
+
 	UE_MIR_CHECKPOINT(Em);
 
-	BuildTextureSampleExpression(Em, this, TextureValue, Em.GetTextureFromValue(TextureValue)->GetMaterialType());
+	UObject* DefaultTexture = Em.GetTextureFromValue(TextureValue);
+	if (!DefaultTexture)
+	{
+		Em.Error(TEXT("Missing texture object from input"));
+		return;
+	}
+
+	BuildTextureSampleExpression(Em, this, TextureValue, MIR::Internal::GetTextureMaterialValueType(DefaultTexture));
 }
 
 // Inherits from UMaterialExpressionTextureSample, but uses different particle specific UVs, and optionally supports blending two different texture samples.
@@ -1387,7 +1392,7 @@ void UMaterialExpressionParticleSubUV::Build(MIR::FEmitter& Em)
 	FValueRef TextureValue = BuildTextureValue(Em, this);
 	UE_MIR_CHECKPOINT(Em);
 
-	EMaterialValueType TextureType = Em.GetTextureFromValue(TextureValue)->GetMaterialType();
+	EMaterialValueType TextureType = MIR::Internal::GetTextureMaterialValueType(Em.GetTextureFromValue(TextureValue));
 
 	// Although the parent UMaterialExpressionTextureSample class includes an automatic view mip bias flag, it is specifically ignored by ParticleSubUV.
 	const bool bAutomaticViewMipBias = false;
@@ -1556,7 +1561,17 @@ void UMaterialExpressionTextureProperty::Build(MIR::FEmitter& Em)
 			? UE::Shader::EPreshaderOpcode::TexelSize
 			: UE::Shader::EPreshaderOpcode::TextureSize;
 
-	Em.Output(0, Em.PreshaderParameter(PreshaderOpcode, TextureValue));
+	UObject* SourceParameterTexture = Em.GetTextureFromValue(TextureValue);
+	if (!SourceParameterTexture)
+	{
+		Em.Error(TEXT("Missing default texture from source parameter"));
+		return;
+	}
+
+	const EMaterialValueType TextureType = MIR::Internal::GetTextureMaterialValueType(SourceParameterTexture);
+	const EMaterialValueType PropertyType = UE::MaterialTranslatorUtils::GetTexturePropertyValueType(TextureType);
+
+	Em.Output(0, Em.PreshaderParameter(MIR::FType::FromMaterialValueType(PropertyType), PreshaderOpcode, TextureValue));
 }
 
 #if WITH_EDITOR
@@ -1602,8 +1617,7 @@ void UMaterialExpressionFontSample::Build(MIR::FEmitter& Em)
 				Em.TextureObject(Texture, ExpectedSamplerType),
 				Em.ExternalInput(MIR::TexCoordIndexToExternalInput(0)),
 				false,
-				SSM_FromTextureAsset,
-				ExpectedSamplerType
+				{ SSM_FromTextureAsset, ExpectedSamplerType }
 			)
 		);
 	}
@@ -1644,14 +1658,307 @@ void UMaterialExpressionFontSampleParameter::Build(MIR::FEmitter& Em)
 				TextureParameter,
 				Em.ExternalInput(MIR::TexCoordIndexToExternalInput(0)),
 				false,
-				SSM_FromTextureAsset,
-				ExpectedSamplerType
+				{ SSM_FromTextureAsset, ExpectedSamplerType }
 			)
 		);
 	}
 }
 #endif
 
+#if WITH_EDITOR
+
+static MIR::FValueRef BuildVirtualTextureWorldToUV(MIR::FEmitter& Em, MIR::FValueRef WorldPositionValue, MIR::FValueRef P0, MIR::FValueRef P1, MIR::FValueRef P2, EPositionOrigin PositionOrigin)
+{
+	return Em.InlineHLSL(MIR::FType::MakeFloatVector(2), TEXTVIEW("VirtualTextureWorldToUV($0, $1, $2, $3)"), { WorldPositionValue, P0, P1, P2 });
+}
+
+static MIR::FValueRef BuildConstantVector(MIR::FEmitter& Em, TArrayView<float> InConstants)
+{
+	switch (InConstants.Num())
+	{
+	case 1:
+		return Em.ConstantFloat(InConstants[0]);
+	case 2:
+		return Em.ConstantFloat2(FVector2f{ InConstants[0], InConstants[1] });
+	case 3:
+		return Em.ConstantFloat3(FVector3f{ InConstants[0], InConstants[1], InConstants[2] });
+	case 4:
+		return Em.ConstantFloat4(FVector4f{ InConstants[0], InConstants[1], InConstants[2], InConstants[3] });
+	default:
+		UE_MIR_UNREACHABLE();
+	}
+}
+
+static FValueRef EmitWorldPosition(MIR::FEmitter& Em, EWorldPositionIncludedOffsets WorldPositionShaderOffset);
+
+static FValueRef EmitTransformVectorBase(
+	MIR::FEmitter& Em, FValueRef InputValue, EMaterialCommonBasis TransformSourceBasis, EMaterialCommonBasis TransformDestBasis,
+	bool bIsPositionTransform, MIR::FValueRef PeriodicWorldTileSizeValue, MIR::FValueRef FirstPersonInterpolationAlphaValue);
+
+static FValueRef BuildVirtualTextureUnpack(MIR::FEmitter& Em, FValueRef SampleCode0, FValueRef SampleCode1, FValueRef SampleCode2, FValueRef P0, EVirtualTextureUnpackType UnpackType)
+{
+	switch (UnpackType)
+	{
+	case EVirtualTextureUnpackType::BaseColorYCoCg:		return Em.InlineHLSL(MIR::FType::MakeFloatVector(3), TEXTVIEW("VirtualTextureUnpackBaseColorYCoCg($0)"), { SampleCode0 });
+	case EVirtualTextureUnpackType::NormalBC3:			return Em.InlineHLSL(MIR::FType::MakeFloatVector(3), TEXTVIEW("VirtualTextureUnpackNormalBC3($0)"), { SampleCode1 });
+	case EVirtualTextureUnpackType::NormalBC5:			return Em.InlineHLSL(MIR::FType::MakeFloatVector(3), TEXTVIEW("VirtualTextureUnpackNormalBC5($0)"), { SampleCode1 });
+	case EVirtualTextureUnpackType::NormalBC3BC3:		return Em.InlineHLSL(MIR::FType::MakeFloatVector(3), TEXTVIEW("VirtualTextureUnpackNormalBC3BC3($0, $1)"), { SampleCode0, SampleCode1 });
+	case EVirtualTextureUnpackType::NormalBC5BC1:		return Em.InlineHLSL(MIR::FType::MakeFloatVector(3), TEXTVIEW("VirtualTextureUnpackNormalBC5BC1($0, $1)"), { SampleCode1, SampleCode2 });
+	case EVirtualTextureUnpackType::HeightR16:			return Em.InlineHLSL(MIR::FType::MakeFloatScalar(), TEXTVIEW("VirtualTextureUnpackHeight($0, $1)"), { SampleCode0, P0 });
+	case EVirtualTextureUnpackType::DisplacementR16:	return Em.Swizzle(SampleCode0, MIR::EVectorComponent::X);
+	case EVirtualTextureUnpackType::NormalBGR565:		return Em.InlineHLSL(MIR::FType::MakeFloatVector(3), TEXTVIEW("VirtualTextureUnpackNormalBGR565($0)"), { SampleCode1 });
+	case EVirtualTextureUnpackType::BaseColorSRGB:		return Em.InlineHLSL(MIR::FType::MakeFloatVector(3), TEXTVIEW("VirtualTextureUnpackBaseColorSRGB($0)"), { SampleCode0 });
+	default:											UE_MIR_UNREACHABLE();
+	}
+}
+
+void UMaterialExpressionRuntimeVirtualTextureSample::Build(MIR::FEmitter& Em)
+{
+	if (!UseVirtualTexturing(Em.GetShaderPlatform()))
+	{
+		Em.Errorf(TEXT("Virtual texturing not supported on platform '%s'"), *ShaderPlatformToPlatformName(Em.GetShaderPlatform()).ToString());
+		return;
+	}
+
+	// Is this a valid UMaterialExpressionRuntimeVirtualTextureSampleParameter?
+	const bool bIsParameter = IsParameter();
+
+	// Check validity of current virtual texture
+	FString TextureValidityError;
+	const bool bIsVirtualTextureValid = ValidateVirtualTextureParameters(TextureValidityError);
+	if (!bIsVirtualTextureValid)
+	{
+		Em.Error(*TextureValidityError);
+		if (!VirtualTexture)
+		{
+			return;
+		}
+	}
+
+	// Compile the texture object references
+	const int32 TextureLayerCount = URuntimeVirtualTexture::GetLayerCount(MaterialType);
+	check(TextureLayerCount <= RuntimeVirtualTexture::MaxTextureLayers);
+
+	MIR::FValueRef TextureObjects[RuntimeVirtualTexture::MaxTextureLayers];
+	for (int32 TexureLayerIndex = 0; TexureLayerIndex < TextureLayerCount; TexureLayerIndex++)
+	{
+		const int32 PageTableLayerIndex = bSinglePhysicalSpace ? 0 : TexureLayerIndex;
+
+		if (bIsParameter)
+		{
+			FMaterialParameterMetadata Metadata;
+			GetParameterValue(Metadata);
+			TextureObjects[TexureLayerIndex] = Em.Parameter(GetParameterName(), Metadata, SAMPLERTYPE_VirtualMasks, TexureLayerIndex, PageTableLayerIndex);
+		}
+		else
+		{
+			TextureObjects[TexureLayerIndex] = Em.RuntimeVirtualTextureObject(VirtualTexture, SAMPLERTYPE_VirtualMasks, TexureLayerIndex, PageTableLayerIndex);
+		}
+	}
+
+	UE_MIR_CHECKPOINT(Em);
+
+	// Compile the runtime virtual texture uniforms
+	MIR::FValueRef Uniforms[ERuntimeVirtualTextureShaderUniform_Count];
+
+	for (int32 UniformIndex = 0; UniformIndex < ERuntimeVirtualTextureShaderUniform_Count; ++UniformIndex)
+	{
+		const UE::Shader::EValueType UniformType = URuntimeVirtualTexture::GetUniformParameterType(UniformIndex);
+		Uniforms[UniformIndex] = Em.PreshaderParameter(
+			MIR::FType::FromShaderType(UniformType), UE::Shader::EPreshaderOpcode::RuntimeVirtualTextureUniform,
+			TextureObjects[0], MIR::FPreshaderParameterPayload{ .UniformIndex = UniformIndex });
+	}
+
+	// Compile the coordinates
+	// We use the virtual texture world space transform by default
+	if (Coordinates.GetTracedInput().Expression != nullptr && WorldPosition.GetTracedInput().Expression != nullptr)
+	{
+		Em.Error(TEXT("Only one of 'Coordinates' and 'WorldPosition' can be used"));
+	}
+
+	MIR::FValueRef CoordinateValue = Em.TryInput(&Coordinates);
+	if (!CoordinateValue)
+	{
+		MIR::FValueRef WorldPositionValue;
+		if (WorldPosition.GetTracedInput().Expression != nullptr)
+		{
+			WorldPositionValue = Em.Input(&WorldPosition);
+		}
+		else
+		{
+			WorldPositionValue = EmitWorldPosition(Em, UE::MaterialTranslatorUtils::GetWorldPositionTypeWithOrigin(WorldPositionOriginType));
+			ensure(WorldPositionValue);
+		}
+
+		if (WorldPositionValue)
+		{
+			if (WorldPositionOriginType == EPositionOrigin::Absolute)
+			{
+				MIR::FValueRef P0 = Uniforms[ERuntimeVirtualTextureShaderUniform_WorldToUVTransform0];
+				MIR::FValueRef P1 = Uniforms[ERuntimeVirtualTextureShaderUniform_WorldToUVTransform1];
+				MIR::FValueRef P2 = Uniforms[ERuntimeVirtualTextureShaderUniform_WorldToUVTransform2];
+				CoordinateValue = BuildVirtualTextureWorldToUV(Em, WorldPositionValue, P0, P1, P2, EPositionOrigin::Absolute);
+			}
+			else if (WorldPositionOriginType == EPositionOrigin::CameraRelative)
+			{
+				//TODO: optimize by calculating translated world to VT directly.
+				//This requires some more work as the transform is currently fed in through a preshader variable, which is cached.
+				MIR::FValueRef AbsWorldPosIndex = EmitTransformVectorBase(Em, WorldPositionValue, EMaterialCommonBasis::MCB_TranslatedWorld, EMaterialCommonBasis::MCB_World, true, {}, {});
+
+				MIR::FValueRef P0 = Uniforms[ERuntimeVirtualTextureShaderUniform_WorldToUVTransform0];
+				MIR::FValueRef P1 = Uniforms[ERuntimeVirtualTextureShaderUniform_WorldToUVTransform1];
+				MIR::FValueRef P2 = Uniforms[ERuntimeVirtualTextureShaderUniform_WorldToUVTransform2];
+				CoordinateValue = BuildVirtualTextureWorldToUV(Em, AbsWorldPosIndex, P0, P1, P2, EPositionOrigin::Absolute);
+			}
+			else
+			{
+				checkNoEntry();
+			}
+		}
+	}
+
+	// Compile the mip level for the current mip value mode
+	ETextureMipValueMode TextureMipLevelMode = TMVM_None;
+	MIR::FValueRef MipValue0Value;
+	MIR::FValueRef MipValue1Value;
+	const bool bMipValueExpressionValid = MipValue.GetTracedInput().Expression != nullptr;
+
+	if (MipValueMode == RVTMVM_MipLevel)
+	{
+		TextureMipLevelMode = TMVM_MipLevel;
+		MipValue0Value = bMipValueExpressionValid ? Em.Input(&MipValue) : Em.ConstantFloat(0.0f);
+	}
+	else if (MipValueMode == RVTMVM_MipBias)
+	{
+		TextureMipLevelMode = TMVM_MipBias;
+		MipValue0Value = bMipValueExpressionValid ? Em.Input(&MipValue) : Em.ConstantFloat(0.0f);
+	}
+	else if (MipValueMode == RVTMVM_DerivativeUV || MipValueMode == RVTMVM_DerivativeWorld)
+	{
+		if (DDX.GetTracedInput().Expression == nullptr || DDY.GetTracedInput().Expression == nullptr)
+		{
+			Em.Error(TEXT("Derivative MipValueMode requires connected DDX and DDY pins."));
+		}
+
+		TextureMipLevelMode = TMVM_Derivative;
+		MIR::FValueRef Ddx = Em.Input(&DDX);
+		MIR::FValueRef Ddy = Em.Input(&DDY);
+
+		if (MipValueMode == RVTMVM_DerivativeUV)
+		{
+			MipValue0Value = Ddx;
+			MipValue1Value = Ddy;
+		}
+		else if (MipValueMode == RVTMVM_DerivativeWorld)
+		{
+			MIR::FValueRef UDdx = Em.Dot(Ddx, Uniforms[ERuntimeVirtualTextureShaderUniform_WorldToUVTransform1]);
+			MIR::FValueRef VDdx = Em.Dot(Ddx, Uniforms[ERuntimeVirtualTextureShaderUniform_WorldToUVTransform2]);
+			MipValue0Value = Em.Vector2(UDdx, VDdx);
+
+			MIR::FValueRef UDdy = Em.Dot(Ddy, Uniforms[ERuntimeVirtualTextureShaderUniform_WorldToUVTransform1]);
+			MIR::FValueRef VDdy = Em.Dot(Ddy, Uniforms[ERuntimeVirtualTextureShaderUniform_WorldToUVTransform2]);
+			MipValue1Value = Em.Vector2(UDdy, VDdy);
+		}
+	}
+	else if (MipValueMode == RVTMVM_RecalculateDerivatives)
+	{
+		// Calculate derivatives from world position.
+		// This is legacy/hidden, and is better implemented in the material graph using RVTMVM_DerivativeWorld.
+		TextureMipLevelMode = TMVM_Derivative;
+
+		MIR::FValueRef WorldPos = Em.ExternalInput(MIR::EExternalInput::WorldPosition_CameraRelative);
+		MIR::FValueRef WorldPositionDdx = Em.AnalyticalPartialDerivative(WorldPos, MIR::EDerivativeAxis::X);
+		MIR::FValueRef UDdx = Em.Dot(WorldPositionDdx, Uniforms[ERuntimeVirtualTextureShaderUniform_WorldToUVTransform1]);
+		MIR::FValueRef VDdx = Em.Dot(WorldPositionDdx, Uniforms[ERuntimeVirtualTextureShaderUniform_WorldToUVTransform2]);
+		MipValue0Value = Em.Vector2(UDdx, VDdx);
+
+		MIR::FValueRef WorldPositionDdy = Em.AnalyticalPartialDerivative(WorldPos, MIR::EDerivativeAxis::Y);
+		MIR::FValueRef UDdy = Em.Dot(WorldPositionDdy, Uniforms[ERuntimeVirtualTextureShaderUniform_WorldToUVTransform1]);
+		MIR::FValueRef VDdy = Em.Dot(WorldPositionDdy, Uniforms[ERuntimeVirtualTextureShaderUniform_WorldToUVTransform2]);
+		MipValue1Value = Em.Vector2(UDdy, VDdy);
+	}
+
+	// We can support disabling feedback for MipLevel mode.
+	const bool bForceEnableFeedback = TextureMipLevelMode != TMVM_MipLevel;
+
+	// Compile the texture sample code
+	constexpr bool bAutomaticMipViewBias = true;
+
+	const MIR::FTextureSampleBaseAttributes SampleAttributes
+	{
+		.SamplerSourceMode = GetSamplerSourceMode(),
+		.SamplerType = SAMPLERTYPE_VirtualMasks,
+		.bEnableFeedback = bEnableFeedback || bForceEnableFeedback,
+		.bIsAdaptive = bAdaptive
+	};
+
+	MIR::FValueRef SampleCodeValues[RuntimeVirtualTexture::MaxTextureLayers];
+	for (int32 TexureLayerIndex = 0; TexureLayerIndex < TextureLayerCount; TexureLayerIndex++)
+	{
+		switch (TextureMipLevelMode)
+		{
+			case TMVM_None:
+				SampleCodeValues[TexureLayerIndex] = Em.TextureSample(TextureObjects[TexureLayerIndex], CoordinateValue, bAutomaticMipViewBias, SampleAttributes);
+				break;
+
+			case TMVM_MipBias:
+				SampleCodeValues[TexureLayerIndex] = Em.TextureSampleBias(TextureObjects[TexureLayerIndex], CoordinateValue, MipValue0Value, bAutomaticMipViewBias, SampleAttributes);
+				break;
+
+			case TMVM_MipLevel:
+				SampleCodeValues[TexureLayerIndex] = Em.TextureSampleLevel(TextureObjects[TexureLayerIndex], CoordinateValue, MipValue0Value, bAutomaticMipViewBias, SampleAttributes);
+				break;
+
+			case TMVM_Derivative:
+				SampleCodeValues[TexureLayerIndex] = Em.TextureSampleGrad(TextureObjects[TexureLayerIndex], CoordinateValue, MipValue0Value, MipValue1Value, bAutomaticMipViewBias, SampleAttributes);
+				break;
+
+			default:
+				UE_MIR_UNREACHABLE();
+		}
+	}
+
+	UE_MIR_CHECKPOINT(Em);
+
+	// Compile unpacking code
+	for (int32 OutputIndex = 0; OutputIndex < 8; ++OutputIndex)
+	{
+		// Calculate the virtual texture layer and sampling/unpacking functions for this output
+		// Fallback to a sensible default value if the output isn't valid for the bound virtual texture
+		FRuntimeVirtualTextureUnpackProperties UnpackProperties;
+		if (!GetRVTUnpackProperties(OutputIndex, bIsVirtualTextureValid, UnpackProperties))
+		{
+			Em.Errorf(TEXT("Failed to retrieve unpack properties from RuntimeVirtualTexture for output pin %d"), OutputIndex);
+			return;
+		}
+
+		if (UnpackProperties.ConstantVector.IsEmpty())
+		{
+			if (UnpackProperties.UnpackType != EVirtualTextureUnpackType::None)
+			{
+				MIR::FValueRef P0 = Uniforms[ERuntimeVirtualTextureShaderUniform_WorldHeightUnpack];
+				Em.Output(OutputIndex,
+					BuildVirtualTextureUnpack(Em, SampleCodeValues[0], SampleCodeValues[1], SampleCodeValues[2], P0, UnpackProperties.UnpackType));
+			}
+			else
+			{
+				Em.Output(OutputIndex,
+					Em.Swizzle(SampleCodeValues[UnpackProperties.UnpackTarget], MIR::FSwizzleMask(
+						(UnpackProperties.UnpackMask     ) & 1,
+						(UnpackProperties.UnpackMask >> 1) & 1,
+						(UnpackProperties.UnpackMask >> 2) & 1,
+						(UnpackProperties.UnpackMask >> 3) & 1)));
+			}
+		}
+		else
+		{
+			Em.Output(OutputIndex, BuildConstantVector(Em, UnpackProperties.ConstantVector));
+		}
+	}
+}
+
+#endif // WITH_EDITOR
+
 void UMaterialExpressionTime::Build(MIR::FEmitter& Em)
 {
 	MIR::FType ScalarFloatType = MIR::FType::MakeScalar(MIR::EScalarKind::Float);
diff --git a/Engine/Source/Runtime/Engine/Private/Materials/MaterialIR.cpp b/Engine/Source/Runtime/Engine/Private/Materials/MaterialIR.cpp
--- a/Engine/Source/Runtime/Engine/Private/Materials/MaterialIR.cpp
+++ b/Engine/Source/Runtime/Engine/Private/Materials/MaterialIR.cpp
@@ -6,6 +6,7 @@
 #include "Materials/MaterialAttributeDefinitionMap.h"
 #include "Misc/LargeWorldRenderPosition.h"
 #include "Engine/Texture.h"
+#include "VT/RuntimeVirtualTexture.h"
 
 #if WITH_EDITOR
 
@@ -35,6 +36,7 @@ const TCHAR* LexToString(EValueKind Kind)
 		case VK_ScreenTexture: return TEXT("ScreenTexture");
 		case VK_ShadingModel: return TEXT("ShadingModel");
 		case VK_TextureObject: return TEXT("TextureObject");
+		case VK_RuntimeVirtualTextureObject: return TEXT("RuntimeVirtualTextureObject");
 		case VK_UniformParameter: return TEXT("UniformParameter");
 		case VK_Composite: return TEXT("Composite");
 		case VK_SetMaterialOutput: return TEXT("SetMaterialOutput");
@@ -101,6 +103,7 @@ uint32 FValue::GetSizeInBytes() const
 		case VK_ScreenTexture: return sizeof(FScreenTexture);
 		case VK_ShadingModel: return sizeof(FShadingModel);
 		case VK_TextureObject: return sizeof(FTextureObject);
+		case VK_RuntimeVirtualTextureObject: return sizeof(FRuntimeVirtualTextureObject);
 		case VK_UniformParameter: return sizeof(FUniformParameter);
 		case VK_Composite: return sizeof(FComposite) + sizeof(FValue*) * static_cast<const FComposite*>(this)->GetComponents().Num();
 		case VK_SetMaterialOutput: return sizeof(FSetMaterialOutput);
@@ -441,12 +444,16 @@ bool FValue::EqualsConstant(FVector4f TestValue) const
 	}
 }
 
-UTexture* FValue::GetTexture() const
+UObject* FValue::GetTextureObject() const
 {
 	if (const FTextureObject* TextureObject = As<FTextureObject>())
 	{
 		return TextureObject->Texture;
 	}
+	if (const FRuntimeVirtualTextureObject* RVTextureObject = As<FRuntimeVirtualTextureObject>())
+	{
+		return RVTextureObject->RVTexture;
+	}
 	return nullptr;
 }
 
@@ -454,11 +461,15 @@ int32 FValue::GetUniformParameterIndex() const
 {
 	if (const MIR::FTextureObject* TextureObject = As<MIR::FTextureObject>())
 	{
-		return (int32)TextureObject->Analysis_UniformParameterIndex;
+		return TextureObject->Analysis_UniformParameterIndex;
+	}
+	if (const MIR::FRuntimeVirtualTextureObject* RVTextureObject = As<MIR::FRuntimeVirtualTextureObject>())
+	{
+		return RVTextureObject->Analysis_UniformParameterIndex;
 	}
 	if (const MIR::FUniformParameter* UniformParameter = As<MIR::FUniformParameter>())
 	{
-		return (int32)UniformParameter->Analysis_UniformParameterIndex;
+		return UniformParameter->Analysis_UniformParameterIndex;
 	}
 	return INDEX_NONE;
 }
diff --git a/Engine/Source/Runtime/Engine/Private/Materials/MaterialIREmitter.cpp b/Engine/Source/Runtime/Engine/Private/Materials/MaterialIREmitter.cpp
--- a/Engine/Source/Runtime/Engine/Private/Materials/MaterialIREmitter.cpp
+++ b/Engine/Source/Runtime/Engine/Private/Materials/MaterialIREmitter.cpp
@@ -59,13 +59,32 @@ FSwizzleMask::FSwizzleMask(EVectorComponent X, EVectorComponent Y, EVectorCompon
 FSwizzleMask::FSwizzleMask(EVectorComponent X, EVectorComponent Y, EVectorComponent Z, EVectorComponent W)
 : NumComponents{ 4 }
 {
-
 	Components[0] = X;
 	Components[1] = Y;
 	Components[2] = Z;
 	Components[3] = W;
 }
 
+FSwizzleMask::FSwizzleMask(bool bMaskX, bool bMaskY, bool bMaskZ, bool bMaskW)
+{
+	if (bMaskX)
+	{
+		Append(MIR::EVectorComponent::X);
+	}
+	if (bMaskY)
+	{
+		Append(MIR::EVectorComponent::Y);
+	}
+	if (bMaskZ)
+	{
+		Append(MIR::EVectorComponent::Z);
+	}
+	if (bMaskW)
+	{
+		Append(MIR::EVectorComponent::W);
+	}
+}
+
 FSwizzleMask FSwizzleMask::XYZ()
 {
 	return { EVectorComponent::X, EVectorComponent::Y, EVectorComponent::Z };
@@ -972,6 +991,7 @@ FValueRef FEmitter::ShadingModel(EMaterialShadingModel Id)
 
 FValueRef FEmitter::TextureObject(UTexture* Texture, EMaterialSamplerType SamplerType)
 {
+	check(Texture);
 	FString SamplerTypeError;
 	if (!UMaterialExpressionTextureBase::VerifySamplerType(GetShaderPlatform(), GetTargetPlatform(), Texture, SamplerType, SamplerTypeError))
 	{
@@ -986,7 +1006,111 @@ FValueRef FEmitter::TextureObject(UTexture* Texture, EMaterialSamplerType Sample
 	return EmitPrototype(*this, Proto);
 }
 
-FValueRef FEmitter::Parameter(FName Name, FMaterialParameterMetadata& Metadata, EMaterialSamplerType SamplerType)
+// Maps a sampler type from standard texture (ST) to virtual texture (VT).
+static EMaterialSamplerType PromoteSamplerTypeFromSTtoVT(EMaterialSamplerType InSamplerType)
+{
+	switch (InSamplerType)
+	{
+		case SAMPLERTYPE_Color: return SAMPLERTYPE_VirtualColor;
+		case SAMPLERTYPE_Grayscale: return SAMPLERTYPE_VirtualGrayscale;
+		case SAMPLERTYPE_Alpha: return SAMPLERTYPE_VirtualAlpha;
+		case SAMPLERTYPE_Normal: return SAMPLERTYPE_VirtualNormal;
+		case SAMPLERTYPE_Masks: return SAMPLERTYPE_VirtualMasks;
+		case SAMPLERTYPE_LinearColor: return SAMPLERTYPE_VirtualLinearColor;
+		case SAMPLERTYPE_LinearGrayscale: return SAMPLERTYPE_VirtualLinearGrayscale;
+	}
+	return InSamplerType;
+}
+
+// Maps a sampler type from virtual texture (VT) to standard texture (ST).
+static EMaterialSamplerType DemoteSamplerTypeFromVTtoST(EMaterialSamplerType InSamplerType)
+{
+	switch (InSamplerType)
+	{
+		case SAMPLERTYPE_VirtualColor: return SAMPLERTYPE_Color;
+		case SAMPLERTYPE_VirtualGrayscale: return SAMPLERTYPE_Grayscale;
+		case SAMPLERTYPE_VirtualAlpha: return SAMPLERTYPE_Alpha;
+		case SAMPLERTYPE_VirtualNormal: return SAMPLERTYPE_Normal;
+		case SAMPLERTYPE_VirtualMasks: return SAMPLERTYPE_Masks;
+		case SAMPLERTYPE_VirtualLinearColor: return SAMPLERTYPE_LinearColor;
+		case SAMPLERTYPE_VirtualLinearGrayscale: return SAMPLERTYPE_LinearGrayscale;
+	}
+	return InSamplerType;
+}
+
+static bool IsVirtualTexture(FValueRef Texture)
+{
+	if (FUniformParameter* UniformParameter = Texture->As<FUniformParameter>())
+	{
+		return IsVirtualSamplerType(UniformParameter->SamplerType);
+	}
+	if (FTextureObject* TextureObject = Texture->As<FTextureObject>())
+	{
+		return (TextureObject->Texture->GetMaterialType() & MCT_TextureVirtual) != 0;
+	}
+	if (Texture->IsA(VK_RuntimeVirtualTextureObject))
+	{
+		return true;
+	}
+	return false;
+}
+
+// Returns this value's texture sampler type if it has one (SAMPLERTYPE_MAX otherwise).
+static EMaterialSamplerType GetValueMaterialSamplerType(FValueRef Value)
+{
+	if (MIR::FTextureObject* TextureObject = MIR::As<MIR::FTextureObject>(Value))
+	{
+		return TextureObject->SamplerType;
+	}
+	if (MIR::FRuntimeVirtualTextureObject* RVTextureObject = MIR::As<MIR::FRuntimeVirtualTextureObject>(Value))
+	{
+		return RVTextureObject->SamplerType;
+	}
+	if (MIR::FUniformParameter* UniformParameter = MIR::As<MIR::FUniformParameter>(Value))
+	{
+		return UniformParameter->SamplerType;
+	}
+	return SAMPLERTYPE_MAX;
+}
+
+static EMaterialSamplerType MapSamplerTypeForTexture(FValueRef InTexture, EMaterialSamplerType InSamplerType)
+{
+	// Can't sample with virtual texturing if input texture is not a virtual texture
+	if (InSamplerType == SAMPLERTYPE_MAX)
+	{
+		InSamplerType = GetValueMaterialSamplerType(InTexture);
+	}
+	return IsVirtualTexture(InTexture) ? PromoteSamplerTypeFromSTtoVT(InSamplerType) : DemoteSamplerTypeFromVTtoST(InSamplerType);
+}
+
+static FValueRef VTPageTableLoadFromSamplerSource(
+	FEmitter& Em, FValueRef Texture, const FTextureSampleBaseAttributes& BaseAttributes, FValueRef TexCoord,
+	FValueRef TexCoordDdx = {}, FValueRef TexCoordDdy = {}, ETextureMipValueMode MipValueMode = TMVM_None, FValueRef MipValue = {})
+{
+	// Cast input texture to UTexture. If it's a URuntimeVirtualTexture, we accept the cast to be null when passed to GetTextureAddressForSamplerSource().
+	TextureAddress StaticAddressX = TA_Wrap;
+	TextureAddress StaticAddressY = TA_Wrap;
+	TextureAddress StaticAddressZ = TA_Wrap;
+	UE::MaterialTranslatorUtils::GetTextureAddressForSamplerSource(Cast<UTexture>(Texture->GetTextureObject()), BaseAttributes.SamplerSourceMode, StaticAddressX, StaticAddressY, StaticAddressZ);
+
+	return Em.VTPageTableLoad(Texture, StaticAddressX, StaticAddressY, TexCoord, TexCoordDdx, TexCoordDdy, BaseAttributes.bEnableFeedback, BaseAttributes.bIsAdaptive, MipValueMode, MipValue);
+}
+
+FValueRef FEmitter::RuntimeVirtualTextureObject(URuntimeVirtualTexture* RVTexture, EMaterialSamplerType SamplerType, int32 VTLayerIndex, int32 VTPageTableIndex)
+{
+	check(RVTexture);
+	check(IsVirtualSamplerType(SamplerType));
+
+	FRuntimeVirtualTextureObject Prototype = MakePrototype<FRuntimeVirtualTextureObject>(FType::MakeRuntimeVirtualTexture());
+	Prototype.RVTexture = RVTexture;
+	Prototype.SamplerType = PromoteSamplerTypeFromSTtoVT(SamplerType);
+	Prototype.VTLayerIndex = VTLayerIndex;
+	Prototype.VTPageTableIndex = VTPageTableIndex;
+	Prototype.Analysis_UniformParameterIndex = INDEX_NONE;
+	return EmitPrototype(*this, Prototype);
+}
+
+FValueRef FEmitter::Parameter(FName Name, FMaterialParameterMetadata& Metadata, EMaterialSamplerType SamplerType, int32 VTLayerIndex, int32 VTPageTableIndex)
 {
 	// Helper local function that register a parameter (info and metadata) to the module, and returns some uint32 ID.
 	auto RegisterParameter = [] (FMaterialIRModule* InModule, FMaterialParameterInfo InInfo, const FMaterialParameterMetadata& InMetadata) -> uint32
@@ -1004,11 +1128,15 @@ FValueRef FEmitter::Parameter(FName Name, FMaterialParameterMetadata& Metadata,
 		return Id;
 	};
 
-	auto MakeUniformParameter = [this, &RegisterParameter](FType Type, FMaterialParameterInfo InInfo, const FMaterialParameterMetadata& InMetadata, EMaterialSamplerType InSamplerType = SAMPLERTYPE_Color) -> FValueRef
+	auto MakeUniformParameter = [this, &RegisterParameter](
+		FType Type, FMaterialParameterInfo InInfo, const FMaterialParameterMetadata& InMetadata,
+		EMaterialSamplerType InSamplerType = SAMPLERTYPE_Color, int32 VTLayerIndex = INDEX_NONE, int32 VTPageTableIndex = INDEX_NONE) -> FValueRef
 	{
 		FUniformParameter Prototype = MakePrototype<FUniformParameter>(MoveTemp(Type));
 		Prototype.ParameterIdInModule = RegisterParameter(Module, MoveTemp(InInfo), InMetadata);
 		Prototype.SamplerType = InSamplerType;
+		Prototype.VTLayerIndex = VTLayerIndex;
+		Prototype.VTPageTableIndex = VTPageTableIndex;
 		Prototype.Analysis_UniformParameterIndex = INDEX_NONE;
 		return EmitPrototype(*this, Prototype);
 	};
@@ -1048,6 +1176,11 @@ FValueRef FEmitter::Parameter(FName Name, FMaterialParameterMetadata& Metadata,
 			return MakeUniformParameter(FType::MakeTexture(), Info, Metadata, SamplerType);
 		}
 
+		case EMaterialParameterType::RuntimeVirtualTexture:
+		{
+			return MakeUniformParameter(FType::MakeRuntimeVirtualTexture(), Info, Metadata, PromoteSamplerTypeFromSTtoVT(SamplerType), VTLayerIndex, VTPageTableIndex);
+		}
+
 		case EMaterialParameterType::StaticSwitch:
 		{
 			// Apply eventual parameter override
@@ -1085,7 +1218,7 @@ FValueRef FEmitter::Parameter(FName Name, FMaterialParameterMetadata& Metadata,
 	}
 }
 
-UTexture* FEmitter::GetTextureFromValue(MIR::FValueRef Texture)
+UObject* FEmitter::GetTextureFromValue(MIR::FValueRef Texture)
 {
 	if (!Texture)
 	{
@@ -1095,37 +1228,33 @@ UTexture* FEmitter::GetTextureFromValue(MIR::FValueRef Texture)
 	if (FUniformParameter* UniformParameter = Texture->As<FUniformParameter>())
 	{
 		const FMaterialParameterMetadata& ParameterMetadata = Module->GetParameterMetadata(UniformParameter->ParameterIdInModule);
-		if (UTexture* ParameterDefaultTexture = ::Cast<UTexture>(ParameterMetadata.Value.AsTextureObject()))
-		{
-			return ParameterDefaultTexture;
-		}
+		return ParameterMetadata.Value.AsTextureObject();
 	}
-	return Texture->GetTexture();
+	return Texture->GetTextureObject();
 }
 
-// For now, only a small subset of opcodes are supported: TextureSize, TexelSize.
-FValueRef FEmitter::PreshaderParameter(UE::Shader::EPreshaderOpcode Opcode, FValueRef SourceParameter)
+// For now, only a small subset of opcodes are supported: TextureSize, TexelSize, and RuntimeVirtualTextureUniform.
+FValueRef FEmitter::PreshaderParameter(FType Type, UE::Shader::EPreshaderOpcode Opcode, FValueRef SourceParameter, FPreshaderParameterPayload Payload)
 {
 	checkf(
 		Opcode == UE::Shader::EPreshaderOpcode::TextureSize ||
-		Opcode == UE::Shader::EPreshaderOpcode::TexelSize,
+		Opcode == UE::Shader::EPreshaderOpcode::TexelSize ||
+		Opcode == UE::Shader::EPreshaderOpcode::RuntimeVirtualTextureUniform,
 		TEXT("Preshader opcode (0x%X) not supported for parameters in new material translator"), (int32)Opcode
 	);
 
-	UTexture* SourceParameterTexture = GetTextureFromValue(SourceParameter);
+	UObject* SourceParameterTexture = GetTextureFromValue(SourceParameter);
 	if (!SourceParameterTexture)
 	{
 		Error(TEXT("Missing default texture from source parameter"));
 		return Poison();
 	}
 
-	const EMaterialValueType TextureType = SourceParameterTexture->GetMaterialType();
-	const EMaterialValueType PropertyType = UE::MaterialTranslatorUtils::GetTexturePropertyValueType(TextureType);
-
-	FPreshaderParameter Prototype = MakePrototype<FPreshaderParameter>(FType::FromMaterialValueType(PropertyType));
+	FPreshaderParameter Prototype = MakePrototype<FPreshaderParameter>(Type);
 	Prototype.SourceParameter = SourceParameter;
 	Prototype.Opcode = Opcode;
 	Prototype.TextureIndex = Material->GetReferencedTextures().Find(SourceParameterTexture);
+	Prototype.Payload = MoveTemp(Payload);
 	return EmitPrototype(*this, Prototype);
 }
 
@@ -3203,68 +3332,7 @@ FValueRef FEmitter::Nop(FValueRef Arg)
 	return EmitPrototype(*this, Prototype);
 }
 
-// Maps a sampler type from standard texture (ST) to virtual texture (VT).
-static EMaterialSamplerType PromoteSamplerTypeFromSTtoVT(EMaterialSamplerType InSamplerType)
-{
-	switch (InSamplerType)
-	{
-		case SAMPLERTYPE_Color: return SAMPLERTYPE_VirtualColor;
-		case SAMPLERTYPE_Grayscale: return SAMPLERTYPE_VirtualGrayscale;
-		case SAMPLERTYPE_Alpha: return SAMPLERTYPE_VirtualAlpha;
-		case SAMPLERTYPE_Normal: return SAMPLERTYPE_VirtualNormal;
-		case SAMPLERTYPE_Masks: return SAMPLERTYPE_VirtualMasks;
-		case SAMPLERTYPE_LinearColor: return SAMPLERTYPE_VirtualLinearColor;
-		case SAMPLERTYPE_LinearGrayscale: return SAMPLERTYPE_VirtualLinearGrayscale;
-	}
-	return InSamplerType;
-}
-
-// Maps a sampler type from virtual texture (VT) to standard texture (ST).
-static EMaterialSamplerType DemoteSamplerTypeFromVTtoST(EMaterialSamplerType InSamplerType)
-{
-	switch (InSamplerType)
-	{
-		case SAMPLERTYPE_VirtualColor: return SAMPLERTYPE_Color;
-		case SAMPLERTYPE_VirtualGrayscale: return SAMPLERTYPE_Grayscale;
-		case SAMPLERTYPE_VirtualAlpha: return SAMPLERTYPE_Alpha;
-		case SAMPLERTYPE_VirtualNormal: return SAMPLERTYPE_Normal;
-		case SAMPLERTYPE_VirtualMasks: return SAMPLERTYPE_Masks;
-		case SAMPLERTYPE_VirtualLinearColor: return SAMPLERTYPE_LinearColor;
-		case SAMPLERTYPE_VirtualLinearGrayscale: return SAMPLERTYPE_LinearGrayscale;
-	}
-	return InSamplerType;
-}
-
-static bool IsVirtualTexture(FValueRef Texture)
-{
-	if (FUniformParameter* UniformParameter = Texture->As<FUniformParameter>())
-	{
-		return IsVirtualSamplerType(UniformParameter->SamplerType);
-	}
-	else if (FTextureObject* TextureObject = Texture->As<FTextureObject>())
-	{
-		return (TextureObject->Texture->GetMaterialType() & MCT_TextureVirtual) != 0;
-	}
-	return false;
-}
-
-static EMaterialSamplerType MapSamplerTypeForTexture(FValueRef InTexture, EMaterialSamplerType InSamplerType)
-{
-	// Can't sample with virtual texturing if input texture is not a virtual texture
-	return IsVirtualTexture(InTexture) ? PromoteSamplerTypeFromSTtoVT(InSamplerType) : DemoteSamplerTypeFromVTtoST(InSamplerType);
-}
-
-static FValueRef VTPageTableLoadFromSamplerSource(FEmitter& Em, FValueRef Texture, ESamplerSourceMode SamplerSourceMode, FValueRef TexCoord, FValueRef TexCoordDdx, FValueRef TexCoordDdy)
-{
-	TextureAddress StaticAddressX = TA_Wrap;
-	TextureAddress StaticAddressY = TA_Wrap;
-	TextureAddress StaticAddressZ = TA_Wrap;
-	UE::MaterialTranslatorUtils::GetTextureAddressForSamplerSource(Texture->GetTexture(), SamplerSourceMode, StaticAddressX, StaticAddressY, StaticAddressZ);
-
-	return Em.VTPageTableLoad(Texture, StaticAddressX, StaticAddressY, TexCoord, TexCoordDdx, TexCoordDdy);
-}
-
-FValueRef FEmitter::TextureGather(FValueRef Texture, FValueRef TexCoord, ETextureReadMode GatherMode, ESamplerSourceMode SamplerSourceMode, EMaterialSamplerType SamplerType)
+FValueRef FEmitter::TextureGather(FValueRef Texture, FValueRef TexCoord, ETextureReadMode GatherMode, const FTextureSampleBaseAttributes& BaseAttributes)
 {
 	check(GatherMode >= ETextureReadMode::GatherRed && GatherMode <= ETextureReadMode::GatherAlpha);
 
@@ -3273,37 +3341,37 @@ FValueRef FEmitter::TextureGather(FValueRef Texture, FValueRef TexCoord, ETextur
 		return Poison();
 	}
 
-	SamplerType = MapSamplerTypeForTexture(Texture, SamplerType);
+	const EMaterialSamplerType SamplerType = MapSamplerTypeForTexture(Texture, BaseAttributes.SamplerType);
 
 	FTextureRead Prototype = MakePrototype<FTextureRead>(FType::MakeFloatVector(4));
 	Prototype.TextureObject = Texture;
 	Prototype.TexCoord = TexCoord;
 	Prototype.Mode = GatherMode;
-	Prototype.SamplerSourceMode = SamplerSourceMode;
+	Prototype.SamplerSourceMode = BaseAttributes.SamplerSourceMode;
 	Prototype.SamplerType = SamplerType;
 
 	if (IsVirtualSamplerType(SamplerType))
 	{
-		Prototype.VTPageTable = VTPageTableLoadFromSamplerSource(*this, Texture, SamplerSourceMode, Prototype.TexCoord, {}, {});
+		Prototype.VTPageTable = VTPageTableLoadFromSamplerSource(*this, Texture, BaseAttributes, Prototype.TexCoord);
 	}
 
 	return EmitPrototype(*this, Prototype);
 }
 
-FValueRef FEmitter::TextureSample(FValueRef Texture, FValueRef TexCoord, bool bAutomaticViewMipBias, ESamplerSourceMode SamplerSourceMode, EMaterialSamplerType SamplerType)
+FValueRef FEmitter::TextureSample(FValueRef Texture, FValueRef TexCoord, bool bAutomaticViewMipBias, const FTextureSampleBaseAttributes& BaseAttributes)
 {
 	if (IsAnyNotValid(Texture, TexCoord))
 	{
 		return Poison();
 	}
 
-	SamplerType = MapSamplerTypeForTexture(Texture, SamplerType);
+	const EMaterialSamplerType SamplerType = MapSamplerTypeForTexture(Texture, BaseAttributes.SamplerType);
 
 	FTextureRead PrototypePixel = MakePrototype<FTextureRead>(FType::MakeFloatVector(4));
 	PrototypePixel.TextureObject = Texture;
 	PrototypePixel.TexCoord = TexCoord;
 	PrototypePixel.Mode = ETextureReadMode::MipAuto;
-	PrototypePixel.SamplerSourceMode = SamplerSourceMode;
+	PrototypePixel.SamplerSourceMode = BaseAttributes.SamplerSourceMode;
 	PrototypePixel.SamplerType = SamplerType;
 
 	FTextureRead PrototypeCompute = PrototypePixel;
@@ -3328,9 +3396,9 @@ FValueRef FEmitter::TextureSample(FValueRef Texture, FValueRef TexCoord, bool bA
 
 	if (IsVirtualSamplerType(SamplerType))
 	{
-		PrototypePixel.VTPageTable = VTPageTableLoadFromSamplerSource(*this, Texture, SamplerSourceMode, PrototypePixel.TexCoord, {}, {});
-		PrototypeCompute.VTPageTable = VTPageTableLoadFromSamplerSource(*this, Texture, SamplerSourceMode, PrototypeCompute.TexCoord, PrototypeCompute.TexCoordDdx, PrototypeCompute.TexCoordDdy);
-		PrototypeVertex.VTPageTable = VTPageTableLoadFromSamplerSource(*this, Texture, SamplerSourceMode, PrototypeVertex.TexCoord, {}, {});
+		PrototypePixel.VTPageTable = VTPageTableLoadFromSamplerSource(*this, Texture, BaseAttributes, PrototypePixel.TexCoord);
+		PrototypeCompute.VTPageTable = VTPageTableLoadFromSamplerSource(*this, Texture, BaseAttributes, PrototypeCompute.TexCoord, PrototypeCompute.TexCoordDdx, PrototypeCompute.TexCoordDdy);
+		PrototypeVertex.VTPageTable = VTPageTableLoadFromSamplerSource(*this, Texture, BaseAttributes, PrototypeVertex.TexCoord);
 	}
 
 	FStageSwitch StageSwitch = MakePrototype<FStageSwitch>(PrototypePixel.Type);
@@ -3341,21 +3409,21 @@ FValueRef FEmitter::TextureSample(FValueRef Texture, FValueRef TexCoord, bool bA
 	return EmitPrototype(*this, StageSwitch);
 }
 
-FValueRef FEmitter::TextureSampleLevel(FValueRef Texture, FValueRef TexCoord, FValueRef MipLevel, bool bAutomaticViewMipBias, ESamplerSourceMode SamplerSourceMode, EMaterialSamplerType SamplerType)
+FValueRef FEmitter::TextureSampleLevel(FValueRef Texture, FValueRef TexCoord, FValueRef MipLevel, bool bAutomaticViewMipBias, const FTextureSampleBaseAttributes& BaseAttributes)
 {
 	if (IsAnyNotValid(Texture, TexCoord, MipLevel))
 	{
 		return Poison();
 	}
 
-	SamplerType = MapSamplerTypeForTexture(Texture, SamplerType);
+	const EMaterialSamplerType SamplerType = MapSamplerTypeForTexture(Texture, BaseAttributes.SamplerType);
 
 	FTextureRead Prototype = MakePrototype<FTextureRead>(FType::MakeFloatVector(4));
 	Prototype.TextureObject = Texture;
 	Prototype.TexCoord = TexCoord;
 	Prototype.MipValue = MipLevel;
 	Prototype.Mode = ETextureReadMode::MipLevel;
-	Prototype.SamplerSourceMode = SamplerSourceMode;
+	Prototype.SamplerSourceMode = BaseAttributes.SamplerSourceMode;
 	Prototype.SamplerType = SamplerType;
 
 	if (bAutomaticViewMipBias)
@@ -3365,13 +3433,13 @@ FValueRef FEmitter::TextureSampleLevel(FValueRef Texture, FValueRef TexCoord, FV
 
 	if (IsVirtualSamplerType(SamplerType))
 	{
-		Prototype.VTPageTable = VTPageTableLoadFromSamplerSource(*this, Texture, SamplerSourceMode, Prototype.TexCoord, {}, {});
+		Prototype.VTPageTable = VTPageTableLoadFromSamplerSource(*this, Texture, BaseAttributes, Prototype.TexCoord, {}, {}, TMVM_MipLevel, Prototype.MipValue);
 	}
 
 	return EmitPrototype(*this, Prototype);
 }
 
-FValueRef FEmitter::TextureSampleBias(FValueRef Texture, FValueRef TexCoord, FValueRef MipBias, bool bAutomaticViewMipBias, ESamplerSourceMode SamplerSourceMode, EMaterialSamplerType SamplerType)
+FValueRef FEmitter::TextureSampleBias(FValueRef Texture, FValueRef TexCoord, FValueRef MipBias, bool bAutomaticViewMipBias, const FTextureSampleBaseAttributes& BaseAttributes)
 {
 	if (IsAnyNotValid(Texture, TexCoord, MipBias))
 	{
@@ -3383,14 +3451,14 @@ FValueRef FEmitter::TextureSampleBias(FValueRef Texture, FValueRef TexCoord, FVa
 		MipBias = Operator(BO_Add, MipBias, ExternalInput(EExternalInput::ViewMaterialTextureMipBias));
 	}
 
-	SamplerType = MapSamplerTypeForTexture(Texture, SamplerType);
+	const EMaterialSamplerType SamplerType = MapSamplerTypeForTexture(Texture, BaseAttributes.SamplerType);
 
 	FTextureRead PrototypePixel = MakePrototype<FTextureRead>(FType::MakeFloatVector(4));
 	PrototypePixel.TextureObject = Texture;
 	PrototypePixel.TexCoord = TexCoord;
 	PrototypePixel.MipValue = MipBias;
 	PrototypePixel.Mode = ETextureReadMode::MipBias;
-	PrototypePixel.SamplerSourceMode = SamplerSourceMode;
+	PrototypePixel.SamplerSourceMode = BaseAttributes.SamplerSourceMode;
 	PrototypePixel.SamplerType = SamplerType;
 
 	FTextureRead PrototypeCompute = PrototypePixel;
@@ -3411,8 +3479,8 @@ FValueRef FEmitter::TextureSampleBias(FValueRef Texture, FValueRef TexCoord, FVa
 
 	if (IsVirtualSamplerType(SamplerType))
 	{
-		PrototypePixel.VTPageTable = VTPageTableLoadFromSamplerSource(*this, Texture, SamplerSourceMode, PrototypePixel.TexCoord, {}, {});
-		PrototypeCompute.VTPageTable = VTPageTableLoadFromSamplerSource(*this, Texture, SamplerSourceMode, PrototypeCompute.TexCoord, PrototypeCompute.TexCoordDdx, PrototypeCompute.TexCoordDdy);
+		PrototypePixel.VTPageTable = VTPageTableLoadFromSamplerSource(*this, Texture, BaseAttributes, PrototypePixel.TexCoord, {}, {}, TMVM_MipBias, PrototypePixel.MipValue);
+		PrototypeCompute.VTPageTable = VTPageTableLoadFromSamplerSource(*this, Texture, BaseAttributes, PrototypeCompute.TexCoord, PrototypeCompute.TexCoordDdx, PrototypeCompute.TexCoordDdy, TMVM_MipBias, PrototypeCompute.MipValue);
 	}
 
 	FStageSwitch StageSwitch = MakePrototype<FStageSwitch>(PrototypePixel.Type);
@@ -3423,14 +3491,14 @@ FValueRef FEmitter::TextureSampleBias(FValueRef Texture, FValueRef TexCoord, FVa
 	return EmitPrototype(*this, StageSwitch);
 }
 
-FValueRef FEmitter::TextureSampleGrad(FValueRef Texture, FValueRef TexCoord, FValueRef TexCoordDdx, FValueRef TexCoordDdy, bool bAutomaticViewMipBias, ESamplerSourceMode SamplerSourceMode, EMaterialSamplerType SamplerType)
+FValueRef FEmitter::TextureSampleGrad(FValueRef Texture, FValueRef TexCoord, FValueRef TexCoordDdx, FValueRef TexCoordDdy, bool bAutomaticViewMipBias, const FTextureSampleBaseAttributes& BaseAttributes)
 {
 	if (IsAnyNotValid(Texture, TexCoord, TexCoordDdx, TexCoordDdy))
 	{
 		return Poison();
 	}
 
-	SamplerType = MapSamplerTypeForTexture(Texture, SamplerType);
+	const EMaterialSamplerType SamplerType = MapSamplerTypeForTexture(Texture, BaseAttributes.SamplerType);
 
 	FTextureRead Prototype = MakePrototype<FTextureRead>(FType::MakeFloatVector(4));
 	Prototype.TextureObject = Texture;
@@ -3438,7 +3506,7 @@ FValueRef FEmitter::TextureSampleGrad(FValueRef Texture, FValueRef TexCoord, FVa
 	Prototype.TexCoordDdx = TexCoordDdx;
 	Prototype.TexCoordDdy = TexCoordDdy;
 	Prototype.Mode = ETextureReadMode::Derivatives;
-	Prototype.SamplerSourceMode = SamplerSourceMode;
+	Prototype.SamplerSourceMode = BaseAttributes.SamplerSourceMode;
 	Prototype.SamplerType = SamplerType;
 
 	if (bAutomaticViewMipBias)
@@ -3450,13 +3518,15 @@ FValueRef FEmitter::TextureSampleGrad(FValueRef Texture, FValueRef TexCoord, FVa
 
 	if (IsVirtualSamplerType(SamplerType))
 	{
-		Prototype.VTPageTable = VTPageTableLoadFromSamplerSource(*this, Texture, SamplerSourceMode, Prototype.TexCoord, Prototype.TexCoordDdx, Prototype.TexCoordDdy);
+		Prototype.VTPageTable = VTPageTableLoadFromSamplerSource(*this, Texture, BaseAttributes, Prototype.TexCoord, Prototype.TexCoordDdx, Prototype.TexCoordDdy, TMVM_Derivative);
 	}
 
 	return EmitPrototype(*this, Prototype);
 }
 
-FValueRef FEmitter::VTPageTableLoad(FValueRef Texture, TextureAddress AddressU, TextureAddress AddressV, FValueRef TexCoord, FValueRef TexCoordDdx, FValueRef TexCoordDdy, bool bEnableFeedback)
+FValueRef FEmitter::VTPageTableLoad(
+	FValueRef Texture, TextureAddress AddressU, TextureAddress AddressV, FValueRef TexCoord,
+	FValueRef TexCoordDdx, FValueRef TexCoordDdy, bool bEnableFeedback, bool bIsAdaptive, ETextureMipValueMode MipValueMode, FValueRef MipValue)
 {
 	// VT stack and layer indices are initializes during IR analysis
 	FVTPageTableRead Prototype = MakePrototype<FVTPageTableRead>(FType::MakeVTPageTableResult());
@@ -3464,9 +3534,12 @@ FValueRef FEmitter::VTPageTableLoad(FValueRef Texture, TextureAddress AddressU,
 	Prototype.TexCoord = TexCoord;
 	Prototype.TexCoordDdx = TexCoordDdx;
 	Prototype.TexCoordDdy = TexCoordDdy;
+	Prototype.MipValue = MipValue;
 	Prototype.AddressU = AddressU;
 	Prototype.AddressV = AddressV;
+	Prototype.MipValueMode = MipValueMode;
 	Prototype.bEnableFeedback = bEnableFeedback;
+	Prototype.bIsAdaptive = bIsAdaptive;
 	return EmitPrototype(*this, Prototype);
 }
 
diff --git a/Engine/Source/Runtime/Engine/Private/Materials/MaterialIRInternal.cpp b/Engine/Source/Runtime/Engine/Private/Materials/MaterialIRInternal.cpp
--- a/Engine/Source/Runtime/Engine/Private/Materials/MaterialIRInternal.cpp
+++ b/Engine/Source/Runtime/Engine/Private/Materials/MaterialIRInternal.cpp
@@ -5,11 +5,29 @@
 #include "Materials/MaterialIREmitter.h"
 #include "Materials/Material.h"
 #include "MaterialShared.h"
+#include "Engine/Texture.h"
+#include "VT/RuntimeVirtualTexture.h"
 
 #if WITH_EDITOR
 
 namespace MIR::Internal {
 
+EMaterialValueType GetTextureMaterialValueType(const UObject* TextureObject)
+{
+	if (TextureObject)
+	{
+		if (TextureObject->IsA<UTexture>())
+		{
+			return Cast<UTexture>(TextureObject)->GetMaterialType();
+		}
+		if (TextureObject->IsA<URuntimeVirtualTexture>())
+		{
+			return MCT_TextureVirtual;
+		}
+	}
+	return MCT_Unknown;
+}
+
 EMaterialTextureParameterType TextureMaterialValueTypeToParameterType(EMaterialValueType Type)
 {
 	switch (Type)
diff --git a/Engine/Source/Runtime/Engine/Private/Materials/MaterialIRInternal.h b/Engine/Source/Runtime/Engine/Private/Materials/MaterialIRInternal.h
--- a/Engine/Source/Runtime/Engine/Private/Materials/MaterialIRInternal.h
+++ b/Engine/Source/Runtime/Engine/Private/Materials/MaterialIRInternal.h
@@ -7,7 +7,10 @@
 #if WITH_EDITOR
 
 namespace MIR::Internal {
- 
+
+// Returns the material value tpye of the specified UTexture or URuntimeVirtualTexture.
+EMaterialValueType GetTextureMaterialValueType(const UObject* TextureObject);
+
 //
 EMaterialTextureParameterType TextureMaterialValueTypeToParameterType(EMaterialValueType Type);
 
diff --git a/Engine/Source/Runtime/Engine/Private/Materials/MaterialIRToHLSLTranslator.cpp b/Engine/Source/Runtime/Engine/Private/Materials/MaterialIRToHLSLTranslator.cpp
--- a/Engine/Source/Runtime/Engine/Private/Materials/MaterialIRToHLSLTranslator.cpp
+++ b/Engine/Source/Runtime/Engine/Private/Materials/MaterialIRToHLSLTranslator.cpp
@@ -5,6 +5,7 @@
 #include "Engine/Font.h"
 #include "Engine/Texture.h"
 #include "Engine/Texture2D.h"
+#include "VT/RuntimeVirtualTexture.h"
 #include "MaterialDomain.h"
 #include "MaterialIRInternal.h"
 #include "Materials/Material.h"
@@ -420,6 +421,7 @@ struct FPrivate : FMaterialIRToHLSLTranslation
 			case MIR::VK_ScreenTexture: LowerScreenTexture(static_cast<const MIR::FScreenTexture*>(Value)); break;
 			case MIR::VK_ShadingModel: LowerShadingModel(static_cast<const MIR::FShadingModel*>(Value)); break;
 			case MIR::VK_TextureObject: LowerTextureObject(static_cast<const MIR::FTextureObject*>(Value)); break;
+			case MIR::VK_RuntimeVirtualTextureObject: LowerRuntimeVirtualTextureObject(static_cast<const MIR::FRuntimeVirtualTextureObject*>(Value)); break;
 			case MIR::VK_UniformParameter: LowerUniformParameter(static_cast<const MIR::FUniformParameter*>(Value)); break;
 			default: UE_MIR_UNREACHABLE();
 		}
@@ -779,14 +781,18 @@ struct FPrivate : FMaterialIRToHLSLTranslation
 		LowerTextureReference(TextureObject->Texture->GetMaterialType(), TextureObject->Analysis_UniformParameterIndex);
 	}
 
+	void LowerRuntimeVirtualTextureObject(const MIR::FRuntimeVirtualTextureObject* RVTextureObject)
+	{
+		LowerTextureReference(MCT_TextureVirtual, RVTextureObject->Analysis_UniformParameterIndex);
+	}
+
 	void LowerUniformParameter(const MIR::FUniformParameter* UniformParameter)
 	{
-		if (UniformParameter->Type.IsTexture())
+		if (UniformParameter->Type.IsTexture() || UniformParameter->Type.IsRuntimeVirtualTexture())
 		{
-			UTexture* TextureObject = GetTextureFromUniformParameter(UniformParameter);
+			UObject* TextureObject = GetTextureFromUniformParameter(UniformParameter);
 			check(TextureObject);
-			EMaterialValueType TextureType = TextureObject->GetMaterialType();
-			LowerTextureReference(TextureType, UniformParameter->Analysis_UniformParameterIndex);
+			LowerTextureReference(MIR::Internal::GetTextureMaterialValueType(TextureObject), UniformParameter->Analysis_UniformParameterIndex);
 		}
 		else
 		{
@@ -1651,17 +1657,17 @@ struct FPrivate : FMaterialIRToHLSLTranslation
 		LowerTextureMaterialType(TextureObject->Texture->GetMaterialType());
 	}
 
-	UTexture* GetTextureFromUniformParameter(const MIR::FUniformParameter* UniformParameter)
+	UObject* GetTextureFromUniformParameter(const MIR::FUniformParameter* UniformParameter) const
 	{
 		const FMaterialParameterValue& TextureParameter = Module->GetParameterMetadata(UniformParameter->ParameterIdInModule).Value;
-		return Cast<UTexture>(TextureParameter.AsTextureObject());
+		return TextureParameter.AsTextureObject();
 	}
 
 	void LowerTextureParameterType(const MIR::FUniformParameter* UniformParameter)
 	{
-		UTexture* TextureObject = GetTextureFromUniformParameter(UniformParameter);
-		check(TextureObject);
-		LowerTextureMaterialType(TextureObject->GetMaterialType());
+		UObject* Texture = GetTextureFromUniformParameter(UniformParameter);
+		check(Texture);
+		LowerTextureMaterialType(MIR::Internal::GetTextureMaterialValueType(Texture));
 	}
 
 	void LowerStandardTextureRead(const MIR::FTextureRead* TextureRead)
@@ -1669,6 +1675,7 @@ struct FPrivate : FMaterialIRToHLSLTranslation
 		switch (TextureRead->TextureObject->Kind)
 		{
 			case MIR::EValueKind::VK_TextureObject: LowerTextureType(TextureRead->TextureObject->As<MIR::FTextureObject>()); break;
+			case MIR::EValueKind::VK_RuntimeVirtualTextureObject: LowerTextureMaterialType(MCT_TextureVirtual); break;
 			case MIR::EValueKind::VK_UniformParameter: LowerTextureParameterType(TextureRead->TextureObject->As<MIR::FUniformParameter>()); break;
 			default: UE_MIR_UNREACHABLE();
 		}
@@ -1710,8 +1717,6 @@ struct FPrivate : FMaterialIRToHLSLTranslation
 		const int32 VirtualTextureIndex = TextureRead->TextureObject->GetUniformParameterIndex();
 		check(VirtualTextureIndex >= 0);
 
-		const int32 VTLayerIndex = VTPageTableResult->VTLayerIndex;
-
 		// Sampling function
 		Printer << TEXTVIEW("TextureVirtualSample");
 
@@ -1736,7 +1741,7 @@ struct FPrivate : FMaterialIRToHLSLTranslation
 		}
 
 		Printer << ListSeparator << LowerValue(TextureRead->VTPageTable)
-			<< ListSeparator << VTLayerIndex
+			<< ListSeparator << VTPageTableResult->VTPageTableIndex
 			<< ListSeparator << TEXTVIEW("VTUniform_Unpack(") << TEXT("Material.VTPackedUniform[") << VirtualTextureIndex << TEXTVIEW("]") << TEXTVIEW(")")
 			<< EndArgs;
 	}
@@ -1750,7 +1755,11 @@ struct FPrivate : FMaterialIRToHLSLTranslation
 			Printer << TEXTVIEW("(");
 		}
 
-		checkf(TextureRead->TextureObject->Type.IsTexture(), TEXT("Invalid texture object type"));
+		checkf(
+			TextureRead->TextureObject->Type.IsTexture() ||
+			TextureRead->TextureObject->Type.IsRuntimeVirtualTexture(),
+			TEXT("Invalid texture object type")
+		);
 
 		if (IsVirtualSamplerType(TextureRead->SamplerType))
 		{
@@ -1771,40 +1780,88 @@ struct FPrivate : FMaterialIRToHLSLTranslation
 	{
 		using namespace UE::MaterialTranslatorUtils;
 
-		FStringView VTPageTableReadFunctionName = TEXTVIEW("TextureLoadVirtualPageTable");
-
 		const bool bHasDerivativeTexCoords = VTPageTableRead->TexCoordDdx != nullptr && VTPageTableRead->TexCoordDdy != nullptr;
 
+		// Construct VT page table load function name 'TextureLoadVirtualPageTable [ Adaptive ] [ * | Grad | Level ]'
+		Printer << TEXTVIEW("TextureLoadVirtualPageTable");
+		if (VTPageTableRead->bIsAdaptive)
+		{
+			Printer << TEXTVIEW("Adaptive");
+		}
+
+		switch (VTPageTableRead->MipValueMode)
+		{
+			case TMVM_None: [[fallthrough]];
+			case TMVM_MipBias: if (bHasDerivativeTexCoords) { Printer << TEXTVIEW("Grad"); } break;
+			case TMVM_MipLevel: Printer << TEXTVIEW("Level"); break;
+			case TMVM_Derivative: Printer << TEXTVIEW("Grad"); break;
+			default: UE_MIR_UNREACHABLE();
+		}
+
+		// Lower common parameters that are shared across all VT page table load functions
 		Printer
-			<< VTPageTableReadFunctionName << (bHasDerivativeTexCoords ? TEXTVIEW("Grad") : TEXTVIEW(""))
 			<< BeginArgs
 			<< ListSeparator << TEXTVIEW("VIRTUALTEXTURE_PAGETABLE_") << VTPageTableRead->VTStackIndex
 			<< ListSeparator << TEXTVIEW("VTPageTableUniform_Unpack(VIRTUALTEXTURE_PAGETABLE_UNIFORM_") << VTPageTableRead->VTStackIndex << TEXTVIEW(")")
 			<< ListSeparator << LowerValue(VTPageTableRead->TexCoord)
 			<< ListSeparator << FStringView{ GetVTAddressMode(VTPageTableRead->AddressU) }
-			<< ListSeparator << FStringView{ GetVTAddressMode(VTPageTableRead->AddressV) };
+			<< ListSeparator << FStringView{ GetVTAddressMode(VTPageTableRead->AddressV) }
+			;
 
-		if (bHasDerivativeTexCoords)
+		// Lower additional parameters depening on VT page table load function
+		switch (VTPageTableRead->MipValueMode)
 		{
-			Printer
-				<< ListSeparator << LowerValue(VTPageTableRead->TexCoordDdx)
-				<< ListSeparator << LowerValue(VTPageTableRead->TexCoordDdy);
-		}
-		else
-		{
-			constexpr float MipBias = 0.0f;
-			Printer << ListSeparator << MipBias;
+			case TMVM_None:
+				if (bHasDerivativeTexCoords)
+				{
+					Printer
+						<< ListSeparator << LowerValue(VTPageTableRead->TexCoordDdx)
+						<< ListSeparator << LowerValue(VTPageTableRead->TexCoordDdy);
+				}
+				else
+				{
+					constexpr float kZeroMipBias = 0.0f;
+					Printer << ListSeparator << kZeroMipBias;
+				}
+				break;
+
+			case TMVM_MipBias:
+				if (bHasDerivativeTexCoords)
+				{
+					Printer
+						<< ListSeparator << LowerValue(VTPageTableRead->TexCoordDdx)
+						<< ListSeparator << LowerValue(VTPageTableRead->TexCoordDdy);
+				}
+				else
+				{
+					check(VTPageTableRead->MipValue);
+					Printer << ListSeparator << LowerValue(VTPageTableRead->MipValue);
+				}
+				break;
+
+			case TMVM_MipLevel:
+				check(VTPageTableRead->MipValue);
+				Printer << ListSeparator << LowerValue(VTPageTableRead->MipValue);
+				break;
+
+			case TMVM_Derivative:
+				check(bHasDerivativeTexCoords);
+				Printer
+					<< ListSeparator << LowerValue(VTPageTableRead->TexCoordDdx)
+					<< ListSeparator << LowerValue(VTPageTableRead->TexCoordDdy);
+				break;
 		}
 
 		Printer << ListSeparator << TEXTVIEW("Parameters.SvPosition.xy");
 
-		// VT feedback
+		// Lower final arguments for VT feedback
 		const bool bGenerateFeedback = VTPageTableRead->VTFeedbackRequestIndex != INDEX_NONE;
 		if (bGenerateFeedback)
 		{
 			Printer
 				<< ListSeparator << VTPageTableRead->VTFeedbackRequestIndex << TEXTVIEW("u + LIGHTMAP_VT_ENABLED")
-				<< ListSeparator << TEXTVIEW("Parameters.VirtualTextureFeedback");
+				<< ListSeparator << TEXTVIEW("Parameters.VirtualTextureFeedback")
+				;
 		}
 
 		Printer << EndArgs;
@@ -1812,6 +1869,8 @@ struct FPrivate : FMaterialIRToHLSLTranslation
 
 	void LowerSamplerType(EMaterialSamplerType SamplerType, bool& bOutSamplerNeedsBrackets)
 	{
+		bOutSamplerNeedsBrackets = true;
+
 		switch (SamplerType)
 		{
 			case SAMPLERTYPE_External:
@@ -1862,8 +1921,6 @@ struct FPrivate : FMaterialIRToHLSLTranslation
 			default:
 				UE_MIR_UNREACHABLE();
 		}
-
-		bOutSamplerNeedsBrackets = true;
 	}
 
 	ENoOp LowerTextureSamplerReference(MIR::FValue* TextureValue, ESamplerSourceMode SamplerSource)
diff --git a/Engine/Source/Runtime/Engine/Private/Materials/MaterialIRTypes.cpp b/Engine/Source/Runtime/Engine/Private/Materials/MaterialIRTypes.cpp
--- a/Engine/Source/Runtime/Engine/Private/Materials/MaterialIRTypes.cpp
+++ b/Engine/Source/Runtime/Engine/Private/Materials/MaterialIRTypes.cpp
@@ -18,6 +18,7 @@ const TCHAR* TypeKindToString(ETypeKind Kind)
 		case ETypeKind::Aggregate: return TEXT("Aggregate");
 		case ETypeKind::ShadingModel: return TEXT("ShadingModel");
 		case ETypeKind::Texture: return TEXT("Texture");
+		case ETypeKind::RuntimeVirtualTexture: return TEXT("RuntimeVirtualTexture");
 		case ETypeKind::ParameterCollection: return TEXT("ParameterCollection");
 		case ETypeKind::SubstrateData: return TEXT("SubstrateData");
 		case ETypeKind::VTPageTableResult: return TEXT("VTPageTableResult");
@@ -130,7 +131,7 @@ FType FType::FromMaterialParameterType(EMaterialParameterType Type)
 		case EMaterialParameterType::TextureCollection: UE_MIR_TODO();
 		case EMaterialParameterType::ParameterCollection: UE_MIR_TODO();
 		case EMaterialParameterType::Font: UE_MIR_TODO();
-		case EMaterialParameterType::RuntimeVirtualTexture: UE_MIR_TODO();
+		case EMaterialParameterType::RuntimeVirtualTexture: return FType::MakeRuntimeVirtualTexture();
 		case EMaterialParameterType::SparseVolumeTexture: UE_MIR_TODO();
 		case EMaterialParameterType::StaticSwitch: return FType::MakeBoolScalar();
 		default: UE_MIR_UNREACHABLE();
@@ -194,6 +195,13 @@ FType FType::MakeTexture()
 	return Type;
 }
 
+FType FType::MakeRuntimeVirtualTexture()
+{
+	FType Type{};
+	Type.Kind = ETypeKind::RuntimeVirtualTexture;
+	return Type;
+}
+
 FType FType::MakeSubstrateData()
 {
 	FType Type{};
diff --git a/Engine/Source/Runtime/Engine/Private/Materials/MaterialIRValueAnalyzer.cpp b/Engine/Source/Runtime/Engine/Private/Materials/MaterialIRValueAnalyzer.cpp
--- a/Engine/Source/Runtime/Engine/Private/Materials/MaterialIRValueAnalyzer.cpp
+++ b/Engine/Source/Runtime/Engine/Private/Materials/MaterialIRValueAnalyzer.cpp
@@ -21,6 +21,7 @@
 #include "Engine/Font.h"
 #include "Engine/Texture.h"
 #include "Engine/Texture2D.h"
+#include "VT/RuntimeVirtualTexture.h"
 #include "UObject/Package.h"
 #include "PostProcess/PostProcessMaterialInputs.h"
 
@@ -142,18 +143,18 @@ static void AnalyzeInlineHLSL(FMaterialIRValueAnalyzer& Analyzer, MIR::FInlineHL
 	}
 }
 
-static int32 AcquireVTStackIndex(FMaterialIRValueAnalyzer& Analyzer, MIR::FVTPageTableRead& VTPageTableRead, int32 TextureUniformIndex, bool bGenerateFeedback)
+static int32 AcquireVTStackIndex(FMaterialIRValueAnalyzer& Analyzer, MIR::FVTPageTableRead& VTPageTableRead, bool bGenerateFeedback)
 {
 	// Try to find matching VT stack entry
-	const uint64 CoordinateHash = CityHash64(reinterpret_cast<const char*>(&VTPageTableRead.TexCoord), sizeof(void*));
-
 	for (int32 VTStackIndex = 0; VTStackIndex < Analyzer.VTStacks.Num(); ++VTStackIndex)
 	{
 		const FMaterialIRValueAnalyzer::FVTStackEntry& VTStack = Analyzer.VTStacks[VTStackIndex];
 		if (VTStack.TexCoord == VTPageTableRead.TexCoord &&
 			VTStack.bGenerateFeedback == bGenerateFeedback &&
 			VTStack.AddressU == VTPageTableRead.AddressU &&
-			VTStack.AddressV == VTPageTableRead.AddressV)
+			VTStack.AddressV == VTPageTableRead.AddressV &&
+			VTStack.MipValue == VTPageTableRead.MipValue &&
+			VTStack.MipValueMode == VTPageTableRead.MipValueMode)
 		{
 			return VTStackIndex;
 		}
@@ -165,25 +166,27 @@ static int32 AcquireVTStackIndex(FMaterialIRValueAnalyzer& Analyzer, MIR::FVTPag
 	VTStack.bGenerateFeedback = bGenerateFeedback;
 	VTStack.AddressU = VTPageTableRead.AddressU;
 	VTStack.AddressV = VTPageTableRead.AddressV;
+	VTStack.MipValue = VTPageTableRead.MipValue;
+	VTStack.MipValueMode = VTPageTableRead.MipValueMode;
 
 	return Analyzer.CompilationOutput->UniformExpressionSet.AddVTStack(INDEX_NONE);
 }
 
-static int32 RegisterTextureParameter(FMaterialIRValueAnalyzer& Analyzer, UTexture* Texture, const FHashedMaterialParameterInfo& ParameterInfo, uint8 VTLayerIndex = 0xFF)
+static int32 RegisterTextureParameter(FMaterialIRValueAnalyzer& Analyzer, UObject* Texture, const FHashedMaterialParameterInfo& ParameterInfo, int32 VTLayerIndex = INDEX_NONE)
 {
 	check(Texture);
+	check(VTLayerIndex == INDEX_NONE || (VTLayerIndex >= 0 && VTLayerIndex < UINT8_MAX));
 
-	FMaterialTextureParameterInfo TextureParameterInfo
+	const FMaterialTextureParameterInfo TextureParameterInfo
 	{
 		.ParameterInfo = ParameterInfo,
 		.TextureIndex = Analyzer.Material->GetReferencedTextures().Find(Texture),
 		.SamplerSource = SSM_FromTextureAsset,
-		.VirtualTextureLayerIndex = VTLayerIndex,
+		.VirtualTextureLayerIndex = static_cast<uint8>(VTLayerIndex == INDEX_NONE ? UINT8_MAX : VTLayerIndex),
 	};
-
 	check(TextureParameterInfo.TextureIndex != INDEX_NONE);
 
-	EMaterialTextureParameterType ParamType = MIR::Internal::TextureMaterialValueTypeToParameterType(Texture->GetMaterialType());
+	const EMaterialTextureParameterType ParamType = MIR::Internal::TextureMaterialValueTypeToParameterType(MIR::Internal::GetTextureMaterialValueType(Texture));
 
 	return Analyzer.CompilationOutput->UniformExpressionSet.FindOrAddTextureParameter(ParamType, TextureParameterInfo);
 }
@@ -197,13 +200,22 @@ static void AccessTextureObject(FMaterialIRValueAnalyzer& Analyzer, MIR::FTextur
 	}
 }
 
+static void AccessRuntimeVirtualTextureObject(FMaterialIRValueAnalyzer& Analyzer, MIR::FRuntimeVirtualTextureObject* RVTextureObject)
+{
+	if (RVTextureObject->Analysis_UniformParameterIndex == INDEX_NONE)
+	{
+		const FHashedMaterialParameterInfo ParameterInfo{ RVTextureObject->RVTexture->GetFName(), EMaterialParameterAssociation::GlobalParameter, INDEX_NONE };
+		RVTextureObject->Analysis_UniformParameterIndex = RegisterTextureParameter(Analyzer, RVTextureObject->RVTexture, ParameterInfo, RVTextureObject->VTLayerIndex);
+	}
+}
+
 static void AccessTextureUniformParameter(FMaterialIRValueAnalyzer& Analyzer, MIR::FUniformParameter* UniformParameter)
 {
 	if (UniformParameter->Analysis_UniformParameterIndex == INDEX_NONE)
 	{
 		const FMaterialParameterValue& ParameterValue = Analyzer.Module->GetParameterMetadata(UniformParameter->ParameterIdInModule).Value;
-		UTexture* Texture = Cast<UTexture>(ParameterValue.AsTextureObject());
-		UniformParameter->Analysis_UniformParameterIndex = RegisterTextureParameter(Analyzer, Texture, Analyzer.Module->GetParameterInfo(UniformParameter->ParameterIdInModule));
+		const FHashedMaterialParameterInfo& ParameterInfo = Analyzer.Module->GetParameterInfo(UniformParameter->ParameterIdInModule);
+		UniformParameter->Analysis_UniformParameterIndex = RegisterTextureParameter(Analyzer, ParameterValue.AsTextureObject(), ParameterInfo, UniformParameter->VTLayerIndex);
 	}
 }
 
@@ -213,6 +225,10 @@ static void AccessTexture(FMaterialIRValueAnalyzer& Analyzer, MIR::FValue* Textu
 	{
 		AccessTextureObject(Analyzer, TextureObject);
 	}
+	else if (MIR::FRuntimeVirtualTextureObject* RVTextureObject = TextureValue->As<MIR::FRuntimeVirtualTextureObject>())
+	{
+		AccessRuntimeVirtualTextureObject(Analyzer, RVTextureObject);
+	}
 	else if (MIR::FUniformParameter* UniformParameter = TextureValue->As<MIR::FUniformParameter>())
 	{
 		AccessTextureUniformParameter(Analyzer, UniformParameter);
@@ -223,6 +239,32 @@ static void AccessTexture(FMaterialIRValueAnalyzer& Analyzer, MIR::FValue* Textu
 	}
 }
 
+static int32 GetVTLayerIndex(MIR::FValue* TextureObject)
+{
+	if (MIR::FRuntimeVirtualTextureObject* RVTextureObject = TextureObject->As<MIR::FRuntimeVirtualTextureObject>())
+	{
+		return RVTextureObject->VTLayerIndex;
+	}
+	if (MIR::FUniformParameter* UniformParameter = TextureObject->As<MIR::FUniformParameter>())
+	{
+		return UniformParameter->VTLayerIndex;
+	}
+	return INDEX_NONE;
+}
+
+static int32 GetPageTableLayerIndex(MIR::FValue* TextureObject)
+{
+	if (MIR::FRuntimeVirtualTextureObject* RVTextureObject = TextureObject->As<MIR::FRuntimeVirtualTextureObject>())
+	{
+		return RVTextureObject->VTPageTableIndex;
+	}
+	if (MIR::FUniformParameter* UniformParameter = TextureObject->As<MIR::FUniformParameter>())
+	{
+		return UniformParameter->VTPageTableIndex;
+	}
+	return INDEX_NONE;
+}
+
 static void AnalyzeInStageVTPageTableRead(FMaterialIRValueAnalyzer& Analyzer, MIR::FVTPageTableRead* VTPageTableRead, MIR::EStage Stage)
 {
 	check(VTPageTableRead->TextureObject);
@@ -234,18 +276,28 @@ static void AnalyzeInStageVTPageTableRead(FMaterialIRValueAnalyzer& Analyzer, MI
 	// Only support GPU feedback from pixel shader
 	const bool bGenerateFeedback = VTPageTableRead->bEnableFeedback && Stage == MIR::Stage_Pixel;
 
-	const int32 VTStackIndex = AcquireVTStackIndex(Analyzer, *VTPageTableRead, TextureUniformIndex, bGenerateFeedback);
+	const int32 VTStackIndex = AcquireVTStackIndex(Analyzer, *VTPageTableRead, bGenerateFeedback);
 
 	VTPageTableRead->VTStackIndex = VTStackIndex;
 	VTPageTableRead->VTFeedbackRequestIndex = INDEX_NONE;
 
-	// Acquire VT layer
-	int32 VTLayerIndex = Analyzer.CompilationOutput->UniformExpressionSet.GetVTStack(VTStackIndex).FindLayer(TextureUniformIndex);
-	if (VTLayerIndex == INDEX_NONE)
+	// Check if VT layer is already known. Otherwise, acquire VT layer
+	int32 VTLayerIndex = GetVTLayerIndex(VTPageTableRead->TextureObject);
+	if (VTLayerIndex != INDEX_NONE)
+	{
+		// The layer index in the VT stack is already known, so fetch the page table from the texture object and assign it to the VT stack
+		VTPageTableRead->VTPageTableIndex = GetPageTableLayerIndex(VTPageTableRead->TextureObject);
+		Analyzer.CompilationOutput->UniformExpressionSet.SetVTLayer(VTStackIndex, VTLayerIndex, TextureUniformIndex);
+	}
+	else
 	{
-		VTLayerIndex = Analyzer.CompilationOutput->UniformExpressionSet.AddVTLayer(VTStackIndex, TextureUniformIndex);
+		VTLayerIndex = Analyzer.CompilationOutput->UniformExpressionSet.GetVTStack(VTStackIndex).FindLayer(TextureUniformIndex);
+		if (VTLayerIndex == INDEX_NONE)
+		{
+			VTLayerIndex = Analyzer.CompilationOutput->UniformExpressionSet.AddVTLayer(VTStackIndex, TextureUniformIndex);
+		}
+		VTPageTableRead->VTPageTableIndex = VTLayerIndex;
 	}
-	VTPageTableRead->VTLayerIndex = VTLayerIndex;
 
 	// Acquire VT feedback request index
 	if (bGenerateFeedback)
@@ -565,8 +617,11 @@ static void AnalyzeTextureRead(FMaterialIRValueAnalyzer& Analyzer, MIR::FTexture
 }
 
 // Returns the next available offset into the preshader buffer for a float vector with the specified number of components (1-4).
-static uint32 NextGlobalComponentOffset(FMaterialIRValueAnalyzer& Analyzer, uint32 NumComponents, bool bAreComponentsDouble = false)
+static uint32 NextGlobalComponentOffset(FMaterialIRValueAnalyzer& Analyzer, const MIR::FPrimitive& PrimitiveType)
 {
+	const uint32 NumComponents = PrimitiveType.NumComponents();
+	const bool bAreComponentsDouble = PrimitiveType.IsDouble();
+
 	// Following code calculates GlobalComponentOffset.
 	// GlobalComponentOffset is the i-th component in the array of float4s that make the uniform buffer.
 	// For example a GlobalComponentOffset of 13 references PreshaderBuffer[3].y.
@@ -659,8 +714,7 @@ static void AnalyzePrimitiveUniformParameter(FMaterialIRValueAnalyzer& Analyzer,
 	// Only int, float and LWC parameters supported for now
 	check(PrimitiveType.IsInteger() || PrimitiveType.IsAnyFloat());
 
-	const uint32 NumComponents = PrimitiveType.NumComponents();
-	const uint32 GlobalComponentOffset = NextGlobalComponentOffset(Analyzer, NumComponents, PrimitiveType.IsDouble());
+	const uint32 GlobalComponentOffset = NextGlobalComponentOffset(Analyzer, PrimitiveType);
 
 	// Add the parameter evaluation to the uniform data
 	UniformExpressionSet.AddNumericParameterEvaluation(Parameter->Analysis_UniformParameterIndex, GlobalComponentOffset);
@@ -671,7 +725,7 @@ static void AnalyzePrimitiveUniformParameter(FMaterialIRValueAnalyzer& Analyzer,
 		FMaterialInsights::FUniformParameterAllocationInsight ParamInsight;
 		ParamInsight.BufferSlotIndex = GlobalComponentOffset / 4;
 		ParamInsight.BufferSlotOffset = GlobalComponentOffset % 4;
-		ParamInsight.ComponentsCount = NumComponents;
+		ParamInsight.ComponentsCount = PrimitiveType.NumComponents();
 		ParamInsight.ParameterName = ParameterInfo.Name;
 
 		switch (PrimitiveType.ScalarKind)
@@ -720,11 +774,10 @@ static void AnalyzePreshaderParameter(FMaterialIRValueAnalyzer& Analyzer, MIR::F
 	// Only int, float and LWC parameters supported for now
 	check(PrimitiveType.IsInteger() || PrimitiveType.IsAnyFloat());
 
-	const uint32 NumComponents = PrimitiveType.NumComponents();
-	const uint32 GlobalComponentOffset = NextGlobalComponentOffset(Analyzer, NumComponents);
+	const uint32 GlobalComponentOffset = NextGlobalComponentOffset(Analyzer, PrimitiveType);
 
 	UniformExpressionSet.WriteUniformPreshaderEntry(
-		GlobalComponentOffset, UE::Shader::MakeValueType(UE::Shader::EValueComponentType::Float, NumComponents),
+		GlobalComponentOffset, UE::Shader::MakeValueType(UE::Shader::EValueComponentType::Float, PrimitiveType.NumComponents()),
 		[Parameter, SourceParameterName](UE::Shader::FPreshaderData& UniformPreshaderData) -> void
 		{
 			const FHashedMaterialParameterInfo HashedSourceParameterInfo{ SourceParameterName, EMaterialParameterAssociation::GlobalParameter, INDEX_NONE };
@@ -733,7 +786,11 @@ static void AnalyzePreshaderParameter(FMaterialIRValueAnalyzer& Analyzer, MIR::F
 			{
 			case UE::Shader::EPreshaderOpcode::TextureSize:
 			case UE::Shader::EPreshaderOpcode::TexelSize:
-				UniformPreshaderData.WriteOpcode(Parameter->Opcode).Write(HashedSourceParameterInfo).Write((int32)Parameter->TextureIndex);
+				UniformPreshaderData.WriteOpcode(Parameter->Opcode).Write(HashedSourceParameterInfo).Write(Parameter->TextureIndex);
+				break;
+
+			case UE::Shader::EPreshaderOpcode::RuntimeVirtualTextureUniform:
+				UniformPreshaderData.WriteOpcode(Parameter->Opcode).Write(HashedSourceParameterInfo).Write(Parameter->TextureIndex).Write(Parameter->Payload.UniformIndex);
 				break;
 
 			default:
@@ -789,11 +846,11 @@ void FMaterialIRValueAnalyzer::Setup(UMaterial* InMaterial, FMaterialIRModule* I
 void FMaterialIRValueAnalyzer::Analyze(MIR::FValue* Value)
 {
 	#define EXPAND_CASE(ValueType, ...) \
-	case MIR::VK_##ValueType: \
-	{ \
-		Analyze##ValueType(*this, static_cast<MIR::F##ValueType*>(Value) , ## __VA_ARGS__); \
-		break; \
-	}
+		case MIR::VK_##ValueType: \
+		{ \
+			Analyze##ValueType(*this, static_cast<MIR::F##ValueType*>(Value) , ## __VA_ARGS__); \
+			break; \
+		}
 
 	// IMPORTANT: Before adding a case here, read the FMaterialIRValueAnalyzer documentation.
 	switch (Value->Kind)
diff --git a/Engine/Source/Runtime/Engine/Private/Materials/MaterialIRValueAnalyzer.h b/Engine/Source/Runtime/Engine/Private/Materials/MaterialIRValueAnalyzer.h
--- a/Engine/Source/Runtime/Engine/Private/Materials/MaterialIRValueAnalyzer.h
+++ b/Engine/Source/Runtime/Engine/Private/Materials/MaterialIRValueAnalyzer.h
@@ -50,6 +50,8 @@ struct FMaterialIRValueAnalyzer
 		bool bGenerateFeedback{};
 		TextureAddress AddressU{ TA_MAX };
 		TextureAddress AddressV{ TA_MAX };
+		MIR::FValue* MipValue{};
+		ETextureMipValueMode MipValueMode{ TMVM_None };
 	};
 
 	// Resets the analyzer to process a new translation run.
diff --git a/Engine/Source/Runtime/Engine/Private/Materials/MaterialShared.cpp b/Engine/Source/Runtime/Engine/Private/Materials/MaterialShared.cpp
--- a/Engine/Source/Runtime/Engine/Private/Materials/MaterialShared.cpp
+++ b/Engine/Source/Runtime/Engine/Private/Materials/MaterialShared.cpp
@@ -5,6 +5,7 @@
 =============================================================================*/
 
 #include "MaterialShared.h"
+#include "MaterialSharedPrivate.h"
 #include "Misc/DelayedAutoRegister.h"
 #include "Stats/StatsMisc.h"
 #include "Stats/StatsTrace.h"
@@ -5950,6 +5951,28 @@ FString SceneTextureIdToHLSLString(ESceneTextureId TexId)
 	}
 }
 
+EWorldPositionIncludedOffsets GetWorldPositionTypeWithOrigin(EPositionOrigin PositionOrigin, bool IncludeOffsets)
+{
+	switch (PositionOrigin)
+	{
+	case EPositionOrigin::Absolute: return IncludeOffsets ? WPT_Default : WPT_ExcludeAllShaderOffsets;
+	case EPositionOrigin::CameraRelative: return IncludeOffsets ? WPT_CameraRelative : WPT_CameraRelativeNoOffsets;
+	default: checkNoEntry();
+	}
+	return WPT_Default;
+}
+
+FName GetWorldPositionInputName(EPositionOrigin PositionOrigin)
+{
+	switch (PositionOrigin)
+	{
+	case EPositionOrigin::CameraRelative: return TEXT("Translated World Position");
+	case EPositionOrigin::Absolute: return TEXT("World Position");
+	default: checkNoEntry();
+	}
+	return FName();
+}
+
 // FMaterialResource::IsTranslucencyAfterDOFEnabled
 static bool IsTranslucencyAfterDOFEnabled(UMaterial* Material)
 {
diff --git a/Engine/Source/Runtime/Engine/Private/Materials/MaterialSharedPrivate.h b/Engine/Source/Runtime/Engine/Private/Materials/MaterialSharedPrivate.h
--- a/Engine/Source/Runtime/Engine/Private/Materials/MaterialSharedPrivate.h
+++ b/Engine/Source/Runtime/Engine/Private/Materials/MaterialSharedPrivate.h
@@ -11,6 +11,8 @@
 #include "MaterialShared.h"
 #include "Materials/MaterialExpressionTransform.h"
 #include "Materials/MaterialExpressionTransformPosition.h"
+#include "Materials/MaterialExpressionWorldPosition.h"
+#include "Materials/MaterialExpressionRuntimeVirtualTextureSample.h"
 #include "RHIFeatureLevel.h"
 #include "RHIShaderPlatform.h"
 #include "SceneTypes.h"
@@ -91,6 +93,12 @@ namespace UE::MaterialTranslatorUtils
 	// Returns the value type of a texture property with the specified texture value type, e.g. MCT_Float3 for Volume textures.
 	EMaterialValueType GetTexturePropertyValueType(EMaterialValueType TextureType);
 
+	// Maps the specified position origin to its corresponding EWorldPositionIncludedOffsets value with an include offsets option.
+	EWorldPositionIncludedOffsets GetWorldPositionTypeWithOrigin(EPositionOrigin PositionOrigin, bool IncludeOffsets = true);
+
+	// Returns a human readable name (en-US) for the specified position origin enumeration value.
+	FName GetWorldPositionInputName(EPositionOrigin PositionOrigin);
+
 	// Final compile validation logic, run at the end of translation to HLSL.
 	void FinalCompileValidation(
 		UMaterial* Material,
diff --git a/Engine/Source/Runtime/Engine/Public/MaterialCompiler.h b/Engine/Source/Runtime/Engine/Public/MaterialCompiler.h
--- a/Engine/Source/Runtime/Engine/Public/MaterialCompiler.h
+++ b/Engine/Source/Runtime/Engine/Public/MaterialCompiler.h
@@ -16,6 +16,7 @@
 #include "Materials/MaterialExpressionWorldPosition.h"
 #include "Materials/MaterialExpressionTextureProperty.h"
 #include "Materials/MaterialExpressionViewProperty.h"
+#include "Materials/MaterialExpressionRuntimeVirtualTextureSample.h"
 
 
 namespace UE
@@ -44,20 +45,6 @@ enum EMaterialForceCastFlags
 	MFCF_ReplicateValue	= 1<<3	// Replicates a Float1 value when up-casting, else appends zero
 };
 
-enum class EVirtualTextureUnpackType
-{
-	None,
-	BaseColorYCoCg,
-	NormalBC3,
-	NormalBC5,
-	NormalBC3BC3,
-	NormalBC5BC1,
-	HeightR16,
-	NormalBGR565,
-	BaseColorSRGB,
-	DisplacementR16,
-};
-
 /** What type of compiler is this? Used by material expressions that select input based on compile context */
 enum class EMaterialCompilerType
 {
diff --git a/Engine/Source/Runtime/Engine/Public/Materials/MaterialExpressionRuntimeVirtualTextureSample.h b/Engine/Source/Runtime/Engine/Public/Materials/MaterialExpressionRuntimeVirtualTextureSample.h
--- a/Engine/Source/Runtime/Engine/Public/Materials/MaterialExpressionRuntimeVirtualTextureSample.h
+++ b/Engine/Source/Runtime/Engine/Public/Materials/MaterialExpressionRuntimeVirtualTextureSample.h
@@ -3,6 +3,7 @@
 #pragma once
 
 #include "CoreMinimal.h"
+#include "Engine/EngineTypes.h"
 #include "Materials/MaterialExpression.h"
 #include "UObject/ObjectMacros.h"
 #include "VT/RuntimeVirtualTexture.h"
@@ -67,6 +68,37 @@ enum ERuntimeVirtualTextureTextureAddressMode : int
 	RVTTA_MAX,
 };
 
+enum class EVirtualTextureUnpackType
+{
+	None,
+	BaseColorYCoCg,
+	NormalBC3,
+	NormalBC5,
+	NormalBC3BC3,
+	NormalBC5BC1,
+	HeightR16,
+	NormalBGR565,
+	BaseColorSRGB,
+	DisplacementR16,
+};
+
+struct FRuntimeVirtualTextureUnpackProperties
+{
+	bool bIsBaseColorValid : 1 = false;
+	bool bIsSpecularValid : 1 = false;
+	bool bIsRoughnessValid : 1 = false;
+	bool bIsNormalValid : 1 = false;
+	bool bIsWorldHeightValid : 1 = false;
+	bool bIsMaskValid : 1 = false;
+	bool bIsMask4Valid : 1 = false;
+	bool bIsDisplacementValid : 1 = false;
+
+	uint32 UnpackTarget = 0;
+	uint32 UnpackMask = 0;
+	EVirtualTextureUnpackType UnpackType = EVirtualTextureUnpackType::None;
+	TArray<float, TInlineAllocator<4>> ConstantVector;
+};
+
 /** Material expression for sampling from a runtime virtual texture. */
 UCLASS(collapsecategories, hidecategories=Object, MinimalAPI)
 class UMaterialExpressionRuntimeVirtualTextureSample : public UMaterialExpression
@@ -132,6 +164,9 @@ class UMaterialExpressionRuntimeVirtualTextureSample : public UMaterialExpressio
 	/** Init settings that affect shader compilation and need to match the current VirtualTexture */
 	ENGINE_API bool InitVirtualTextureDependentSettings();
 
+	/** Returns the sampler source mode for the respective TextureAddressMode value. */
+	ESamplerSourceMode GetSamplerSourceMode() const;
+
 protected:
 	/** Initialize the output pins. */
 	ENGINE_API void InitOutputs();
@@ -144,9 +179,19 @@ protected:
 	ENGINE_API virtual void PostLoad() override;
 	FExpressionInput* GetInput(int32 InputIndex) override;
 	ENGINE_API virtual FName GetInputName(int32 InputIndex) const override;
+	ENGINE_API virtual void Build(MIR::FEmitter& Emitter) override;
 	ENGINE_API virtual int32 Compile(class FMaterialCompiler* Compiler, int32 OutputIndex) override;
 	ENGINE_API virtual void GetCaption(TArray<FString>& OutCaptions) const override;
 	ENGINE_API virtual bool CanEditChange(const FProperty* InProperty) const override;
+
+	/** Returns true if this material expression has a valid virtual texture object that matches the parameters of this expression. */
+	bool ValidateVirtualTextureParameters(FString& OutError) const;
+
+	/** Returns true if this is a valid UMaterialExpressionRuntimeVirtualTextureSampleParameter. */
+	bool IsParameter() const;
+
+	bool GetRVTUnpackProperties(int32 OutputIndex, bool bIsVirtualTextureValid, FRuntimeVirtualTextureUnpackProperties& Output) const;
+
 public:
 	ENGINE_API virtual void PostEditChangeProperty(FPropertyChangedEvent& PropertyChangedEvent) override;
 
diff --git a/Engine/Source/Runtime/Engine/Public/Materials/MaterialIR.h b/Engine/Source/Runtime/Engine/Public/Materials/MaterialIR.h
--- a/Engine/Source/Runtime/Engine/Public/Materials/MaterialIR.h
+++ b/Engine/Source/Runtime/Engine/Public/Materials/MaterialIR.h
@@ -13,6 +13,8 @@
 
 struct FMaterialExternalCodeDeclaration;
 
+class URuntimeVirtualTexture;
+
 namespace MIR
 {
 
@@ -103,6 +105,7 @@ enum EValueKind
 	VK_ScreenTexture,
 	VK_ShadingModel,
 	VK_TextureObject,
+	VK_RuntimeVirtualTextureObject,
 	VK_UniformParameter,
 
 	// Instructions
@@ -241,10 +244,10 @@ struct FValue
 	// Returns whether this value equals the given constant vector.
 	bool EqualsConstant(FVector4f Value) const;
 
-	// Returns the UTexture object if this value is of type FTextureObject. Returns null otherwise.
-	UTexture* GetTexture() const;
+	// Returns the UTexture or URuntimeVirtualTexture object if this value is of type FTextureObject or FRuntimeVirtualTextureObject. Returns null otherwise.
+	UObject* GetTextureObject() const;
 
-	// Returns the uniform parameter index if this value is of type FTextureObject or FUniformParameter. Retunrs INDEX_NONE otherwise.
+	// Returns the uniform parameter index if this value is of type FTextureObject, FRuntimeVirtualTextureObject, or FUniformParameter. Returns INDEX_NONE otherwise.
 	int32 GetUniformParameterIndex() const;
 
 	// Tries to cast this value to specified type T and returns the casted pointer, if possible (nullptr otherwise).
@@ -514,15 +517,41 @@ struct FTextureObject : TValue<VK_TextureObject>
 	int32 Analysis_UniformParameterIndex;
 };
 
+// It represents a runtime virtual texture (RVT) object value. Those *do not* extend UTexture so they have their own MIR value.
+struct FRuntimeVirtualTextureObject : TValue<VK_RuntimeVirtualTextureObject>
+{
+	// The runtime virtual texture object.
+	URuntimeVirtualTexture* RVTexture;
+
+	// The sampler type associated to this RVT object.
+	EMaterialSamplerType SamplerType;
+
+	// Field to explicitly assign a VT stack layer index (INDEX_NONE if unset). Regular VTs are assigned these indices automatically.
+	int32 VTLayerIndex : 16;
+
+	// Field to explicitly assign a VT page table index. Regular VTs are assigned these indices automatically.
+	int32 VTPageTableIndex : 16;
+
+	// Index of this parameter in the uniform expression set.
+	// Note: This field is automatically set by the builder.
+	int32 Analysis_UniformParameterIndex;
+};
+
 // It represents a material uniform parameter.
 struct FUniformParameter: TValue<VK_UniformParameter>
 {
 	// Index of the parameter as registered in the module.
 	uint32 ParameterIdInModule;
 
-	// Eventual sampler type to use to sample this parameter's texture (if it is one)
+	// Eventual sampler type to use to sample this parameter's texture (if it is one).
 	EMaterialSamplerType SamplerType;
 
+	// Field to explicitly assign a VT stack layer index (INDEX_NONE if unset). Regular VTs are assigned these indices automatically.
+	int32 VTLayerIndex : 16;
+
+	// Field to explicitly assign a VT page table index. Regular VTs are assigned these indices automatically.
+	int32 VTPageTableIndex : 16;
+
 	// Index of this parameter in the uniform expression set
 	// Note: This field is automatically set by the builder.
 	int32 Analysis_UniformParameterIndex;
@@ -961,7 +990,8 @@ struct FTextureRead : TInstruction<VK_TextureRead, 6>
 	bool bUseAnisoSampler : 1;
 };
 
-struct FVTPageTableRead : TInstruction<VK_VTPageTableRead, 4>
+// This instruction reads a virtual texture (VT) page table and is allocated alongside VT sampling instructions.
+struct FVTPageTableRead : TInstruction<VK_VTPageTableRead, 5>
 {
 	// The texture object to load the page table for.
 	FValue* TextureObject;
@@ -975,11 +1005,14 @@ struct FVTPageTableRead : TInstruction<VK_VTPageTableRead, 4>
 	// Optional. The analytical partial derivative of the coordinates along the Y axis.
 	FValue* TexCoordDdy;
 
+	// Used with TMVM_MipBias and TMVM_MipLevel.
+	FValue* MipValue;
+
 	// Index of the virtual texture slot. This is determined in the analysis stage.
 	int32 VTStackIndex : 16;
 
-	// Index of the VT layer within the stack. This is determined in the analysis stage.
-	int32 VTLayerIndex : 16;
+	// Index of the VT page table. This is determined in the analysis stage.
+	int32 VTPageTableIndex : 16;
 
 	// Index of the VT feedback request. This is determined in the analysis stage.
 	int32 VTFeedbackRequestIndex : 16;
@@ -990,8 +1023,14 @@ struct FVTPageTableRead : TInstruction<VK_VTPageTableRead, 4>
 	// Texture address mode for V axis.
 	TextureAddress AddressV : 8;
 
+	// Specifies the sampler function to read the page table, i.e. TextureLoadVirtualPageTable(), TextureLoadVirtualPageTableGrad(), or TextureLoadVirtualPageTableLevel().
+	ETextureMipValueMode MipValueMode : 3;
+
 	// Enables VT texture sampling feedback
 	bool bEnableFeedback : 1;
+
+	// Determines whether adaptive or regular VT page table loads will be emitted, i.e. TextureLoadVirtualPageTable*() or TextureLoadVirtualPageTableAdaptive*().
+	bool bIsAdaptive : 1;
 };
 
 // Utility value for selecting a different value based on the execution stage.
@@ -1139,6 +1178,12 @@ struct FCallParameterOutput : TInstruction<VK_CallParameterOutput, 1>
 	int32 	Index;
 };
 
+struct FPreshaderParameterPayload
+{
+	// Uniform index for RVTs. Only used for UE::Shader::EPreshaderOpcode::RuntimeVirtualTextureUniform.
+	int32 UniformIndex;
+};
+
 struct FPreshaderParameter : TInstruction<VK_PreshaderParameter, 1>
 {
 	// Argument that points to the source parameter this dynamic parameter depends on.
@@ -1150,6 +1195,9 @@ struct FPreshaderParameter : TInstruction<VK_PreshaderParameter, 1>
 	// Index into UMaterialInterface::GetReferencedTextures() pointing to the source parameter's texture.
 	int32 TextureIndex;
 
+	// Payload data that depends on the opcode.
+	FPreshaderParameterPayload Payload;
+
 	// Offset into the preshader buffer of the material resource.
 	uint32 Analysis_PreshaderOffset;
 };
diff --git a/Engine/Source/Runtime/Engine/Public/Materials/MaterialIREmitter.h b/Engine/Source/Runtime/Engine/Public/Materials/MaterialIREmitter.h
--- a/Engine/Source/Runtime/Engine/Public/Materials/MaterialIREmitter.h
+++ b/Engine/Source/Runtime/Engine/Public/Materials/MaterialIREmitter.h
@@ -36,6 +36,7 @@ struct FSwizzleMask
 	FSwizzleMask(EVectorComponent X, EVectorComponent Y);
 	FSwizzleMask(EVectorComponent X, EVectorComponent Y, EVectorComponent Z);
 	FSwizzleMask(EVectorComponent X, EVectorComponent Y, EVectorComponent Z, EVectorComponent W);
+	FSwizzleMask(bool bMaskX, bool bMaskY, bool bMaskZ, bool bMaskW);
 
 	// Pushes a component to this swizzle mask.
 	void Append(EVectorComponent Component);
@@ -165,6 +166,23 @@ struct FFunctionHLSLDesc
 	uint32 GetNumParameters() const { return NumInputOnlyParams + NumInputOutputParams + NumOutputOnlyParams; }
 };
 
+// Helper structure to simplify passing attributes to texture sampling emitter functions that share the same attributes.
+struct FTextureSampleBaseAttributes
+{
+	// Specifies which sampler source mode to use. By default SSM_FromTextureAsset.
+	ESamplerSourceMode SamplerSourceMode { SSM_FromTextureAsset };
+	
+	// Specifies Which sampler type to use. By default SAMPLERTYPE_MAX which implies the sampler type is fetched from the texture object.
+	// See SamplerType field in FTextureObject, FRuntimeVirtualTextureObject, and FUniformParameter.
+	EMaterialSamplerType SamplerType { SAMPLERTYPE_MAX };
+
+	// Specifies whether to enable virtual texture sampling feedback. This is only used for pixel shaders. Enabled by default.
+	bool bEnableFeedback : 1 { true };
+
+	// True if adaptive VT sampling is used. Only used for virtual textures.
+	bool bIsAdaptive : 1 { false };
+};
+
 // The FEmitter is a helper object responsible for emitting MIR values and instructions.
 // It is primarily used within UMaterialExpression::Build() implementations to lower the
 // high-level semantics of a material expression into corresponding lower-level MIR nodes.
@@ -239,8 +257,8 @@ public:
 	// Returns the type of the i-th attribute in the given material aggregate.
 	FType GetMaterialAggregateAttributeType(const UMaterialAggregate* Aggregate, int32 AttributeIndex);
 
-	// Returns the UTexture of the specified texture object or uniform parameter.
-	UTexture* GetTextureFromValue(MIR::FValueRef Texture);
+	// Returns the UTexture or URuntimeVirtualTexture of the specified texture object or uniform parameter.
+	UObject* GetTextureFromValue(MIR::FValueRef Texture);
 
 	/*-------------------------------- Input management --------------------------------*/
 
@@ -403,11 +421,14 @@ public:
 	// Emits a texture object value of fiven sampler type and returns it.
 	FValueRef TextureObject(UTexture* Texture, EMaterialSamplerType SamplerType);
 
-	// Emits a material parameter value with given name, metadata and sampler type (if it is a texture) and returns it.
-	FValueRef Parameter(FName Name, FMaterialParameterMetadata& Metadata, EMaterialSamplerType SamplerType = SAMPLERTYPE_Color);
+	// Emits a runtime virtual texture object value.
+	FValueRef RuntimeVirtualTextureObject(URuntimeVirtualTexture* RVTexture, EMaterialSamplerType SamplerType, int32 VTLayerIndex, int32 VTPageTableIndex);
+
+	// Emits a material parameter value with given name, metadata, sampler type (if it's a texture) and optional parameters for RVT layer index and page table index.
+	FValueRef Parameter(FName Name, FMaterialParameterMetadata& Metadata, EMaterialSamplerType SamplerType = SAMPLERTYPE_Color, int32 VTLayerIndex = INDEX_NONE, int32 VTPageTableIndex = INDEX_NONE);
 
 	// Emits a preshader parameter value with given opcode.
-	FValueRef PreshaderParameter(UE::Shader::EPreshaderOpcode Opcode, FValueRef SourceParameter);
+	FValueRef PreshaderParameter(FType Type, UE::Shader::EPreshaderOpcode Opcode, FValueRef SourceParameter, FPreshaderParameterPayload Payload = {});
 
 	// Emits a custom primitive data value with given index.
 	FValueRef CustomPrimitiveData(uint32 PrimitiveDataIndex);
@@ -525,63 +546,63 @@ public:
 
 	// Emits a texture gather instruction.
 	// 
-	// @param Texture The texture to sample. Must be of type MIR::FObjectType and represent a valid texture.
+	// @param Texture The texture to sample. Must be of type MIR::FTextureObject, MIR::FRuntimeVirtualTextureObject, or MIR::FUniformParameter.
 	// @param TexCoord The UV coordinates at which to sample the texture.
 	// @param GatherMode What channel to gather (any ::GatherXXX value).
-	// @param SamplerSourceMode Which sampler source mode to use.
-	// @param SamplerType Which sampler type to use.
-	FValueRef TextureGather(FValueRef Texture, FValueRef TexCoord, ETextureReadMode GatherMode, ESamplerSourceMode SamplerSourceMode, EMaterialSamplerType SamplerType);
+	// @param BaseAttributes Base sampling attributes for sampler source and type as well as VT attributes.
+	FValueRef TextureGather(FValueRef Texture, FValueRef TexCoord, ETextureReadMode GatherMode, const FTextureSampleBaseAttributes& BaseAttributes);
 
 	// Emits a texture sample instruction.
 	// 
-	// @param Texture The texture to sample. Must be of type MIR::FObjectType and represent a valid texture.
+	// @param Texture The texture to sample. Must be of type MIR::FTextureObject, MIR::FRuntimeVirtualTextureObject, or MIR::FUniformParameter.
 	// @param TexCoord The UV coordinates at which to sample the texture.
 	// @param bAddViewMaterialMipBias Whether to add the view material texture mip bias to the mip level.
-	// @param SamplerSourceMode Which sampler source mode to use.
-	// @param SamplerType Which sampler type to use.
-	FValueRef TextureSample(FValueRef Texture, FValueRef TexCoord, bool bAddViewMaterialMipBias, ESamplerSourceMode SamplerSourceMode, EMaterialSamplerType SamplerType);
+	// @param BaseAttributes Base sampling attributes for sampler source and type as well as VT attributes.
+	FValueRef TextureSample(FValueRef Texture, FValueRef TexCoord, bool bAddViewMaterialMipBias, const FTextureSampleBaseAttributes& BaseAttributes);
 
 	// Emits a texture sample instruction with a manually given mip level.
 	// 
-	// @param Texture The texture to sample. Must be of type MIR::FObjectType and represent a valid texture.
+	// @param Texture The texture to sample. Must be of type MIR::FTextureObject, MIR::FRuntimeVirtualTextureObject, or MIR::FUniformParameter.
 	// @param TexCoord The UV coordinates at which to sample the texture.
 	// @param MipLevel Which mip level to use for this sample.
 	// @param bAddViewMaterialMipBias Whether to add the view material texture mip bias to the mip level.
-	// @param SamplerSourceMode Which sampler source mode to use.
-	// @param SamplerType Which sampler type to use.
-	FValueRef TextureSampleLevel(FValueRef Texture, FValueRef TexCoord, FValueRef MipLevel, bool bAddViewMaterialMipBias, ESamplerSourceMode SamplerSourceMode, EMaterialSamplerType SamplerType);
+	// @param BaseAttributes Base sampling attributes for sampler source and type as well as VT attributes.
+	FValueRef TextureSampleLevel(FValueRef Texture, FValueRef TexCoord, FValueRef MipLevel, bool bAddViewMaterialMipBias, const FTextureSampleBaseAttributes& BaseAttributes);
 
 	// Emits a texture sample instruction with bias added to the automatically selected mip level.
 	// 
-	// @param Texture The texture to sample. Must be of type MIR::FObjectType and represent a valid texture.
+	// @param Texture The texture to sample. Must be of type MIR::FTextureObject, MIR::FRuntimeVirtualTextureObject, or MIR::FUniformParameter.
 	// @param TexCoord The UV coordinates at which to sample the texture.
 	// @param MipBias The bias to add to the automatically selected mip level.
 	// @param bAddViewMaterialMipBias Whether to add the view material texture mip bias to the mip level.
-	// @param SamplerSourceMode Which sampler source mode to use.
-	// @param SamplerType Which sampler type to use.
-	FValueRef TextureSampleBias(FValueRef Texture, FValueRef TexCoord, FValueRef MipBias, bool bAddViewMaterialMipBias, ESamplerSourceMode SamplerSourceMode, EMaterialSamplerType SamplerType);
+	// @param BaseAttributes Base sampling attributes for sampler source and type as well as VT attributes.
+	FValueRef TextureSampleBias(FValueRef Texture, FValueRef TexCoord, FValueRef MipBias, bool bAddViewMaterialMipBias, const FTextureSampleBaseAttributes& BaseAttributes);
 	
 	// Emits a texture sample instruction using user provided partial derivatives.
 	// 
-	// @param Texture The texture to sample. Must be of type MIR::FObjectType and represent a valid texture.
+	// @param Texture The texture to sample. Must be of type MIR::FTextureObject, MIR::FRuntimeVirtualTextureObject, or MIR::FUniformParameter.
 	// @param TexCoord The UV coordinates at which to sample the texture.
 	// @param TexCoordDdx The partial derivative of the texture coordinates along the screen space x axis.
 	// @param TexCoordDdy The partial derivative of the texture coordinates along the screen space y axis.
 	// @param bAddViewMaterialMipBias Whether to add the view material texture mip bias to the mip level.
-	// @param SamplerSourceMode Which sampler source mode to use.
-	// @param SamplerType Which sampler type to use.
-	FValueRef TextureSampleGrad(FValueRef Texture, FValueRef TexCoord, FValueRef TexCoordDdx, FValueRef TexCoordDdy, bool bAddViewMaterialMipBias, ESamplerSourceMode SamplerSourceMode, EMaterialSamplerType SamplerType);
+	// @param BaseAttributes Base sampling attributes for sampler source and type as well as VT attributes.
+	FValueRef TextureSampleGrad(FValueRef Texture, FValueRef TexCoord, FValueRef TexCoordDdx, FValueRef TexCoordDdy, bool bAddViewMaterialMipBias, const FTextureSampleBaseAttributes& BaseAttributes);
 
 	// Emits the virtual texture page load instruction.
 	//
-	// @param Texture The texture to sample. Must be of type MIR::FObjectType and represent a valid virtual texture.
+	// @param Texture The texture to sample. Must be of type MIR::FTextureObject, MIR::FRuntimeVirtualTextureObject, or MIR::FUniformParameter.
 	// @param AddressU Texture address mode for the U-coordinate.
 	// @param AddressV Texture address mode for the V-coordinate.
 	// @param TexCoord The UV coordinates at which to sample the virtual texture.
 	// @param TexCoordDdx The partial derivative of the texture coordinates along the screen space x axis.
 	// @param TexCoordDdy The partial derivative of the texture coordinates along the screen space y axis.
 	// @param bEnableFeedback Specifies whether to enable virtual texture sampling feedback. This is only used for pixel shaders. Enabled by default.
-	FValueRef VTPageTableLoad(FValueRef Texture, TextureAddress AddressU, TextureAddress AddressV, FValueRef TexCoord, FValueRef TexCoordDdx = {}, FValueRef TexCoordDdy = {}, bool bEnableFeedback = true);
+	// @param bIsAdaptive Specifies whether the page table instruction corresponds to an adaptive RVT.
+	// @param MipValueMode Specifies the MIP value mode. This determines what page table load function in the shader will be emitted, i.e. TextureLoadVirtualPageTable(), TextureLoadVirtualPageTableLevel(), etc..
+	// @param MipValue Optional value for an explicit MIP level or bias.
+	FValueRef VTPageTableLoad(
+		FValueRef Texture, TextureAddress AddressU, TextureAddress AddressV, FValueRef TexCoord,
+		FValueRef TexCoordDdx = {}, FValueRef TexCoordDdy = {}, bool bEnableFeedback = true, bool bIsAdaptive = false, ETextureMipValueMode MipValueMode = TMVM_None, FValueRef MipValue = {});
 	
 	// Emits the partial derivative of given value.
 	// This function differentiates the given value with respect to the given axis.
diff --git a/Engine/Source/Runtime/Engine/Public/Materials/MaterialIRTypes.h b/Engine/Source/Runtime/Engine/Public/Materials/MaterialIRTypes.h
--- a/Engine/Source/Runtime/Engine/Public/Materials/MaterialIRTypes.h
+++ b/Engine/Source/Runtime/Engine/Public/Materials/MaterialIRTypes.h
@@ -23,6 +23,7 @@ enum class ETypeKind : uint8
 	
 	ShadingModel,			// Dynamically specified material shading model.
 	Texture,				// A texture object.
+	RuntimeVirtualTexture,	// A runtime virtual texture object.
 	ParameterCollection,	// Material parameter collection object.
 	SubstrateData,			// Substrate-specific data structure.
 	VTPageTableResult,		// A virtual texture page table lookup result.
@@ -188,6 +189,9 @@ struct FType
 
 	// Returns the Texture object type.
 	static FType MakeTexture();
+
+	// Returns the RuntimeVirtualTexture object type.
+	static FType MakeRuntimeVirtualTexture();
 	
 	// Returns the SubstrateData object type.
 	static FType MakeSubstrateData();
@@ -246,9 +250,12 @@ struct FType
 	// Returns true if the given type is a matrix primitive.
 	bool IsMatrix() const { return Is(ETypeKind::Primitive) && Primitive.IsMatrix(); }
 
-	// Returns whether this type is a Texture2D object.
+	// Returns whether this type is a Texture object.
 	bool IsTexture() const { return Kind == ETypeKind::Texture; }
 
+	// Returns whether this type is a RuntimeVirtualTexture object.
+	bool IsRuntimeVirtualTexture() const { return Kind == ETypeKind::RuntimeVirtualTexture; }
+
 	// Returns whether this type is a SubstrateData object.
 	bool IsSubstrateData() const { return Kind == ETypeKind::SubstrateData; }
 
