[MIR] Fixed compile error on Linux for new MIR instruction (follow up on 45241464).

#rnx

diff --git a/Engine/Source/Runtime/Engine/Public/Materials/MaterialIR.h b/Engine/Source/Runtime/Engine/Public/Materials/MaterialIR.h
--- a/Engine/Source/Runtime/Engine/Public/Materials/MaterialIR.h
+++ b/Engine/Source/Runtime/Engine/Public/Materials/MaterialIR.h
@@ -1024,7 +1024,7 @@ struct FVTPageTableRead : TInstruction<VK_VTPageTableRead, 5>
 	TextureAddress AddressV : 8;
 
 	// Specifies the sampler function to read the page table, i.e. TextureLoadVirtualPageTable(), TextureLoadVirtualPageTableGrad(), or TextureLoadVirtualPageTableLevel().
-	ETextureMipValueMode MipValueMode : 3;
+	ETextureMipValueMode MipValueMode : 4;
 
 	// Enables VT texture sampling feedback
 	bool bEnableFeedback : 1;
