[DataTable] reenabled visitor based exporter and updated visitor to only step into container and structs (not in other property types like object, class, etc.)

diff --git a/Engine/Source/Runtime/Engine/Private/DataTableJSON.cpp b/Engine/Source/Runtime/Engine/Private/DataTableJSON.cpp
--- a/Engine/Source/Runtime/Engine/Private/DataTableJSON.cpp
+++ b/Engine/Source/Runtime/Engine/Private/DataTableJSON.cpp
@@ -65,7 +65,7 @@ namespace UE::DataTableJSON
 	 * in case issues get reported and validate potential regression.
 	 * We could eventually get rid of the FieldIterator implementation.
 	 */
-	bool bExportUsingPropertyVisitor = false;
+	bool bExportUsingPropertyVisitor = true;
 	static FAutoConsoleVariableRef CVarRuntimeValidationContext(
 		TEXT("DataTableJSON.ExportUsingPropertyVisitor"),
 		bExportUsingPropertyVisitor,
@@ -321,7 +321,7 @@ bool TDataTableExporterJSON<CharType>::WriteStructEntry(const void* InRowData, c
 	}
 	else if (const FStructProperty* StructProp = CastField<const FStructProperty>(InProperty))
 	{
-		if (!!(DTExportFlags & EDataTableExportFlags::UseJsonObjectsForStructs))
+		if (EnumHasAnyFlags(DTExportFlags, EDataTableExportFlags::UseJsonObjectsForStructs))
 		{
 			JsonWriter->WriteObjectStart(Identifier);
 			WriteStruct(StructProp->Struct, InPropertyData);
@@ -377,7 +377,7 @@ bool TDataTableExporterJSON<CharType>::WriteContainerEntry(const FProperty* InPr
 	}
 	else if (const FStructProperty* StructProp = CastField<const FStructProperty>(InProperty))
 	{
-		if (!!(DTExportFlags & EDataTableExportFlags::UseJsonObjectsForStructs))
+		if (EnumHasAnyFlags(DTExportFlags, EDataTableExportFlags::UseJsonObjectsForStructs))
 		{
 			WriteJSONObjectStartWithOptionalIdentifier<CharType>(*JsonWriter, InIdentifier);
 			WriteStruct(StructProp->Struct, InPropertyData);
@@ -465,7 +465,7 @@ bool TDataTableExporterJSON<CharType>::WriteStructUsingVisitor(TNotNull<const US
 			}
 			else if (Property->IsA<FStructProperty>())
 			{
-				if (!!(DTExportFlags & EDataTableExportFlags::UseJsonObjectsForStructs))
+				if (EnumHasAnyFlags(DTExportFlags, EDataTableExportFlags::UseJsonObjectsForStructs))
 				{
 					if (IdentifierOverrides.Num() && IdentifierOverrides.Top().IsSet())
 					{
@@ -512,7 +512,7 @@ bool TDataTableExporterJSON<CharType>::WriteStructUsingVisitor(TNotNull<const US
 			}
 			else if (Property->IsA<FStructProperty>())
 			{
-				if (!!(DTExportFlags & EDataTableExportFlags::UseJsonObjectsForStructs))
+				if (EnumHasAnyFlags(DTExportFlags, EDataTableExportFlags::UseJsonObjectsForStructs))
 				{
 					JsonWriter->WriteObjectEnd();
 				}
@@ -608,18 +608,19 @@ EPropertyVisitorControlFlow TDataTableExporterJSON<CharType>::WriteVisitedEntry(
 	}
 	else if (InProperty->IsA<const FStructProperty>())
 	{
-		if (!(DTExportFlags & EDataTableExportFlags::UseJsonObjectsForStructs))
+		if (EnumHasAnyFlags(DTExportFlags, EDataTableExportFlags::UseJsonObjectsForStructs))
 		{
-			const FString PropertyValue = DataTableUtils::GetPropertyValueAsStringDirect(InProperty, (uint8*)InValueAddress, DTExportFlags);
-			WriteJSONValueWithOptionalIdentifier<CharType>(*JsonWriter, InIdentifier, PropertyValue);
-			return EPropertyVisitorControlFlow::StepOver;
+			return EPropertyVisitorControlFlow::StepInto;
 		}
+
+		const FString PropertyValue = DataTableUtils::GetPropertyValueAsStringDirect(InProperty, (uint8*)InValueAddress, DTExportFlags);
+		WriteJSONValueWithOptionalIdentifier<CharType>(*JsonWriter, InIdentifier, PropertyValue);
 	}
 	else if (InProperty->IsA<FArrayProperty>()
 		|| InProperty->IsA<FSetProperty>()
 		|| InProperty->IsA<FMapProperty>())
 	{
-		// Cannot nest arrays, sets and maps
+		return EPropertyVisitorControlFlow::StepInto;
 	}
 	else
 	{
@@ -627,7 +628,7 @@ EPropertyVisitorControlFlow TDataTableExporterJSON<CharType>::WriteVisitedEntry(
 		WriteJSONValueWithOptionalIdentifier<CharType>(*JsonWriter, InIdentifier, PropertyValue);
 	}
 
-	return EPropertyVisitorControlFlow::StepInto;
+	return EPropertyVisitorControlFlow::StepOver;
 };
 
 template class TDataTableExporterJSON<TCHAR>;
