Proper fix to remove duplicated actor data layers.

diff --git a/Engine/Source/Runtime/Engine/Private/ActorEditor.cpp b/Engine/Source/Runtime/Engine/Private/ActorEditor.cpp
--- a/Engine/Source/Runtime/Engine/Private/ActorEditor.cpp
+++ b/Engine/Source/Runtime/Engine/Private/ActorEditor.cpp
@@ -41,7 +41,6 @@
 #include "UObject/SoftObjectPtr.h"
 #include "Algo/Transform.h"
 #include "Algo/RemoveIf.h"
-#include "Algo/Unique.h"
 #include "Misc/DataValidation.h"
 #include "DeletedObjectPlaceholder.h"
 #include "DataStorage/Features.h"
@@ -1870,8 +1869,9 @@ void AActor::FixupDataLayers(bool bRevertChangesOnLockedDataLayer /*= false*/)
 	}));
 
 	// Remove duplicates
-	DataLayerAssets.SetNum(Algo::Unique(DataLayerAssets));
-
+	TSet<TSoftObjectPtr<UDataLayerAsset>> UniqueDataLayerAssets(DataLayerAssets);
+	DataLayerAssets = UniqueDataLayerAssets.Array();
+	
 	// Use Actor's DataLayerManager since the fixup is relative to this level
 	UDataLayerManager* DataLayerManager = UDataLayerManager::GetDataLayerManager(this);
 	if (!DataLayerManager || !DataLayerManager->CanResolveDataLayers())
