Make sure a world subsystem executes PostInitialize and OnWorldBeginPlay if it is added to the world after those functions have been executed for the world.

diff --git a/Engine/Plugins/Experimental/GPULightmass/Source/GPULightmass/Private/GPULightmassSettings.cpp b/Engine/Plugins/Experimental/GPULightmass/Source/GPULightmass/Private/GPULightmassSettings.cpp
--- a/Engine/Plugins/Experimental/GPULightmass/Source/GPULightmass/Private/GPULightmassSettings.cpp
+++ b/Engine/Plugins/Experimental/GPULightmass/Source/GPULightmass/Private/GPULightmassSettings.cpp
@@ -121,6 +121,8 @@ bool UGPULightmassSettings::CanEditChange(const FProperty* InProperty) const
 
 void UGPULightmassSubsystem::PostInitialize()
 {
+	Super::PostInitialize();
+
 	UWorld* World = Cast<UWorld>(GetOuter());
 	
 	if (!World) return;
diff --git a/Engine/Plugins/Experimental/WaterAdvanced/Source/WaterAdvanced/Private/ShallowWaterSubsystem.cpp b/Engine/Plugins/Experimental/WaterAdvanced/Source/WaterAdvanced/Private/ShallowWaterSubsystem.cpp
--- a/Engine/Plugins/Experimental/WaterAdvanced/Source/WaterAdvanced/Private/ShallowWaterSubsystem.cpp
+++ b/Engine/Plugins/Experimental/WaterAdvanced/Source/WaterAdvanced/Private/ShallowWaterSubsystem.cpp
@@ -115,6 +115,8 @@ void UShallowWaterSubsystem::PostInitialize()
 
 void UShallowWaterSubsystem::OnWorldBeginPlay(UWorld& InWorld)
 {
+	Super::OnWorldBeginPlay(InWorld);
+
 	TArray<ULocalPlayer*> LocalPlayers = InWorld.GetGameInstance() != nullptr ? InWorld.GetGameInstance()->GetLocalPlayers() : TArray<ULocalPlayer*>();
 	
 	// we don't support split screen
diff --git a/Engine/Source/Runtime/Engine/Private/Subsystems/WorldSubsystem.cpp b/Engine/Source/Runtime/Engine/Private/Subsystems/WorldSubsystem.cpp
--- a/Engine/Source/Runtime/Engine/Private/Subsystems/WorldSubsystem.cpp
+++ b/Engine/Source/Runtime/Engine/Private/Subsystems/WorldSubsystem.cpp
@@ -41,6 +41,47 @@ bool UWorldSubsystem::DoesSupportWorldType(const EWorldType::Type WorldType) con
 	return WorldType == EWorldType::Game || WorldType == EWorldType::Editor || WorldType == EWorldType::PIE;
 }
 
+void UWorldSubsystem::Initialize(FSubsystemCollectionBase& Collection)
+{
+	Super::Initialize(Collection);
+	//run relevant functions we may have missed if we are initialized after the world has set up
+	if (UWorld* World = GetWorld())
+	{
+		if (!bHasCalledPostInitialize && World->bIsWorldInitialized)
+		{
+			PostInitialize();
+			EnsureHasCalledPostInitialize();
+		}
+
+		if (!bHasCalledBeginPlay && World->HasBegunPlay())
+		{
+			OnWorldBeginPlay(*World);
+			EnsureHasCalledBeginPlay();
+		}
+	}
+}
+
+void UWorldSubsystem::PostInitialize()
+{
+	ensureAlwaysMsgf(!bHasCalledPostInitialize, TEXT("PostInitialize has already been called for subsystem %s"), *GetClass()->GetName());
+	bHasCalledPostInitialize = true;
+}
+
+void UWorldSubsystem::OnWorldBeginPlay(UWorld& InWorld)
+{
+	ensureAlwaysMsgf(!bHasCalledBeginPlay, TEXT("OnWorldBeginPlay has already been called for subsystem %s"), *GetClass()->GetName());
+	bHasCalledBeginPlay = true;
+}
+
+void UWorldSubsystem::EnsureHasCalledPostInitialize() const
+{
+	ensureAlwaysMsgf(bHasCalledPostInitialize, TEXT("PostInitialize has not been called for subsystem %s"), *GetClass()->GetName());
+}
+void UWorldSubsystem::EnsureHasCalledBeginPlay() const
+{
+	ensureAlwaysMsgf(bHasCalledBeginPlay, TEXT("OnWorldBeginPlay has not been called for subsystem %s"), *GetClass()->GetName());
+}
+
 PRAGMA_DISABLE_DEPRECATION_WARNINGS
 void UWorldSubsystem::UpdateStreamingState()
 {
@@ -92,6 +133,7 @@ void UTickableWorldSubsystem::Tick(float DeltaTime)
 
 void UTickableWorldSubsystem::Initialize(FSubsystemCollectionBase& Collection)
 {
+	Super::Initialize(Collection);
 	check(!bInitialized);
 	bInitialized = true;
 
@@ -101,6 +143,7 @@ void UTickableWorldSubsystem::Initialize(FSubsystemCollectionBase& Collection)
 
 void UTickableWorldSubsystem::Deinitialize()
 {
+	Super::Deinitialize();
 	check(bInitialized);
 	bInitialized = false;
 
diff --git a/Engine/Source/Runtime/Engine/Private/World.cpp b/Engine/Source/Runtime/Engine/Private/World.cpp
--- a/Engine/Source/Runtime/Engine/Private/World.cpp
+++ b/Engine/Source/Runtime/Engine/Private/World.cpp
@@ -6045,6 +6045,7 @@ void UWorld::BeginPlay()
 	SubsystemCollection.ForEachSubsystem([this](UWorldSubsystem* WorldSubsystem)
 	{
 		WorldSubsystem->OnWorldBeginPlay(*this);
+		WorldSubsystem->EnsureHasCalledBeginPlay();
 	});
 
 	AGameModeBase* const GameMode = GetAuthGameMode();
@@ -10293,6 +10294,7 @@ void UWorld::PostInitializeSubsystems()
 	SubsystemCollection.ForEachSubsystem([](UWorldSubsystem* WorldSubsystem)
 	{
 		WorldSubsystem->PostInitialize();
+		WorldSubsystem->EnsureHasCalledPostInitialize();
 	});
 }
 
diff --git a/Engine/Source/Runtime/Engine/Public/Subsystems/WorldSubsystem.h b/Engine/Source/Runtime/Engine/Public/Subsystems/WorldSubsystem.h
--- a/Engine/Source/Runtime/Engine/Public/Subsystems/WorldSubsystem.h
+++ b/Engine/Source/Runtime/Engine/Public/Subsystems/WorldSubsystem.h
@@ -29,14 +29,20 @@ public:
 
 	ENGINE_API virtual bool ShouldCreateSubsystem(UObject* Outer) const override;
 
+	/** 
+	* Initializes the world subsystem.  
+	* Will execute PostInitialize if the world has already been Initialize
+	* Will execute OnWorldBeginPlay if the world has already begun play*/
+	ENGINE_API virtual void Initialize(FSubsystemCollectionBase& Collection) override;
+
 	/** Called once all UWorldSubsystems have been initialized */
-	virtual void PostInitialize() {}
+	ENGINE_API virtual void PostInitialize();
 
 	/** Called once all UWorldSubsystems are about to be deinitialized */
 	virtual void PreDeinitialize() {}
 	
 	/** Called when world is ready to start gameplay before the game mode transitions to the correct state and call BeginPlay on all actors */
-	virtual void OnWorldBeginPlay(UWorld& InWorld) {}
+	ENGINE_API virtual void OnWorldBeginPlay(UWorld& InWorld);
 
 	/** Called when world is ready to stop gameplay before the game mode transitions to the correct state and call EndPlay on all actors */
 	virtual void OnWorldEndPlay(UWorld& InWorld) {}
@@ -44,11 +50,18 @@ public:
 	/** Called after world components (e.g. line batcher and all level components) have been updated */
 	virtual void OnWorldComponentsUpdated(UWorld& World) {}
 
+	ENGINE_API void EnsureHasCalledPostInitialize() const;
+
+	ENGINE_API void EnsureHasCalledBeginPlay() const;
+
 	UE_DEPRECATED(5.5, "Implement IStreamingWorldSubsystemInterface instead.")
 	ENGINE_API virtual void UpdateStreamingState();
 
 protected:
 	ENGINE_API virtual bool DoesSupportWorldType(const EWorldType::Type WorldType) const;
+
+	bool bHasCalledPostInitialize : 1 = false;
+	bool bHasCalledBeginPlay : 1 = false;
 };
 
 /**
