Fix for editor GetData() assert during anim bank compilation and the data isn't available.

diff --git a/Engine/Source/Runtime/Engine/Private/Animation/AnimBank.cpp b/Engine/Source/Runtime/Engine/Private/Animation/AnimBank.cpp
--- a/Engine/Source/Runtime/Engine/Private/Animation/AnimBank.cpp
+++ b/Engine/Source/Runtime/Engine/Private/Animation/AnimBank.cpp
@@ -1222,11 +1222,6 @@ bool UAnimBankData::IsEnabled() const
 			return false;
 		}
 	#endif
-
-		if (BankItem.SequenceIndex >= BankItem.BankAsset->GetData().Entries.Num())
-		{
-			return false;
-		}
 	}
 
 	return bEnabled;
