Allow systems to prepare to be captured by [REDACTED_PERSON] by adding a delegate that can be registered to.

diff --git a/Engine/Source/Editor/UnrealEd/Private/WorldPartition/WorldPartitionMiniMapBuilder.cpp b/Engine/Source/Editor/UnrealEd/Private/WorldPartition/WorldPartitionMiniMapBuilder.cpp
--- a/Engine/Source/Editor/UnrealEd/Private/WorldPartition/WorldPartitionMiniMapBuilder.cpp
+++ b/Engine/Source/Editor/UnrealEd/Private/WorldPartition/WorldPartitionMiniMapBuilder.cpp
@@ -33,6 +33,8 @@
 
 DEFINE_LOG_CATEGORY_STATIC(LogWorldPartitionMiniMapBuilder, All, All);
 
+UWorldPartitionMiniMapBuilder::FOnWarmupTickDelegate UWorldPartitionMiniMapBuilder::OnWarmupTick;
+
 UWorldPartitionMiniMapBuilder::UWorldPartitionMiniMapBuilder(const FObjectInitializer& ObjectInitializer)
 	: Super(ObjectInitializer)
 {
@@ -120,7 +122,11 @@ bool UWorldPartitionMiniMapBuilder::RunInternal(UWorld* World, const FCellInfo&
 		TileTexture->Source.Init(CaptureWidthPixels, CaptureHeightPixels, 1, 1, TSF_BGRA8);
 		TileTexture->PowerOfTwoMode = ETexturePowerOfTwoSetting::PadToPowerOfTwo;
 
-		FWorldPartitionMiniMapHelper::CaptureBoundsMiniMapToTexture(World, GetTransientPackage(), CaptureWidthPixels, CaptureHeightPixels, TileTexture, TextureName, ClampedBounds, WorldMiniMap->CaptureSource, WorldMiniMap->CaptureWarmupFrames);
+		FWorldPartitionMiniMapHelper::CaptureBoundsMiniMapToTexture(World, GetTransientPackage(), CaptureWidthPixels, CaptureHeightPixels, TileTexture, TextureName, ClampedBounds, WorldMiniMap->CaptureSource, WorldMiniMap->CaptureWarmupFrames,
+			[World]()
+			{
+				OnWarmupTick.Broadcast(World);
+			});
 
 		// Copy captured image to VT minimap
 		const uint32 BPP = TileTexture->Source.GetBytesPerPixel();
diff --git a/Engine/Source/Editor/UnrealEd/Public/WorldPartition/WorldPartitionMiniMapBuilder.h b/Engine/Source/Editor/UnrealEd/Public/WorldPartition/WorldPartitionMiniMapBuilder.h
--- a/Engine/Source/Editor/UnrealEd/Public/WorldPartition/WorldPartitionMiniMapBuilder.h
+++ b/Engine/Source/Editor/UnrealEd/Public/WorldPartition/WorldPartitionMiniMapBuilder.h
@@ -20,6 +20,9 @@ public:
 	virtual bool RequiresCommandletRendering() const override { return true; }
 	virtual ELoadingMode GetLoadingMode() const override { return ELoadingMode::IterativeCells2D; }
 
+	DECLARE_MULTICAST_DELEGATE_OneParam(FOnWarmupTickDelegate, UWorld*);
+	static UNREALED_API FOnWarmupTickDelegate OnWarmupTick;
+
 protected:
 	UNREALED_API virtual bool PreRun(UWorld* World, FPackageSourceControlHelper& PackageHelper) override;
 	UNREALED_API virtual bool RunInternal(UWorld* World, const FCellInfo& InCellInfo, FPackageSourceControlHelper& PackageHelper) override;
diff --git a/Engine/Source/Runtime/Engine/Private/WorldPartition/WorldPartitionMiniMapHelper.cpp b/Engine/Source/Runtime/Engine/Private/WorldPartition/WorldPartitionMiniMapHelper.cpp
--- a/Engine/Source/Runtime/Engine/Private/WorldPartition/WorldPartitionMiniMapHelper.cpp
+++ b/Engine/Source/Runtime/Engine/Private/WorldPartition/WorldPartitionMiniMapHelper.cpp
@@ -55,10 +55,12 @@ AWorldPartitionMiniMap* FWorldPartitionMiniMapHelper::GetWorldPartitionMiniMap(U
 void FWorldPartitionMiniMapHelper::CaptureBoundsMiniMapToTexture(UWorld* InWorld, UObject* InOuterForTexture, uint32 InMiniMapSize, UTexture2D*& InOutMiniMapTexture, const FString& InTextureName, const FBox& InBounds)
 {
 	const uint32 DefaultNumWarmupFrames = 5;
-	CaptureBoundsMiniMapToTexture(InWorld, InOuterForTexture, InMiniMapSize, InMiniMapSize, InOutMiniMapTexture, InTextureName, InBounds, SCS_BaseColor, DefaultNumWarmupFrames);
+	auto TickFunction = []() {};
+	CaptureBoundsMiniMapToTexture(InWorld, InOuterForTexture, InMiniMapSize, InMiniMapSize, InOutMiniMapTexture, InTextureName, InBounds, SCS_BaseColor, DefaultNumWarmupFrames, TickFunction);
 }
 
-void FWorldPartitionMiniMapHelper::CaptureBoundsMiniMapToTexture(UWorld* InWorld, UObject* InOuterForTexture, uint32 InMiniMapSizeX, uint32 InMiniMapSizeY, UTexture2D*& InOutMiniMapTexture, const FString& InTextureName, const FBox& InBounds, ESceneCaptureSource InCaptureSource, uint32 InNumWarmupFrames)
+void FWorldPartitionMiniMapHelper::CaptureBoundsMiniMapToTexture(UWorld* InWorld, UObject* InOuterForTexture, uint32 InMiniMapSizeX, uint32 InMiniMapSizeY, UTexture2D*& InOutMiniMapTexture, const FString& InTextureName, const FBox& InBounds, ESceneCaptureSource InCaptureSource, uint32 InNumWarmupFrames,
+	TFunctionRef<void()> WarmupTickCallback)
 {
 	// Before capturing the scene, make sure all assets are finished compiling
 	FAssetCompilingManager::Get().FinishAllCompilation();
@@ -105,6 +107,7 @@ void FWorldPartitionMiniMapHelper::CaptureBoundsMiniMapToTexture(UWorld* InWorld
 
 	for (uint32 i = 0; i < InNumWarmupFrames; i++)
 	{
+		WarmupTickCallback();
 		CaptureComponent->CaptureScene();
 	}
 
diff --git a/Engine/Source/Runtime/Engine/Public/WorldPartition/WorldPartitionMiniMapHelper.h b/Engine/Source/Runtime/Engine/Public/WorldPartition/WorldPartitionMiniMapHelper.h
--- a/Engine/Source/Runtime/Engine/Public/WorldPartition/WorldPartitionMiniMapHelper.h
+++ b/Engine/Source/Runtime/Engine/Public/WorldPartition/WorldPartitionMiniMapHelper.h
@@ -23,7 +23,7 @@ class FWorldPartitionMiniMapHelper
 public:
 	static ENGINE_API AWorldPartitionMiniMap* GetWorldPartitionMiniMap(UWorld* World, bool bCreateNewMiniMap=false);
 	static ENGINE_API void CaptureBoundsMiniMapToTexture(UWorld* InWorld, UObject* InOuterForTexture, uint32 InMiniMapSize, UTexture2D*& InOutMiniMapTexture, const FString& InTextureName, const FBox& InBounds);
-	static ENGINE_API void CaptureBoundsMiniMapToTexture(UWorld* InWorld, UObject* InOuterForTexture, uint32 InMiniMapSizeX, uint32 InMiniMapSizeY, UTexture2D*& InOutMiniMapTexture, const FString& InTextureName, const FBox& InBounds, ESceneCaptureSource InCaptureSource, uint32 InNumWarmupFrames);
+	static ENGINE_API void CaptureBoundsMiniMapToTexture(UWorld* InWorld, UObject* InOuterForTexture, uint32 InMiniMapSizeX, uint32 InMiniMapSizeY, UTexture2D*& InOutMiniMapTexture, const FString& InTextureName, const FBox& InBounds, ESceneCaptureSource InCaptureSource, uint32 InNumWarmupFrames, TFunctionRef<void()> WarmupTickCallback);
 
 private:
 	static void CalTopViewOfWorld(FMatrix& OutProjectionMatrix, const FBox& WorldBox, uint32 ViewportWidth, uint32 ViewportHeight);
