SceneProxyDesc cleanup
- Moving all FPrimitiveSceneProxyDesc members together instead of being mixed between methods.
- Moving all FPrimitiveSceneProxyDesc method definitions to its own cpp file.
- Reorganizing FPrimitiveSceneProxyDesc and FStaticMeshSceneProxyDesc data for memory savings:
-- FPrimitiveSceneProxyDesc from 208 -> 200 bytes
-- FStaticMeshSceneProxyDesc from 472 -> 448 bytes

diff --git a/Engine/Source/Runtime/Engine/Private/PrimitiveSceneProxy.cpp b/Engine/Source/Runtime/Engine/Private/PrimitiveSceneProxy.cpp
--- a/Engine/Source/Runtime/Engine/Private/PrimitiveSceneProxy.cpp
+++ b/Engine/Source/Runtime/Engine/Private/PrimitiveSceneProxy.cpp
@@ -24,14 +24,11 @@
 #include "NaniteSceneProxy.h" // TODO: PROG_RASTER
 #include "ComponentRecreateRenderStateContext.h"
 #include "DataDrivenShaderPlatformInfo.h"
-#include "SceneInterface.h"
 #include "PrimitiveUniformShaderParametersBuilder.h"
-#include "Components/InstancedStaticMeshComponent.h"
 #include "InstanceDataSceneProxy.h"
 #include "GameFramework/ActorPrimitiveColorHandler.h"
 
 #if WITH_EDITOR
-#include "FoliageHelper.h"
 #include "ObjectCacheEventSink.h"
 #endif
 
@@ -282,145 +279,6 @@ FPrimitiveSceneProxyDesc::FPrimitiveSceneProxyDesc(const UPrimitiveComponent* In
 
 FPrimitiveSceneProxy::FPrimitiveSceneProxy(FPrimitiveSceneProxy const&) = default;
 
-void FPrimitiveSceneProxyDesc::InitializeFromPrimitiveComponent(const UPrimitiveComponent* InComponent)
-{
-	CastShadow = InComponent->CastShadow;
-	bReceivesDecals = InComponent->bReceivesDecals;
-	bOnlyOwnerSee = InComponent->bOnlyOwnerSee;
-	bOwnerNoSee = InComponent->bOwnerNoSee;
-	bLevelInstanceEditingState = InComponent->GetLevelInstanceEditingState();
-	bUseViewOwnerDepthPriorityGroup  = InComponent->bUseViewOwnerDepthPriorityGroup ;
-	bVisibleInReflectionCaptures = InComponent->bVisibleInReflectionCaptures;
-	bVisibleInRealTimeSkyCaptures = InComponent->bVisibleInRealTimeSkyCaptures;
-	bVisibleInRayTracing = InComponent->bVisibleInRayTracing;
-	bRenderInDepthPass = InComponent->bRenderInDepthPass;
-	bRenderInMainPass = InComponent->bRenderInMainPass;
-	bTreatAsBackgroundForOcclusion = InComponent->bTreatAsBackgroundForOcclusion;
-	bCastDynamicShadow = InComponent->bCastDynamicShadow;
-	bCastStaticShadow = InComponent->bCastStaticShadow;
-	bEmissiveLightSource = InComponent->bEmissiveLightSource;
-	bAffectDynamicIndirectLighting = InComponent->bAffectDynamicIndirectLighting;
-	bAffectIndirectLightingWhileHidden = InComponent->bAffectIndirectLightingWhileHidden;
-	bAffectDistanceFieldLighting = InComponent->bAffectDistanceFieldLighting;
-	bCastVolumetricTranslucentShadow = InComponent->bCastVolumetricTranslucentShadow;
-	bCastContactShadow = InComponent->bCastContactShadow;
-	bCastHiddenShadow = InComponent->bCastHiddenShadow;
-	bCastShadowAsTwoSided = InComponent->bCastShadowAsTwoSided;
-	bSelfShadowOnly = InComponent->bSelfShadowOnly;
-	bCastInsetShadow = InComponent->bCastInsetShadow;
-	bCastCinematicShadow = InComponent->bCastCinematicShadow;
-	bCastFarShadow = InComponent->bCastFarShadow;
-	bLightAttachmentsAsGroup = InComponent->bLightAttachmentsAsGroup;
-	bSingleSampleShadowFromStationaryLights = InComponent->bSingleSampleShadowFromStationaryLights;
-	bUseAsOccluder = InComponent->bUseAsOccluder;
-	bSelectable = InComponent->bSelectable;
-	bHasPerInstanceHitProxies = InComponent->bHasPerInstanceHitProxies;
-	bUseEditorCompositing = InComponent->bUseEditorCompositing;
-	bIsBeingMovedByEditor = InComponent->bIsBeingMovedByEditor;
-	bReceiveMobileCSMShadows = InComponent->bReceiveMobileCSMShadows;
-	bRenderCustomDepth = InComponent->bRenderCustomDepth;
-	bVisibleInSceneCaptureOnly = InComponent->bVisibleInSceneCaptureOnly;
-	bHiddenInSceneCapture = InComponent->bHiddenInSceneCapture;
-	bForceMipStreaming = InComponent->bForceMipStreaming;
-	bRayTracingFarField = InComponent->bRayTracingFarField;
-	bHoldout = InComponent->bHoldout;
-	bWantsEditorEffects = InComponent->bWantsEditorEffects;
-	bIsFirstPerson = InComponent->FirstPersonPrimitiveType == EFirstPersonPrimitiveType::FirstPerson;
-	bIsFirstPersonWorldSpaceRepresentation = InComponent->FirstPersonPrimitiveType == EFirstPersonPrimitiveType::WorldSpaceRepresentation;
-	bLumenHeightfield = InComponent->bLumenHeightfield;
-	
-	bIsVisible = InComponent->IsVisible();
-	bIsVisibleEditor = InComponent->GetVisibleFlag();
-	bSelected = InComponent->IsSelected();
-	bIndividuallySelected = InComponent->IsComponentIndividuallySelected();
-	bShouldRenderSelected = InComponent->ShouldRenderSelected();
-	bCollisionEnabled = InComponent->IsCollisionEnabled(); 
-	
-	if (const AActor* ActorOwner = InComponent->GetOwner())
-	{
-		bIsHidden = ActorOwner->IsHidden(); 
-#if WITH_EDITOR
-		bIsHiddenEd = ActorOwner->IsHiddenEd(); 	
-		bIsOwnedByFoliage = FFoliageHelper::IsOwnedByFoliage(ActorOwner);
-#endif
-
-		if(bOnlyOwnerSee || bOwnerNoSee || bUseViewOwnerDepthPriorityGroup || bIsFirstPersonWorldSpaceRepresentation)
-		{
-			// Make a list of the actors which directly or indirectly own the InComponent.
-			for(const AActor* CurrentOwner = ActorOwner;CurrentOwner;CurrentOwner = CurrentOwner->GetOwner())
-			{
-				ActorOwners.Add(CurrentOwner->GetUniqueID());
-			}
-		}
-		
-		bIsOwnerEditorOnly = InComponent->GetOwner()->IsEditorOnly(); 
-	}
-	bSupportsWorldPositionOffsetVelocity = InComponent->SupportsWorldPositionOffsetVelocity();
-	bIsInstancedStaticMesh = Cast<UInstancedStaticMeshComponent>(InComponent) != nullptr; 
-
-	Mobility = InComponent->Mobility;;
-	TranslucencySortPriority = InComponent->TranslucencySortPriority;
-	TranslucencySortDistanceOffset = InComponent->TranslucencySortDistanceOffset;
-	LightmapType = InComponent->GetLightmapType();
-	ViewOwnerDepthPriorityGroup = InComponent->ViewOwnerDepthPriorityGroup;
-	CustomDepthStencilValue = InComponent->CustomDepthStencilValue;
-	CustomDepthStencilWriteMask = InComponent->CustomDepthStencilWriteMask;
-	LightingChannels = InComponent->LightingChannels;
-	RayTracingGroupCullingPriority = InComponent->RayTracingGroupCullingPriority;
-	IndirectLightingCacheQuality = InComponent->IndirectLightingCacheQuality;
-	ShadowCacheInvalidationBehavior = InComponent->ShadowCacheInvalidationBehavior;
-	DepthPriorityGroup = InComponent->GetStaticDepthPriorityGroup();
-	
-	VirtualTextureLodBias = InComponent->VirtualTextureLodBias ;
-	VirtualTextureCullMips = InComponent->VirtualTextureCullMips ;
-	VirtualTextureMinCoverage = InComponent->VirtualTextureMinCoverage ;
-	ComponentId = InComponent->GetPrimitiveSceneId() ;
-	VisibilityId = InComponent->VisibilityId ;
-	CachedMaxDrawDistance = InComponent->CachedMaxDrawDistance ;
-	MinDrawDistance = InComponent->MinDrawDistance ;
-	BoundsScale = InComponent->BoundsScale ;
-	RayTracingGroupId = InComponent->GetRayTracingGroupId() ;
-
-	bHasStaticLighting = InComponent->HasStaticLighting();
-	bHasValidSettingsForStaticLighting = InComponent->HasValidSettingsForStaticLighting(false);
-	bIsPrecomputedLightingValid = InComponent->IsPrecomputedLightingValid();
-	bShadowIndirectOnly = InComponent->GetShadowIndirectOnly();	
-
-	Component = const_cast<UPrimitiveComponent*>(InComponent); 
-	Owner = InComponent->GetOwner();
-
-	#if !WITH_STATE_STREAM
-	World = InComponent->GetWorld();
-	#endif
-
-	CustomPrimitiveData = &InComponent->GetCustomPrimitiveData();
-	Scene = InComponent->GetScene();
-	PrimitiveComponentInterface = InComponent->GetPrimitiveComponentInterface();
-	
-	if (Scene)
-	{
-		FeatureLevel = Scene->GetFeatureLevel();
-	}
-
-#if WITH_EDITOR
-	HiddenEditorViews = InComponent->GetHiddenEditorViews();
-	OverlayColor = InComponent->OverlayColor;
-#endif
-	bShouldRenderProxyFallbackToDefaultMaterial = InComponent->ShouldRenderProxyFallbackToDefaultMaterial();
-
-	AdditionalStatObjectPtr = InComponent->AdditionalStatObject();
-	StatId = AdditionalStatObjectPtr? AdditionalStatObjectPtr->GetStatID(true) : InComponent->GetStatID(true);
-
-	TArray<URuntimeVirtualTexture*> const& VirtualTextures = InComponent->GetRuntimeVirtualTextures();	
-	RuntimeVirtualTextures = MakeArrayView( const_cast<URuntimeVirtualTexture**>(VirtualTextures.GetData()), VirtualTextures.Num());	
-	VirtualTextureRenderPassType = InComponent->GetVirtualTextureRenderPassType();
-	VirtualTextureMainPassMaxDrawDistance = InComponent->GetVirtualTextureMainPassMaxDrawDistance();
-
-#if MESH_DRAW_COMMAND_STATS
-	MeshDrawCommandStatsCategory = InComponent->GetMeshDrawCommandStatsCategory();
-#endif
-}
-
 
 FPrimitiveSceneProxy::FPrimitiveSceneProxy(const UPrimitiveComponent* InComponent, FName InResourceName)
 	: FPrimitiveSceneProxy(FPrimitiveSceneProxyDesc(InComponent), InResourceName)
@@ -2069,13 +1927,3 @@ ERayTracingProxyType FPrimitiveSceneProxy::GetRayTracingProxyType() const
 	return ProxyType;
 }
 #endif // RHI_RAYTRACING
-
-void FPrimitiveSceneProxyDesc::GetUsedMaterials(TArray<UMaterialInterface*>& OutMaterials, bool bGetDebugMaterials) const 
-{
-	// Only UPrimitiveComponent should rely on this method 
-	const UPrimitiveComponent* AsComponent = Cast<UPrimitiveComponent>(Component);
-	check(AsComponent);
-	
-	return AsComponent->GetUsedMaterials(OutMaterials, bGetDebugMaterials);	
-}
-
diff --git a/Engine/Source/Runtime/Engine/Private/PrimitiveSceneProxyDesc.cpp b/Engine/Source/Runtime/Engine/Private/PrimitiveSceneProxyDesc.cpp
new file mode 100644
--- /dev/null
+++ b/Engine/Source/Runtime/Engine/Private/PrimitiveSceneProxyDesc.cpp
@@ -0,0 +1,228 @@
+// Copyright Epic Games, Inc. All Rights Reserved.
+
+#include "PrimitiveSceneProxyDesc.h"
+#include "Components/InstancedStaticMeshComponent.h"
+#include "SceneInterface.h"
+
+#if WITH_EDITOR
+#include "FoliageHelper.h"
+#endif
+
+FPrimitiveSceneProxyDesc::FPrimitiveSceneProxyDesc()
+{
+	CastShadow = false;
+	bReceivesDecals = true;
+	bOnlyOwnerSee = false;
+	bOwnerNoSee = false;
+	bLevelInstanceEditingState = false;
+	bUseViewOwnerDepthPriorityGroup = false;
+	bVisibleInReflectionCaptures = true;
+	bVisibleInRealTimeSkyCaptures = true;
+	bVisibleInRayTracing = true;
+	bRenderInDepthPass = true;
+	bRenderInMainPass = true;
+	bTreatAsBackgroundForOcclusion = false;
+	bCastDynamicShadow = true;
+	bCastStaticShadow = true;
+	bEmissiveLightSource = false;
+	bAffectDynamicIndirectLighting = true;
+	bAffectIndirectLightingWhileHidden = false;
+	bAffectDistanceFieldLighting = true;
+	bCastVolumetricTranslucentShadow = false;
+	bCastContactShadow = true;
+	bCastHiddenShadow = false;
+	bCastShadowAsTwoSided = false;
+	bSelfShadowOnly = false;
+	bCastInsetShadow = false;
+	bCastCinematicShadow = false;
+	bCastFarShadow = false;
+	bLightAttachmentsAsGroup = false;
+	bSingleSampleShadowFromStationaryLights = false;
+	bUseAsOccluder = false;
+	bSelectable = true;
+	bHasPerInstanceHitProxies = false;
+	bUseEditorCompositing = false;
+	bIsBeingMovedByEditor = false;
+	bReceiveMobileCSMShadows = true;
+	bRenderCustomDepth = false;
+	bVisibleInSceneCaptureOnly = false;
+	bHiddenInSceneCapture = false;
+	bForceMipStreaming = false;
+	bRayTracingFarField = false;
+	bIsVisible = true;
+	bIsVisibleEditor = true;
+	bSelected = false;
+	bIndividuallySelected = false;
+	bCollisionEnabled = false;
+	bIsHidden = false;
+	bIsHiddenEd = false;
+	bSupportsWorldPositionOffsetVelocity = true;
+	bIsOwnerEditorOnly = false;
+	bIsInstancedStaticMesh = false;
+	bHoldout = false;
+	bIsFirstPerson = false;
+	bIsFirstPersonWorldSpaceRepresentation = false;
+	bLumenHeightfield = false;
+
+	bHasStaticLighting = false;
+	bHasValidSettingsForStaticLighting = false;
+	bIsPrecomputedLightingValid = false;
+	bShadowIndirectOnly = false;
+	bShouldRenderProxyFallbackToDefaultMaterial = false;
+	bShouldRenderSelected = false;
+	bWantsEditorEffects = false;
+
+#if WITH_EDITOR
+	bIsOwnedByFoliage = false;
+#endif
+}
+
+FPrimitiveSceneProxyDesc::~FPrimitiveSceneProxyDesc() = default;
+
+void FPrimitiveSceneProxyDesc::InitializeFromPrimitiveComponent(const UPrimitiveComponent* InComponent)
+{
+	CastShadow = InComponent->CastShadow;
+	bReceivesDecals = InComponent->bReceivesDecals;
+	bOnlyOwnerSee = InComponent->bOnlyOwnerSee;
+	bOwnerNoSee = InComponent->bOwnerNoSee;
+	bLevelInstanceEditingState = InComponent->GetLevelInstanceEditingState();
+	bUseViewOwnerDepthPriorityGroup = InComponent->bUseViewOwnerDepthPriorityGroup;
+	bVisibleInReflectionCaptures = InComponent->bVisibleInReflectionCaptures;
+	bVisibleInRealTimeSkyCaptures = InComponent->bVisibleInRealTimeSkyCaptures;
+	bVisibleInRayTracing = InComponent->bVisibleInRayTracing;
+	bRenderInDepthPass = InComponent->bRenderInDepthPass;
+	bRenderInMainPass = InComponent->bRenderInMainPass;
+	bTreatAsBackgroundForOcclusion = InComponent->bTreatAsBackgroundForOcclusion;
+	bCastDynamicShadow = InComponent->bCastDynamicShadow;
+	bCastStaticShadow = InComponent->bCastStaticShadow;
+	bEmissiveLightSource = InComponent->bEmissiveLightSource;
+	bAffectDynamicIndirectLighting = InComponent->bAffectDynamicIndirectLighting;
+	bAffectIndirectLightingWhileHidden = InComponent->bAffectIndirectLightingWhileHidden;
+	bAffectDistanceFieldLighting = InComponent->bAffectDistanceFieldLighting;
+	bCastVolumetricTranslucentShadow = InComponent->bCastVolumetricTranslucentShadow;
+	bCastContactShadow = InComponent->bCastContactShadow;
+	bCastHiddenShadow = InComponent->bCastHiddenShadow;
+	bCastShadowAsTwoSided = InComponent->bCastShadowAsTwoSided;
+	bSelfShadowOnly = InComponent->bSelfShadowOnly;
+	bCastInsetShadow = InComponent->bCastInsetShadow;
+	bCastCinematicShadow = InComponent->bCastCinematicShadow;
+	bCastFarShadow = InComponent->bCastFarShadow;
+	bLightAttachmentsAsGroup = InComponent->bLightAttachmentsAsGroup;
+	bSingleSampleShadowFromStationaryLights = InComponent->bSingleSampleShadowFromStationaryLights;
+	bUseAsOccluder = InComponent->bUseAsOccluder;
+	bSelectable = InComponent->bSelectable;
+	bHasPerInstanceHitProxies = InComponent->bHasPerInstanceHitProxies;
+	bUseEditorCompositing = InComponent->bUseEditorCompositing;
+	bIsBeingMovedByEditor = InComponent->bIsBeingMovedByEditor;
+	bReceiveMobileCSMShadows = InComponent->bReceiveMobileCSMShadows;
+	bRenderCustomDepth = InComponent->bRenderCustomDepth;
+	bVisibleInSceneCaptureOnly = InComponent->bVisibleInSceneCaptureOnly;
+	bHiddenInSceneCapture = InComponent->bHiddenInSceneCapture;
+	bForceMipStreaming = InComponent->bForceMipStreaming;
+	bRayTracingFarField = InComponent->bRayTracingFarField;
+	bHoldout = InComponent->bHoldout;
+	bWantsEditorEffects = InComponent->bWantsEditorEffects;
+	bIsFirstPerson = InComponent->FirstPersonPrimitiveType == EFirstPersonPrimitiveType::FirstPerson;
+	bIsFirstPersonWorldSpaceRepresentation = InComponent->FirstPersonPrimitiveType == EFirstPersonPrimitiveType::WorldSpaceRepresentation;
+	bLumenHeightfield = InComponent->bLumenHeightfield;
+
+	bIsVisible = InComponent->IsVisible();
+	bIsVisibleEditor = InComponent->GetVisibleFlag();
+	bSelected = InComponent->IsSelected();
+	bIndividuallySelected = InComponent->IsComponentIndividuallySelected();
+	bShouldRenderSelected = InComponent->ShouldRenderSelected();
+	bCollisionEnabled = InComponent->IsCollisionEnabled();
+
+	if (const AActor* ActorOwner = InComponent->GetOwner())
+	{
+		bIsHidden = ActorOwner->IsHidden();
+#if WITH_EDITOR
+		bIsHiddenEd = ActorOwner->IsHiddenEd();
+		bIsOwnedByFoliage = FFoliageHelper::IsOwnedByFoliage(ActorOwner);
+#endif
+
+		if (bOnlyOwnerSee || bOwnerNoSee || bUseViewOwnerDepthPriorityGroup || bIsFirstPersonWorldSpaceRepresentation)
+		{
+			// Make a list of the actors which directly or indirectly own the InComponent.
+			for (const AActor* CurrentOwner = ActorOwner; CurrentOwner; CurrentOwner = CurrentOwner->GetOwner())
+			{
+				ActorOwners.Add(CurrentOwner->GetUniqueID());
+			}
+		}
+
+		bIsOwnerEditorOnly = InComponent->GetOwner()->IsEditorOnly();
+	}
+	bSupportsWorldPositionOffsetVelocity = InComponent->SupportsWorldPositionOffsetVelocity();
+	bIsInstancedStaticMesh = Cast<UInstancedStaticMeshComponent>(InComponent) != nullptr;
+
+	Mobility = InComponent->Mobility;;
+	TranslucencySortPriority = InComponent->TranslucencySortPriority;
+	TranslucencySortDistanceOffset = InComponent->TranslucencySortDistanceOffset;
+	LightmapType = InComponent->GetLightmapType();
+	ViewOwnerDepthPriorityGroup = InComponent->ViewOwnerDepthPriorityGroup;
+	CustomDepthStencilValue = InComponent->CustomDepthStencilValue;
+	CustomDepthStencilWriteMask = InComponent->CustomDepthStencilWriteMask;
+	LightingChannels = InComponent->LightingChannels;
+	RayTracingGroupCullingPriority = InComponent->RayTracingGroupCullingPriority;
+	IndirectLightingCacheQuality = InComponent->IndirectLightingCacheQuality;
+	ShadowCacheInvalidationBehavior = InComponent->ShadowCacheInvalidationBehavior;
+	DepthPriorityGroup = InComponent->GetStaticDepthPriorityGroup();
+
+	VirtualTextureLodBias = InComponent->VirtualTextureLodBias;
+	VirtualTextureCullMips = InComponent->VirtualTextureCullMips;
+	VirtualTextureMinCoverage = InComponent->VirtualTextureMinCoverage;
+	ComponentId = InComponent->GetPrimitiveSceneId();
+	VisibilityId = InComponent->VisibilityId;
+	CachedMaxDrawDistance = InComponent->CachedMaxDrawDistance;
+	MinDrawDistance = InComponent->MinDrawDistance;
+	BoundsScale = InComponent->BoundsScale;
+	RayTracingGroupId = InComponent->GetRayTracingGroupId();
+
+	bHasStaticLighting = InComponent->HasStaticLighting();
+	bHasValidSettingsForStaticLighting = InComponent->HasValidSettingsForStaticLighting(false);
+	bIsPrecomputedLightingValid = InComponent->IsPrecomputedLightingValid();
+	bShadowIndirectOnly = InComponent->GetShadowIndirectOnly();
+
+	Component = const_cast<UPrimitiveComponent*>(InComponent);
+	Owner = InComponent->GetOwner();
+
+#if !WITH_STATE_STREAM
+	World = InComponent->GetWorld();
+#endif
+
+	CustomPrimitiveData = &InComponent->GetCustomPrimitiveData();
+	Scene = InComponent->GetScene();
+	PrimitiveComponentInterface = InComponent->GetPrimitiveComponentInterface();
+
+	if (Scene)
+	{
+		FeatureLevel = Scene->GetFeatureLevel();
+	}
+
+#if WITH_EDITOR
+	HiddenEditorViews = InComponent->GetHiddenEditorViews();
+	OverlayColor = InComponent->OverlayColor;
+#endif
+	bShouldRenderProxyFallbackToDefaultMaterial = InComponent->ShouldRenderProxyFallbackToDefaultMaterial();
+
+	AdditionalStatObjectPtr = InComponent->AdditionalStatObject();
+	StatId = AdditionalStatObjectPtr ? AdditionalStatObjectPtr->GetStatID(true) : InComponent->GetStatID(true);
+
+	TArray<URuntimeVirtualTexture*> const& VirtualTextures = InComponent->GetRuntimeVirtualTextures();
+	RuntimeVirtualTextures = MakeArrayView(const_cast<URuntimeVirtualTexture**>(VirtualTextures.GetData()), VirtualTextures.Num());
+	VirtualTextureRenderPassType = InComponent->GetVirtualTextureRenderPassType();
+	VirtualTextureMainPassMaxDrawDistance = InComponent->GetVirtualTextureMainPassMaxDrawDistance();
+
+#if MESH_DRAW_COMMAND_STATS
+	MeshDrawCommandStatsCategory = InComponent->GetMeshDrawCommandStatsCategory();
+#endif
+}
+
+void FPrimitiveSceneProxyDesc::GetUsedMaterials(TArray<UMaterialInterface*>& OutMaterials, bool bGetDebugMaterials) const
+{
+	// Only UPrimitiveComponent should rely on this method 
+	const UPrimitiveComponent* AsComponent = Cast<UPrimitiveComponent>(Component);
+	check(AsComponent);
+
+	return AsComponent->GetUsedMaterials(OutMaterials, bGetDebugMaterials);
+}
diff --git a/Engine/Source/Runtime/Engine/Public/PrimitiveSceneProxyDesc.h b/Engine/Source/Runtime/Engine/Public/PrimitiveSceneProxyDesc.h
--- a/Engine/Source/Runtime/Engine/Public/PrimitiveSceneProxyDesc.h
+++ b/Engine/Source/Runtime/Engine/Public/PrimitiveSceneProxyDesc.h
@@ -10,83 +10,29 @@
 
 struct FPrimitiveSceneProxyDesc
 {	
-	FPrimitiveSceneProxyDesc()
-	{
-		CastShadow = false;
-		bReceivesDecals = true;
-		bOnlyOwnerSee = false;
-		bOwnerNoSee = false;
-		bLevelInstanceEditingState = false;
-		bUseViewOwnerDepthPriorityGroup = false;
-		bVisibleInReflectionCaptures = true;
-		bVisibleInRealTimeSkyCaptures = true;
-		bVisibleInRayTracing = true;
-		bRenderInDepthPass = true;
-		bRenderInMainPass = true;
-		bTreatAsBackgroundForOcclusion = false;
-		bCastDynamicShadow = true;
-		bCastStaticShadow = true;
-		bEmissiveLightSource = false;
-		bAffectDynamicIndirectLighting = true;
-		bAffectIndirectLightingWhileHidden = false;
-		bAffectDistanceFieldLighting = true;
-		bCastVolumetricTranslucentShadow = false;
-		bCastContactShadow = true;
-		bCastHiddenShadow = false;
-		bCastShadowAsTwoSided = false;
-		bSelfShadowOnly = false;
-		bCastInsetShadow = false;
-		bCastCinematicShadow = false;
-		bCastFarShadow = false;
-		bLightAttachmentsAsGroup = false;
-		bSingleSampleShadowFromStationaryLights = false;
-		bUseAsOccluder = false;
-		bSelectable = true;
-		bHasPerInstanceHitProxies = false;
-		bUseEditorCompositing = false;
-		bIsBeingMovedByEditor = false;
-		bReceiveMobileCSMShadows = true;
-		bRenderCustomDepth = false;
-		bVisibleInSceneCaptureOnly = false;
-		bHiddenInSceneCapture = false;
-		bForceMipStreaming = false;
-		bRayTracingFarField = false;
-		bIsVisible = true;
-		bIsVisibleEditor = true;
-		bSelected = false;
-		bIndividuallySelected = false;
-		bCollisionEnabled = false;
-		bIsHidden = false;
-		bIsHiddenEd = false;
-		bSupportsWorldPositionOffsetVelocity = true;
-		bIsOwnerEditorOnly = false;
-		bIsInstancedStaticMesh = false;
-		bHoldout = false;
-		bIsFirstPerson = false;
-		bIsFirstPersonWorldSpaceRepresentation = false;
-		bLumenHeightfield = false;
-
-		bHasStaticLighting = false;
-		bHasValidSettingsForStaticLighting = false;
-		bIsPrecomputedLightingValid = false;
-		bShadowIndirectOnly = false;
-		bShouldRenderProxyFallbackToDefaultMaterial = false;
-		bShouldRenderSelected = false;
-		bWantsEditorEffects = false;
-
-#if WITH_EDITOR
-		bIsOwnedByFoliage = false;
-#endif
-	}
-
+	ENGINE_API FPrimitiveSceneProxyDesc();
 	ENGINE_API FPrimitiveSceneProxyDesc(const UPrimitiveComponent*);
+	ENGINE_API virtual ~FPrimitiveSceneProxyDesc();
 
 	ENGINE_API void InitializeFromPrimitiveComponent(const UPrimitiveComponent*);
-	
-	UE_DEPRECATED(5.5, "Use InitializeFromPrimitiveComponent instead.")
-	void InitializeFrom(const UPrimitiveComponent* InComponent) { InitializeFromPrimitiveComponent(InComponent); }
 
-	virtual ~FPrimitiveSceneProxyDesc() = default;
+	UObject* Component = nullptr;
+	UObject* Owner = nullptr;
+	UWorld* World = nullptr;
+
+	TArrayView<URuntimeVirtualTexture*>  RuntimeVirtualTextures;
+
+	// Only used by actors for now, explicitly intended to be moved to the FPrimitiveSceneProxy
+	mutable TArray<uint32> ActorOwners;
+	const FCustomPrimitiveData* CustomPrimitiveData = nullptr;
+	FSceneInterface* Scene = nullptr;
+	IPrimitiveComponent* PrimitiveComponentInterface = nullptr;
+
+	const UObject* AdditionalStatObjectPtr = nullptr;
+
+	uint64 HiddenEditorViews = 0;
+
+	TStatId StatId;
 
 	uint32 CastShadow : 1;
 	uint32 bReceivesDecals : 1;
@@ -154,25 +100,28 @@ struct FPrimitiveSceneProxyDesc
 	uint32 bIsOwnedByFoliage:1;
 #endif
 
-
-
-	TEnumAsByte<EComponentMobility::Type> Mobility = EComponentMobility::Movable;
 	int32 TranslucencySortPriority = 0;
 	float TranslucencySortDistanceOffset = 0.0f;
-	ELightmapType LightmapType = ELightmapType::Default;
-	TEnumAsByte<enum ESceneDepthPriorityGroup> ViewOwnerDepthPriorityGroup = ESceneDepthPriorityGroup::SDPG_World;
 	int32 CustomDepthStencilValue = 0;
-	ERendererStencilMask CustomDepthStencilWriteMask = ERendererStencilMask::ERSM_Default;
+
 	FLightingChannels LightingChannels;
 	ERayTracingGroupCullingPriority RayTracingGroupCullingPriority = ERayTracingGroupCullingPriority::CP_4_DEFAULT;
-	TEnumAsByte<EIndirectLightingCacheQuality> IndirectLightingCacheQuality = EIndirectLightingCacheQuality::ILCQ_Point;	
 	EShadowCacheInvalidationBehavior ShadowCacheInvalidationBehavior = EShadowCacheInvalidationBehavior::Auto;
 
-	TEnumAsByte<enum ESceneDepthPriorityGroup> DepthPriorityGroup = ESceneDepthPriorityGroup::SDPG_World;
+	TEnumAsByte<ESceneDepthPriorityGroup> ViewOwnerDepthPriorityGroup = ESceneDepthPriorityGroup::SDPG_World;
+	TEnumAsByte<EComponentMobility::Type> Mobility = EComponentMobility::Movable;
+	TEnumAsByte<EIndirectLightingCacheQuality> IndirectLightingCacheQuality = EIndirectLightingCacheQuality::ILCQ_Point;
+	TEnumAsByte<ESceneDepthPriorityGroup> DepthPriorityGroup = ESceneDepthPriorityGroup::SDPG_World;
 	
+	ELightmapType LightmapType = ELightmapType::Default;
+	ERendererStencilMask CustomDepthStencilWriteMask = ERendererStencilMask::ERSM_Default;
+	ERuntimeVirtualTextureMainPassType VirtualTextureRenderPassType = ERuntimeVirtualTextureMainPassType::Exclusive;
 	int8 VirtualTextureLodBias = 0;
-	int32 VirtualTextureCullMips = 0;
 	int8 VirtualTextureMinCoverage = 0;
+
+	int32 VirtualTextureCullMips = 0;
+	float VirtualTextureMainPassMaxDrawDistance = 0.0f;
+
 	FPrimitiveComponentId ComponentId;
 	int32 VisibilityId = 0;
 	float CachedMaxDrawDistance = 0.0f;
@@ -181,100 +130,208 @@ struct FPrimitiveSceneProxyDesc
 	int32 RayTracingGroupId = FPrimitiveSceneProxy::InvalidRayTracingGroupId;
 
 	ERHIFeatureLevel::Type FeatureLevel;
-	
-	UObject* Component = nullptr; 
-	UObject* Owner = nullptr;
-	UWorld* World = nullptr;	
-
-	// Only used by actors for now, explicitly intended to be moved to the FPrimitiveSceneProxy
-	mutable TArray<uint32> ActorOwners;	
-	const FCustomPrimitiveData* CustomPrimitiveData = nullptr;
-	FSceneInterface* Scene = nullptr;
-	IPrimitiveComponent* PrimitiveComponentInterface = nullptr;
 
-	uint64 HiddenEditorViews = 0;	
+#if MESH_DRAW_COMMAND_STATS
+	FName MeshDrawCommandStatsCategory;
+#endif
 
 #if WITH_EDITOR
 	FColor OverlayColor = FColor(EForceInit::ForceInitToZero);
 #endif
-#if MESH_DRAW_COMMAND_STATS
-	FName MeshDrawCommandStatsCategory;
-	FName GetMeshDrawCommandStatsCategory() const { return MeshDrawCommandStatsCategory; }
-#endif //!MESH_DRAW_COMMAND_STATS
-
-
-	const FCustomPrimitiveData& GetCustomPrimitiveData() const { check(CustomPrimitiveData); return *CustomPrimitiveData; }
-	bool IsVisible()  const { return bIsVisible; }
-	bool IsVisibleEditor() const { return bIsVisibleEditor; }
-	bool ShouldRenderSelected() const { return bShouldRenderSelected; }
-	bool IsComponentIndividuallySelected() const { return bIndividuallySelected; }
-	ESceneDepthPriorityGroup GetStaticDepthPriorityGroup() const { return DepthPriorityGroup; }
-	bool HasStaticLighting() const { return bHasStaticLighting; } 
-	bool IsCollisionEnabled() const { return bCollisionEnabled; }
-	bool IsPrecomputedLightingValid() const { return false; }
-	bool HasValidSettingsForStaticLighting() const { return bHasValidSettingsForStaticLighting; } 
-	bool GetShadowIndirectOnly() const { return bShadowIndirectOnly; }
-	int32 GetRayTracingGroupId() const { return RayTracingGroupId; }
-	FSceneInterface* GetScene()  const { return Scene; }
-	bool GetLevelInstanceEditingState() const { return bLevelInstanceEditingState; }
-
-	bool IsFirstPersonRelevant() const { return bIsFirstPerson || bIsFirstPersonWorldSpaceRepresentation; }
 
-	UE_DEPRECATED(5.6, "Use FPrimitiveSceneProxyDesc::IsFirstPersonRelevant() instead.")
-	bool IsStaticMeshFirstPersonRelevant() const { return IsFirstPersonRelevant(); }
+	bool IsVisible() const
+	{
+		return bIsVisible;
+	}
 
-	UObject* GetOwner() const { return Owner; }	
-	template<class T>
-	T* GetOwner() const { return Cast<T>(Owner); }	
+	bool IsVisibleEditor() const
+	{
+		return bIsVisibleEditor;
+	}
 
-	ULevel* GetLevel() const { return  Owner ? Owner->GetTypedOuter<ULevel>() : nullptr; }
+	bool ShouldRenderSelected() const
+	{
+		return bShouldRenderSelected;
+	}
 
-	UE_DEPRECATED(5.6, "Use GetLevel instead.")
-	ULevel* GetComponentLevel() const { return GetLevel(); }
+	bool IsComponentIndividuallySelected() const
+	{
+		return bIndividuallySelected;
+	}
 
-	FString GetPathName() const { return Component->GetPathName(); }
-	
-	bool IsHidden() const { return bIsHidden; }
-	bool IsOwnerEditorOnly() const { return bIsOwnerEditorOnly; }
+	ESceneDepthPriorityGroup GetStaticDepthPriorityGroup() const
+	{
+		return DepthPriorityGroup;
+	}
+
+	bool HasStaticLighting() const
+	{
+		return bHasStaticLighting;
+	}
+
+	bool IsCollisionEnabled() const
+	{
+		return bCollisionEnabled;
+	}
+
+	bool IsPrecomputedLightingValid() const
+	{
+		return false;
+	}
+
+	bool HasValidSettingsForStaticLighting() const
+	{
+		return bHasValidSettingsForStaticLighting;
+	}
+
+	bool GetShadowIndirectOnly() const
+	{
+		return bShadowIndirectOnly;
+	}
+
+	bool GetLevelInstanceEditingState() const
+	{
+		return bLevelInstanceEditingState;
+	}
+
+	bool IsHidden() const
+	{
+		return bIsHidden;
+	}
+
+	bool IsOwnerEditorOnly() const
+	{
+		return bIsOwnerEditorOnly;
+	}
+
+	bool ShouldRenderProxyFallbackToDefaultMaterial() const
+	{
+		return bShouldRenderProxyFallbackToDefaultMaterial;
+	}
+
+	bool SupportsWorldPositionOffsetVelocity() const
+	{
+		return bSupportsWorldPositionOffsetVelocity;
+	}
+
+	bool IsFirstPersonRelevant() const
+	{
+		return bIsFirstPerson || bIsFirstPersonWorldSpaceRepresentation;
+	}
+
+	int32 GetRayTracingGroupId() const
+	{
+		return RayTracingGroupId;
+	}
+
+	FSceneInterface* GetScene() const
+	{
+		return Scene;
+	}
+
+	const FCustomPrimitiveData& GetCustomPrimitiveData() const
+	{
+		check(CustomPrimitiveData);
+		return *CustomPrimitiveData;
+	}
+
+	TArrayView<URuntimeVirtualTexture*> GetRuntimeVirtualTextures() const
+	{
+		return RuntimeVirtualTextures;
+	}
+
+	ERuntimeVirtualTextureMainPassType GetVirtualTextureRenderPassType() const
+	{
+		return VirtualTextureRenderPassType;
+	}
+
+	float GetVirtualTextureMainPassMaxDrawDistance() const
+	{
+		return VirtualTextureMainPassMaxDrawDistance;
+	}
 
-#if WITH_EDITOR
-	bool IsHiddenEd() const { return bIsHiddenEd; }
-	uint64 GetHiddenEditorViews() const { return HiddenEditorViews; }
-	bool IsOwnedByFoliage() const { return bIsOwnedByFoliage; }	
-#endif
-	
-	const UObject* AdditionalStatObjectPtr = nullptr;
 	const UObject* AdditionalStatObject() const
 	{
 		return AdditionalStatObjectPtr;
 	}
 
-	TStatId StatId;
 	TStatId GetStatID(bool bForDeferredUse = false) const
 	{
 		return StatId;
 	}
 
-	TArrayView<URuntimeVirtualTexture*>  RuntimeVirtualTextures;
-	ERuntimeVirtualTextureMainPassType VirtualTextureRenderPassType = ERuntimeVirtualTextureMainPassType::Exclusive;
+	IPrimitiveComponent* GetPrimitiveComponentInterface() const
+	{
+		return PrimitiveComponentInterface;
+	}
 
-	TArrayView<URuntimeVirtualTexture*> const& GetRuntimeVirtualTextures() const { return RuntimeVirtualTextures; }
-	ERuntimeVirtualTextureMainPassType GetVirtualTextureRenderPassType() const { return VirtualTextureRenderPassType; }
-	float VirtualTextureMainPassMaxDrawDistance = 0.0f;
-	float GetVirtualTextureMainPassMaxDrawDistance() const { return VirtualTextureMainPassMaxDrawDistance; }
+	UWorld* GetWorld() const
+	{
+		check(World);
+		return World;
+	}
+
+	FPrimitiveComponentId GetPrimitiveSceneId() const
+	{
+		return ComponentId;
+	}
+
+	UObject* GetOwner() const
+	{
+		return Owner;
+	}
+
+	template<class T>
+	T* GetOwner() const
+	{
+		return Cast<T>(GetOwner());
+	}
+
+	ULevel* GetLevel() const
+	{
+		return  Owner ? Owner->GetTypedOuter<ULevel>() : nullptr;
+	}
+
+	FString GetPathName() const
+	{
+		return Component->GetPathName();
+	}
 	
-	bool ShouldRenderProxyFallbackToDefaultMaterial() const { return bShouldRenderProxyFallbackToDefaultMaterial; }
-	ENGINE_API virtual void GetUsedMaterials(TArray<UMaterialInterface*>& OutMaterials, bool bGetDebugMaterials = false) const;
+#if WITH_EDITOR
+	bool IsHiddenEd() const
+	{
+		return bIsHiddenEd;
+	}
 
-	bool SupportsWorldPositionOffsetVelocity() const { return bSupportsWorldPositionOffsetVelocity; }
+	uint64 GetHiddenEditorViews() const
+	{
+		return HiddenEditorViews;
+	}
 
-	IPrimitiveComponent* GetPrimitiveComponentInterface() const { return PrimitiveComponentInterface; }	
+	bool IsOwnedByFoliage() const
+	{
+		return bIsOwnedByFoliage;
+	}
+#endif
 
-	UWorld* GetWorld() const 
+#if MESH_DRAW_COMMAND_STATS
+	FName GetMeshDrawCommandStatsCategory() const
 	{
-		check(World);
-		return World;
+		return MeshDrawCommandStatsCategory;
+	}
+#endif
+
+	ENGINE_API virtual void GetUsedMaterials(TArray<UMaterialInterface*>& OutMaterials, bool bGetDebugMaterials = false) const;
+
+	UE_DEPRECATED(5.6, "Use FPrimitiveSceneProxyDesc::IsFirstPersonRelevant() instead.")
+	bool IsStaticMeshFirstPersonRelevant() const
+	{
+		return IsFirstPersonRelevant();
 	}
 
-	FPrimitiveComponentId GetPrimitiveSceneId() const { return ComponentId; }
+	UE_DEPRECATED(5.6, "Use GetLevel instead.")
+	ULevel* GetComponentLevel() const
+	{
+		return GetLevel();
+	}
 };
diff --git a/Engine/Source/Runtime/Engine/Public/StaticMeshSceneProxyDesc.h b/Engine/Source/Runtime/Engine/Public/StaticMeshSceneProxyDesc.h
--- a/Engine/Source/Runtime/Engine/Public/StaticMeshSceneProxyDesc.h
+++ b/Engine/Source/Runtime/Engine/Public/StaticMeshSceneProxyDesc.h
@@ -17,17 +17,37 @@ struct FStaticMeshSceneProxyDesc : public FPrimitiveSceneProxyDesc
 
 	ENGINE_API void InitializeFromStaticMeshComponent(const UStaticMeshComponent*);
 
-	UStaticMesh* StaticMesh = nullptr;
 	TArrayView<TObjectPtr<UMaterialInterface>>	OverrideMaterials;
-	TObjectPtr<class UMaterialInterface> OverlayMaterial;
-	float OverlayMaterialMaxDrawDistance = 0.0f;
 	TArray<TObjectPtr<UMaterialInterface>> MaterialSlotsOverlayMaterial;
+	TArray<UMaterialCacheVirtualTexture*> MaterialCacheTextures;
+
+	TArrayView<FStaticMeshComponentLODInfo> LODData;
+
+	UStaticMesh* StaticMesh = nullptr;
+	UBodySetup* BodySetup = nullptr;
+	UObject* LODParentPrimitive = nullptr;
+	UTexture* MeshPaintTexture = nullptr;
+	const Nanite::FResources* NaniteResources = nullptr;
+
+	TObjectPtr<UMaterialInterface> OverlayMaterial;
+	FMaterialRelevance MaterialRelevance;
+
+	float OverlayMaterialMaxDrawDistance = 0.0f;
+	float NanitePixelProgrammableDistance = 0.0f;
+
+	float DistanceFieldSelfShadowBias = 0;
+	float DistanceFieldIndirectShadowMinVisibility = 0.1f;
 
 	int32 ForcedLodModel = 0;
 	int32 MinLOD = 0;
 	int32 WorldPositionOffsetDisableDistance = 0;
-	float NanitePixelProgrammableDistance = 0.0f;
-	
+	int32 StaticLightMapResolution = 0;
+	int32 MeshPaintTextureCoordinateIndex = 0;
+	int32 MaterialCacheTextureCoordinateIndex = 0;
+
+	//@todo: share color selection logic according to mobility with USMC?
+	FColor	WireframeColor = FColor(0, 255, 255, 255);
+
 	uint32 bReverseCulling : 1 = false;
 #if STATICMESH_ENABLE_DEBUG_RENDERING
 	uint32 bDrawMeshCollisionIfComplex : 1;
@@ -38,20 +58,18 @@ struct FStaticMeshSceneProxyDesc : public FPrimitiveSceneProxyDesc
 
 	uint32 bCastDistanceFieldIndirectShadow : 1;
 	uint32 bOverrideDistanceFieldSelfShadowBias : 1;
-	uint32 bEvaluateWorldPositionOffsetInRayTracing :1;
-	uint32 bSortTriangles :1;
+	uint32 bEvaluateWorldPositionOffsetInRayTracing : 1;
+	uint32 bSortTriangles : 1;
 
 	uint32 bDisplayNaniteFallbackMesh : 1;
 	uint32 bDisallowNanite : 1;
 	uint32 bForceDisableNanite : 1;
 	uint32 bForceNaniteForMasked : 1;
 
-	uint32 bUseProvidedMaterialRelevance:1;
-	uint32 bUseProvidedCollisionResponseContainer:1;
+	uint32 bUseProvidedMaterialRelevance : 1;
+	uint32 bUseProvidedCollisionResponseContainer : 1;
 
-	float DistanceFieldSelfShadowBias = 0;
-	float DistanceFieldIndirectShadowMinVisibility = 0.1f;
-	int32 StaticLightMapResolution = 0;
+	TOptional<FCollisionResponseContainer> CollisionResponseContainer;
 
 	ELightmapType LightmapType = ELightmapType::Default;
 
@@ -66,26 +84,6 @@ struct FStaticMeshSceneProxyDesc : public FPrimitiveSceneProxyDesc
 	float TextureStreamingTransformScale = 1.0f;
 #endif
 
-	const Nanite::FResources* NaniteResources = nullptr;
-
-	TArrayView<FStaticMeshComponentLODInfo> LODData;
-
-	FMaterialRelevance MaterialRelevance;
-
-	UTexture* MeshPaintTexture = nullptr;
-	int32 MeshPaintTextureCoordinateIndex = 0;
-
-	TArray<UMaterialCacheVirtualTexture*> MaterialCacheTextures;
-	int32 MaterialCacheTextureCoordinateIndex = 0;
-
-	UBodySetup* BodySetup = nullptr;
-	UObject* LODParentPrimitive = nullptr;
-
-	TOptional<FCollisionResponseContainer> CollisionResponseContainer;
-
-	//@todo: share color selection logic according to mobility with USMC?
-	FColor	WireframeColor = FColor(0, 255, 255, 255);
-
 	bool IsReverseCulling() const
 	{
 		return bReverseCulling;
