Cleaning up some ContentStreaming code.

diff --git a/Engine/Source/Runtime/Engine/Private/ContentStreaming.cpp b/Engine/Source/Runtime/Engine/Private/ContentStreaming.cpp
--- a/Engine/Source/Runtime/Engine/Private/ContentStreaming.cpp
+++ b/Engine/Source/Runtime/Engine/Private/ContentStreaming.cpp
@@ -142,11 +142,6 @@ void TrackRenderAssetInit()
 	}
 }
 
-void TrackTextureInit()
-{
-	TrackRenderAssetInit();
-}
-
 /**
  * Adds a (partial) texture/mesh name to track in the streaming system and updates the .ini setting.
  *
@@ -173,11 +168,6 @@ bool TrackRenderAsset( const FString& AssetName )
 	return false;
 }
 
-bool TrackTexture(const FString& TextureName)
-{
-	return TrackRenderAsset(TextureName);
-}
-
 /**
  * Removes a texture/mesh name from being tracked in the streaming system and updates the .ini setting.
  * The name must match an existing tracking name, but isn't case-sensitive.
@@ -209,11 +199,6 @@ bool UntrackRenderAsset( const FString& AssetName )
 	return false;
 }
 
-bool UntrackTexture(const FString& TextureName)
-{
-	return UntrackRenderAsset(TextureName);
-}
-
 /**
  * Lists all currently tracked texture/mesh names in the specified log.
  *
@@ -231,11 +216,6 @@ void ListTrackedRenderAssets( FOutputDevice& Ar, int32 NumAssets )
 	Ar.Logf( TEXT("%d render assets are being tracked."), NumAssets );
 }
 
-void ListTrackedTextures(FOutputDevice& Ar, int32 NumTextures)
-{
-	ListTrackedRenderAssets(Ar, NumTextures);
-}
-
 /**
  * Checks a texture/mesh and tracks it if its name contains any of the tracked render asset names (GTrackedRenderAssetNames).
  *
@@ -325,42 +305,22 @@ bool TrackRenderAssetEvent(FStreamingRenderAsset* StreamingRenderAsset, UStreama
 	return false;
 }
 
-bool TrackTextureEvent(FStreamingRenderAsset* StreamingTexture, UStreamableRenderAsset* Texture, bool bForceMipLevelsToBeResident, const FRenderAssetStreamingManager* Manager)
-{
-	return TrackRenderAssetEvent(StreamingTexture, Texture, bForceMipLevelsToBeResident, Manager);
-}
 #else
 bool TrackRenderAsset(const FString& AssetName)
 {
 	return false;
 }
-bool TrackTexture( const FString& TextureName )
-{
-	return TrackRenderAsset(TextureName);
-}
 bool UntrackRenderAsset(const FString& AssetName)
 {
 	return false;
 }
-bool UntrackTexture( const FString& TextureName )
-{
-	return UntrackRenderAsset(TextureName);
-}
 void ListTrackedRenderAssets(FOutputDevice& Ar, int32 NumAssets)
 {
 }
-void ListTrackedTextures( FOutputDevice& Ar, int32 NumTextures )
-{
-	ListTrackedRenderAssets(Ar, NumTextures);
-}
 bool TrackRenderAssetEvent(FStreamingRenderAsset* StreamingRenderAsset, UStreamableRenderAsset* RenderAsset, bool bForceMipLevelsToBeResident, const FRenderAssetStreamingManager* Manager)
 {
 	return false;
 }
-bool TrackTextureEvent( FStreamingRenderAsset* StreamingTexture, UStreamableRenderAsset* Texture, bool bForceMipLevelsToBeResident, const FRenderAssetStreamingManager* Manager )
-{
-	return TrackRenderAssetEvent(StreamingTexture, Texture, bForceMipLevelsToBeResident, Manager);
-}
 #endif
 
 /*-----------------------------------------------------------------------------
@@ -1620,4 +1580,4 @@ FAudioChunkHandle IAudioStreamingManager::BuildChunkHandle(const uint8* InData,
 void IAudioStreamingManager::LogWarning()
 {
 	UE_LOG(LogContentStreaming, Warning, TEXT("This function is now deprecated on IAudioStreamingManager and no-ops."));
-}
\ No newline at end of file
+}
diff --git a/Engine/Source/Runtime/Engine/Private/Texture2D.cpp b/Engine/Source/Runtime/Engine/Private/Texture2D.cpp
--- a/Engine/Source/Runtime/Engine/Private/Texture2D.cpp
+++ b/Engine/Source/Runtime/Engine/Private/Texture2D.cpp
@@ -58,9 +58,6 @@
 
 #include UE_INLINE_GENERATED_CPP_BY_NAME(Texture2D)
 
-#if WITH_EDITORONLY_DATA
-#endif
-
 #define LOCTEXT_NAMESPACE "UTexture2D"
 
 UTexture2D::UTexture2D(const FObjectInitializer& ObjectInitializer)
@@ -140,7 +137,7 @@ struct FStreamingRenderAsset;
 struct FRenderAssetStreamingManager;
 
 /** Turn on ENABLE_RENDER_ASSET_TRACKING in ContentStreaming.cpp and setup GTrackedTextures to track specific textures/meshes through the streaming system. */
-extern bool TrackTextureEvent( FStreamingRenderAsset* StreamingTexture, UStreamableRenderAsset* Texture, bool bForceMipLevelsToBeResident, const FRenderAssetStreamingManager* Manager );
+extern bool TrackRenderAssetEvent( FStreamingRenderAsset* StreamingTexture, UStreamableRenderAsset* Texture, bool bForceMipLevelsToBeResident, const FRenderAssetStreamingManager* Manager );
 
 /*-----------------------------------------------------------------------------
 	FTexture2DMipMap
@@ -809,7 +806,7 @@ void UTexture2D::BeginDestroy()
 	// Route BeginDestroy.
 	Super::BeginDestroy();
 
-	TrackTextureEvent( NULL, this, false, 0 );
+	TrackRenderAssetEvent(nullptr, this, false, 0);
 }
 
 FString UTexture2D::GetDesc()
diff --git a/Engine/Source/Runtime/Engine/Public/ContentStreaming.h b/Engine/Source/Runtime/Engine/Public/ContentStreaming.h
--- a/Engine/Source/Runtime/Engine/Public/ContentStreaming.h
+++ b/Engine/Source/Runtime/Engine/Public/ContentStreaming.h
@@ -252,12 +252,6 @@ struct IStreamingManager
 	 */
 	ENGINE_API void AddViewLocation(const FVector& Location, float BoostFactor = 1.0f, bool bOverrideLocation = false, float Duration = 0.0f);
 
-	UE_DEPRECATED(5.1, "This is deprecated to follow inclusive naming rules. Use AddViewLocation() instead.")
-	void AddViewSlaveLocation(const FVector& Location, float BoostFactor = 1.0f, bool bOverrideLocation = false, float Duration = 0.0f)
-	{
-		AddViewLocation(Location, BoostFactor, bOverrideLocation, Duration);
-	}
-
 	/** Don't stream world resources for the next NumFrames. */
 	virtual void SetDisregardWorldResourcesForFrames(int32 NumFrames) = 0;
 
@@ -606,46 +600,6 @@ protected:
 // DEPRECATED VIRTUAL INTERFACE
 private:
 	static void LogWarning();
-
-public:
-	/** Adds a new Sound Wave to the streaming manager. */
-	UE_DEPRECATED(5.5, "Use These methods are no longer used (removing the unused legacy streaming manager)")
-	virtual void AddStreamingSoundWave(const FSoundWaveProxyPtr& SoundWave) { LogWarning(); }
-
-	/** Removes a Sound Wave from the streaming manager. */
-	UE_DEPRECATED(5.5, "Use These methods are no longer used (removing the unused legacy streaming manager)")
-	virtual void RemoveStreamingSoundWave(const FSoundWaveProxyPtr& SoundWave) { LogWarning(); }
-
-	/** Adds the decoder to the streaming manager to prevent stream chunks from getting reaped from underneath it */
-	UE_DEPRECATED(5.5, "Use These methods are no longer used (removing the unused legacy streaming manager)")
-	virtual void AddDecoder(ICompressedAudioInfo* CompressedAudioInfo) { LogWarning(); }
-
-	/** Removes the decoder from the streaming manager. */
-	UE_DEPRECATED(5.5, "Use These methods are no longer used (removing the unused legacy streaming manager)")
-	virtual void RemoveDecoder(ICompressedAudioInfo* CompressedAudioInfo) { LogWarning(); }
-
-	/** Returns true if this is a Sound Wave that is managed by the streaming manager. */
-	UE_DEPRECATED(5.5, "Use These methods are no longer used (removing the unused legacy streaming manager)")
-	virtual bool IsManagedStreamingSoundWave(const FSoundWaveProxyPtr&  SoundWave) const { LogWarning(); return false; }
-
-	/** Returns true if this Sound Wave is currently streaming a chunk. */
-	UE_DEPRECATED(5.5, "Use These methods are no longer used (removing the unused legacy streaming manager)")
-	virtual bool IsStreamingInProgress(const FSoundWaveProxyPtr&  SoundWave) { LogWarning(); return false; }
-
-	UE_DEPRECATED(5.5, "Use These methods are no longer used (removing the unused legacy streaming manager)")
-	virtual bool CanCreateSoundSource(const FWaveInstance* WaveInstance) const { LogWarning(); return false; }
-
-	/** Adds a new Sound Source to the streaming manager. */
-	UE_DEPRECATED(5.5, "Use These methods are no longer used (removing the unused legacy streaming manager)")
-	virtual void AddStreamingSoundSource(FSoundSource* SoundSource) { LogWarning(); }
-
-	/** Removes a Sound Source from the streaming manager. */
-	UE_DEPRECATED(5.5, "Use These methods are no longer used (removing the unused legacy streaming manager)")
-	virtual void RemoveStreamingSoundSource(FSoundSource* SoundSource) { LogWarning(); }
-
-	/** Returns true if this is a streaming Sound Source that is managed by the streaming manager. */
-	UE_DEPRECATED(5.5, "Use These methods are no longer used (removing the unused legacy streaming manager)")
-	virtual bool IsManagedStreamingSoundSource(const FSoundSource* SoundSource) const { LogWarning(); return false; }
 };
 
 /**
