Make sure FeatureLevel is always initialized.

diff --git a/Engine/Source/Runtime/Engine/Private/PrimitiveSceneProxyDesc.cpp b/Engine/Source/Runtime/Engine/Private/PrimitiveSceneProxyDesc.cpp
--- a/Engine/Source/Runtime/Engine/Private/PrimitiveSceneProxyDesc.cpp
+++ b/Engine/Source/Runtime/Engine/Private/PrimitiveSceneProxyDesc.cpp
@@ -198,6 +198,10 @@ void FPrimitiveSceneProxyDesc::InitializeFromPrimitiveComponent(const UPrimitive
 	{
 		FeatureLevel = Scene->GetFeatureLevel();
 	}
+	else
+	{
+		FeatureLevel = ERHIFeatureLevel::Num;
+	}
 
 #if WITH_EDITOR
 	HiddenEditorViews = InComponent->GetHiddenEditorViews();
