Function abstract a function to simplify binding textures to new serialization stack

#rnx

diff --git a/Engine/Source/Runtime/Engine/Private/TextureDerivedData.cpp b/Engine/Source/Runtime/Engine/Private/TextureDerivedData.cpp
--- a/Engine/Source/Runtime/Engine/Private/TextureDerivedData.cpp
+++ b/Engine/Source/Runtime/Engine/Private/TextureDerivedData.cpp
@@ -4215,6 +4215,77 @@ void UTexture::CleanupCachedRunningPlatformData()
 	}
 }
 
+#if WITH_EDITOR
+static TArray<FTexturePlatformData*> GetPlatformDataToSerialize(UTexture& Texture, const ITargetPlatform& CookingTarget)
+{
+	if (Texture.GetOutermost()->bIsCookedForEditor)
+	{
+		// For cooked packages, simply grab the current platform data and serialize it
+		FTexturePlatformData** RunningPlatformDataPtr = Texture.GetRunningPlatformData();
+		if (RunningPlatformDataPtr && *RunningPlatformDataPtr)
+		{
+			return {*RunningPlatformDataPtr};	
+		}
+	}
+	
+	TMap<FString, FTexturePlatformData*>* CookedPlatformDataPtr = Texture.GetCookedPlatformData();
+	if (!Texture.CanBuildPlatformData() || !CookedPlatformDataPtr)
+	{
+		return {};
+	}
+
+	// Kick off builds for anything we don't have on hand already.
+	ETextureEncodeSpeed EncodeSpeed = Texture.GetDesiredEncodeSpeed();
+
+	TArray< TArray<FTextureBuildSettings> > BuildSettingsToCacheFetch;
+	TArray< TArray<FTextureBuildSettings> > BuildSettingsToCacheFetchOrBuild;
+	if (EncodeSpeed == ETextureEncodeSpeed::FinalIfAvailable)
+	{
+		FTextureBuildSettings BuildSettingsFetch;
+		GetTextureBuildSettings(Texture, CookingTarget.GetTextureLODSettings(), CookingTarget, ETextureEncodeSpeed::Final, BuildSettingsFetch, nullptr);
+		GetBuildSettingsPerFormat(Texture, BuildSettingsFetch, nullptr, &CookingTarget, ETextureEncodeSpeed::Final, BuildSettingsToCacheFetch, nullptr);
+
+		FTextureBuildSettings BuildSettingsFetchOrBuild;
+		GetTextureBuildSettings(Texture, CookingTarget.GetTextureLODSettings(), CookingTarget, ETextureEncodeSpeed::Fast, BuildSettingsFetchOrBuild, nullptr);
+		GetBuildSettingsPerFormat(Texture, BuildSettingsFetchOrBuild, nullptr, &CookingTarget, ETextureEncodeSpeed::Fast, BuildSettingsToCacheFetchOrBuild, nullptr);
+	}
+	else
+	{
+		FTextureBuildSettings BuildSettingsFetchOrBuild;
+		GetTextureBuildSettings(Texture, CookingTarget.GetTextureLODSettings(), CookingTarget, EncodeSpeed, BuildSettingsFetchOrBuild, nullptr);
+		GetBuildSettingsPerFormat(Texture, BuildSettingsFetchOrBuild, nullptr, &CookingTarget, EncodeSpeed, BuildSettingsToCacheFetchOrBuild, nullptr);
+	}
+		
+	TArray<FTexturePlatformData*> Out;
+	Out.SetNum(BuildSettingsToCacheFetchOrBuild.Num());
+	for (int32 SettingIndex = 0; SettingIndex < Out.Num(); SettingIndex++)
+	{
+		check(BuildSettingsToCacheFetchOrBuild[SettingIndex].Num() == Texture.Source.GetNumLayers());
+
+		// CookedPlatformData is keyed off of the fetchorbuild key.
+		FString DerivedDataKeyFetchOrBuild;
+		GetTextureDerivedDataKey(Texture, BuildSettingsToCacheFetchOrBuild[SettingIndex].GetData(), DerivedDataKeyFetchOrBuild);
+
+		FTexturePlatformData *PlatformDataPtr = (*CookedPlatformDataPtr).FindRef(DerivedDataKeyFetchOrBuild);
+		if (PlatformDataPtr == NULL)
+		{
+			PlatformDataPtr = new FTexturePlatformData();
+			PlatformDataPtr->Cache(Texture, 
+				BuildSettingsToCacheFetch.Num() ? BuildSettingsToCacheFetch[SettingIndex].GetData() : nullptr,
+				BuildSettingsToCacheFetchOrBuild[SettingIndex].GetData(), 
+				nullptr,
+				nullptr,
+				uint32(ETextureCacheFlags::InlineMips | ETextureCacheFlags::Async), 
+				nullptr);
+
+			CookedPlatformDataPtr->Add(DerivedDataKeyFetchOrBuild, PlatformDataPtr);
+		}
+		Out[SettingIndex] = PlatformDataPtr;
+	}
+	return Out;
+}
+#endif
+
 void UTexture::SerializeCookedPlatformData(FArchive& Ar, const bool bSerializeMipData)
 {
 	if (IsTemplate() )
@@ -4231,77 +4302,10 @@ void UTexture::SerializeCookedPlatformData(FArchive& Ar, const bool bSerializeMi
 	{
 		if (Ar.CookingTarget()->AllowAudioVisualData())
 		{
-			TArray<FTexturePlatformData*> PlatformDataToSerialize;
-
-			if (GetOutermost()->bIsCookedForEditor)
-			{
-				// For cooked packages, simply grab the current platform data and serialize it
-				FTexturePlatformData** RunningPlatformDataPtr = GetRunningPlatformData();
-				if (RunningPlatformDataPtr == nullptr)
-				{
-					return;
-				}
-				FTexturePlatformData* RunningPlatformData = *RunningPlatformDataPtr;
-				if (RunningPlatformData == nullptr)
-				{
-					return;
-				}
-				PlatformDataToSerialize.Add(RunningPlatformData);
-			}
-			else if (CanBuildPlatformData())
+			TArray<FTexturePlatformData*> PlatformDataToSerialize = GetPlatformDataToSerialize(*this, *Ar.CookingTarget());
+			if (PlatformDataToSerialize.IsEmpty())
 			{
-				TMap<FString, FTexturePlatformData*> *CookedPlatformDataPtr = GetCookedPlatformData();
-				if (CookedPlatformDataPtr == nullptr)
-				{
-					return;
-				}
-
-				// Kick off builds for anything we don't have on hand already.
-				ETextureEncodeSpeed EncodeSpeed = GetDesiredEncodeSpeed();
-
-				TArray< TArray<FTextureBuildSettings> > BuildSettingsToCacheFetch;
-				TArray< TArray<FTextureBuildSettings> > BuildSettingsToCacheFetchOrBuild;
-				if (EncodeSpeed == ETextureEncodeSpeed::FinalIfAvailable)
-				{
-					FTextureBuildSettings BuildSettingsFetch;
-					GetTextureBuildSettings(*this, Ar.CookingTarget()->GetTextureLODSettings(), *Ar.CookingTarget(), ETextureEncodeSpeed::Final, BuildSettingsFetch, nullptr);
-					GetBuildSettingsPerFormat(*this, BuildSettingsFetch, nullptr, Ar.CookingTarget(), ETextureEncodeSpeed::Final, BuildSettingsToCacheFetch, nullptr);
-
-					FTextureBuildSettings BuildSettingsFetchOrBuild;
-					GetTextureBuildSettings(*this, Ar.CookingTarget()->GetTextureLODSettings(), *Ar.CookingTarget(), ETextureEncodeSpeed::Fast, BuildSettingsFetchOrBuild, nullptr);
-					GetBuildSettingsPerFormat(*this, BuildSettingsFetchOrBuild, nullptr, Ar.CookingTarget(), ETextureEncodeSpeed::Fast, BuildSettingsToCacheFetchOrBuild, nullptr);
-				}
-				else
-				{
-					FTextureBuildSettings BuildSettingsFetchOrBuild;
-					GetTextureBuildSettings(*this, Ar.CookingTarget()->GetTextureLODSettings(), *Ar.CookingTarget(), EncodeSpeed, BuildSettingsFetchOrBuild, nullptr);
-					GetBuildSettingsPerFormat(*this, BuildSettingsFetchOrBuild, nullptr, Ar.CookingTarget(), EncodeSpeed, BuildSettingsToCacheFetchOrBuild, nullptr);
-				}
-
-				for (int32 SettingIndex = 0; SettingIndex < BuildSettingsToCacheFetchOrBuild.Num(); SettingIndex++)
-				{
-					check(BuildSettingsToCacheFetchOrBuild[SettingIndex].Num() == Source.GetNumLayers());
-
-					// CookedPlatformData is keyed off of the fetchorbuild key.
-					FString DerivedDataKeyFetchOrBuild;
-					GetTextureDerivedDataKey(*this, BuildSettingsToCacheFetchOrBuild[SettingIndex].GetData(), DerivedDataKeyFetchOrBuild);
-
-					FTexturePlatformData *PlatformDataPtr = (*CookedPlatformDataPtr).FindRef(DerivedDataKeyFetchOrBuild);
-					if (PlatformDataPtr == NULL)
-					{
-						PlatformDataPtr = new FTexturePlatformData();
-						PlatformDataPtr->Cache(*this, 
-							BuildSettingsToCacheFetch.Num() ? BuildSettingsToCacheFetch[SettingIndex].GetData() : nullptr,
-							BuildSettingsToCacheFetchOrBuild[SettingIndex].GetData(), 
-							nullptr,
-							nullptr,
-							uint32(ETextureCacheFlags::InlineMips | ETextureCacheFlags::Async), 
-							nullptr);
-
-						CookedPlatformDataPtr->Add(DerivedDataKeyFetchOrBuild, PlatformDataPtr);
-					}
-					PlatformDataToSerialize.Add(PlatformDataPtr);
-				}
+				return;
 			}
 
 			// set DidSerializeStreamingMipsForPlatform to false, then it will change to true if any SerializeCooked makes streaming mips
