[Engine] Call FTraceAuxiliary::WriteSnapshot instead of FTraceAuxiliary::SendSnapshot in FTraceService::OnSnapshotFile

diff --git a/Engine/Source/Runtime/Engine/Private/TraceService.cpp b/Engine/Source/Runtime/Engine/Private/TraceService.cpp
--- a/Engine/Source/Runtime/Engine/Private/TraceService.cpp
+++ b/Engine/Source/Runtime/Engine/Private/TraceService.cpp
@@ -130,7 +130,7 @@ void FTraceService::OnSnapshotSend(const FTraceControlSnapshotSend& Message, con
 
 void FTraceService::OnSnapshotFile(const FTraceControlSnapshotFile& Message, const TSharedRef<IMessageContext>& Context)
 {
-	FTraceAuxiliary::SendSnapshot(*Message.File);
+	FTraceAuxiliary::WriteSnapshot(*Message.File);
 }
 
 void FTraceService::OnPause(const FTraceControlPause& Message, const TSharedRef<IMessageContext>& Context)
