Added ToRotationVector, MakeFromRotationVector and GetShortestArcWith to Kismet Math Library

diff --git a/Engine/Source/Runtime/Engine/Classes/Kismet/KismetMathLibrary.h b/Engine/Source/Runtime/Engine/Classes/Kismet/KismetMathLibrary.h
--- a/Engine/Source/Runtime/Engine/Classes/Kismet/KismetMathLibrary.h
+++ b/Engine/Source/Runtime/Engine/Classes/Kismet/KismetMathLibrary.h
@@ -2813,6 +2813,10 @@ class UKismetMathLibrary : public UBlueprintFunctionLibrary
 	UFUNCTION(BlueprintCallable, meta = (DisplayName = "Ensure Shortest Arc To (Quat)", ScriptMethod = "EnsureShortestArcTo"), Category = "Math|Quat")
 	static UE_INL_API void Quat_EnforceShortestArcWith(UPARAM(ref) FQuat& A, const FQuat& B);
 
+	/** Returns the version of this quaternion that represents the shortest rotation angle to the other quaternion */
+	UFUNCTION(BlueprintPure, meta = (DisplayName = "Get Shortest Arc To (Quat)", ScriptMethod = "GetShortestArcTo"), Category = "Math|Quat")
+	static UE_INL_API FQuat Quat_GetShortestArcWith(const FQuat& A, const FQuat& B);
+
 	/**	Convert a Quaternion into floating-point Euler angles (in degrees). */
 	UFUNCTION(BlueprintPure, meta = (DisplayName = "Euler (Quat)", ScriptMethod = "Euler"), Category = "Math|Quat")
 	static UE_INL_API FVector Quat_Euler(const FQuat& Q);
@@ -2825,7 +2829,7 @@ class UKismetMathLibrary : public UBlueprintFunctionLibrary
 	UFUNCTION(BlueprintPure, meta = (DisplayName = "Exp (Quat)", ScriptMethod = "Exp"), Category = "Math|Quat")
 	static UE_INL_API FQuat Quat_Exp(const FQuat& Q);
 
-	/** Get the angle of this quaternion */
+	/** Get the angle of this quaternion in radians */
 	UFUNCTION(BlueprintPure, meta = (DisplayName = "Angle (Quat)", ScriptMethod = "GetAngle"), Category = "Math|Quat")
 	static UE_INL_API float Quat_GetAngle(const FQuat& Q);
 
@@ -2917,6 +2921,22 @@ class UKismetMathLibrary : public UBlueprintFunctionLibrary
 	UFUNCTION(BlueprintPure, meta=(DisplayName = "ToQuaternion (Rotator)", CompactNodeTitle = "->", ScriptMethod = "Quaternion", Keywords="cast convert", BlueprintAutocast), Category="Math|Conversions")
 	static UE_INL_API FQuat Conv_RotatorToQuaternion(FRotator InRot);
 
+	/**
+	 * Get the rotation vector corresponding to this quaternion. 
+	 * The direction of the vector represents the rotation axis, 
+	 * and the magnitude the angle in degrees.
+	 */
+	UFUNCTION(BlueprintPure, meta = (DisplayName = "To Rotation Vector (Quat)", ScriptMethod = "ToRotationVector"), Category = "Math|Quat")
+	static UE_INL_API FVector Quat_ToRotationVector(const FQuat& Q);
+
+	/**
+	 * Constructs a quaternion corresponding to the rotation vector. 
+	 * The direction of the vector represents the rotation axis, 
+	 * and the magnitude the angle in degrees.
+	 */
+	UFUNCTION(BlueprintPure, meta = (DisplayName = "Make from Rotation Vector (Quat)", ScriptMethod = "MakeFromRotationVector"), Category = "Math|Quat")
+	static UE_INL_API FQuat Quat_MakeFromRotationVector(const FVector& V);
+
 	/**
 	 * Get the length of the quaternion.
 	 *
diff --git a/Engine/Source/Runtime/Engine/Classes/Kismet/KismetMathLibrary.inl b/Engine/Source/Runtime/Engine/Classes/Kismet/KismetMathLibrary.inl
--- a/Engine/Source/Runtime/Engine/Classes/Kismet/KismetMathLibrary.inl
+++ b/Engine/Source/Runtime/Engine/Classes/Kismet/KismetMathLibrary.inl
@@ -3289,6 +3289,12 @@ void UKismetMathLibrary::Quat_EnforceShortestArcWith(FQuat& A, const FQuat& B)
 	A.EnforceShortestArcWith(B);
 }
 
+KISMET_MATH_INLINE
+FQuat UKismetMathLibrary::Quat_GetShortestArcWith(const FQuat& A, const FQuat& B)
+{
+	return A.GetShortestArcWith(B);
+}
+
 KISMET_MATH_INLINE
 FVector UKismetMathLibrary::Quat_Euler(const FQuat& Q)
 {
@@ -3400,6 +3406,18 @@ FRotator UKismetMathLibrary::Quat_Rotator(const FQuat& Q)
 	return Q.Rotator();
 }
 
+KISMET_MATH_INLINE
+FVector UKismetMathLibrary::Quat_ToRotationVector(const FQuat& Q)
+{
+	return FMath::RadiansToDegrees(Q.ToRotationVector());
+}
+
+KISMET_MATH_INLINE
+FQuat UKismetMathLibrary::Quat_MakeFromRotationVector(const FVector& V)
+{
+	return FQuat::MakeFromRotationVector(FMath::DegreesToRadians(V));
+}
+
 KISMET_MATH_INLINE
 FQuat UKismetMathLibrary::Conv_RotatorToQuaternion(FRotator InRot)
 {
