[Materials] Fix copy-and-paste bug in VT unpack encoding of HLSLMaterialTranslator.

#rnx

diff --git a/Engine/Source/Runtime/Engine/Private/Materials/HLSLMaterialTranslator.cpp b/Engine/Source/Runtime/Engine/Private/Materials/HLSLMaterialTranslator.cpp
--- a/Engine/Source/Runtime/Engine/Private/Materials/HLSLMaterialTranslator.cpp
+++ b/Engine/Source/Runtime/Engine/Private/Materials/HLSLMaterialTranslator.cpp
@@ -9005,7 +9005,7 @@ int32 FHLSLMaterialTranslator::VirtualTextureUnpack(int32 CodeIndex0, int32 Code
 	else if (UnpackType == EVirtualTextureUnpackType::NormalBC5BC1)
 	{
 		FString	SampleCode(TEXT("VirtualTextureUnpackNormalBC5BC1(%s, %s)"));
-		return CodeIndex0 == INDEX_NONE || CodeIndex1 == INDEX_NONE ? INDEX_NONE : AddCodeChunk(MCT_Float3, *SampleCode, *GetParameterCode(CodeIndex1), *GetParameterCode(CodeIndex2));
+		return CodeIndex1 == INDEX_NONE || CodeIndex2 == INDEX_NONE ? INDEX_NONE : AddCodeChunk(MCT_Float3, *SampleCode, *GetParameterCode(CodeIndex1), *GetParameterCode(CodeIndex2));
 	}
 	else if (UnpackType == EVirtualTextureUnpackType::HeightR16)
 	{
@@ -9020,7 +9020,7 @@ int32 FHLSLMaterialTranslator::VirtualTextureUnpack(int32 CodeIndex0, int32 Code
 	else if (UnpackType == EVirtualTextureUnpackType::NormalBGR565)
 	{
 		FString	SampleCode(TEXT("VirtualTextureUnpackNormalBGR565(%s)"));
-		return CodeIndex0 == INDEX_NONE || CodeIndex1 == INDEX_NONE ? INDEX_NONE : AddCodeChunk(MCT_Float3, *SampleCode, *GetParameterCode(CodeIndex1));
+		return CodeIndex1 == INDEX_NONE ? INDEX_NONE : AddCodeChunk(MCT_Float3, *SampleCode, *GetParameterCode(CodeIndex1));
 	}
 	if (UnpackType == EVirtualTextureUnpackType::BaseColorSRGB)
 	{
