Substrate - renamed variable from HLLSLMaterialTranslator.

#tests pc editor dx12
#rnx

diff --git a/Engine/Source/Runtime/Core/Private/UObject/FortniteMainBranchObjectVersion.cpp b/Engine/Source/Runtime/Core/Private/UObject/FortniteMainBranchObjectVersion.cpp
--- a/Engine/Source/Runtime/Core/Private/UObject/FortniteMainBranchObjectVersion.cpp
+++ b/Engine/Source/Runtime/Core/Private/UObject/FortniteMainBranchObjectVersion.cpp
@@ -13,7 +13,7 @@ TMap<FGuid, FGuid> FFortniteMainBranchObjectVersion::GetSystemGuids()
 	SystemGuids.Add(DevGuids.MATERIALSHADERMAP_DERIVEDDATA_VER, FGuid("F051877974B44D409435F321BDB39DBB"));
 	SystemGuids.Add(DevGuids.NIAGARASHADERMAP_DERIVEDDATA_VER, FGuid("6360A977062842A29ED30E3A7ACB0E64"));
 	SystemGuids.Add(DevGuids.Niagara_LatestScriptCompileVersion, FGuid("783FF3700BB641C799B2D43552FBF91E"));
-	SystemGuids.Add(DevGuids.MaterialTranslationDDCVersion, FGuid("F3D6B573B70C4F459D74F0EAED22179D"));
+	SystemGuids.Add(DevGuids.MaterialTranslationDDCVersion, FGuid("240CBCCB056D4E96BE11C168C07AA73C"));
 
 	// These GUIDs are in AppendRestrictedFortniteMainStreamObjectVersionGuids as they should rarely be changed.
 	//SystemGuids.Add(DevGuids.NANITE_DERIVEDDATA_VER, FGuid("ACF98184978742A592537DA223A1E6BE"));
diff --git a/Engine/Source/Runtime/Engine/Private/Materials/HLSLMaterialTranslator.cpp b/Engine/Source/Runtime/Engine/Private/Materials/HLSLMaterialTranslator.cpp
--- a/Engine/Source/Runtime/Engine/Private/Materials/HLSLMaterialTranslator.cpp
+++ b/Engine/Source/Runtime/Engine/Private/Materials/HLSLMaterialTranslator.cpp
@@ -294,7 +294,7 @@ struct FHLSLMaterialTranslator::FEnvironmentDefines
 	bool bMaterialVolumetricAdvancedAdvancedGroundContribution;
 	bool bMaterialVolumetricCloudEmptySpaceSkippingOutput;
 	bool bMaterialIsSubstrate;
-	bool bUseRootNodeToSubstrateHiddenConversion;
+	bool bMaterialUsesRootNodeToSubstrateHiddenConversion;
 	bool bDualSourceColorBlendingEnabled;
 	bool bSubstratePremultipliedAlphaOpacityOverridden;
 	bool bSubstrateUsesConversionFromLegacy;
@@ -402,7 +402,7 @@ struct FHLSLMaterialTranslator::FEnvironmentDefines
 		Ar << bMaterialVolumetricAdvancedAdvancedGroundContribution;
 		Ar << bMaterialVolumetricCloudEmptySpaceSkippingOutput;
 		Ar << bMaterialIsSubstrate;
-		Ar << bUseRootNodeToSubstrateHiddenConversion;
+		Ar << bMaterialUsesRootNodeToSubstrateHiddenConversion;
 		Ar << bDualSourceColorBlendingEnabled;
 		Ar << bSubstratePremultipliedAlphaOpacityOverridden;
 		Ar << bSubstrateUsesConversionFromLegacy;
@@ -723,8 +723,8 @@ FHLSLMaterialTranslator::FHLSLMaterialTranslator(FMaterial* InMaterial,
 ,	bUsesAnisotropy(false)
 ,	bSubstrateEnabled(false)
 ,	bMaterialIsSubstrate(false)
-,	bUseRootNodeToSubstrateHiddenConversion(false)
-,	bSubstrateHiddenMaterialAssetConversion(false)
+,	bMaterialUsesRootNodeToSubstrateHiddenConversion(false)
+,	bProjectSubstrateHiddenMaterialAssetConversionEnabled(false)
 ,	bUsesCurvature(false)
 ,	bUsesPerInstanceFadeAmount(false)
 ,	bCullIntermediateUniformExpressions(GCullIntermediateUniformExpressions)
@@ -1430,7 +1430,7 @@ void FHLSLMaterialTranslator::TranslateMaterial()
 	// Process the Substrate tree representing the material topology.
 	//
 	bSubstrateEnabled = Substrate::IsSubstrateEnabled();
-	bSubstrateHiddenMaterialAssetConversion = Substrate::IsHiddenMaterialAssetConversionEnabled();
+	bProjectSubstrateHiddenMaterialAssetConversionEnabled = Substrate::IsHiddenMaterialAssetConversionEnabled();
 	UMaterialExpression* FrontMaterialExpr = nullptr;
 	int32 FrontMaterialOutputIndex = INDEX_NONE;
 
@@ -1472,16 +1472,16 @@ void FHLSLMaterialTranslator::TranslateMaterial()
 	FExpressionInput* ClearCoatNormal = &NullInput;
 	FExpressionInput* CustomTangent = &NullInput;
 
-	bUseRootNodeToSubstrateHiddenConversion = bSubstrateEnabled && bSubstrateHiddenMaterialAssetConversion && !FrontMaterialExpr;
+	bMaterialUsesRootNodeToSubstrateHiddenConversion = bSubstrateEnabled && bProjectSubstrateHiddenMaterialAssetConversionEnabled && !FrontMaterialExpr;
 
 	// If bSubstrateMaterialUsesLegacyMaterialCompilation logic is changed, then FMaterialResource::IsSubstrateMaterial must be updated.
-	const bool bSubstrateMaterialUsesLegacyMaterialCompilation = bUseRootNodeToSubstrateHiddenConversion && Substrate::IsSubstrateBlendableGBufferEnabled(Platform);
+	const bool bSubstrateMaterialUsesLegacyMaterialCompilation = bMaterialUsesRootNodeToSubstrateHiddenConversion && Substrate::IsSubstrateBlendableGBufferEnabled(Platform);
 	if (bSubstrateMaterialUsesLegacyMaterialCompilation)
 	{
 		// In this case we use the legacy material shader: faster to compile and less instructions (no substrate conversion)
 		bSubstrateEnabled = false;
-		bSubstrateHiddenMaterialAssetConversion = false;
-		bUseRootNodeToSubstrateHiddenConversion = false;
+		bProjectSubstrateHiddenMaterialAssetConversionEnabled = false;
+		bMaterialUsesRootNodeToSubstrateHiddenConversion = false;
 
 		// Still set some Substrate material compilation output that are needed at runtime.
 		FMaterialShadingModelField MaterialShadingModels = GetCompiledShadingModels();
@@ -1522,7 +1522,7 @@ void FHLSLMaterialTranslator::TranslateMaterial()
 		}
 	}
 
-	if (bUseRootNodeToSubstrateHiddenConversion)
+	if (bMaterialUsesRootNodeToSubstrateHiddenConversion)
 	{
 		UMaterialInterface* MaterialInterface = Material->GetMaterialInterface();
 		UMaterial* BaseMaterial = MaterialInterface ? MaterialInterface->GetMaterial() : nullptr;
@@ -1808,7 +1808,7 @@ void FHLSLMaterialTranslator::TranslateMaterial()
 		NormalCodeChunkEnd = SharedPropertyCodeChunks[NormalShaderFrequency].Num();
 	}
 
-	if (!bSubstrateEnabled || !bUseRootNodeToSubstrateHiddenConversion 
+	if (!bSubstrateEnabled || !bMaterialUsesRootNodeToSubstrateHiddenConversion
 		|| GetSubstrateMaterialExportType() != ESubstrateMaterialExport::SME_None)	// Needed for material preview on node or any other export.
 	{
 		// Rest of properties
@@ -1840,7 +1840,7 @@ void FHLSLMaterialTranslator::TranslateMaterial()
 	Chunk[MP_SurfaceThickness] = Material->CompilePropertyAndSetMaterialProperty(MP_SurfaceThickness		,this);
 
 	// This is what is run when hidden conversion is not executed
-	if (!bUseRootNodeToSubstrateHiddenConversion)
+	if (!bMaterialUsesRootNodeToSubstrateHiddenConversion)
 	{
 		Chunk[MP_FrontMaterial] = Material->CompilePropertyAndSetMaterialProperty(MP_FrontMaterial, this);
 
@@ -1857,7 +1857,7 @@ void FHLSLMaterialTranslator::TranslateMaterial()
 			FullySimplifiedFrontMaterialCodeChunkEnd = SharedPropertyCodeChunks[FrontMaterialShaderFrequency].Num();
 		}
 	}
-	else // if (bUseRootNodeToSubstrateHiddenConversion)
+	else // if (bMaterialUsesRootNodeToSubstrateHiddenConversion)
 	{
 		auto CompileFrontMaterial = [&]()
 			{
@@ -1959,7 +1959,7 @@ void FHLSLMaterialTranslator::TranslateMaterial()
 		CHECK_DDC_QUERY_FINISHED_ELSE_RETURN();
 	}
 
-	if (!bSubstrateEnabled || !bUseRootNodeToSubstrateHiddenConversion)
+	if (!bSubstrateEnabled || !bMaterialUsesRootNodeToSubstrateHiddenConversion)
 	{
 		Chunk[MP_CustomData0] = Material->CompilePropertyAndSetMaterialProperty(MP_CustomData0, this);
 		Chunk[MP_CustomData1] = Material->CompilePropertyAndSetMaterialProperty(MP_CustomData1, this);
@@ -2412,7 +2412,7 @@ void FHLSLMaterialTranslator::TranslateMaterial()
 				}
 			}
 
-			if (bUseRootNodeToSubstrateHiddenConversion)
+			if (bMaterialUsesRootNodeToSubstrateHiddenConversion)
 			{
 				// Hardcoded substrate tree node visit when using hidden conversion.
 				bMaterialIsSubstrate = true;
@@ -3050,7 +3050,7 @@ void FHLSLMaterialTranslator::GetMaterialEnvironment(EShaderPlatform InPlatform,
 	}
 
 	OutEnvironment.SetDefine(TEXT("MATERIAL_IS_SUBSTRATE"), EnvironmentDefines->bMaterialIsSubstrate);
-	OutEnvironment.SetDefine(TEXT("MATERIAL_IS_SUBSTRATEHIDDENCONVERSION"), EnvironmentDefines->bUseRootNodeToSubstrateHiddenConversion);
+	OutEnvironment.SetDefine(TEXT("MATERIAL_IS_SUBSTRATEHIDDENCONVERSION"), EnvironmentDefines->bMaterialUsesRootNodeToSubstrateHiddenConversion);
 	OutEnvironment.SetDefine(TEXT("DUAL_SOURCE_COLOR_BLENDING_ENABLED"), EnvironmentDefines->bDualSourceColorBlendingEnabled);
 	OutEnvironment.SetDefine(TEXT("SUBSTRATE_PREMULTIPLIED_ALPHA_OPACITY_OVERRIDEN"), EnvironmentDefines->bSubstratePremultipliedAlphaOpacityOverridden);
 
@@ -16208,7 +16208,7 @@ void FHLSLMaterialTranslator::PrepareEnvironmentDefines()
 	}
 
 	EnvironmentDefines->bMaterialIsSubstrate = bMaterialIsSubstrate;
-	EnvironmentDefines->bUseRootNodeToSubstrateHiddenConversion = bUseRootNodeToSubstrateHiddenConversion;
+	EnvironmentDefines->bMaterialUsesRootNodeToSubstrateHiddenConversion = bMaterialUsesRootNodeToSubstrateHiddenConversion;
 
 	// Substrate requests dual source blending only for BLEND_TranslucentColoredTransmittance
 	bMaterialRequestsDualSourceBlending |= EnvironmentDefines->bMaterialIsSubstrate && Material->GetBlendMode() == EBlendMode::BLEND_TranslucentColoredTransmittance;
diff --git a/Engine/Source/Runtime/Engine/Private/Materials/HLSLMaterialTranslator.h b/Engine/Source/Runtime/Engine/Private/Materials/HLSLMaterialTranslator.h
--- a/Engine/Source/Runtime/Engine/Private/Materials/HLSLMaterialTranslator.h
+++ b/Engine/Source/Runtime/Engine/Private/Materials/HLSLMaterialTranslator.h
@@ -453,10 +453,10 @@ protected:
 	uint32 bMaterialIsSubstrate : 1;
 
 	/** True if Substrate hidden conversion from the root node is enabled. */
-	uint32 bUseRootNodeToSubstrateHiddenConversion : 1;
+	uint32 bMaterialUsesRootNodeToSubstrateHiddenConversion : 1;
 
 	/** True if this material is using a hidden conversion from the root node (hidden conversion enabled and FrontMaterial is not plugged in). */
-	uint32 bSubstrateHiddenMaterialAssetConversion : 1;
+	uint32 bProjectSubstrateHiddenMaterialAssetConversionEnabled : 1;
 
 	/** True if the opacity input is plugged in */
 	uint32 bOpacityPropertyIsUsed : 1;
