Substrate - updated comments.

#rnx

diff --git a/Engine/Source/Runtime/Engine/Private/Materials/HLSLMaterialTranslator.h b/Engine/Source/Runtime/Engine/Private/Materials/HLSLMaterialTranslator.h
--- a/Engine/Source/Runtime/Engine/Private/Materials/HLSLMaterialTranslator.h
+++ b/Engine/Source/Runtime/Engine/Private/Materials/HLSLMaterialTranslator.h
@@ -452,10 +452,10 @@ protected:
 	 */
 	uint32 bMaterialIsSubstrate : 1;
 
-	/** True if Substrate hidden conversion from the root node is enabled. */
+	/** True if Substrate hidden conversion from the root node is used by this material. (enabled for the project, FrontMaterial is not plugged in). */
 	uint32 bMaterialUsesRootNodeToSubstrateHiddenConversion : 1;
 
-	/** True if this material is using a hidden conversion from the root node (hidden conversion enabled and FrontMaterial is not plugged in). */
+	/** True if Substrate hidden conversion from the root node is enabled for the project.*/
 	uint32 bProjectSubstrateHiddenMaterialAssetConversionEnabled : 1;
 
 	/** True if the opacity input is plugged in */
