SplineComponent: Create attribute channel deletion method and expose to BP

diff --git a/Engine/Source/Runtime/Engine/Classes/Components/SplineComponent.h b/Engine/Source/Runtime/Engine/Classes/Components/SplineComponent.h
--- a/Engine/Source/Runtime/Engine/Classes/Components/SplineComponent.h
+++ b/Engine/Source/Runtime/Engine/Classes/Components/SplineComponent.h
@@ -902,6 +902,13 @@ public:
 	UFUNCTION(BlueprintCallable, Category = "Spline|Metadata")
 	ENGINE_API bool SupportsAttributes() const;
 
+	/**
+	 * Remove an attribute channel.
+	 * @param PropertyName The name of the channel to remove.
+	 */
+	UFUNCTION(BlueprintCallable, Category = "Spline|Metadata")
+	ENGINE_API bool RemovePropertyChannel(FName PropertyName);
+	
 	/**
 	 * Get the number of attribute values in a specified attribute channel.
 	 * @param PropertyName The name of the attribute channel.
diff --git a/Engine/Source/Runtime/Engine/Classes/Curves/Spline.h b/Engine/Source/Runtime/Engine/Classes/Curves/Spline.h
--- a/Engine/Source/Runtime/Engine/Classes/Curves/Spline.h
+++ b/Engine/Source/Runtime/Engine/Classes/Curves/Spline.h
@@ -77,6 +77,7 @@ struct FSpline
 	// Non-Templated functions
 	bool SupportsAttributes() const { return IsNew(); }
 	ENGINE_API bool HasAttributeChannel(FName AttributeName) const;
+	ENGINE_API bool RemoveAttributeChannel(FName AttributeName) const;
 	ENGINE_API TArray<FName> GetFloatPropertyChannels() const;
 	ENGINE_API TArray<FName> GetVectorPropertyChannels() const;
 
diff --git a/Engine/Source/Runtime/Engine/Private/Components/SplineComponent.cpp b/Engine/Source/Runtime/Engine/Private/Components/SplineComponent.cpp
--- a/Engine/Source/Runtime/Engine/Private/Components/SplineComponent.cpp
+++ b/Engine/Source/Runtime/Engine/Private/Components/SplineComponent.cpp
@@ -831,6 +831,16 @@ bool USplineComponent::SupportsAttributes() const
 	return UE::SplineComponent::ShouldUseSplineCurves() == false && Spline.SupportsAttributes();
 }
 
+bool USplineComponent::RemovePropertyChannel(FName PropertyName)
+{
+	if (UE::SplineComponent::ShouldUseSplineCurves() == false && Spline.SupportsAttributes())
+	{
+		return Spline.RemoveAttributeChannel(PropertyName);
+	}
+
+	return false;
+}
+
 void USplineComponent::CreateFloatPropertyChannel(FName PropertyName)
 {
 	CreatePropertyChannel<float>(GetSplinePointsMetadata(), Spline, PropertyName);
diff --git a/Engine/Source/Runtime/Engine/Private/Curves/Spline.cpp b/Engine/Source/Runtime/Engine/Private/Curves/Spline.cpp
--- a/Engine/Source/Runtime/Engine/Private/Curves/Spline.cpp
+++ b/Engine/Source/Runtime/Engine/Private/Curves/Spline.cpp
@@ -1210,6 +1210,16 @@ bool FSpline::HasAttributeChannel(FName AttributeName) const
 	return false;
 }
 
+bool FSpline::RemoveAttributeChannel(FName AttributeName) const
+{
+	if (NewData && SupportsAttributes())
+	{
+		return NewData->RemoveAttributeChannel(AttributeName);
+	}
+
+	return false;
+}
+
 TArray<FName> FSpline::GetFloatPropertyChannels() const
 {
 	TArray<FName> ChannelNames;
